Remove hard coded engine version from nDisplay plugin.

diff --git a/Engine/Plugins/Runtime/nDisplay/nDisplay.uplugin b/Engine/Plugins/Runtime/nDisplay/nDisplay.uplugin
--- a/Engine/Plugins/Runtime/nDisplay/nDisplay.uplugin
+++ b/Engine/Plugins/Runtime/nDisplay/nDisplay.uplugin
@@ -5,7 +5,6 @@
 	"FriendlyName": "nDisplay",
 	"CreatedBy" : "Epic Games Inc",
 	"CreatedByURL" : "http://epicgames.com",
-	"EngineVersion" : "4.20.0",
 	"Description": "Mono/Stereo standalone/cluster virtual reality",
 	"Category": "Misc",
 	"DocsURL": "",
