Improve nDisplay plugin description

diff --git a/Engine/Plugins/Runtime/nDisplay/nDisplay.uplugin b/Engine/Plugins/Runtime/nDisplay/nDisplay.uplugin
--- a/Engine/Plugins/Runtime/nDisplay/nDisplay.uplugin
+++ b/Engine/Plugins/Runtime/nDisplay/nDisplay.uplugin
@@ -5,7 +5,7 @@
 	"FriendlyName": "nDisplay",
 	"CreatedBy" : "Epic Games Inc",
 	"CreatedByURL" : "http://epicgames.com",
-	"Description": "Mono/Stereo standalone/cluster virtual reality",
+	"Description": "Support for synchronized clustered rendering using multiple PCs in mono or stereo",
 	"Category": "Misc",
 	"DocsURL": "",
 	"MarketplaceURL": "",
