Fix nDisplay viewport positioning issue when using an x offset
Fix BP API comments

diff --git a/Engine/Plugins/Runtime/nDisplay/Source/DisplayCluster/Private/Cluster/Controller/DisplayClusterClusterNodeCtrlBase.cpp b/Engine/Plugins/Runtime/nDisplay/Source/DisplayCluster/Private/Cluster/Controller/DisplayClusterClusterNodeCtrlBase.cpp
--- a/Engine/Plugins/Runtime/nDisplay/Source/DisplayCluster/Private/Cluster/Controller/DisplayClusterClusterNodeCtrlBase.cpp
+++ b/Engine/Plugins/Runtime/nDisplay/Source/DisplayCluster/Private/Cluster/Controller/DisplayClusterClusterNodeCtrlBase.cpp
@@ -41,6 +41,7 @@ bool FDisplayClusterClusterNodeCtrlBase::InitializeStereo()
 		return false;
 	}
 
+	//@todo: Move this logic to the render manager
 	IDisplayClusterStereoDevice* const StereoDevice = GDisplayCluster->GetPrivateRenderMgr()->GetStereoDevice();
 	if (StereoDevice)
 	{
diff --git a/Engine/Plugins/Runtime/nDisplay/Source/DisplayCluster/Private/Render/Devices/QuadBufferStereo/DisplayClusterDeviceQuadBufferStereoBase.cpp b/Engine/Plugins/Runtime/nDisplay/Source/DisplayCluster/Private/Render/Devices/QuadBufferStereo/DisplayClusterDeviceQuadBufferStereoBase.cpp
--- a/Engine/Plugins/Runtime/nDisplay/Source/DisplayCluster/Private/Render/Devices/QuadBufferStereo/DisplayClusterDeviceQuadBufferStereoBase.cpp
+++ b/Engine/Plugins/Runtime/nDisplay/Source/DisplayCluster/Private/Render/Devices/QuadBufferStereo/DisplayClusterDeviceQuadBufferStereoBase.cpp
@@ -42,8 +42,8 @@ void FDisplayClusterDeviceQuadBufferStereoBase::CalculateRenderTargetSize(const
 {
 	check(IsInGameThread());
 
-	InOutSizeX = Viewport.GetSizeXY().X;
-	InOutSizeY = Viewport.GetSizeXY().Y * 2;
+	InOutSizeX = Viewport.GetSizeXY().X * 2;
+	InOutSizeY = Viewport.GetSizeXY().Y;
 
 	check(InOutSizeX > 0 && InOutSizeY > 0);
 }
@@ -56,12 +56,12 @@ bool FDisplayClusterDeviceQuadBufferStereoBase::ShouldUseSeparateRenderTarget()
 
 void FDisplayClusterDeviceQuadBufferStereoBase::AdjustViewRect(enum EStereoscopicPass StereoPass, int32& X, int32& Y, uint32& SizeX, uint32& SizeY) const
 {
-	uint32 screentHeight = SizeY;
+	const uint32 screentWidth = SizeX;
 	FDisplayClusterDeviceBase::AdjustViewRect(StereoPass, X, Y, SizeX, SizeY);
 
 	if (StereoPass == EStereoscopicPass::eSSP_RIGHT_EYE)
 	{
-		Y += screentHeight;
+		X += screentWidth;
 	}
 }
 
diff --git a/Engine/Plugins/Runtime/nDisplay/Source/DisplayCluster/Private/Render/Devices/QuadBufferStereo/DisplayClusterDeviceQuadBufferStereoD3D11.cpp b/Engine/Plugins/Runtime/nDisplay/Source/DisplayCluster/Private/Render/Devices/QuadBufferStereo/DisplayClusterDeviceQuadBufferStereoD3D11.cpp
--- a/Engine/Plugins/Runtime/nDisplay/Source/DisplayCluster/Private/Render/Devices/QuadBufferStereo/DisplayClusterDeviceQuadBufferStereoD3D11.cpp
+++ b/Engine/Plugins/Runtime/nDisplay/Source/DisplayCluster/Private/Render/Devices/QuadBufferStereo/DisplayClusterDeviceQuadBufferStereoD3D11.cpp
@@ -72,19 +72,19 @@ void FDisplayClusterDeviceQuadBufferStereoD3D11::RenderTexture_RenderThread(FRHI
 	check(IsInRenderingThread());
 
 	//calculate sub regions to copy
-	const int halfSizeY = BackBuffSize.Y / 2;
+	const int halfSizeX = BackBuffSize.X / 2;
 
 	FResolveParams copyParamsLeft;
 	copyParamsLeft.DestArrayIndex = 0;
 	copyParamsLeft.SourceArrayIndex = 0;
 	copyParamsLeft.Rect.X1 = 0;
-	copyParamsLeft.Rect.X2 = ViewportArea.GetSize().X;
 	copyParamsLeft.Rect.Y1 = 0;
-	copyParamsLeft.Rect.Y2 = halfSizeY;
+	copyParamsLeft.Rect.X2 = halfSizeX;
+	copyParamsLeft.Rect.Y2 = BackBuffSize.Y;
 	copyParamsLeft.DestRect.X1 = 0;
 	copyParamsLeft.DestRect.Y1 = 0;
-	copyParamsLeft.DestRect.X2 = ViewportArea.GetSize().X;
-	copyParamsLeft.DestRect.Y2 = halfSizeY;
+	copyParamsLeft.DestRect.X2 = halfSizeX;
+	copyParamsLeft.DestRect.Y2 = BackBuffSize.Y;
 
 	RHICmdList.CopyToResolveTarget(SrcTexture, BackBuffer, copyParamsLeft);
 
@@ -93,13 +93,10 @@ void FDisplayClusterDeviceQuadBufferStereoD3D11::RenderTexture_RenderThread(FRHI
 	copyParamsRight.SourceArrayIndex = 0;
 
 	copyParamsRight.Rect = copyParamsLeft.Rect;
+	copyParamsRight.Rect.X1 = halfSizeX;
+	copyParamsRight.Rect.X2 = halfSizeX * 2;
 
-	copyParamsRight.Rect.Y1 = halfSizeY;
-	copyParamsRight.Rect.Y2 = halfSizeY * 2;
-	copyParamsRight.DestRect.X1 = 0;
-	copyParamsRight.DestRect.Y1 = 0;
-	copyParamsRight.DestRect.X2 = ViewportArea.GetSize().X;
-	copyParamsRight.DestRect.Y2 = halfSizeY;
+	copyParamsRight.DestRect = copyParamsLeft.DestRect;
 
 	RHICmdList.CopyToResolveTarget(SrcTexture, BackBuffer, copyParamsRight);
 }
diff --git a/Engine/Plugins/Runtime/nDisplay/Source/DisplayCluster/Private/Render/Devices/QuadBufferStereo/DisplayClusterDeviceQuadBufferStereoD3D12.cpp b/Engine/Plugins/Runtime/nDisplay/Source/DisplayCluster/Private/Render/Devices/QuadBufferStereo/DisplayClusterDeviceQuadBufferStereoD3D12.cpp
--- a/Engine/Plugins/Runtime/nDisplay/Source/DisplayCluster/Private/Render/Devices/QuadBufferStereo/DisplayClusterDeviceQuadBufferStereoD3D12.cpp
+++ b/Engine/Plugins/Runtime/nDisplay/Source/DisplayCluster/Private/Render/Devices/QuadBufferStereo/DisplayClusterDeviceQuadBufferStereoD3D12.cpp
@@ -70,19 +70,19 @@ void FDisplayClusterDeviceQuadBufferStereoD3D12::RenderTexture_RenderThread(FRHI
 	check(IsInRenderingThread());
 	
 	//calculate sub regions to copy
-	const int halfSizeY = BackBuffSize.Y / 2;
+	const int halfSizeX = BackBuffSize.X / 2;
 		
 	FResolveParams copyParamsLeft; 
 	copyParamsLeft.DestArrayIndex = 0;
 	copyParamsLeft.SourceArrayIndex = 0;
 	copyParamsLeft.Rect.X1 = 0;
-	copyParamsLeft.Rect.X2 = ViewportArea.GetSize().X;
 	copyParamsLeft.Rect.Y1 = 0;
-	copyParamsLeft.Rect.Y2 = halfSizeY;
+	copyParamsLeft.Rect.X2 = halfSizeX;
+	copyParamsLeft.Rect.Y2 = BackBuffSize.Y;
 	copyParamsLeft.DestRect.X1 = 0;
 	copyParamsLeft.DestRect.Y1 = 0;
-	copyParamsLeft.DestRect.X2 = ViewportArea.GetSize().X;
-	copyParamsLeft.DestRect.Y2 = halfSizeY;
+	copyParamsLeft.DestRect.X2 = halfSizeX;
+	copyParamsLeft.DestRect.Y2 = BackBuffSize.Y;
 
 	RHICmdList.CopyToResolveTarget(SrcTexture, BackBuffer, copyParamsLeft);
 	
@@ -92,12 +92,10 @@ void FDisplayClusterDeviceQuadBufferStereoD3D12::RenderTexture_RenderThread(FRHI
 	
 	copyParamsRight.Rect = copyParamsLeft.Rect;
 
-	copyParamsRight.Rect.Y1 = halfSizeY;
-	copyParamsRight.Rect.Y2 = halfSizeY*2;
-	copyParamsRight.DestRect.X1 = 0;
-	copyParamsRight.DestRect.Y1 = 0;
-	copyParamsRight.DestRect.X2 = ViewportArea.GetSize().X;
-	copyParamsRight.DestRect.Y2 = halfSizeY;
+	copyParamsRight.Rect.X1 = halfSizeX;
+	copyParamsRight.Rect.X2 = halfSizeX * 2;
+
+	copyParamsRight.DestRect = copyParamsLeft.DestRect;
 
 	RHICmdList.CopyToResolveTarget(SrcTexture, BackBuffer, copyParamsRight);
 }
diff --git a/Engine/Plugins/Runtime/nDisplay/Source/DisplayCluster/Private/Render/Devices/QuadBufferStereo/DisplayClusterDeviceQuadBufferStereoOpenGL.cpp b/Engine/Plugins/Runtime/nDisplay/Source/DisplayCluster/Private/Render/Devices/QuadBufferStereo/DisplayClusterDeviceQuadBufferStereoOpenGL.cpp
--- a/Engine/Plugins/Runtime/nDisplay/Source/DisplayCluster/Private/Render/Devices/QuadBufferStereo/DisplayClusterDeviceQuadBufferStereoOpenGL.cpp
+++ b/Engine/Plugins/Runtime/nDisplay/Source/DisplayCluster/Private/Render/Devices/QuadBufferStereo/DisplayClusterDeviceQuadBufferStereoOpenGL.cpp
@@ -52,9 +52,9 @@ bool FDisplayClusterDeviceQuadBufferStereoOpenGL::Present(int32& InOutSyncInterv
 {
 	UE_LOG(LogDisplayClusterRender, Verbose, TEXT("FDisplayClusterDeviceQuadBufferStereoOpenGL::Present"));
 
-	const int halfSizeY = BackBuffSize.Y / 2;
+	const int halfSizeX = BackBuffSize.X / 2;
 	int dstX1 = 0;
-	int dstX2 = BackBuffSize.X;
+	int dstX2 = halfSizeX;
 
 	// Convert to left bottom origin and flip Y
 	int dstY1 = ViewportSize.Y;
@@ -83,14 +83,14 @@ bool FDisplayClusterDeviceQuadBufferStereoOpenGL::Present(int32& InOutSyncInterv
 
 	glDrawBuffer(GL_BACK_LEFT);
 	glBlitFramebuffer(
-		0, 0, BackBuffSize.X, halfSizeY,
+		0, 0, halfSizeX, BackBuffSize.Y,
 		dstX1, dstY1, dstX2, dstY2,
 		GL_COLOR_BUFFER_BIT,
 		GL_NEAREST);
 
 	glDrawBuffer(GL_BACK_RIGHT);
 	glBlitFramebuffer(
-		0, halfSizeY, BackBuffSize.X, BackBuffSize.Y,
+		halfSizeX, 0, BackBuffSize.X, BackBuffSize.Y,
 		dstX1, dstY1, dstX2, dstY2,
 		GL_COLOR_BUFFER_BIT,
 		GL_NEAREST);
diff --git a/Engine/Plugins/Runtime/nDisplay/Source/DisplayCluster/Public/Blueprints/DisplayClusterBlueprintLib.h b/Engine/Plugins/Runtime/nDisplay/Source/DisplayCluster/Public/Blueprints/DisplayClusterBlueprintLib.h
--- a/Engine/Plugins/Runtime/nDisplay/Source/DisplayCluster/Public/Blueprints/DisplayClusterBlueprintLib.h
+++ b/Engine/Plugins/Runtime/nDisplay/Source/DisplayCluster/Public/Blueprints/DisplayClusterBlueprintLib.h
@@ -18,7 +18,7 @@ class UDisplayClusterBlueprintLib
 	GENERATED_UCLASS_BODY()
 
 public:
-
-	UFUNCTION(BlueprintCallable, meta = (DisplayName = "Get plugin API"), Category = "DisplayCluster")
+	/** Return Display Cluster API interface. */
+	UFUNCTION(BlueprintCallable, meta = (DisplayName = "Get Display Cluster Plugin API"), Category = "DisplayCluster")
 	static void GetAPI(TScriptInterface<IDisplayClusterBlueprintAPI>& OutAPI);
 };
diff --git a/Engine/Plugins/Runtime/nDisplay/Source/DisplayCluster/Public/Blueprints/IDisplayClusterBlueprintAPI.h b/Engine/Plugins/Runtime/nDisplay/Source/DisplayCluster/Public/Blueprints/IDisplayClusterBlueprintAPI.h
--- a/Engine/Plugins/Runtime/nDisplay/Source/DisplayCluster/Public/Blueprints/IDisplayClusterBlueprintAPI.h
+++ b/Engine/Plugins/Runtime/nDisplay/Source/DisplayCluster/Public/Blueprints/IDisplayClusterBlueprintAPI.h
@@ -25,21 +25,27 @@ public:
 	//////////////////////////////////////////////////////////////////////////////////////////////
 	// Cluster API
 	//////////////////////////////////////////////////////////////////////////////////////////////
+	/** Return if current node is a master computer in cluster. */
 	UFUNCTION(BlueprintCallable, meta = (DisplayName = "Is master node"), Category = "DisplayCluster|Cluster")
 	virtual bool IsMaster() = 0;
 	
+	/** Return if current node is not a master computer in cluster. */
 	UFUNCTION(BlueprintCallable, meta = (DisplayName = "Is slave node"), Category = "DisplayCluster|Cluster")
 	virtual bool IsSlave() = 0;
 
+	/** Whether application is in cluster mode or not. */
 	UFUNCTION(BlueprintCallable, meta = (DisplayName = "Is cluster mode"), Category = "DisplayCluster|Cluster")
 	virtual bool IsCluster() = 0;
 
+	/** Whether application is in standalone mode or not. */
 	UFUNCTION(BlueprintCallable, meta = (DisplayName = "Is standalone mode"), Category = "DisplayCluster|Cluster")
 	virtual bool IsStandalone() = 0;
 
+	/** Returns node name of the current application instance. */
 	UFUNCTION(BlueprintCallable, meta = (DisplayName = "Get node ID"), Category = "DisplayCluster|Cluster")
 	virtual FString GetNodeId() = 0;
 
+	/** Returns amount of nodes in cluster. */
 	UFUNCTION(BlueprintCallable, meta = (DisplayName = "Get nodes amount"), Category = "DisplayCluster|Cluster")
 	virtual int32 GetNodesAmount() = 0;
 
@@ -54,19 +60,24 @@ public:
 	// Game API
 	//////////////////////////////////////////////////////////////////////////////////////////////
 	// Root
+	/** Returns Cluster Pawn. */
 	UFUNCTION(BlueprintCallable, meta = (DisplayName = "Get root"), Category = "DisplayCluster|Game")
 	virtual ADisplayClusterPawn* GetRoot() = 0;
 
 	// Screens
+	/** Returns screen reference used for computing frustum output. */
 	UFUNCTION(BlueprintCallable, meta = (DisplayName = "Get active screen"), Category = "DisplayCluster|Game")
 	virtual UDisplayClusterScreenComponent* GetActiveScreen() = 0;
 
+	/** Returns screen reference by id name. */
 	UFUNCTION(BlueprintCallable, meta = (DisplayName = "Get screen by ID"), Category = "DisplayCluster|Game")
 	virtual UDisplayClusterScreenComponent* GetScreenById(const FString& id) = 0;
 
+	/** Returns array of all screen references. */
 	UFUNCTION(BlueprintCallable, meta = (DisplayName = "Get all screens"), Category = "DisplayCluster|Game")
 	virtual TArray<UDisplayClusterScreenComponent*> GetAllScreens() = 0;
 
+	/** Returns amount of screens defined in configuration file. */
 	UFUNCTION(BlueprintCallable, meta = (DisplayName = "Get amount of screens"), Category = "DisplayCluster|Game")
 	virtual int32 GetScreensAmount() = 0;
 
@@ -81,28 +92,36 @@ public:
 	*/
 
 	// Nodes
+	/** Returns node reference by its id name. */
 	UFUNCTION(BlueprintCallable, meta = (DisplayName = "Get node by ID"), Category = "DisplayCluster|Game")
 	virtual UDisplayClusterSceneComponent* GetNodeById(const FString& id) = 0;
 
+	/** Returns array of all nodes references by its id name. */
 	UFUNCTION(BlueprintCallable, meta = (DisplayName = "Get all nodes"), Category = "DisplayCluster|Game")
 	virtual TArray<UDisplayClusterSceneComponent*> GetAllNodes() = 0;
 
 	// Navigation
+	/** Returns scene component used for default pawn navigation. */
 	UFUNCTION(BlueprintCallable, meta = (DisplayName = "Get translation direction component"), Category = "DisplayCluster|Game")
 	virtual USceneComponent* GetTranslationDirectionComponent() = 0;
 
+	/** Set scene component to be used for default pawn navigation. */
 	UFUNCTION(BlueprintCallable, meta = (DisplayName = "Set translation direction component"), Category = "DisplayCluster|Game")
 	virtual void SetTranslationDirectionComponent(USceneComponent* pComp) = 0;
 
+	/** Set scene component to be used for default pawn navigation by id name. */
 	UFUNCTION(BlueprintCallable, meta = (DisplayName = "Set translation direction component by ID"), Category = "DisplayCluster|Game")
 	virtual void SetTranslationDirectionComponentId(const FString& id) = 0;
 
+	/** Return scene component used as a pivot point for rotation of the scene node hierarchy. */
 	UFUNCTION(BlueprintCallable, meta = (DisplayName = "Get rotate around component"), Category = "DisplayCluster|Game")
 	virtual USceneComponent* GetRotateAroundComponent() = 0;
 
+	/** Set scene component used as a pivot point for rotation of the scene node hierarchy. */
 	UFUNCTION(BlueprintCallable, meta = (DisplayName = "Set rotate around component"), Category = "DisplayCluster|Game")
 	virtual void SetRotateAroundComponent(USceneComponent* pComp) = 0;
 
+	/** Set scene component used as a pivot point for rotation of the scene node hierarchy by id name. */
 	UFUNCTION(BlueprintCallable, meta = (DisplayName = "Set rotate around component by ID"), Category = "DisplayCluster|Game")
 	virtual void SetRotateAroundComponentId(const FString& id) = 0;
 
@@ -111,48 +130,62 @@ public:
 	// Input API
 	//////////////////////////////////////////////////////////////////////////////////////////////
 	// Device information
+	/** Return amount of VRPN axis devices. */
 	UFUNCTION(BlueprintCallable, meta = (DisplayName = "Get amount of VRPN axis devices"), Category = "DisplayCluster|Input")
 	virtual int32 GetAxisDeviceAmount() = 0;
 
+	/** Return amount of VRPN button devices. */
 	UFUNCTION(BlueprintCallable, meta = (DisplayName = "Get amount of VRPN button devices"), Category = "DisplayCluster|Input")
 	virtual int32 GetButtonDeviceAmount() = 0;
 
+	/** Return amount of VRPN tracker devices. */
 	UFUNCTION(BlueprintCallable, meta = (DisplayName = "Get amount of VRPN tracker devices"), Category = "DisplayCluster|Input")
 	virtual int32 GetTrackerDeviceAmount() = 0;
 
+	/** Return array of names of all VRPN axis devices. */
 	UFUNCTION(BlueprintCallable, meta = (DisplayName = "Get IDs of VRPN axis devices"), Category = "DisplayCluster|Input")
 	virtual bool GetAxisDeviceIds(TArray<FString>& IDs) = 0;
 
+	/** Return array of names of all VRPN button devices. */
 	UFUNCTION(BlueprintCallable, meta = (DisplayName = "Get IDs of VRPN button devices"), Category = "DisplayCluster|Input")
 	virtual bool GetButtonDeviceIds(TArray<FString>& IDs) = 0;
 
+	/** Return array of names of all VRPN tracker devices. */
 	UFUNCTION(BlueprintCallable, meta = (DisplayName = "Get IDs of VRPN tracker devices"), Category = "DisplayCluster|Input")
 	virtual bool GetTrackerDeviceIds(TArray<FString>& IDs) = 0;
 
 	// Buttons
+	/** Return state of VRPN button at specified device and channel. */
 	UFUNCTION(BlueprintCallable, meta = (DisplayName = "Get VRPN button state"), Category = "DisplayCluster|Input")
 	virtual void GetButtonState(const FString& DeviceId, uint8 DeviceChannel, bool& CurState, bool& IsChannelAvailable) = 0;
 
+	/** Return whether VRPN button is pressed at specified device and channel. */
 	UFUNCTION(BlueprintCallable, meta = (DisplayName = "Is VRPN button pressed"), Category = "DisplayCluster|Input")
 	virtual void IsButtonPressed(const FString& DeviceId, uint8 DeviceChannel, bool& CurPressed, bool& IsChannelAvailable) = 0;
 
+	/** Return whether VRPN button is released at specified device and channel. */
 	UFUNCTION(BlueprintCallable, meta = (DisplayName = "Is VRPN button released"), Category = "DisplayCluster|Input")
 	virtual void IsButtonReleased(const FString& DeviceId, uint8 DeviceChannel, bool& CurReleased, bool& IsChannelAvailable) = 0;
 
+	/** Return whether VRPN button was released at specified device and channel. */
 	UFUNCTION(BlueprintCallable, meta = (DisplayName = "Was VRPN button pressed"), Category = "DisplayCluster|Input")
 	virtual void WasButtonPressed(const FString& DeviceId, uint8 DeviceChannel, bool& WasPressed, bool& IsChannelAvailable) = 0;
 
+	/** Return whether VRPN button was released at specified device and channel. */
 	UFUNCTION(BlueprintCallable, meta = (DisplayName = "Was VRPN button released"), Category = "DisplayCluster|Input")
 	virtual void WasButtonReleased(const FString& DeviceId, uint8 DeviceChannel, bool& WasReleased, bool& IsChannelAvailable) = 0;
 
 	// Axes
+	/** Return axis value at specified device and channel. */
 	UFUNCTION(BlueprintCallable, meta = (DisplayName = "Get VRPN axis value"), Category = "DisplayCluster|Input")
 	virtual void GetAxis(const FString& DeviceId, uint8 DeviceChannel, float& Value, bool& IsAvailable) = 0;
 
 	// Trackers
+	/** Return tracker location values at specified device and channel. */
 	UFUNCTION(BlueprintCallable, meta = (DisplayName = "Get VRPN tracker location"), Category = "DisplayCluster|Input")
 	virtual void GetTrackerLocation(const FString& DeviceId, uint8 DeviceChannel, FVector& Location, bool& IsChannelAvailable) = 0;
 
+	/** Return tracker quanternion values at specified device and channel. */
 	UFUNCTION(BlueprintCallable, meta = (DisplayName = "Get VRPN tracker rotation (as quaternion)"), Category = "DisplayCluster|Input")
 	virtual void GetTrackerQuat(const FString& DeviceId, uint8 DeviceChannel, FQuat& Rotation, bool& IsChannelAvailable) = 0;
 
@@ -160,18 +193,23 @@ public:
 	//////////////////////////////////////////////////////////////////////////////////////////////
 	// Render API
 	//////////////////////////////////////////////////////////////////////////////////////////////
+	/** Set eye interpupillary distance (eye separation) for stereoscopic rendering. */
 	UFUNCTION(BlueprintCallable, meta = (DisplayName = "Set interpuppillary distance"), Category = "DisplayCluster|Render")
 	virtual void SetInterpupillaryDistance(float dist) = 0;
 
+	/** Return eye interpupillary distance (eye separation) for stereoscopic rendering. */
 	UFUNCTION(BlueprintCallable, meta = (DisplayName = "Get interpuppillary distance"), Category = "DisplayCluster|Render")
 	virtual float GetInterpupillaryDistance() = 0;
 
+	/** Swap eye rendering. */
 	UFUNCTION(BlueprintCallable, meta = (DisplayName = "Set eye swap"), Category = "DisplayCluster|Render")
 	virtual void SetEyesSwap(bool swap) = 0;
 
+	/** Get Swap eye rendering state. */
 	UFUNCTION(BlueprintCallable, meta = (DisplayName = "Get eye swap"), Category = "DisplayCluster|Render")
 	virtual bool GetEyesSwap() = 0;
 
+	/** Toggle current eye swap state. */
 	UFUNCTION(BlueprintCallable, meta = (DisplayName = "Toggle eye swap"), Category = "DisplayCluster|Render")
 	virtual bool ToggleEyesSwap() = 0;
 
@@ -181,9 +219,11 @@ public:
 	UFUNCTION(BlueprintCallable, meta = (DisplayName = "Get output flip"), Category = "DisplayCluster|Render")
 	virtual void GetOutputFlip(bool& flipH, bool& flipV) = 0;
 
+	/** Return near and far plane clip plane distances. */
 	UFUNCTION(BlueprintCallable, meta = (DisplayName = "Get near and far clipping distance"), Category = "DisplayCluster|Render")
 	virtual void GetCullingDistance(float& NearClipPlane, float& FarClipPlane) = 0;
 
+	/** Set near and far plane clip plane distances. */
 	UFUNCTION(BlueprintCallable, meta = (DisplayName = "Set near and far clipping distance"), Category = "DisplayCluster|Render")
 	virtual void SetCullingDistance(float NearClipPlane, float FarClipPlane) = 0;
 };
