GetPlugin Api garbage collected on stop play
Change private access for DisplayClusterGameEngine

diff --git a/Engine/Plugins/Runtime/nDisplay/Source/DisplayCluster/Private/Blueprints/DisplayClusterBlueprintLib.cpp b/Engine/Plugins/Runtime/nDisplay/Source/DisplayCluster/Private/Blueprints/DisplayClusterBlueprintLib.cpp
--- a/Engine/Plugins/Runtime/nDisplay/Source/DisplayCluster/Private/Blueprints/DisplayClusterBlueprintLib.cpp
+++ b/Engine/Plugins/Runtime/nDisplay/Source/DisplayCluster/Private/Blueprints/DisplayClusterBlueprintLib.cpp
@@ -2,16 +2,17 @@
 
 #include "Blueprints/DisplayClusterBlueprintLib.h"
 #include "DisplayClusterBlueprintAPIImpl.h"
+#include "UObject/Package.h"
 
 
-UDisplayClusterBlueprintLib::UDisplayClusterBlueprintLib(class FObjectInitializer const & ObjectInitializer) :
-	Super(ObjectInitializer)
+UDisplayClusterBlueprintLib::UDisplayClusterBlueprintLib(class FObjectInitializer const & ObjectInitializer)
+	: Super(ObjectInitializer)
 {
 
 }
 
 void UDisplayClusterBlueprintLib::GetAPI(TScriptInterface<IDisplayClusterBlueprintAPI>& OutAPI)
 {
-	static TUniquePtr<UDisplayClusterBlueprintAPIImpl> Obj(NewObject<UDisplayClusterBlueprintAPIImpl>());
-	OutAPI = Obj.Get();
+	static UDisplayClusterBlueprintAPIImpl* Obj = NewObject<UDisplayClusterBlueprintAPIImpl>(GetTransientPackage(), NAME_None, RF_MarkAsRootSet);
+	OutAPI = Obj;
 }
diff --git a/Engine/Plugins/Runtime/nDisplay/Source/DisplayCluster/Public/DisplayClusterGameEngine.h b/Engine/Plugins/Runtime/nDisplay/Source/DisplayCluster/Public/DisplayClusterGameEngine.h
--- a/Engine/Plugins/Runtime/nDisplay/Source/DisplayCluster/Public/DisplayClusterGameEngine.h
+++ b/Engine/Plugins/Runtime/nDisplay/Source/DisplayCluster/Public/DisplayClusterGameEngine.h
@@ -24,7 +24,7 @@ class DISPLAYCLUSTER_API UDisplayClusterGameEngine
 {
 	GENERATED_BODY()
 	
-private:
+public:
 	virtual void Init(class IEngineLoop* InEngineLoop) override;
 	virtual void PreExit() override;
 	virtual void Tick(float DeltaSeconds, bool bIdleMode) override;
