Fixed visibility of some inherited members

diff --git a/Engine/Plugins/Runtime/nDisplay/Source/DisplayCluster/Public/DisplayClusterCameraComponent.h b/Engine/Plugins/Runtime/nDisplay/Source/DisplayCluster/Public/DisplayClusterCameraComponent.h
--- a/Engine/Plugins/Runtime/nDisplay/Source/DisplayCluster/Public/DisplayClusterCameraComponent.h
+++ b/Engine/Plugins/Runtime/nDisplay/Source/DisplayCluster/Public/DisplayClusterCameraComponent.h
@@ -22,7 +22,7 @@ public:
 	virtual void SetSettings(const FDisplayClusterConfigSceneNode* pConfig) override;
 	virtual bool ApplySettings() override;
 
-protected:
+public:
 	virtual void BeginPlay() override;
 	virtual void TickComponent(float DeltaTime, ELevelTick TickType, FActorComponentTickFunction* ThisTickFunction) override;
 };
diff --git a/Engine/Plugins/Runtime/nDisplay/Source/DisplayCluster/Public/DisplayClusterGameMode.h b/Engine/Plugins/Runtime/nDisplay/Source/DisplayCluster/Public/DisplayClusterGameMode.h
--- a/Engine/Plugins/Runtime/nDisplay/Source/DisplayCluster/Public/DisplayClusterGameMode.h
+++ b/Engine/Plugins/Runtime/nDisplay/Source/DisplayCluster/Public/DisplayClusterGameMode.h
@@ -30,17 +30,17 @@ public:
 	bool IsDisplayClusterActive() const
 	{ return bIsDisplayClusterActive; }
 
-protected:
-	UPROPERTY(EditAnywhere, BlueprintReadWrite, Category = "DisplayCluster")
-	bool bIsDisplayClusterActive = true;
-
-protected:
+public:
 	virtual void InitGame(const FString& MapName, const FString& Options, FString& ErrorMessage) override;
 	virtual void StartPlay() override;
 	virtual void Tick(float DeltaSeconds) override;
 	virtual void BeginPlay() override;
 	virtual void BeginDestroy() override;
 
+protected:
+	UPROPERTY(EditAnywhere, BlueprintReadWrite, Category = "DisplayCluster")
+	bool bIsDisplayClusterActive = true;
+
 protected:
 	bool bGameStarted = false;
 
diff --git a/Engine/Plugins/Runtime/nDisplay/Source/DisplayCluster/Public/DisplayClusterHUD.h b/Engine/Plugins/Runtime/nDisplay/Source/DisplayCluster/Public/DisplayClusterHUD.h
--- a/Engine/Plugins/Runtime/nDisplay/Source/DisplayCluster/Public/DisplayClusterHUD.h
+++ b/Engine/Plugins/Runtime/nDisplay/Source/DisplayCluster/Public/DisplayClusterHUD.h
@@ -18,7 +18,7 @@ class DISPLAYCLUSTER_API ADisplayClusterHUD
 public:
 	ADisplayClusterHUD(const FObjectInitializer& ObjectInitializer);
 
-protected:
+public:
 	virtual void BeginPlay() override;
 
 	/** Primary draw call for the HUD */
diff --git a/Engine/Plugins/Runtime/nDisplay/Source/DisplayCluster/Public/DisplayClusterPawn.h b/Engine/Plugins/Runtime/nDisplay/Source/DisplayCluster/Public/DisplayClusterPawn.h
--- a/Engine/Plugins/Runtime/nDisplay/Source/DisplayCluster/Public/DisplayClusterPawn.h
+++ b/Engine/Plugins/Runtime/nDisplay/Source/DisplayCluster/Public/DisplayClusterPawn.h
@@ -43,7 +43,7 @@ public:
 	UPROPERTY(EditAnywhere, Category = "DisplayCluster")
 	USceneComponent* RotationAround;
 
-protected:
+public:
 	//////////////////////////////////////////////////////////////////////////////////////////////
 	// APawn
 	//////////////////////////////////////////////////////////////////////////////////////////////
diff --git a/Engine/Plugins/Runtime/nDisplay/Source/DisplayCluster/Public/DisplayClusterPawnDefault.h b/Engine/Plugins/Runtime/nDisplay/Source/DisplayCluster/Public/DisplayClusterPawnDefault.h
--- a/Engine/Plugins/Runtime/nDisplay/Source/DisplayCluster/Public/DisplayClusterPawnDefault.h
+++ b/Engine/Plugins/Runtime/nDisplay/Source/DisplayCluster/Public/DisplayClusterPawnDefault.h
@@ -33,16 +33,18 @@ public:
 	virtual UPawnMovementComponent* GetMovementComponent() const override
 	{ return MovementComponent; }
 
-protected:
+public:
 	//////////////////////////////////////////////////////////////////////////////////////////////
 	// APawn
 	//////////////////////////////////////////////////////////////////////////////////////////////
-	virtual void SetupPlayerInputComponent(UInputComponent* InInputComponent) override;
 	virtual void BeginPlay() override;
 	virtual void BeginDestroy() override;
 	virtual void Tick(float DeltaSeconds) override;
 
 protected:
+	virtual void SetupPlayerInputComponent(UInputComponent* InInputComponent) override;
+
+public:
 	/**
 	* Input callback to move forward in local space (or backward if Val is negative).
 	* @param Val Amount of movement in the forward direction (or backward if negative).
diff --git a/Engine/Plugins/Runtime/nDisplay/Source/DisplayCluster/Public/DisplayClusterPlayerController.h b/Engine/Plugins/Runtime/nDisplay/Source/DisplayCluster/Public/DisplayClusterPlayerController.h
--- a/Engine/Plugins/Runtime/nDisplay/Source/DisplayCluster/Public/DisplayClusterPlayerController.h
+++ b/Engine/Plugins/Runtime/nDisplay/Source/DisplayCluster/Public/DisplayClusterPlayerController.h
@@ -14,7 +14,7 @@ class DISPLAYCLUSTER_API ADisplayClusterPlayerController
 {
 	GENERATED_BODY()
 
-protected:
+public:
 	virtual void BeginPlay() override;
 	virtual void PlayerTick(float DeltaTime) override;
 };
diff --git a/Engine/Plugins/Runtime/nDisplay/Source/DisplayCluster/Public/DisplayClusterSceneComponent.h b/Engine/Plugins/Runtime/nDisplay/Source/DisplayCluster/Public/DisplayClusterSceneComponent.h
--- a/Engine/Plugins/Runtime/nDisplay/Source/DisplayCluster/Public/DisplayClusterSceneComponent.h
+++ b/Engine/Plugins/Runtime/nDisplay/Source/DisplayCluster/Public/DisplayClusterSceneComponent.h
@@ -33,7 +33,7 @@ public:
 	inline FString GetParentId() const
 	{ return Config.ParentId; }
 
-protected:
+public:
 	virtual void BeginPlay() override;
 	virtual void BeginDestroy() override;
 	virtual void TickComponent(float DeltaTime, ELevelTick TickType, FActorComponentTickFunction* ThisTickFunction) override;
diff --git a/Engine/Plugins/Runtime/nDisplay/Source/DisplayCluster/Public/DisplayClusterSceneComponentSync.h b/Engine/Plugins/Runtime/nDisplay/Source/DisplayCluster/Public/DisplayClusterSceneComponentSync.h
--- a/Engine/Plugins/Runtime/nDisplay/Source/DisplayCluster/Public/DisplayClusterSceneComponentSync.h
+++ b/Engine/Plugins/Runtime/nDisplay/Source/DisplayCluster/Public/DisplayClusterSceneComponentSync.h
@@ -45,7 +45,7 @@ public:
 	virtual FString SerializeToString() const override final;
 	virtual bool    DeserializeFromString(const FString& data) override final;
 
-protected:
+public:
 	virtual void BeginPlay() override;
 	virtual void TickComponent(float DeltaTime, ELevelTick TickType, FActorComponentTickFunction* ThisTickFunction) override;
 	virtual void DestroyComponent(bool bPromoteChildren) override;
diff --git a/Engine/Plugins/Runtime/nDisplay/Source/DisplayCluster/Public/DisplayClusterSceneComponentSyncParent.h b/Engine/Plugins/Runtime/nDisplay/Source/DisplayCluster/Public/DisplayClusterSceneComponentSyncParent.h
--- a/Engine/Plugins/Runtime/nDisplay/Source/DisplayCluster/Public/DisplayClusterSceneComponentSyncParent.h
+++ b/Engine/Plugins/Runtime/nDisplay/Source/DisplayCluster/Public/DisplayClusterSceneComponentSyncParent.h
@@ -18,7 +18,7 @@ class DISPLAYCLUSTER_API UDisplayClusterSceneComponentSyncParent
 public:
 	UDisplayClusterSceneComponentSyncParent(const FObjectInitializer& ObjectInitializer);
 
-protected:
+public:
 	virtual void BeginPlay() override;
 	virtual void TickComponent( float DeltaTime, ELevelTick TickType, FActorComponentTickFunction* ThisTickFunction ) override;
 	virtual void DestroyComponent(bool bPromoteChildren) override;
diff --git a/Engine/Plugins/Runtime/nDisplay/Source/DisplayCluster/Public/DisplayClusterSceneComponentSyncThis.h b/Engine/Plugins/Runtime/nDisplay/Source/DisplayCluster/Public/DisplayClusterSceneComponentSyncThis.h
--- a/Engine/Plugins/Runtime/nDisplay/Source/DisplayCluster/Public/DisplayClusterSceneComponentSyncThis.h
+++ b/Engine/Plugins/Runtime/nDisplay/Source/DisplayCluster/Public/DisplayClusterSceneComponentSyncThis.h
@@ -18,7 +18,7 @@ class DISPLAYCLUSTER_API UDisplayClusterSceneComponentSyncThis
 public:
 	UDisplayClusterSceneComponentSyncThis(const FObjectInitializer& ObjectInitializer);
 
-protected:
+public:
 	virtual void BeginPlay() override;
 	virtual void TickComponent( float DeltaTime, ELevelTick TickType, FActorComponentTickFunction* ThisTickFunction ) override;
 	virtual void DestroyComponent(bool bPromoteChildren) override;
diff --git a/Engine/Plugins/Runtime/nDisplay/Source/DisplayCluster/Public/DisplayClusterScreenComponent.h b/Engine/Plugins/Runtime/nDisplay/Source/DisplayCluster/Public/DisplayClusterScreenComponent.h
--- a/Engine/Plugins/Runtime/nDisplay/Source/DisplayCluster/Public/DisplayClusterScreenComponent.h
+++ b/Engine/Plugins/Runtime/nDisplay/Source/DisplayCluster/Public/DisplayClusterScreenComponent.h
@@ -26,7 +26,7 @@ public:
 	inline FVector2D GetScreenSize() const
 	{ return Size; }
 
-protected:
+public:
 	virtual void BeginPlay() override;
 	virtual void TickComponent(float DeltaTime, ELevelTick TickType, FActorComponentTickFunction* ThisTickFunction) override;
 
