Fixed crash in DC components when the DC is not initialized

diff --git a/Engine/Plugins/Runtime/nDisplay/Source/DisplayCluster/Private/Blueprints/DisplayClusterBlueprintAPIImpl.cpp b/Engine/Plugins/Runtime/nDisplay/Source/DisplayCluster/Private/Blueprints/DisplayClusterBlueprintAPIImpl.cpp
--- a/Engine/Plugins/Runtime/nDisplay/Source/DisplayCluster/Private/Blueprints/DisplayClusterBlueprintAPIImpl.cpp
+++ b/Engine/Plugins/Runtime/nDisplay/Source/DisplayCluster/Private/Blueprints/DisplayClusterBlueprintAPIImpl.cpp
@@ -13,6 +13,20 @@
 
 
 
+//////////////////////////////////////////////////////////////////////////////////////////////
+// DisplayCluster module API
+//////////////////////////////////////////////////////////////////////////////////////////////
+/** Return if the module has been initialized. */
+bool UDisplayClusterBlueprintAPIImpl::IsModuleInitialized()
+{
+	return IDisplayCluster::Get().IsModuleInitialized();
+}
+
+EDisplayClusterOperationMode UDisplayClusterBlueprintAPIImpl::GetOperationMode()
+{
+	return IDisplayCluster::Get().GetOperationMode();
+}
+
 //////////////////////////////////////////////////////////////////////////////////////////////
 // Cluster API
 //////////////////////////////////////////////////////////////////////////////////////////////
diff --git a/Engine/Plugins/Runtime/nDisplay/Source/DisplayCluster/Private/Blueprints/DisplayClusterBlueprintAPIImpl.h b/Engine/Plugins/Runtime/nDisplay/Source/DisplayCluster/Private/Blueprints/DisplayClusterBlueprintAPIImpl.h
--- a/Engine/Plugins/Runtime/nDisplay/Source/DisplayCluster/Private/Blueprints/DisplayClusterBlueprintAPIImpl.h
+++ b/Engine/Plugins/Runtime/nDisplay/Source/DisplayCluster/Private/Blueprints/DisplayClusterBlueprintAPIImpl.h
@@ -17,6 +17,18 @@ class DISPLAYCLUSTER_API UDisplayClusterBlueprintAPIImpl
 {
 	GENERATED_BODY()
 
+public:
+	//////////////////////////////////////////////////////////////////////////////////////////////
+	// DisplayCluster module API
+	//////////////////////////////////////////////////////////////////////////////////////////////
+	/** Return if the module has been initialized. */
+	UFUNCTION(BlueprintCallable, meta = (DisplayName = "Is module initialized"), Category = "DisplayCluster")
+	virtual bool IsModuleInitialized() override;
+
+	/** Return current operation mode. */
+	UFUNCTION(BlueprintCallable, meta = (DisplayName = "Get operation mode"), Category = "DisplayCluster")
+	virtual EDisplayClusterOperationMode GetOperationMode() override;
+
 public:
 	//////////////////////////////////////////////////////////////////////////////////////////////
 	// Cluster API
diff --git a/Engine/Plugins/Runtime/nDisplay/Source/DisplayCluster/Private/Cluster/Controller/DisplayClusterClusterNodeCtrlBase.cpp b/Engine/Plugins/Runtime/nDisplay/Source/DisplayCluster/Private/Cluster/Controller/DisplayClusterClusterNodeCtrlBase.cpp
--- a/Engine/Plugins/Runtime/nDisplay/Source/DisplayCluster/Private/Cluster/Controller/DisplayClusterClusterNodeCtrlBase.cpp
+++ b/Engine/Plugins/Runtime/nDisplay/Source/DisplayCluster/Private/Cluster/Controller/DisplayClusterClusterNodeCtrlBase.cpp
@@ -34,6 +34,11 @@ FDisplayClusterClusterNodeCtrlBase::~FDisplayClusterClusterNodeCtrlBase()
 //////////////////////////////////////////////////////////////////////////////////////////////
 bool FDisplayClusterClusterNodeCtrlBase::InitializeStereo()
 {
+	if (GDisplayCluster->GetOperationMode() == EDisplayClusterOperationMode::Disabled)
+	{
+		return false;
+	}
+
 	FDisplayClusterConfigViewport ViewportCfg;
 	if (!GDisplayCluster->GetPrivateConfigMgr()->GetLocalViewport(ViewportCfg))
 	{
diff --git a/Engine/Plugins/Runtime/nDisplay/Source/DisplayCluster/Private/Cluster/Controller/DisplayClusterClusterNodeCtrlMaster.cpp b/Engine/Plugins/Runtime/nDisplay/Source/DisplayCluster/Private/Cluster/Controller/DisplayClusterClusterNodeCtrlMaster.cpp
--- a/Engine/Plugins/Runtime/nDisplay/Source/DisplayCluster/Private/Cluster/Controller/DisplayClusterClusterNodeCtrlMaster.cpp
+++ b/Engine/Plugins/Runtime/nDisplay/Source/DisplayCluster/Private/Cluster/Controller/DisplayClusterClusterNodeCtrlMaster.cpp
@@ -48,6 +48,11 @@ void FDisplayClusterClusterNodeCtrlMaster::GetInputData(FDisplayClusterMessage::
 //////////////////////////////////////////////////////////////////////////////////////////////
 bool FDisplayClusterClusterNodeCtrlMaster::InitializeServers()
 {
+	if (GDisplayCluster->GetOperationMode() == EDisplayClusterOperationMode::Disabled)
+	{
+		return false;
+	}
+
 	if (!FDisplayClusterClusterNodeCtrlSlave::InitializeServers())
 	{
 		return false;
diff --git a/Engine/Plugins/Runtime/nDisplay/Source/DisplayCluster/Private/Config/DisplayClusterConfigManager.cpp b/Engine/Plugins/Runtime/nDisplay/Source/DisplayCluster/Private/Config/DisplayClusterConfigManager.cpp
--- a/Engine/Plugins/Runtime/nDisplay/Source/DisplayCluster/Private/Config/DisplayClusterConfigManager.cpp
+++ b/Engine/Plugins/Runtime/nDisplay/Source/DisplayCluster/Private/Config/DisplayClusterConfigManager.cpp
@@ -107,6 +107,11 @@ bool FDisplayClusterConfigManager::GetMasterClusterNode(FDisplayClusterConfigClu
 
 bool FDisplayClusterConfigManager::GetLocalClusterNode(FDisplayClusterConfigClusterNode& node) const
 {
+	if (GDisplayCluster->GetOperationMode() == EDisplayClusterOperationMode::Disabled)
+	{
+		return false;
+	}
+
 	const FString nodeId = GDisplayCluster->GetPrivateClusterMgr()->GetNodeId();
 	return GetItem(CfgClusterNodes, nodeId, node, FString("GetLocalNode"));
 }
diff --git a/Engine/Plugins/Runtime/nDisplay/Source/DisplayCluster/Private/DisplayClusterModule.cpp b/Engine/Plugins/Runtime/nDisplay/Source/DisplayCluster/Private/DisplayClusterModule.cpp
--- a/Engine/Plugins/Runtime/nDisplay/Source/DisplayCluster/Private/DisplayClusterModule.cpp
+++ b/Engine/Plugins/Runtime/nDisplay/Source/DisplayCluster/Private/DisplayClusterModule.cpp
@@ -98,6 +98,9 @@ bool FDisplayClusterModule::Init(EDisplayClusterOperationMode OperationMode)
 		UE_LOG(LogDisplayClusterModule, Error, TEXT("An error occurred during internal initialization"));
 	}
 
+	// Set internal initialization flag
+	bIsModuleInitialized = result;
+
 	return result;
 }
 
diff --git a/Engine/Plugins/Runtime/nDisplay/Source/DisplayCluster/Private/DisplayClusterModule.h b/Engine/Plugins/Runtime/nDisplay/Source/DisplayCluster/Private/DisplayClusterModule.h
--- a/Engine/Plugins/Runtime/nDisplay/Source/DisplayCluster/Private/DisplayClusterModule.h
+++ b/Engine/Plugins/Runtime/nDisplay/Source/DisplayCluster/Private/DisplayClusterModule.h
@@ -25,6 +25,12 @@ public:
 	//////////////////////////////////////////////////////////////////////////////////////////////
 	// IDisplayCluster
 	//////////////////////////////////////////////////////////////////////////////////////////////
+	virtual bool IsModuleInitialized() const override
+	{ return bIsModuleInitialized; }
+	
+	virtual EDisplayClusterOperationMode GetOperationMode() const override
+	{ return CurrentOperationMode; }
+
 	virtual IDisplayClusterRenderManager*    GetRenderMgr()    const override { return MgrRender; }
 	virtual IDisplayClusterClusterManager*   GetClusterMgr()   const override { return MgrCluster; }
 	virtual IDisplayClusterInputManager*     GetInputMgr()     const override { return MgrInput; }
@@ -41,9 +47,6 @@ public:
 	virtual IPDisplayClusterConfigManager*    GetPrivateConfigMgr()    const override { return MgrConfig; }
 	virtual IPDisplayClusterGameManager*      GetPrivateGameMgr()      const override { return MgrGame; }
 
-	virtual EDisplayClusterOperationMode       GetOperationMode() const override
-	{ return CurrentOperationMode; }
-
 public:
 	//////////////////////////////////////////////////////////////////////////////////////////////
 	// IPDisplayClusterManager
@@ -71,6 +74,10 @@ private:
 #endif
 
 private:
+	// Is module initialized.
+	// This flag is not the same as EDisplayClusterOperationMode::Disabled which is used when we turn off the DC functionality in a game mode.
+	bool bIsModuleInitialized = false;
+
 	// DisplayCluster subsystems
 	IPDisplayClusterClusterManager*   MgrCluster   = nullptr;
 	IPDisplayClusterRenderManager*    MgrRender    = nullptr;
diff --git a/Engine/Plugins/Runtime/nDisplay/Source/DisplayCluster/Private/Game/Classes/Basics/DisplayClusterGameMode.cpp b/Engine/Plugins/Runtime/nDisplay/Source/DisplayCluster/Private/Game/Classes/Basics/DisplayClusterGameMode.cpp
--- a/Engine/Plugins/Runtime/nDisplay/Source/DisplayCluster/Private/Game/Classes/Basics/DisplayClusterGameMode.cpp
+++ b/Engine/Plugins/Runtime/nDisplay/Source/DisplayCluster/Private/Game/Classes/Basics/DisplayClusterGameMode.cpp
@@ -52,7 +52,12 @@ void ADisplayClusterGameMode::InitGame(const FString& MapName, const FString& Op
 
 	Super::InitGame(MapName, Options, ErrorMessage);
 
-	UE_LOG(LogDisplayClusterGame, Log, TEXT("%s"), bIsDisplayClusterActive ? 
+	if (!GDisplayCluster->IsModuleInitialized())
+	{
+		return;
+	}
+
+	UE_LOG(LogDisplayClusterGame, Log, TEXT("%s"), bIsDisplayClusterActive ?
 		TEXT("DisplayCluster feature is active for this world.") : 
 		TEXT("DisplayCluster feature has been deactivated for this world by game mode."));
 
@@ -117,7 +122,7 @@ void ADisplayClusterGameMode::StartPlay()
 {
 	DISPLAY_CLUSTER_FUNC_TRACE(LogDisplayClusterGame);
 
-	if (bIsDisplayClusterActive)
+	if (GDisplayCluster->IsModuleInitialized() && bIsDisplayClusterActive)
 	{
 		IPDisplayClusterGameManager* const pGameMgr = GDisplayCluster->GetPrivateGameMgr();
 		if (pGameMgr)
@@ -145,7 +150,7 @@ void ADisplayClusterGameMode::BeginPlay()
 {
 	DISPLAY_CLUSTER_FUNC_TRACE(LogDisplayClusterGame);
 
-	if (bIsDisplayClusterActive)
+	if (GDisplayCluster->IsModuleInitialized() && bIsDisplayClusterActive)
 	{
 		bGameStarted = GDisplayCluster->StartScene(GetWorld());
 		if (!bGameStarted)
@@ -162,7 +167,7 @@ void ADisplayClusterGameMode::BeginDestroy()
 {
 	DISPLAY_CLUSTER_FUNC_TRACE(LogDisplayClusterGame);
 
-	if (bIsDisplayClusterActive)
+	if (GDisplayCluster->IsModuleInitialized() && bIsDisplayClusterActive)
 	{
 		if (bGameStarted)
 		{
@@ -181,7 +186,7 @@ void ADisplayClusterGameMode::Tick(float DeltaSeconds)
 
 	Super::Tick(DeltaSeconds);
 
-	if (!bIsDisplayClusterActive)
+	if (!GDisplayCluster->IsModuleInitialized() || !bIsDisplayClusterActive)
 	{
 		return;
 	}
diff --git a/Engine/Plugins/Runtime/nDisplay/Source/DisplayCluster/Private/Game/Classes/Basics/DisplayClusterGameModeDefault.cpp b/Engine/Plugins/Runtime/nDisplay/Source/DisplayCluster/Private/Game/Classes/Basics/DisplayClusterGameModeDefault.cpp
--- a/Engine/Plugins/Runtime/nDisplay/Source/DisplayCluster/Private/Game/Classes/Basics/DisplayClusterGameModeDefault.cpp
+++ b/Engine/Plugins/Runtime/nDisplay/Source/DisplayCluster/Private/Game/Classes/Basics/DisplayClusterGameModeDefault.cpp
@@ -4,6 +4,8 @@
 #include "DisplayClusterPawnDefault.h"
 
 #include "Misc/DisplayClusterLog.h"
+#include "IPDisplayCluster.h"
+#include "DisplayClusterGlobals.h"
 
 
 ADisplayClusterGameModeDefault::ADisplayClusterGameModeDefault() :
diff --git a/Engine/Plugins/Runtime/nDisplay/Source/DisplayCluster/Private/Game/Classes/Scene/DisplayClusterPawn.cpp b/Engine/Plugins/Runtime/nDisplay/Source/DisplayCluster/Private/Game/Classes/Scene/DisplayClusterPawn.cpp
--- a/Engine/Plugins/Runtime/nDisplay/Source/DisplayCluster/Private/Game/Classes/Scene/DisplayClusterPawn.cpp
+++ b/Engine/Plugins/Runtime/nDisplay/Source/DisplayCluster/Private/Game/Classes/Scene/DisplayClusterPawn.cpp
@@ -67,6 +67,11 @@ void ADisplayClusterPawn::BeginPlay()
 
 	Super::BeginPlay();
 
+	if (!GDisplayCluster->IsModuleInitialized())
+	{
+		return;
+	}
+
 	GameMgr = GDisplayCluster->GetPrivateGameMgr();
 	bIsCluster = (GDisplayCluster->GetOperationMode() == EDisplayClusterOperationMode::Cluster);
 
@@ -74,11 +79,12 @@ void ADisplayClusterPawn::BeginPlay()
 	CollisionComponent->SetCollisionEnabled(ECollisionEnabled::NoCollision);
 
 	// Enable collision if needed
-	if (GameMgr->IsDisplayClusterActive())
+	if (GameMgr && GameMgr->IsDisplayClusterActive())
 	{
 		const ADisplayClusterSettings* const pDisplayClusterSettings = GameMgr->GetDisplayClusterSceneSettings();
 
-		if (GDisplayCluster->GetPrivateClusterMgr()->IsMaster())
+		const IPDisplayClusterClusterManager* const ClusterMgr = GDisplayCluster->GetPrivateClusterMgr();
+		if (ClusterMgr && ClusterMgr->IsMaster())
 		{
 			if (pDisplayClusterSettings && pDisplayClusterSettings->bEnableCollisions)
 			{
diff --git a/Engine/Plugins/Runtime/nDisplay/Source/DisplayCluster/Private/Game/Classes/Scene/DisplayClusterPawnDefault.cpp b/Engine/Plugins/Runtime/nDisplay/Source/DisplayCluster/Private/Game/Classes/Scene/DisplayClusterPawnDefault.cpp
--- a/Engine/Plugins/Runtime/nDisplay/Source/DisplayCluster/Private/Game/Classes/Scene/DisplayClusterPawnDefault.cpp
+++ b/Engine/Plugins/Runtime/nDisplay/Source/DisplayCluster/Private/Game/Classes/Scene/DisplayClusterPawnDefault.cpp
@@ -69,6 +69,11 @@ void ADisplayClusterPawnDefault::BeginPlay()
 
 	Super::BeginPlay();
 
+	if (!GDisplayCluster->IsModuleInitialized())
+	{
+		return;
+	}
+
 	GameMgr = GDisplayCluster->GetPrivateGameMgr();
 	bIsCluster = (GDisplayCluster->GetOperationMode() == EDisplayClusterOperationMode::Cluster);
 
@@ -77,7 +82,7 @@ void ADisplayClusterPawnDefault::BeginPlay()
 	bUseControllerRotationRoll  = !bIsCluster;
 
 	// Enable collision if needed
-	if (GameMgr->IsDisplayClusterActive())
+	if (GameMgr && GameMgr->IsDisplayClusterActive())
 	{
 		const ADisplayClusterSettings* const pDisplayClusterSettings = GameMgr->GetDisplayClusterSceneSettings();
 		if (pDisplayClusterSettings)
@@ -223,6 +228,7 @@ void ADisplayClusterPawnDefault::LookUpAtRate(float Rate)
 
 	if (bIsCluster)
 	{
+		//@note: usually CAVE-like systems don't use roll and pitch rotation since it can cause dizziness.
 #if 0
 		//@todo: rotate around active camera
 		IPDisplayClusterGameManager* const pMgr = GDisplayCluster->GetPrivateGameMgr();
diff --git a/Engine/Plugins/Runtime/nDisplay/Source/DisplayCluster/Private/Game/Classes/Scene/DisplayClusterSceneComponent.cpp b/Engine/Plugins/Runtime/nDisplay/Source/DisplayCluster/Private/Game/Classes/Scene/DisplayClusterSceneComponent.cpp
--- a/Engine/Plugins/Runtime/nDisplay/Source/DisplayCluster/Private/Game/Classes/Scene/DisplayClusterSceneComponent.cpp
+++ b/Engine/Plugins/Runtime/nDisplay/Source/DisplayCluster/Private/Game/Classes/Scene/DisplayClusterSceneComponent.cpp
@@ -34,13 +34,13 @@ void UDisplayClusterSceneComponent::TickComponent( float DeltaTime, ELevelTick T
 	// Update transform if attached to a tracker
 	if (!Config.TrackerId.IsEmpty())
 	{
-		IPDisplayClusterInputManager* const inputMgr = GDisplayCluster->GetPrivateInputMgr();
-		if (inputMgr)
+		const IPDisplayClusterInputManager* const InputMgr = GDisplayCluster->GetPrivateInputMgr();
+		if (InputMgr)
 		{
 			FVector loc;
 			FQuat rot;
-			const bool bLocAvail = inputMgr->GetTrackerLocation(Config.TrackerId, Config.TrackerCh, loc);
-			const bool bRotAvail = inputMgr->GetTrackerQuat(Config.TrackerId, Config.TrackerCh, rot);
+			const bool bLocAvail = InputMgr->GetTrackerLocation(Config.TrackerId, Config.TrackerCh, loc);
+			const bool bRotAvail = InputMgr->GetTrackerQuat(Config.TrackerId, Config.TrackerCh, rot);
 
 			if (bLocAvail && bRotAvail)
 			{
@@ -71,10 +71,14 @@ bool UDisplayClusterSceneComponent::ApplySettings()
 	// Take place in hierarchy
 	if (!GetParentId().IsEmpty())
 	{
-		UE_LOG(LogDisplayClusterGame, Log, TEXT("Attaching %s to %s"), *GetId(), *GetParentId());
-		UDisplayClusterSceneComponent* const pComp = GDisplayCluster->GetPrivateGameMgr()->GetNodeById(GetParentId());
-		AttachToComponent(pComp, FAttachmentTransformRules(EAttachmentRule::KeepRelative, false));
-		//this->SetRelativeTransform(FTransform::Identity);
+		const IPDisplayClusterGameManager* const GameMgr = GDisplayCluster->GetPrivateGameMgr();
+		if (GameMgr)
+		{
+			UE_LOG(LogDisplayClusterGame, Log, TEXT("Attaching %s to %s"), *GetId(), *GetParentId());
+			UDisplayClusterSceneComponent* const pComp = GameMgr->GetNodeById(GetParentId());
+			AttachToComponent(pComp, FAttachmentTransformRules(EAttachmentRule::KeepRelative, false));
+			//this->SetRelativeTransform(FTransform::Identity);
+		}
 	}
 
 	// Set up location and rotation
diff --git a/Engine/Plugins/Runtime/nDisplay/Source/DisplayCluster/Private/Game/Classes/Scene/DisplayClusterSceneComponentSync.cpp b/Engine/Plugins/Runtime/nDisplay/Source/DisplayCluster/Private/Game/Classes/Scene/DisplayClusterSceneComponentSync.cpp
--- a/Engine/Plugins/Runtime/nDisplay/Source/DisplayCluster/Private/Game/Classes/Scene/DisplayClusterSceneComponentSync.cpp
+++ b/Engine/Plugins/Runtime/nDisplay/Source/DisplayCluster/Private/Game/Classes/Scene/DisplayClusterSceneComponentSync.cpp
@@ -19,12 +19,16 @@ void UDisplayClusterSceneComponentSync::BeginPlay()
 {
 	Super::BeginPlay();
 
-	GameMgr = GDisplayCluster->GetPrivateGameMgr();
+	if (!GDisplayCluster->IsModuleInitialized())
+	{
+		return;
+	}
 
 	// Generate unique sync id
 	SyncId = GetSyncId();
 
-	if (GameMgr->IsDisplayClusterActive())
+	GameMgr = GDisplayCluster->GetPrivateGameMgr();
+	if (GameMgr && GameMgr->IsDisplayClusterActive())
 	{
 		// Register sync object
 		ClusterMgr = GDisplayCluster->GetPrivateClusterMgr();
@@ -50,10 +54,16 @@ void UDisplayClusterSceneComponentSync::TickComponent( float DeltaTime, ELevelTi
 
 void UDisplayClusterSceneComponentSync::DestroyComponent(bool bPromoteChildren)
 {
-	if (ClusterMgr)
+	if (GDisplayCluster->IsModuleInitialized())
 	{
-		UE_LOG(LogDisplayClusterGame, Log, TEXT("Unregistering sync object %s..."), *SyncId);
-		ClusterMgr->UnregisterSyncObject(this);
+		if (GameMgr && GameMgr->IsDisplayClusterActive())
+		{
+			if (ClusterMgr)
+			{
+				UE_LOG(LogDisplayClusterGame, Log, TEXT("Unregistering sync object %s..."), *SyncId);
+				ClusterMgr->UnregisterSyncObject(this);
+			}
+		}
 	}
 
 	Super::DestroyComponent(bPromoteChildren);
diff --git a/Engine/Plugins/Runtime/nDisplay/Source/DisplayCluster/Private/Game/Classes/Scene/DisplayClusterScreenComponent.cpp b/Engine/Plugins/Runtime/nDisplay/Source/DisplayCluster/Private/Game/Classes/Scene/DisplayClusterScreenComponent.cpp
--- a/Engine/Plugins/Runtime/nDisplay/Source/DisplayCluster/Private/Game/Classes/Scene/DisplayClusterScreenComponent.cpp
+++ b/Engine/Plugins/Runtime/nDisplay/Source/DisplayCluster/Private/Game/Classes/Scene/DisplayClusterScreenComponent.cpp
@@ -23,22 +23,26 @@ UDisplayClusterScreenComponent::UDisplayClusterScreenComponent(const FObjectInit
 #if WITH_EDITOR
 	if (GEngine && GEngine->IsEditor())
 	{
-		const ADisplayClusterSettings* const pDisplayClusterSettings = GDisplayCluster->GetPrivateGameMgr()->GetDisplayClusterSceneSettings();
-		if(pDisplayClusterSettings && pDisplayClusterSettings->bEditorShowProjectionScreens)
+		const IPDisplayClusterGameManager* const GameMgr = GDisplayCluster->GetPrivateGameMgr();
+		if (GameMgr)
 		{
-			ScreenGeometryComponent = CreateDefaultSubobject<UStaticMeshComponent>(FName(*(GetName() + FString("_impl"))));
-			check(ScreenGeometryComponent);
-
-			if (ScreenGeometryComponent)
+			const ADisplayClusterSettings* const pDisplayClusterSettings = GameMgr->GetDisplayClusterSceneSettings();
+			if (pDisplayClusterSettings && pDisplayClusterSettings->bEditorShowProjectionScreens)
 			{
-				static ConstructorHelpers::FObjectFinder<UStaticMesh> screenMesh(TEXT("StaticMesh'/Engine/BasicShapes/Cube.Cube'"));
-				static ConstructorHelpers::FObjectFinder<UMaterial>   screenMat(TEXT("Material'/Engine/Engine_MI_Shaders/M_Shader_SimpleTranslucent.M_Shader_SimpleTranslucent'"));
-
-				ScreenGeometryComponent->AttachToComponent(this, FAttachmentTransformRules(EAttachmentRule::KeepRelative, false));
-				ScreenGeometryComponent->SetStaticMesh(screenMesh.Object);
-				ScreenGeometryComponent->SetMobility(EComponentMobility::Movable);
-				ScreenGeometryComponent->SetMaterial(0, screenMat.Object);
-				ScreenGeometryComponent->SetCollisionEnabled(ECollisionEnabled::NoCollision);
+				ScreenGeometryComponent = CreateDefaultSubobject<UStaticMeshComponent>(FName(*(GetName() + FString("_impl"))));
+				check(ScreenGeometryComponent);
+
+				if (ScreenGeometryComponent)
+				{
+					static ConstructorHelpers::FObjectFinder<UStaticMesh> screenMesh(TEXT("StaticMesh'/Engine/BasicShapes/Cube.Cube'"));
+					static ConstructorHelpers::FObjectFinder<UMaterial>   screenMat(TEXT("Material'/Engine/Engine_MI_Shaders/M_Shader_SimpleTranslucent.M_Shader_SimpleTranslucent'"));
+
+					ScreenGeometryComponent->AttachToComponent(this, FAttachmentTransformRules(EAttachmentRule::KeepRelative, false));
+					ScreenGeometryComponent->SetStaticMesh(screenMesh.Object);
+					ScreenGeometryComponent->SetMobility(EComponentMobility::Movable);
+					ScreenGeometryComponent->SetMaterial(0, screenMat.Object);
+					ScreenGeometryComponent->SetCollisionEnabled(ECollisionEnabled::NoCollision);
+				}
 			}
 		}
 	}
diff --git a/Engine/Plugins/Runtime/nDisplay/Source/DisplayCluster/Private/IPDisplayCluster.h b/Engine/Plugins/Runtime/nDisplay/Source/DisplayCluster/Private/IPDisplayCluster.h
--- a/Engine/Plugins/Runtime/nDisplay/Source/DisplayCluster/Private/IPDisplayCluster.h
+++ b/Engine/Plugins/Runtime/nDisplay/Source/DisplayCluster/Private/IPDisplayCluster.h
@@ -30,6 +30,4 @@ struct IPDisplayCluster
 	virtual IPDisplayClusterInputManager*     GetPrivateInputMgr() const = 0;
 	virtual IPDisplayClusterConfigManager*    GetPrivateConfigMgr() const = 0;
 	virtual IPDisplayClusterGameManager*      GetPrivateGameMgr() const = 0;
-
-	virtual EDisplayClusterOperationMode       GetOperationMode() const = 0;
 };
diff --git a/Engine/Plugins/Runtime/nDisplay/Source/DisplayCluster/Private/Input/DisplayClusterInputManager.cpp b/Engine/Plugins/Runtime/nDisplay/Source/DisplayCluster/Private/Input/DisplayClusterInputManager.cpp
--- a/Engine/Plugins/Runtime/nDisplay/Source/DisplayCluster/Private/Input/DisplayClusterInputManager.cpp
+++ b/Engine/Plugins/Runtime/nDisplay/Source/DisplayCluster/Private/Input/DisplayClusterInputManager.cpp
@@ -233,6 +233,11 @@ void FDisplayClusterInputManager::Update()
 {
 	DISPLAY_CLUSTER_FUNC_TRACE(LogDisplayClusterInput);
 
+	if (GDisplayCluster->GetOperationMode() == EDisplayClusterOperationMode::Disabled)
+	{
+		return;
+	}
+
 	// Perform input update on master only. Slaves' state will be replicated later.
 	if (GDisplayCluster->GetPrivateClusterMgr()->IsMaster())
 	{
@@ -376,6 +381,11 @@ bool FDisplayClusterInputManager::InitDevices()
 {
 	DISPLAY_CLUSTER_FUNC_TRACE(LogDisplayClusterInput);
 
+	if (GDisplayCluster->GetOperationMode() == EDisplayClusterOperationMode::Disabled)
+	{
+		return false;
+	}
+
 	FScopeLock ScopeLock(&InternalsSyncScope);
 
 	UE_LOG(LogDisplayClusterInput, Log, TEXT("Initializing input devices..."));
diff --git a/Engine/Plugins/Runtime/nDisplay/Source/DisplayCluster/Private/Network/Service/DisplayClusterService.cpp b/Engine/Plugins/Runtime/nDisplay/Source/DisplayCluster/Private/Network/Service/DisplayClusterService.cpp
--- a/Engine/Plugins/Runtime/nDisplay/Source/DisplayCluster/Private/Network/Service/DisplayClusterService.cpp
+++ b/Engine/Plugins/Runtime/nDisplay/Source/DisplayCluster/Private/Network/Service/DisplayClusterService.cpp
@@ -18,6 +18,11 @@ FDisplayClusterService::FDisplayClusterService(const FString& name, const FStrin
 
 bool FDisplayClusterService::IsClusterIP(const FIPv4Endpoint& ep)
 {
+	if (GDisplayCluster->GetOperationMode() == EDisplayClusterOperationMode::Disabled)
+	{
+		return false;
+	}
+
 	TArray<FDisplayClusterConfigClusterNode> nodes = GDisplayCluster->GetPrivateConfigMgr()->GetClusterNodes();
 	const FString addr = ep.Address.ToString();
 	
diff --git a/Engine/Plugins/Runtime/nDisplay/Source/DisplayCluster/Private/Render/Devices/DisplayClusterDeviceBase.cpp b/Engine/Plugins/Runtime/nDisplay/Source/DisplayCluster/Private/Render/Devices/DisplayClusterDeviceBase.cpp
--- a/Engine/Plugins/Runtime/nDisplay/Source/DisplayCluster/Private/Render/Devices/DisplayClusterDeviceBase.cpp
+++ b/Engine/Plugins/Runtime/nDisplay/Source/DisplayCluster/Private/Render/Devices/DisplayClusterDeviceBase.cpp
@@ -34,6 +34,11 @@ FDisplayClusterDeviceBase::~FDisplayClusterDeviceBase()
 
 bool FDisplayClusterDeviceBase::Initialize()
 {
+	if (GDisplayCluster->GetOperationMode() == EDisplayClusterOperationMode::Disabled)
+	{
+		return false;
+	}
+
 	const FDisplayClusterConfigCustom cfgCustom = GDisplayCluster->GetPrivateConfigMgr()->GetConfigCustom();
 
 	if (cfgCustom.Args.Contains(FString(DisplayClusterStrings::cfg::data::custom::SwapInt)))
@@ -75,6 +80,11 @@ void FDisplayClusterDeviceBase::UpdateProjectionScreenDataForThisFrame()
 	UE_LOG(LogDisplayClusterRender, VeryVerbose, TEXT("UpdateProjectionScreenDataForThisFrame"));
 	check(IsInGameThread());
 
+	if (GDisplayCluster->GetOperationMode() == EDisplayClusterOperationMode::Disabled)
+	{
+		return;
+	}
+
 	// Store transformations of active projection screen
 	UDisplayClusterScreenComponent* pScreen = GDisplayCluster->GetPrivateGameMgr()->GetActiveScreen();
 	if (pScreen)
@@ -87,6 +97,11 @@ void FDisplayClusterDeviceBase::UpdateProjectionScreenDataForThisFrame()
 
 void FDisplayClusterDeviceBase::exec_BarrierWait()
 {
+	if (GDisplayCluster->GetOperationMode() == EDisplayClusterOperationMode::Disabled)
+	{
+		return;
+	}
+
 	double tTime = 0.f;
 	double bTime = 0.f;
 
diff --git a/Engine/Plugins/Runtime/nDisplay/Source/DisplayCluster/Private/Render/DisplayClusterRenderManager.h b/Engine/Plugins/Runtime/nDisplay/Source/DisplayCluster/Private/Render/DisplayClusterRenderManager.h
--- a/Engine/Plugins/Runtime/nDisplay/Source/DisplayCluster/Private/Render/DisplayClusterRenderManager.h
+++ b/Engine/Plugins/Runtime/nDisplay/Source/DisplayCluster/Private/Render/DisplayClusterRenderManager.h
@@ -6,7 +6,6 @@
 #include "IPDisplayClusterRenderManager.h"
 
 class FDisplayClusterDeviceBase;
-enum class EDisplayClusterOperationMode;
 
 
 /**
diff --git a/Engine/Plugins/Runtime/nDisplay/Source/DisplayCluster/Public/Blueprints/IDisplayClusterBlueprintAPI.h b/Engine/Plugins/Runtime/nDisplay/Source/DisplayCluster/Public/Blueprints/IDisplayClusterBlueprintAPI.h
--- a/Engine/Plugins/Runtime/nDisplay/Source/DisplayCluster/Public/Blueprints/IDisplayClusterBlueprintAPI.h
+++ b/Engine/Plugins/Runtime/nDisplay/Source/DisplayCluster/Public/Blueprints/IDisplayClusterBlueprintAPI.h
@@ -4,6 +4,7 @@
 
 #include "CoreMinimal.h"
 #include "UObject/Interface.h"
+#include "DisplayClusterOperationMode.h"
 #include "IDisplayClusterBlueprintAPI.generated.h"
 
 
@@ -21,6 +22,18 @@ class DISPLAYCLUSTER_API IDisplayClusterBlueprintAPI
 {
 	GENERATED_BODY()
 
+public:
+	//////////////////////////////////////////////////////////////////////////////////////////////
+	// DisplayCluster module API
+	//////////////////////////////////////////////////////////////////////////////////////////////
+	/** Return if the module has been initialized. */
+	UFUNCTION(BlueprintCallable, meta = (DisplayName = "Is module initialized"), Category = "DisplayCluster")
+	virtual bool IsModuleInitialized() = 0;
+
+	/** Return current operation mode. */
+	UFUNCTION(BlueprintCallable, meta = (DisplayName = "Get operation mode"), Category = "DisplayCluster")
+	virtual EDisplayClusterOperationMode GetOperationMode() = 0;
+
 public:
 	//////////////////////////////////////////////////////////////////////////////////////////////
 	// Cluster API
diff --git a/Engine/Plugins/Runtime/nDisplay/Source/DisplayCluster/Public/DisplayClusterOperationMode.h b/Engine/Plugins/Runtime/nDisplay/Source/DisplayCluster/Public/DisplayClusterOperationMode.h
--- a/Engine/Plugins/Runtime/nDisplay/Source/DisplayCluster/Public/DisplayClusterOperationMode.h
+++ b/Engine/Plugins/Runtime/nDisplay/Source/DisplayCluster/Public/DisplayClusterOperationMode.h
@@ -2,11 +2,14 @@
 
 #pragma once
 
+#include "CoreMinimal.h"
+#include "DisplayClusterOperationMode.generated.h"
 
 /**
  * Display cluster operation mode
  */
-enum class EDisplayClusterOperationMode
+UENUM(BlueprintType)
+enum class EDisplayClusterOperationMode : uint8
 {
 	Cluster,
 	Standalone,
diff --git a/Engine/Plugins/Runtime/nDisplay/Source/DisplayCluster/Public/DisplayClusterPawn.h b/Engine/Plugins/Runtime/nDisplay/Source/DisplayCluster/Public/DisplayClusterPawn.h
--- a/Engine/Plugins/Runtime/nDisplay/Source/DisplayCluster/Public/DisplayClusterPawn.h
+++ b/Engine/Plugins/Runtime/nDisplay/Source/DisplayCluster/Public/DisplayClusterPawn.h
@@ -65,7 +65,10 @@ protected:
 	UDisplayClusterSceneComponent* CollisionOffsetComponent;
 
 private:
+	UPROPERTY()
 	UDisplayClusterSceneComponentSyncParent* DisplayClusterSyncRoot;
+	
+	UPROPERTY()
 	UDisplayClusterSceneComponentSyncParent* DisplayClusterSyncCollisionOffset;
 
 	IPDisplayClusterGameManager* GameMgr = nullptr;
diff --git a/Engine/Plugins/Runtime/nDisplay/Source/DisplayCluster/Public/DisplayClusterPawnDefault.h b/Engine/Plugins/Runtime/nDisplay/Source/DisplayCluster/Public/DisplayClusterPawnDefault.h
--- a/Engine/Plugins/Runtime/nDisplay/Source/DisplayCluster/Public/DisplayClusterPawnDefault.h
+++ b/Engine/Plugins/Runtime/nDisplay/Source/DisplayCluster/Public/DisplayClusterPawnDefault.h
@@ -41,9 +41,6 @@ public:
 	virtual void BeginDestroy() override;
 	virtual void Tick(float DeltaSeconds) override;
 
-protected:
-	virtual void SetupPlayerInputComponent(UInputComponent* InInputComponent) override;
-
 public:
 	/**
 	* Input callback to move forward in local space (or backward if Val is negative).
@@ -86,6 +83,9 @@ public:
 	UFUNCTION(BlueprintCallable, Category = "Pawn")
 	void LookUpAtRate(float Rate);
 
+protected:
+	virtual void SetupPlayerInputComponent(UInputComponent* InInputComponent) override;
+
 protected:
 	/** Movement component */
 	UPROPERTY(Category = Pawn, VisibleAnywhere, BlueprintReadOnly, meta = (AllowPrivateAccess = "true"))
@@ -101,5 +101,5 @@ protected:
 private:
 	IPDisplayClusterGameManager* GameMgr = nullptr;
 
-	bool bIsCluster;
+	bool bIsCluster = false;
 };
diff --git a/Engine/Plugins/Runtime/nDisplay/Source/DisplayCluster/Public/IDisplayCluster.h b/Engine/Plugins/Runtime/nDisplay/Source/DisplayCluster/Public/IDisplayCluster.h
--- a/Engine/Plugins/Runtime/nDisplay/Source/DisplayCluster/Public/IDisplayCluster.h
+++ b/Engine/Plugins/Runtime/nDisplay/Source/DisplayCluster/Public/IDisplayCluster.h
@@ -5,6 +5,9 @@
 #include "Modules/ModuleManager.h"
 #include "Modules/ModuleInterface.h"
 
+#include "DisplayClusterOperationMode.h"
+
+
 struct IDisplayClusterRenderManager;
 struct IDisplayClusterClusterManager;
 struct IDisplayClusterInputManager;
@@ -44,6 +47,21 @@ struct IDisplayCluster
 		return FModuleManager::Get().IsModuleLoaded(IDisplayCluster::ModuleName);
 	}
 
+
+	/**
+	* Checks if the module has been initialized.
+	*
+	* @return Is initialized
+	*/
+	virtual bool IsModuleInitialized() const = 0;
+
+	/**
+	* Returns current operation mode.
+	*
+	* @return Display Cluster operation mode
+	*/
+	virtual EDisplayClusterOperationMode GetOperationMode() const = 0;
+
 	/**
 	* Access to the device manager.
 	*
