Marking functions as inline to fix CIS errors.

#rnx

diff --git a/Engine/Plugins/Runtime/nDisplay/Source/DisplayCluster/Private/Misc/DisplayClusterTypesConverter.h b/Engine/Plugins/Runtime/nDisplay/Source/DisplayCluster/Private/Misc/DisplayClusterTypesConverter.h
--- a/Engine/Plugins/Runtime/nDisplay/Source/DisplayCluster/Private/Misc/DisplayClusterTypesConverter.h
+++ b/Engine/Plugins/Runtime/nDisplay/Source/DisplayCluster/Private/Misc/DisplayClusterTypesConverter.h
@@ -28,20 +28,20 @@ public:
 	static ConvertTo FromString(const FString& from);
 };
 
-template <> FString FDisplayClusterTypesConverter::ToString<>(const FString& from) { return from; }
-template <> FString FDisplayClusterTypesConverter::ToString<>(const bool& from) { return (from ? DisplayClusterStrings::cfg::spec::ValTrue : DisplayClusterStrings::cfg::spec::ValFalse); }
-template <> FString FDisplayClusterTypesConverter::ToString<>(const int32& from) { return FString::FromInt(from); }
-template <> FString FDisplayClusterTypesConverter::ToString<>(const float& from) { return FString::SanitizeFloat(from); }
-template <> FString FDisplayClusterTypesConverter::ToString<>(const double& from) { return FString::Printf(TEXT("%lf"), from); }
-template <> FString FDisplayClusterTypesConverter::ToString<>(const FVector& from) { return from.ToString(); }
-template <> FString FDisplayClusterTypesConverter::ToString<>(const FVector2D& from) { return from.ToString(); }
-template <> FString FDisplayClusterTypesConverter::ToString<>(const FRotator& from) { return from.ToString(); }
+template <> inline FString FDisplayClusterTypesConverter::ToString<>(const FString& from) { return from; }
+template <> inline FString FDisplayClusterTypesConverter::ToString<>(const bool& from) { return (from ? DisplayClusterStrings::cfg::spec::ValTrue : DisplayClusterStrings::cfg::spec::ValFalse); }
+template <> inline FString FDisplayClusterTypesConverter::ToString<>(const int32& from) { return FString::FromInt(from); }
+template <> inline FString FDisplayClusterTypesConverter::ToString<>(const float& from) { return FString::SanitizeFloat(from); }
+template <> inline FString FDisplayClusterTypesConverter::ToString<>(const double& from) { return FString::Printf(TEXT("%lf"), from); }
+template <> inline FString FDisplayClusterTypesConverter::ToString<>(const FVector& from) { return from.ToString(); }
+template <> inline FString FDisplayClusterTypesConverter::ToString<>(const FVector2D& from) { return from.ToString(); }
+template <> inline FString FDisplayClusterTypesConverter::ToString<>(const FRotator& from) { return from.ToString(); }
 
 // We can't just use FTimecode ToString as that loses information.
-template <> FString FDisplayClusterTypesConverter::ToString<>(const FTimecode& from) { return FString::Printf(TEXT("%d;%d;%d;%d;%d"), from.bDropFrameFormat ? 1 : 0, from.Hours, from.Minutes, from.Seconds, from.Frames); }
-template <> FString FDisplayClusterTypesConverter::ToString<>(const FFrameRate& from) { return FString::Printf(TEXT("%d;%d"), from.Numerator, from.Denominator); }
+template <> inline FString FDisplayClusterTypesConverter::ToString<>(const FTimecode& from) { return FString::Printf(TEXT("%d;%d;%d;%d;%d"), from.bDropFrameFormat ? 1 : 0, from.Hours, from.Minutes, from.Seconds, from.Frames); }
+template <> inline FString FDisplayClusterTypesConverter::ToString<>(const FFrameRate& from) { return FString::Printf(TEXT("%d;%d"), from.Numerator, from.Denominator); }
 
-template <> FString FDisplayClusterTypesConverter::ToString<>(const EDisplayClusterOperationMode& from)
+template <> inline FString FDisplayClusterTypesConverter::ToString<>(const EDisplayClusterOperationMode& from)
 {
 	switch (from)
 	{
@@ -58,15 +58,15 @@ template <> FString FDisplayClusterTypesConverter::ToString<>(const EDisplayClus
 	}
 }
 
-template <> FString   FDisplayClusterTypesConverter::FromString<> (const FString& from) { return from; }
-template <> bool      FDisplayClusterTypesConverter::FromString<> (const FString& from) { return (from == FString("1") || from == DisplayClusterStrings::cfg::spec::ValTrue); }
-template <> int32     FDisplayClusterTypesConverter::FromString<> (const FString& from) { return FCString::Atoi(*from); }
-template <> float     FDisplayClusterTypesConverter::FromString<> (const FString& from) { return FCString::Atof(*from); }
-template <> double    FDisplayClusterTypesConverter::FromString<> (const FString& from) { return FCString::Atod(*from); }
-template <> FVector   FDisplayClusterTypesConverter::FromString<> (const FString& from) { FVector vec;  vec.InitFromString(from); return vec; }
-template <> FVector2D FDisplayClusterTypesConverter::FromString<> (const FString& from) { FVector2D vec;  vec.InitFromString(from); return vec; }
-template <> FRotator  FDisplayClusterTypesConverter::FromString<> (const FString& from) { FRotator rot; rot.InitFromString(from); return rot; }
-template <> FTimecode FDisplayClusterTypesConverter::FromString<> (const FString& from)
+template <> inline FString   FDisplayClusterTypesConverter::FromString<> (const FString& from) { return from; }
+template <> inline bool      FDisplayClusterTypesConverter::FromString<> (const FString& from) { return (from == FString("1") || from == DisplayClusterStrings::cfg::spec::ValTrue); }
+template <> inline int32     FDisplayClusterTypesConverter::FromString<> (const FString& from) { return FCString::Atoi(*from); }
+template <> inline float     FDisplayClusterTypesConverter::FromString<> (const FString& from) { return FCString::Atof(*from); }
+template <> inline double    FDisplayClusterTypesConverter::FromString<> (const FString& from) { return FCString::Atod(*from); }
+template <> inline FVector   FDisplayClusterTypesConverter::FromString<> (const FString& from) { FVector vec;  vec.InitFromString(from); return vec; }
+template <> inline FVector2D FDisplayClusterTypesConverter::FromString<> (const FString& from) { FVector2D vec;  vec.InitFromString(from); return vec; }
+template <> inline FRotator  FDisplayClusterTypesConverter::FromString<> (const FString& from) { FRotator rot; rot.InitFromString(from); return rot; }
+template <> inline FTimecode FDisplayClusterTypesConverter::FromString<> (const FString& from)
 {
 	FTimecode timecode;
 
@@ -86,7 +86,7 @@ template <> FTimecode FDisplayClusterTypesConverter::FromString<> (const FString
 
 	return timecode;
 }
-template <> FFrameRate FDisplayClusterTypesConverter::FromString<> (const FString& from)
+template <> inline FFrameRate FDisplayClusterTypesConverter::FromString<> (const FString& from)
 {
 	FFrameRate frameRate;
 
