Copied cl 4971794
Multi-viewports fix after dev-vr merge

diff --git a/Engine/Plugins/Runtime/nDisplay/Source/DisplayCluster/Private/Render/Devices/DisplayClusterDeviceBase.cpp b/Engine/Plugins/Runtime/nDisplay/Source/DisplayCluster/Private/Render/Devices/DisplayClusterDeviceBase.cpp
--- a/Engine/Plugins/Runtime/nDisplay/Source/DisplayCluster/Private/Render/Devices/DisplayClusterDeviceBase.cpp
+++ b/Engine/Plugins/Runtime/nDisplay/Source/DisplayCluster/Private/Render/Devices/DisplayClusterDeviceBase.cpp
@@ -154,9 +154,9 @@ void FDisplayClusterDeviceBase::PerformSynchronizationPolicyNvSwapLock(int32& In
 
 EStereoscopicPass FDisplayClusterDeviceBase::EncodeStereoscopicPass(int ViewIndex) const
 {
-	EStereoscopicPass EncodedPass = EStereoscopicPass::eSSP_RIGHT_EYE_SIDE;
+	EStereoscopicPass EncodedPass = EStereoscopicPass::eSSP_FULL;
 
-	// We don't care about mono/stereo. We need to fullfil ViewState and StereoViewStates in a proper way.
+	// We don't care about mono/stereo. We need to fulfill ViewState and StereoViewStates in a proper way.
 	// Look at ULocalPlayer::CalcSceneViewInitOptions for view states mapping.
 	if (ViewIndex < 2)
 	{
@@ -164,7 +164,7 @@ EStereoscopicPass FDisplayClusterDeviceBase::EncodeStereoscopicPass(int ViewInde
 	}
 	else
 	{
-		EncodedPass = EStereoscopicPass(int(EStereoscopicPass::eSSP_RIGHT_EYE_SIDE) + ViewIndex - 1);
+		EncodedPass = EStereoscopicPass(int(EStereoscopicPass::eSSP_RIGHT_EYE) + ViewIndex - 1);
 	}
 
 	return EncodedPass;
@@ -172,12 +172,12 @@ EStereoscopicPass FDisplayClusterDeviceBase::EncodeStereoscopicPass(int ViewInde
 
 EStereoscopicPass FDisplayClusterDeviceBase::DecodeStereoscopicPass(const enum EStereoscopicPass StereoPassType) const
 {
-	EStereoscopicPass DecodedPass = EStereoscopicPass::eSSP_RIGHT_EYE_SIDE;
+	EStereoscopicPass DecodedPass = EStereoscopicPass::eSSP_FULL;
 
 	// Monoscopic rendering
 	if (ViewsAmountPerViewport == 1)
 	{
-		DecodedPass = EStereoscopicPass::eSSP_RIGHT_EYE_SIDE;
+		DecodedPass = EStereoscopicPass::eSSP_FULL;
 	}
 	// Stereoscopic rendering
 	else
@@ -190,7 +190,7 @@ EStereoscopicPass FDisplayClusterDeviceBase::DecodeStereoscopicPass(const enum E
 			break;
 
 		default:
-			DecodedPass = ((int(StereoPassType) - int(EStereoscopicPass::eSSP_RIGHT_EYE_SIDE)) % 2 == 0) ? EStereoscopicPass::eSSP_RIGHT_EYE : EStereoscopicPass::eSSP_LEFT_EYE;
+			DecodedPass = ((int(StereoPassType) - int(EStereoscopicPass::eSSP_RIGHT_EYE)) % 2 == 0) ? EStereoscopicPass::eSSP_RIGHT_EYE : EStereoscopicPass::eSSP_LEFT_EYE;
 			break;
 		}
 	}
@@ -219,7 +219,7 @@ FDisplayClusterDeviceBase::EDisplayClusterEyeType FDisplayClusterDeviceBase::Dec
 		EyeType = EDisplayClusterEyeType::StereoLeft;
 		break;
 
-	case EStereoscopicPass::eSSP_RIGHT_EYE_SIDE:
+	case EStereoscopicPass::eSSP_FULL:
 		EyeType = EDisplayClusterEyeType::Mono;
 		break;
 
@@ -417,7 +417,7 @@ uint32 FDisplayClusterDeviceBase::GetViewIndexForPass(EStereoscopicPass StereoPa
 		break;
 
 	default:
-		DecodedViewIndex = (int(StereoPassType) - int(EStereoscopicPass::eSSP_RIGHT_EYE_SIDE) + 1);
+		DecodedViewIndex = (int(StereoPassType) - int(EStereoscopicPass::eSSP_RIGHT_EYE) + 1);
 		break;
 	}
 
