nDisplay : Fixed crash in FDisplayClusterGameManager::GetProjectionScreenData  error handling

diff --git a/Engine/Plugins/Runtime/nDisplay/Source/DisplayCluster/Private/Game/DisplayClusterGameManager.cpp b/Engine/Plugins/Runtime/nDisplay/Source/DisplayCluster/Private/Game/DisplayClusterGameManager.cpp
--- a/Engine/Plugins/Runtime/nDisplay/Source/DisplayCluster/Private/Game/DisplayClusterGameManager.cpp
+++ b/Engine/Plugins/Runtime/nDisplay/Source/DisplayCluster/Private/Game/DisplayClusterGameManager.cpp
@@ -350,20 +350,14 @@ bool FDisplayClusterGameManager::GetProjectionScreenData(const FString& ScreenId
 
 	if (!ActiveScreenComponents.Contains(ScreenId))
 	{
-		UE_LOG(LogDisplayClusterGame, Warning, TEXT("Screen with ID <%s> wasn't found"));
-		return false;
-	}
-
-	if (!ActiveScreenComponents.Contains(ScreenId))
-	{
-		UE_LOG(LogDisplayClusterGame, Error, TEXT("Screen %s not found"), *ScreenId);
+		UE_LOG(LogDisplayClusterGame, Error, TEXT("Screen '%s' wasn't found"), *ScreenId);
 		return false;
 	}
 
 	const UDisplayClusterScreenComponent* const ProjScreen = ActiveScreenComponents[ScreenId];
 	if (!ProjScreen)
 	{
-		UE_LOG(LogDisplayClusterGame, Error, TEXT("Screen %s found but it's nullptr"), *ScreenId);
+		UE_LOG(LogDisplayClusterGame, Error, TEXT("Screen '%s' was found but it's nullptr"), *ScreenId);
 		return false;
 	}
 
