Misc nDisplay fixes.

diff --git a/Engine/Plugins/Runtime/nDisplay/Source/DisplayCluster/Private/Render/Device/DisplayClusterDeviceBase.cpp b/Engine/Plugins/Runtime/nDisplay/Source/DisplayCluster/Private/Render/Device/DisplayClusterDeviceBase.cpp
--- a/Engine/Plugins/Runtime/nDisplay/Source/DisplayCluster/Private/Render/Device/DisplayClusterDeviceBase.cpp
+++ b/Engine/Plugins/Runtime/nDisplay/Source/DisplayCluster/Private/Render/Device/DisplayClusterDeviceBase.cpp
@@ -29,27 +29,27 @@
 
 
 // Custom VSync interval control
-static int CVarVSyncInterval_Value = 1;
-static FAutoConsoleVariableRef  CVarVSyncInterval(
+static TAutoConsoleVariable<int32>  CVarVSyncInterval(
 	TEXT("nDisplay.render.VSyncInterval"),
-	CVarVSyncInterval_Value,
-	TEXT("VSync interval")
+	1,
+	TEXT("VSync interval"),
+	ECVF_RenderThreadSafe
 );
 
 // Enable/disable warp&blend
-static int CVarWarpBlendEnabled_Value = 1;
-static FAutoConsoleVariableRef CVarWarpBlendEnabled(
+static TAutoConsoleVariable<int32> CVarWarpBlendEnabled(
 	TEXT("nDisplay.render.WarpBlendEnabled"),
-	CVarWarpBlendEnabled_Value,
-	TEXT("Warp & Blend (0 = disabled)\n")
+	1,
+	TEXT("Warp & Blend status (0 = disabled)\n"),
+	ECVF_RenderThreadSafe
 );
 
 // Enable/disable nDisplay post-process
-static int CVarCustomPPEnabled_Value = 1;
-static FAutoConsoleVariableRef CVarPPEnabled(
+static TAutoConsoleVariable<int32> CVarCustomPPEnabled(
 	TEXT("nDisplay.render.postprocess"),
-	CVarCustomPPEnabled_Value,
-	TEXT("Custom post-process (0 = disabled)\n")
+	1,
+	TEXT("Custom post-process (0 = disabled)\n"),
+	ECVF_RenderThreadSafe
 );
 
 
@@ -351,8 +351,12 @@ void FDisplayClusterDeviceBase::RenderTexture_RenderThread(FRHICommandListImmedi
 	// This one is to match interface function signatures only. The functions will provide the handlers with a proper data.
 	const FIntRect StubRect(0, 0, 0, 0);
 
+	// Get custom PP and warp&blend flags status
+	const bool bCustomPPEnabled = (CVarCustomPPEnabled.GetValueOnRenderThread() != 0);
+	const bool bWarpBlendEnabled = (CVarWarpBlendEnabled.GetValueOnRenderThread() != 0);
+
 	// Post-process before warp&blend
-	if (CVarCustomPPEnabled_Value != 0)
+	if (bCustomPPEnabled)
 	{
 		// PP round 1: post-process for each view region before warp&blend
 		PerformPostProcessViewBeforeWarpBlend_RenderThread(RHICmdList, SrcTexture, StubRect);
@@ -363,7 +367,7 @@ void FDisplayClusterDeviceBase::RenderTexture_RenderThread(FRHICommandListImmedi
 	}
 
 	// Perform warp&blend
-	if (CVarWarpBlendEnabled_Value != 0)
+	if (bWarpBlendEnabled)
 	{
 		// Iterate over viewports
 		for (int i = 0; i < RenderViewports.Num(); ++i)
@@ -373,14 +377,14 @@ void FDisplayClusterDeviceBase::RenderTexture_RenderThread(FRHICommandListImmedi
 			{
 				for (uint32 j = 0; j < ViewsAmountPerViewport; ++j)
 				{
-					RenderViewports[i].GetProjectionPolicy()->ApplyWarpBlend_RenderThread(j, RHICmdList, SrcTexture, RenderViewports[i].GetArea());
+					RenderViewports[i].GetProjectionPolicy()->ApplyWarpBlend_RenderThread(j, RHICmdList, SrcTexture, RenderViewports[i].GetContext(j).RenderTargetRect);
 				}
 			}
 		}
 	}
 
 	// Post-process after warp&blend
-	if (CVarCustomPPEnabled_Value != 0)
+	if (bCustomPPEnabled)
 	{
 		// PP round 4: post-process for each view region after warp&blend
 		PerformPostProcessViewAfterWarpBlend_RenderThread(RHICmdList, SrcTexture, StubRect);
@@ -621,7 +625,8 @@ uint32 FDisplayClusterDeviceBase::DecodeViewIndex(const enum EStereoscopicPass S
 
 uint32 FDisplayClusterDeviceBase::GetSwapInt() const
 {
-	return (CVarVSyncInterval_Value < 0 ? 1 : (uint32)CVarVSyncInterval_Value);
+	const uint32 SyncInterval = static_cast<uint32>(CVarVSyncInterval.GetValueOnAnyThread());
+	return (SyncInterval);
 }
 
 void FDisplayClusterDeviceBase::AddViewport(const FString& InViewportId, const FIntPoint& InViewportLocation, const FIntPoint& InViewportSize, TSharedPtr<IDisplayClusterProjectionPolicy> InProjPolicy, const FString& InCameraId)
@@ -665,7 +670,7 @@ void FDisplayClusterDeviceBase::CopyTextureToBackBuffer_RenderThread(FRHICommand
 
 	const FIntPoint BBSize = BackBuffer->GetSizeXY();
 
-	// Copy final texture to the back buffer
+	// Copy final texture to the back buffer. This implementation is for simple cases. More complex devices override this method and provide custom copying.
 	FResolveParams copyParams;
 	copyParams.DestArrayIndex = 0;
 	copyParams.SourceArrayIndex = 0;
diff --git a/Engine/Plugins/Runtime/nDisplay/Source/DisplayCluster/Private/Render/Device/DisplayClusterDeviceBase_PostProcess.cpp b/Engine/Plugins/Runtime/nDisplay/Source/DisplayCluster/Private/Render/Device/DisplayClusterDeviceBase_PostProcess.cpp
--- a/Engine/Plugins/Runtime/nDisplay/Source/DisplayCluster/Private/Render/Device/DisplayClusterDeviceBase_PostProcess.cpp
+++ b/Engine/Plugins/Runtime/nDisplay/Source/DisplayCluster/Private/Render/Device/DisplayClusterDeviceBase_PostProcess.cpp
@@ -6,26 +6,29 @@
 #include "Misc/DisplayClusterHelpers.h"
 #include "DisplayClusterLog.h"
 
+#include "HAL/IConsoleManager.h"
+
 
 //////////////////////////////////////////////////////////////////////////////////////////////
 // Round 1: VIEW before warp&blend
 //////////////////////////////////////////////////////////////////////////////////////////////
 
 // Enable/disable PP round 1
-static int CVarPostprocessViewBeforeWarpBlend_Value = 1;
-static FAutoConsoleVariableRef CVarPostprocessViewBeforeWarpBlend(
+static TAutoConsoleVariable<int32> CVarPostprocessViewBeforeWarpBlend(
 	TEXT("nDisplay.render.postprocess.ViewBeforeWarpBlend"),
-	CVarPostprocessViewBeforeWarpBlend_Value,
-	TEXT("Enable PP per view before warp&blend (0 = disabled)\n")
+	1,
+	TEXT("Enable PP per view before warp&blend (0 = disabled)\n"),
+	ECVF_RenderThreadSafe
 );
 
 void FDisplayClusterDeviceBase_PostProcess::PerformPostProcessViewBeforeWarpBlend_RenderThread(FRHICommandListImmediate& RHICmdList, FRHITexture2D* SrcTexture, const FIntRect& ViewRect) const
 {
 	DISPLAY_CLUSTER_FUNC_TRACE(LogDisplayClusterRender);
 
-	UE_LOG(LogDisplayClusterRender, Verbose, TEXT("Postprocess VIEW before WarpBlend: %d"), CVarPostprocessViewBeforeWarpBlend_Value != 0 ? 1 : 0);
+	const bool bEnabled = (CVarPostprocessViewBeforeWarpBlend.GetValueOnRenderThread() != 0);
+	UE_LOG(LogDisplayClusterRender, Verbose, TEXT("Postprocess VIEW before WarpBlend: %d"), bEnabled ? 1 : 0);
 
-	if (CVarPostprocessViewBeforeWarpBlend_Value != 0)
+	if (bEnabled)
 	{
 		for (const IDisplayClusterRenderManager::FDisplayClusterPPInfo& CurPP : PPOperations)
 		{
@@ -50,20 +53,21 @@ void FDisplayClusterDeviceBase_PostProcess::PerformPostProcessViewBeforeWarpBlen
 //////////////////////////////////////////////////////////////////////////////////////////////
 
 // Enable/disable PP round 2
-static int CVarPostprocessFrameBeforeWarpBlend_Value = 1;
-static FAutoConsoleVariableRef CVarPostprocessFrameBeforeWarpBlend(
+static TAutoConsoleVariable<int32> CVarPostprocessFrameBeforeWarpBlend(
 	TEXT("nDisplay.render.postprocess.FrameBeforeWarpBlend"),
-	CVarPostprocessFrameBeforeWarpBlend_Value,
-	TEXT("Enable PP per eye frame before warp&blend (0 = disabled)\n")
+	1,
+	TEXT("Enable PP per eye frame before warp&blend (0 = disabled)\n"),
+	ECVF_RenderThreadSafe
 );
 
 void FDisplayClusterDeviceBase_PostProcess::PerformPostProcessFrameBeforeWarpBlend_RenderThread(FRHICommandListImmediate& RHICmdList, FRHITexture2D* SrcTexture, const FIntRect& ViewRect) const
 {
 	DISPLAY_CLUSTER_FUNC_TRACE(LogDisplayClusterRender);
 
-	UE_LOG(LogDisplayClusterRender, Verbose, TEXT("Postprocess FRAME before WarpBlend: %d"), CVarPostprocessFrameBeforeWarpBlend_Value != 0 ? 1 : 0);
+	const bool bEnabled = (CVarPostprocessFrameBeforeWarpBlend.GetValueOnRenderThread() != 0);
+	UE_LOG(LogDisplayClusterRender, Verbose, TEXT("Postprocess FRAME before WarpBlend: %d"), bEnabled ? 1 : 0);
 
-	if (CVarPostprocessFrameBeforeWarpBlend_Value != 0)
+	if (bEnabled != 0)
 	{
 		for (const IDisplayClusterRenderManager::FDisplayClusterPPInfo& CurPP : PPOperations)
 		{
@@ -88,20 +92,21 @@ void FDisplayClusterDeviceBase_PostProcess::PerformPostProcessFrameBeforeWarpBle
 //////////////////////////////////////////////////////////////////////////////////////////////
 
 // Enable/disable PP round 3
-static int CVarPostprocessTargetBeforeWarpBlend_Value = 1;
-static FAutoConsoleVariableRef CVarPostprocessTargetBeforeWarpBlend(
+static TAutoConsoleVariable<int32> CVarPostprocessTargetBeforeWarpBlend(
 	TEXT("nDisplay.render.postprocess.TargetBeforeWarpBlend"),
-	CVarPostprocessTargetBeforeWarpBlend_Value,
-	TEXT("Enable PP for the whole render target before warp&blend (0 = disabled)\n")
+	1,
+	TEXT("Enable PP for the whole render target before warp&blend (0 = disabled)\n"),
+	ECVF_RenderThreadSafe
 );
 
 void FDisplayClusterDeviceBase_PostProcess::PerformPostProcessRenderTargetBeforeWarpBlend_RenderThread(FRHICommandListImmediate& RHICmdList, FRHITexture2D* SrcTexture) const
 {
 	DISPLAY_CLUSTER_FUNC_TRACE(LogDisplayClusterRender);
 
-	UE_LOG(LogDisplayClusterRender, Verbose, TEXT("Postprocess TARGET before WarpBlend: %d"), CVarPostprocessTargetBeforeWarpBlend_Value != 0 ? 1 : 0);
+	const bool bEnabled = (CVarPostprocessTargetBeforeWarpBlend.GetValueOnRenderThread() != 0);
+	UE_LOG(LogDisplayClusterRender, Verbose, TEXT("Postprocess TARGET before WarpBlend: %d"), bEnabled ? 1 : 0);
 
-	if (CVarPostprocessTargetBeforeWarpBlend_Value != 0)
+	if (bEnabled != 0)
 	{
 		for (const IDisplayClusterRenderManager::FDisplayClusterPPInfo& CurPP : PPOperations)
 		{
@@ -126,20 +131,21 @@ void FDisplayClusterDeviceBase_PostProcess::PerformPostProcessRenderTargetBefore
 //////////////////////////////////////////////////////////////////////////////////////////////
 
 // Enable/disable PP round 4
-static int CVarPostprocessViewAfterWarpBlend_Value = 1;
-static FAutoConsoleVariableRef CVarPostprocessViewAfterWarpBlend(
+static TAutoConsoleVariable<int32> CVarPostprocessViewAfterWarpBlend(
 	TEXT("nDisplay.render.postprocess.ViewAfterWarpBlend"),
-	CVarPostprocessViewAfterWarpBlend_Value,
-	TEXT("Enable PP per view after warp&blend (0 = disabled)\n")
+	1,
+	TEXT("Enable PP per view after warp&blend (0 = disabled)\n"),
+	ECVF_RenderThreadSafe
 );
 
 void FDisplayClusterDeviceBase_PostProcess::PerformPostProcessViewAfterWarpBlend_RenderThread(FRHICommandListImmediate& RHICmdList, FRHITexture2D* SrcTexture, const FIntRect& FrameRect) const
 {
 	DISPLAY_CLUSTER_FUNC_TRACE(LogDisplayClusterRender);
 
-	UE_LOG(LogDisplayClusterRender, Verbose, TEXT("Postprocess VIEW after WarpBlend: %d"), CVarPostprocessViewAfterWarpBlend_Value != 0 ? 1 : 0);
+	const bool bEnabled = (CVarPostprocessViewAfterWarpBlend.GetValueOnRenderThread() != 0);
+	UE_LOG(LogDisplayClusterRender, Verbose, TEXT("Postprocess VIEW after WarpBlend: %d"), bEnabled ? 1 : 0);
 
-	if (CVarPostprocessViewAfterWarpBlend_Value!= 0)
+	if (bEnabled != 0)
 	{
 		for (const IDisplayClusterRenderManager::FDisplayClusterPPInfo& CurPP : PPOperations)
 		{
@@ -164,20 +170,21 @@ void FDisplayClusterDeviceBase_PostProcess::PerformPostProcessViewAfterWarpBlend
 //////////////////////////////////////////////////////////////////////////////////////////////
 
 // Enable/disable PP round 5
-static int CVarPostprocessFrameAfterWarpBlend_Value = 1;
-static FAutoConsoleVariableRef CVarPostprocessFrameAfterWarpBlend(
+static TAutoConsoleVariable<int32> CVarPostprocessFrameAfterWarpBlend(
 	TEXT("nDisplay.render.postprocess.FrameAfterWarpBlend"),
-	CVarPostprocessFrameAfterWarpBlend_Value,
-	TEXT("Enable PP per eye frame after warp&blend (0 = disabled)\n")
+	1,
+	TEXT("Enable PP per eye frame after warp&blend (0 = disabled)\n"),
+	ECVF_RenderThreadSafe
 );
 
 void FDisplayClusterDeviceBase_PostProcess::PerformPostProcessFrameAfterWarpBlend_RenderThread(FRHICommandListImmediate& RHICmdList, FRHITexture2D* SrcTexture, const FIntRect& FrameRect) const
 {
 	DISPLAY_CLUSTER_FUNC_TRACE(LogDisplayClusterRender);
 
-	UE_LOG(LogDisplayClusterRender, Verbose, TEXT("Postprocess VIEW after WarpBlend: %d"), CVarPostprocessFrameAfterWarpBlend_Value != 0 ? 1 : 0);
+	const bool bEnabled = (CVarPostprocessFrameAfterWarpBlend.GetValueOnRenderThread() != 0);
+	UE_LOG(LogDisplayClusterRender, Verbose, TEXT("Postprocess VIEW after WarpBlend: %d"), bEnabled ? 1 : 0);
 
-	if (CVarPostprocessFrameAfterWarpBlend_Value != 0)
+	if (bEnabled != 0)
 	{
 		for (const IDisplayClusterRenderManager::FDisplayClusterPPInfo& CurPP : PPOperations)
 		{
@@ -202,20 +209,21 @@ void FDisplayClusterDeviceBase_PostProcess::PerformPostProcessFrameAfterWarpBlen
 //////////////////////////////////////////////////////////////////////////////////////////////
 
 // Enable/disable PP round 6
-static int CVarPostprocessTargetAfterWarpBlend_Value = 1;
-static FAutoConsoleVariableRef CVarPostprocessTargetAfterWarpBlend(
+static TAutoConsoleVariable<int32> CVarPostprocessTargetAfterWarpBlend(
 	TEXT("nDisplay.render.postprocess.TargetAfterWarpBlend"),
-	CVarPostprocessTargetAfterWarpBlend_Value,
-	TEXT("Enable PP for the whole render target after warp&blend (0 = disabled)\n")
+	1,
+	TEXT("Enable PP for the whole render target after warp&blend (0 = disabled)\n"),
+	ECVF_RenderThreadSafe
 );
 
 void FDisplayClusterDeviceBase_PostProcess::PerformPostProcessRenderTargetAfterWarpBlend_RenderThread(FRHICommandListImmediate& RHICmdList, FRHITexture2D* SrcTexture) const
 {
 	DISPLAY_CLUSTER_FUNC_TRACE(LogDisplayClusterRender);
 
-	UE_LOG(LogDisplayClusterRender, Verbose, TEXT("Postprocess TARGET after WarpBlend: %d"), CVarPostprocessTargetAfterWarpBlend_Value != 0 ? 1 : 0);
+	const bool bEnabled = (CVarPostprocessTargetAfterWarpBlend.GetValueOnRenderThread() != 0);
+	UE_LOG(LogDisplayClusterRender, Verbose, TEXT("Postprocess TARGET after WarpBlend: %d"), bEnabled ? 1 : 0);
 
-	if (CVarPostprocessTargetAfterWarpBlend_Value != 0)
+	if (bEnabled != 0)
 	{
 		for (const IDisplayClusterRenderManager::FDisplayClusterPPInfo& CurPP : PPOperations)
 		{
diff --git a/Engine/Plugins/Runtime/nDisplay/Source/DisplayCluster/Private/Render/DisplayClusterRenderManager.cpp b/Engine/Plugins/Runtime/nDisplay/Source/DisplayCluster/Private/Render/DisplayClusterRenderManager.cpp
--- a/Engine/Plugins/Runtime/nDisplay/Source/DisplayCluster/Private/Render/DisplayClusterRenderManager.cpp
+++ b/Engine/Plugins/Runtime/nDisplay/Source/DisplayCluster/Private/Render/DisplayClusterRenderManager.cpp
@@ -129,6 +129,14 @@ bool FDisplayClusterRenderManager::StartScene(UWorld* InWorld)
 
 void FDisplayClusterRenderManager::EndScene()
 {
+#if WITH_EDITOR
+	if (GIsEditor)
+	{
+		// Since we can run multiple PIE sessions we have to clean device before the next one.
+		GEngine->StereoRenderingDevice.Reset();
+		RenderDevice.Reset();
+	}
+#endif
 }
 
 void FDisplayClusterRenderManager::PreTick(float DeltaSeconds)
diff --git a/Engine/Plugins/Runtime/nDisplay/Source/DisplayCluster/Private/Render/Synchronization/DisplayClusterRenderSyncPolicyNvidiaDX11.cpp b/Engine/Plugins/Runtime/nDisplay/Source/DisplayCluster/Private/Render/Synchronization/DisplayClusterRenderSyncPolicyNvidiaDX11.cpp
--- a/Engine/Plugins/Runtime/nDisplay/Source/DisplayCluster/Private/Render/Synchronization/DisplayClusterRenderSyncPolicyNvidiaDX11.cpp
+++ b/Engine/Plugins/Runtime/nDisplay/Source/DisplayCluster/Private/Render/Synchronization/DisplayClusterRenderSyncPolicyNvidiaDX11.cpp
@@ -14,7 +14,9 @@ FDisplayClusterRenderSyncPolicyNvidiaDX11::~FDisplayClusterRenderSyncPolicyNvidi
 
 bool FDisplayClusterRenderSyncPolicyNvidiaDX11::SynchronizeClusterRendering(int32& InOutSyncInterval)
 {
-	//@todo implement hardware based solution
+	// Unfortunately NVIDIA doesn't export required functions in the public SDK. This make impossible
+	// to implement the hardware base synchronization approach. As a fallback we use our own soft sync.
+	// As soon as the required functions become available, the proper synchonization approach will be implemented.
 	SyncBarrierRenderThread();
 	// Tell a caller that he still needs to present a frame
 	return true;
diff --git a/Engine/Plugins/Runtime/nDisplay/Source/DisplayCluster/Private/Render/Synchronization/DisplayClusterRenderSyncPolicyNvidiaDX12.cpp b/Engine/Plugins/Runtime/nDisplay/Source/DisplayCluster/Private/Render/Synchronization/DisplayClusterRenderSyncPolicyNvidiaDX12.cpp
--- a/Engine/Plugins/Runtime/nDisplay/Source/DisplayCluster/Private/Render/Synchronization/DisplayClusterRenderSyncPolicyNvidiaDX12.cpp
+++ b/Engine/Plugins/Runtime/nDisplay/Source/DisplayCluster/Private/Render/Synchronization/DisplayClusterRenderSyncPolicyNvidiaDX12.cpp
@@ -14,7 +14,9 @@ FDisplayClusterRenderSyncPolicyNvidiaDX12::~FDisplayClusterRenderSyncPolicyNvidi
 
 bool FDisplayClusterRenderSyncPolicyNvidiaDX12::SynchronizeClusterRendering(int32& InOutSyncInterval)
 {
-	//@todo implement hardware based solution
+	// Unfortunately NVIDIA doesn't export required functions in the public SDK. This make impossible
+	// to implement the hardware base synchronization approach. As a fallback we use our own soft sync.
+	// As soon as the required functions become available, the proper synchonization approach will be implemented.
 	SyncBarrierRenderThread();
 	// Tell a caller that he still needs to present a frame
 	return true;
diff --git a/Engine/Plugins/Runtime/nDisplay/Source/DisplayCluster/Private/Render/Synchronization/DisplayClusterRenderSyncPolicySoftwareDX11.cpp b/Engine/Plugins/Runtime/nDisplay/Source/DisplayCluster/Private/Render/Synchronization/DisplayClusterRenderSyncPolicySoftwareDX11.cpp
--- a/Engine/Plugins/Runtime/nDisplay/Source/DisplayCluster/Private/Render/Synchronization/DisplayClusterRenderSyncPolicySoftwareDX11.cpp
+++ b/Engine/Plugins/Runtime/nDisplay/Source/DisplayCluster/Private/Render/Synchronization/DisplayClusterRenderSyncPolicySoftwareDX11.cpp
@@ -29,53 +29,40 @@
 #include "NVIDIA/nvapi/nvapi.h"
 
 
-static int CVarAdvancedSyncEnabled_Value = 0;
-static FAutoConsoleVariableRef CVarAdvancedSyncEnabled(
+static TAutoConsoleVariable<int32> CVarAdvancedSyncEnabled(
 	TEXT("nDisplay.render.softsync.AdvancedSyncEnabled"),
-	CVarAdvancedSyncEnabled_Value,
+	0,
 	TEXT("Advanced DWM based render synchronization (0 = disabled)")
 );
 
-static float CVarVBlankThreshold_Value = 0.003f;
-static FAutoConsoleVariableRef CVarVBlankThreshold(
+static TAutoConsoleVariable<float> CVarVBlankThreshold(
 	TEXT("nDisplay.render.softsync.VBlankThreshold"),
-	CVarVBlankThreshold_Value,
+	0.003f,
 	TEXT("Changes the threshold used to determine whether or not nDisplay SoftSync can make the current frame presentation")
 );
 
-static float CVarVBlankThresholdSleepMultiplier_Value = 1.5f;
-static FAutoConsoleVariableRef CVarVBlankSleepThresholdSleepMultiplier(
+static TAutoConsoleVariable<float> CVarVBlankThresholdSleepMultiplier(
 	TEXT("nDisplay.render.softsync.VBlankThresholdSleepMultipler"),
-	CVarVBlankThresholdSleepMultiplier_Value,
+	1.5f,
 	TEXT("Multiplier applied to VBlankThreshold to compute sleep time to skip frame presentation on upcoming VBlank")
 );
 
-static int CVarVBlankBasisUpdate_Value = 1;
-static FAutoConsoleVariableRef  CVarVBlankBasisUpdate(
+static TAutoConsoleVariable<int32>  CVarVBlankBasisUpdate(
 	TEXT("nDisplay.render.softsync.VBlankBasisUpdate"),
-	CVarVBlankBasisUpdate_Value,
-	TEXT("Update VBlank basis periodically to avoid time drifting")
+	1,
+	TEXT("Update VBlank basis periodically to avoid time drifting (0 = disabled)")
 );
 
-static float CVarVBlankBasisUpdatePeriod_Value = 120.f;
-static FAutoConsoleVariableRef  CVarVBlankBasisUpdatePeriod(
+static TAutoConsoleVariable<float>  CVarVBlankBasisUpdatePeriod(
 	TEXT("nDisplay.render.softsync.VBlankBasisUpdatePeriod"),
-	CVarVBlankBasisUpdatePeriod_Value,
+	120.f,
 	TEXT("VBlank basis update period in seconds")
 );
 
-static int CVarLogDwmStats_Value = 0;
-static FAutoConsoleVariableRef  CVarLogDwmStats(
+static TAutoConsoleVariable<int32>  CVarLogDwmStats(
 	TEXT("nDisplay.render.softsync.LogDwmStats"),
-	CVarLogDwmStats_Value,
-	TEXT("Print DWM stats to log")
-);
-
-static int CVarObtainGpuData_Value = 0;
-static FAutoConsoleVariableRef  CVarObtainGpuData(
-	TEXT("nDisplay.render.softsync.GpuDataObtain"),
-	CVarObtainGpuData_Value,
-	TEXT("Obtain GPU data each frame")
+	0,
+	TEXT("Print DWM stats to log (0 = disabled)")
 );
 
 
@@ -97,7 +84,7 @@ FDisplayClusterRenderSyncPolicySoftwareDX11::FDisplayClusterRenderSyncPolicySoft
 {
 	// We check the console variable once at start. The advanced synchronization policy is not 
 	// supposed to be activated/deactivated at runtime. Use .ini settings to set desired value.
-	bUseAdvancedSynchronization = (CVarAdvancedSyncEnabled_Value != 0);
+	bUseAdvancedSynchronization = (CVarAdvancedSyncEnabled.GetValueOnGameThread() != 0);
 
 	if (bUseAdvancedSynchronization)
 	{
@@ -205,6 +192,11 @@ bool FDisplayClusterRenderSyncPolicySoftwareDX11::SynchronizeClusterRendering(in
 		return true;
 	}
 
+	// Get dynamic values from console variables
+	const float VBlankThreshold = CVarVBlankThreshold.GetValueOnGameThread();
+	const float VBlankThresholdSleepMultiplier = CVarVBlankThresholdSleepMultiplier.GetValueOnGameThread();
+
+
 	if (!bTimersInitialized)
 	{
 		QUICK_SCOPE_CYCLE_COUNTER(STAT_nDisplayPresent_SoftSync_Initialization);
@@ -214,9 +206,9 @@ bool FDisplayClusterRenderSyncPolicySoftwareDX11::SynchronizeClusterRendering(in
 
 		UE_LOG(LogDisplayClusterRender, Log, TEXT("##SYNC_LOG: Refresh period:     %lf"), RefreshPeriod);
 		UE_LOG(LogDisplayClusterRender, Log, TEXT("##SYNC_LOG: VBlank basis:       %lf"), VBlankBasis);
-		UE_LOG(LogDisplayClusterRender, Log, TEXT("##SYNC_LOG: VBlank threshold:   %lf"), CVarVBlankThreshold_Value);
-		UE_LOG(LogDisplayClusterRender, Log, TEXT("##SYNC_LOG: VBlank sleep mult:  %lf"), CVarVBlankThresholdSleepMultiplier_Value);
-		UE_LOG(LogDisplayClusterRender, Log, TEXT("##SYNC_LOG: VBlank sleep:       %lf"), CVarVBlankThreshold_Value * CVarVBlankThresholdSleepMultiplier_Value);
+		UE_LOG(LogDisplayClusterRender, Log, TEXT("##SYNC_LOG: VBlank threshold:   %lf"), VBlankThreshold);
+		UE_LOG(LogDisplayClusterRender, Log, TEXT("##SYNC_LOG: VBlank sleep mult:  %lf"), VBlankThresholdSleepMultiplier);
+		UE_LOG(LogDisplayClusterRender, Log, TEXT("##SYNC_LOG: VBlank sleep:       %lf"), VBlankThreshold * VBlankThresholdSleepMultiplier);
 
 #ifdef UE_BUILD_DEBUG
 		const int SamplesCount = 25;
@@ -251,9 +243,9 @@ bool FDisplayClusterRenderSyncPolicySoftwareDX11::SynchronizeClusterRendering(in
 
 		// Skip upcoming VBlank if we're in red zone
 		SB = FPlatformTime::Seconds();
-		if (LeftToVBlankTime < CVarVBlankThreshold_Value)
+		if (LeftToVBlankTime < VBlankThreshold)
 		{
-			const double SleepTime = CVarVBlankThreshold_Value * CVarVBlankThresholdSleepMultiplier_Value;
+			const double SleepTime = VBlankThreshold * VBlankThresholdSleepMultiplier;
 			UE_LOG(LogDisplayClusterRender, Warning, TEXT("FDisplayClusterNativePresentHandler::Present_SoftSwapSync - Skipped VBlank. Sleeping %f"), SleepTime);
 			FPlatformProcess::Sleep(SleepTime);
 		}
@@ -272,11 +264,15 @@ bool FDisplayClusterRenderSyncPolicySoftwareDX11::SynchronizeClusterRendering(in
 	{
 		QUICK_SCOPE_CYCLE_COUNTER(STAT_nDisplayPresent_SoftSync_SwapChainPresent);
 
+		// Get dynamic values from console variables
+		const int32 VBlankBasisUpdate = CVarVBlankBasisUpdate.GetValueOnGameThread();
+		const float VBlankBasisUpdatePeriod = CVarVBlankBasisUpdatePeriod.GetValueOnGameThread();
+
 		// Regardless of where we are, it's safe to present a frame now. If we need to update the VBlank basis,
 		// we have to wait for a VBlank and store the time. We don't want to miss a frame presentation so we
 		// present it with swap interval 0 right after VBlank sygnal.
 		int SyncIntervalToUse = 1;
-		if ((CVarVBlankBasisUpdate_Value > 0) && (B2A - VBlankBasis) > CVarVBlankBasisUpdatePeriod_Value)
+		if ((VBlankBasisUpdate > 0) && (B2A - VBlankBasis) > VBlankBasisUpdatePeriod)
 		{
 			VBlankBasis = GetVBlankTimestamp(DXOutput);
 			SyncIntervalToUse = 0;
@@ -289,7 +285,8 @@ bool FDisplayClusterRenderSyncPolicySoftwareDX11::SynchronizeClusterRendering(in
 
 	UE_LOG(LogDisplayClusterRender, Verbose, TEXT("##SYNC_LOG - %d:%lf:%lf:%lf:%lf:%lf:%lf:%lf:%lf:%lf"), FrameCounter, B1B, B1A, TToB, SB, SA, B2B, B2A, PB, PA);
 
-	if (CVarLogDwmStats_Value > 0)
+	const bool LogDwmStats = (CVarLogDwmStats.GetValueOnGameThread() != 0);
+	if (LogDwmStats)
 	{
 		PrintDwmStats(FrameCounter);
 	}
diff --git a/Engine/Plugins/Runtime/nDisplay/Source/MPCDI/Private/MPCDIShader.cpp b/Engine/Plugins/Runtime/nDisplay/Source/MPCDI/Private/MPCDIShader.cpp
--- a/Engine/Plugins/Runtime/nDisplay/Source/MPCDI/Private/MPCDIShader.cpp
+++ b/Engine/Plugins/Runtime/nDisplay/Source/MPCDI/Private/MPCDIShader.cpp
@@ -11,19 +11,23 @@
 
 #include "ShaderParameterUtils.h"
 
+#include "HAL/IConsoleManager.h"
+
+
 // Select mpcdi stereo mode
- enum class EVarMPCDIShaderType : uint8
+enum class EVarMPCDIShaderType : uint8
 {
 	Default,
 	Passthrought,
 	ShowWarpTexture,
 	Disable,
 };
-static int CVarMPCDIShaderType_Value = (int)EVarMPCDIShaderType::Default;
-static FAutoConsoleVariableRef CVarMPCDIShaderType(
+
+static TAutoConsoleVariable<int32> CVarMPCDIShaderType(
 	TEXT("nDisplay.render.mpcdi.shader"),
-	CVarMPCDIShaderType_Value,
-	TEXT("Select shader for mpcdi:\n (0 = default warp shader)\n(1 = PassThrought shader)\n(2 = ShowWarpTexture shader)\n(3 = disable warp shader)\n")
+	(int)EVarMPCDIShaderType::Default,
+	TEXT("Select shader for mpcdi:\n0 = default warp shader\n1 = PassThrought shader\n2 = ShowWarpTexture shader\n3 = disable warp shader"),
+	ECVF_RenderThreadSafe
 );
 
 
@@ -208,17 +212,24 @@ bool FMPCDIShader::ApplyWarpBlend(FRHICommandListImmediate& RHICmdList, IMPCDI::
 		//Handle error
 		return false;
 	}
-		
-	switch (CVarMPCDIShaderType_Value) {
-	case (int)EVarMPCDIShaderType::Passthrought:
+
+
+	const EVarMPCDIShaderType ShaderType = (EVarMPCDIShaderType)CVarMPCDIShaderType.GetValueOnAnyThread();
+	switch (ShaderType)
+	{
+	case EVarMPCDIShaderType::Passthrought:
 		PixelShaderType = EMPCDIShader::PassThrought;
 		break;
-	case (int)EVarMPCDIShaderType::ShowWarpTexture:
+
+	case EVarMPCDIShaderType::ShowWarpTexture:
 		PixelShaderType = EMPCDIShader::ShowWarpTexture;
 		break;
-	case (int)EVarMPCDIShaderType::Disable:
+
+	case EVarMPCDIShaderType::Disable:
+		return false;
+
+	default:
 		return false;
-		break;
 	};
 	
 	{
diff --git a/Engine/Plugins/Runtime/nDisplay/Source/MPCDI/Private/MPCDIWarpTexture.cpp b/Engine/Plugins/Runtime/nDisplay/Source/MPCDI/Private/MPCDIWarpTexture.cpp
--- a/Engine/Plugins/Runtime/nDisplay/Source/MPCDI/Private/MPCDIWarpTexture.cpp
+++ b/Engine/Plugins/Runtime/nDisplay/Source/MPCDI/Private/MPCDIWarpTexture.cpp
@@ -29,11 +29,11 @@ enum EVarMPCDIFrustumMethod
 };
 
 
-static int CVarMPCDIFrustumMethod_Value = (int)EVarMPCDIFrustumMethod::PerfectCPU;
-static FAutoConsoleVariableRef CVarMPCDIFrustumMethod(
+static TAutoConsoleVariable<int32> CVarMPCDIFrustumMethod(
 	TEXT("nDisplay.render.mpcdi.Frustum"),
-	CVarMPCDIFrustumMethod_Value,
-	TEXT("Frustum calculation method: \n (0 = AABB based, fast, on CPU)\n(1 = Perfect, slow, on CPU)\n")
+	(int)EVarMPCDIFrustumMethod::PerfectCPU,
+	TEXT("Frustum computation method:\n0 = mesh AABB based, lower quality but fast\n1 = mesh vertices based, best quality but slow\n"),
+	ECVF_RenderThreadSafe
 );
 
 // Select mpcdi stereo mode
@@ -43,11 +43,11 @@ enum EVarMPCDIStereoMode
 	SymmetricAABB,
 };
 
-static int CVarMPCDIStereoMode_Value = (int)EVarMPCDIStereoMode::AsymmetricAABB;
-static FAutoConsoleVariableRef CVarMPCDIStereoMode(
+static TAutoConsoleVariable<int32> CVarMPCDIStereoMode(
 	TEXT("nDisplay.render.mpcdi.StereoMode"),
-	CVarMPCDIStereoMode_Value,
-	TEXT("Stereo mode:\n (0 = Asymmetric to AABB center\n(1 = Symmetric to AABB center)\n")
+	(int)EVarMPCDIStereoMode::AsymmetricAABB,
+	TEXT("Stereo mode:\n (0 = Asymmetric to AABB center\n(1 = Symmetric to AABB center)\n"),
+	ECVF_RenderThreadSafe
 );
 
 // Select mpcdi projection mode
@@ -57,11 +57,11 @@ enum EVarMPCDIProjectionMode
 	StaticAxisAligned,
 };
 
-static int CVarMPCDIProjectionMode_Value = (int)EVarMPCDIProjectionMode::Dynamic;
-static FAutoConsoleVariableRef CVarMPCDIProjectionMode(
+static TAutoConsoleVariable<int32> CVarMPCDIProjectionMode(
 	TEXT("nDisplay.render.mpcdi.Projection"),
-	CVarMPCDIProjectionMode_Value,
-	TEXT("Projection method:\n (0 = Dynamic, to view target\n(1 = Static, aligned to mpcdi origin space axis)\n")
+	(int)EVarMPCDIProjectionMode::Dynamic,
+	TEXT("Projection method:\n0 = Dynamic, to view target\n1 = Static, aligned to mpcdi origin space axis\n"),
+	ECVF_RenderThreadSafe
 );
 
 
@@ -176,7 +176,8 @@ namespace MPCDI
 
 	void FMPCDIWarpTexture::CalcViewProjection(const IMPCDI::FFrustum& Frustum, const FVector& ViewDirection, const FVector& ViewOrigin, const FVector& EyeOrigin, FMatrix& OutViewMatrix) const
 	{
-		switch (CVarMPCDIProjectionMode_Value)
+		const EVarMPCDIProjectionMode ProjMode = (EVarMPCDIProjectionMode)CVarMPCDIProjectionMode.GetValueOnAnyThread();
+		switch (ProjMode)
 		{
 			case EVarMPCDIProjectionMode::Dynamic:
 			{
@@ -228,7 +229,8 @@ namespace MPCDI
 
 		FMatrix Local2world = FMatrix::Identity;
 
-		switch (CVarMPCDIStereoMode_Value)
+		const EVarMPCDIStereoMode StereoMode = (EVarMPCDIStereoMode)CVarMPCDIStereoMode.GetValueOnAnyThread();
+		switch (StereoMode)
 		{
 			case EVarMPCDIStereoMode::AsymmetricAABB:
 			{
@@ -270,7 +272,8 @@ namespace MPCDI
 		float left = FLT_MAX;
 		float right = -FLT_MAX;
 
-		switch (CVarMPCDIFrustumMethod_Value)
+		const EVarMPCDIFrustumMethod FrustumComputeMethod = (EVarMPCDIFrustumMethod)CVarMPCDIFrustumMethod.GetValueOnAnyThread();
+		switch (FrustumComputeMethod)
 		{
 		case EVarMPCDIFrustumMethod::AABB:
 			CalcFrustum_simpleAABB(OutFrustum, World2local, top, bottom, left, right);
diff --git a/Engine/Plugins/Runtime/nDisplay/Source/PicpMPCDI/Private/PicpBlurPostProcess.cpp b/Engine/Plugins/Runtime/nDisplay/Source/PicpMPCDI/Private/PicpBlurPostProcess.cpp
--- a/Engine/Plugins/Runtime/nDisplay/Source/PicpMPCDI/Private/PicpBlurPostProcess.cpp
+++ b/Engine/Plugins/Runtime/nDisplay/Source/PicpMPCDI/Private/PicpBlurPostProcess.cpp
@@ -13,11 +13,12 @@
 
 #include "Engine/TextureRenderTarget2D.h"
 
-static int CVarPicpPostProcess_Value = 1;
-static FAutoConsoleVariableRef CVarPicpPostProcess(
+
+static TAutoConsoleVariable<int32> CVarPicpPostProcess(
 	TEXT("nDisplay.render.picp.EnablePostProcess"),
-	CVarPicpPostProcess_Value,
-	TEXT("Enable postprocess shaders for picp\n(0 - Disable)\n")
+	1,
+	TEXT("Enable postprocess shaders for picp (0 = disabled)"),
+	ECVF_RenderThreadSafe
 );
 
 #define PostProcessShaderFileName TEXT("/Plugin/nDisplay/Private/PicpPostProcessShaders.usf")
@@ -145,7 +146,7 @@ void FPicpBlurPostProcess::ApplyBlur_RenderThread(FRHICommandListImmediate& RHIC
 // BP gamethread api:
 void FPicpBlurPostProcess::ApplyBlur(UTextureRenderTarget2D* InOutRenderTarget, UTextureRenderTarget2D* TemporaryRenderTarget, int KernelRadius, float KernelScale, EPicpBlurPostProcessShaderType BlurType)
 {	
-	if (!InOutRenderTarget || !TemporaryRenderTarget || (CVarPicpPostProcess_Value == 0))
+	if (!InOutRenderTarget || !TemporaryRenderTarget || (CVarPicpPostProcess.GetValueOnAnyThread() == 0))
 	{
 		return;
 	}
@@ -173,16 +174,16 @@ void FPicpBlurPostProcess::ApplyBlur(UTextureRenderTarget2D* InOutRenderTarget,
 
 void FPicpBlurPostProcess::ApplyCompose(UTexture* InputTexture, UTextureRenderTarget2D* OutputRenderTarget, UTextureRenderTarget2D* Result)
 {
-	if (!OutputRenderTarget || (CVarPicpPostProcess_Value == 0))
+	if (!OutputRenderTarget || (CVarPicpPostProcess.GetValueOnAnyThread() == 0))
 	{
 		return;
 	}
+
 	FTextureResource* OverlayTextureResource = InputTexture ? InputTexture->Resource : nullptr;
 	FTextureRenderTargetResource* TextureRenderTarget = OutputRenderTarget->GameThread_GetRenderTargetResource();
 	FTextureRenderTargetResource* TextureRenderResult = Result->GameThread_GetRenderTargetResource();
 
-	ENQUEUE_RENDER_COMMAND(CaptureCommand)(
-		[OverlayTextureResource, TextureRenderTarget, TextureRenderResult](FRHICommandListImmediate& RHICmdList)
+	ENQUEUE_RENDER_COMMAND(CaptureCommand)([OverlayTextureResource, TextureRenderTarget, TextureRenderResult](FRHICommandListImmediate& RHICmdList)
 	{
 		FRHITexture2D* RTTexture = TextureRenderTarget ? TextureRenderTarget->GetRenderTargetTexture()->GetTexture2D(): nullptr;
 		FRHITexture2D* ResultTexture = TextureRenderResult ? TextureRenderResult->GetRenderTargetTexture()->GetTexture2D() : nullptr;
@@ -195,9 +196,7 @@ void FPicpBlurPostProcess::ApplyCompose(UTexture* InputTexture, UTextureRenderTa
 		}
 
 		ApplyCompose_RenderThread(RHICmdList, OverlayTexture, RTTexture, ResultTexture);
-
-	}
-	);
+	});
 }
 
 
diff --git a/Engine/Plugins/Runtime/nDisplay/Source/PicpMPCDI/Private/PicpMPCDIShader.cpp b/Engine/Plugins/Runtime/nDisplay/Source/PicpMPCDI/Private/PicpMPCDIShader.cpp
--- a/Engine/Plugins/Runtime/nDisplay/Source/PicpMPCDI/Private/PicpMPCDIShader.cpp
+++ b/Engine/Plugins/Runtime/nDisplay/Source/PicpMPCDI/Private/PicpMPCDIShader.cpp
@@ -16,11 +16,11 @@
 
 
 
-static int CVarPicpEnableSinglePassRender_Value = 1;
-static FAutoConsoleVariableRef CVarPicpEnableSinglePassRender(
+static TAutoConsoleVariable<int32> CVarPicpEnableSinglePassRender(
 	TEXT("nDisplay.render.picp.EnableSinglePassRender"),
-	CVarPicpEnableSinglePassRender_Value,
-	TEXT("Enable single pass render for picp\n(0 - Disable)\n")
+	1,
+	TEXT("Single render pass for PICP (0 = disabled)"),
+	ECVF_RenderThreadSafe
 );
 
 #define MPCDIShaderFileName TEXT("/Plugin/nDisplay/Private/MPCDIShaders.usf")
@@ -284,7 +284,7 @@ bool FPicpMPCDIShader::ApplyWarpBlend(FRHICommandListImmediate& RHICmdList, IMPC
 				if (ViewportOverlayData)
 				{// Use overlay render tricks:					
 					// First pass render bg
-					bool bIsRenderSinglePass = ((CVarPicpEnableSinglePassRender_Value == 1) && ViewportOverlayData->Cameras.Num() == 1);
+					bool bIsRenderSinglePass = ((CVarPicpEnableSinglePassRender.GetValueOnAnyThread() != 0) && ViewportOverlayData->Cameras.Num() == 1);
 
 					if (bIsRenderSinglePass)
 					{
diff --git a/Engine/Plugins/Runtime/nDisplay/Source/PicpProjection/Private/Policy/MPCDI/PicpProjectionMPCDIPolicy.cpp b/Engine/Plugins/Runtime/nDisplay/Source/PicpProjection/Private/Policy/MPCDI/PicpProjectionMPCDIPolicy.cpp
--- a/Engine/Plugins/Runtime/nDisplay/Source/PicpProjection/Private/Policy/MPCDI/PicpProjectionMPCDIPolicy.cpp
+++ b/Engine/Plugins/Runtime/nDisplay/Source/PicpProjection/Private/Policy/MPCDI/PicpProjectionMPCDIPolicy.cpp
@@ -20,11 +20,11 @@
 #include "Misc/Paths.h"
 
 
-static int CVarPicpShowFakeCamera_Value = 0;
-static FAutoConsoleVariableRef CVarPicpShowFakeCamera(
+static TAutoConsoleVariable<int32> CVarPicpShowFakeCamera(
 	TEXT("nDisplay.render.picp.ShowFakeCamera"),
-	CVarPicpShowFakeCamera_Value,
-	TEXT("Show fake debug camera\n(0 - Disable)\n")
+	0,
+	TEXT("Show fake debug camera (0 = disabled)"),
+	ECVF_RenderThreadSafe
 );
 
 
@@ -154,7 +154,7 @@ bool FPicpProjectionMPCDIPolicy::CalculateView(const uint32 ViewIdx, FVector& In
 
 		IPicpProjection::Get().SetViewport(GetViewportId(), InOutViewRotation, InOutViewLocation, OutPrjMatrix);
 
-		if (CVarPicpShowFakeCamera_Value == 1)
+		if (CVarPicpShowFakeCamera.GetValueOnAnyThread() != 0)
 		{
 			//Add debug content for overlay shaders test purpose:
 			FPicpProjectionOverlayFrameData::GenerateDebugContent(GetViewportId(), this);
diff --git a/Engine/Source/Programs/nDisplayLauncher/MainWindow.xaml.cs b/Engine/Source/Programs/nDisplayLauncher/MainWindow.xaml.cs
--- a/Engine/Source/Programs/nDisplayLauncher/MainWindow.xaml.cs
+++ b/Engine/Source/Programs/nDisplayLauncher/MainWindow.xaml.cs
@@ -26,7 +26,7 @@ namespace nDisplayLauncher
 
 		private void Window_Loaded(object sender, RoutedEventArgs e)
 		{
-			this.Title = "nDisplay Launcher";
+			this.Title = "nDisplay Launcher - UE4.23";
 		}
 
 		private void Exit(object sender, RoutedEventArgs e)
