Copied 7021298, 7021304 from Dev-Enterprise

Misc nDisplay fixes

diff --git a/Engine/Plugins/Runtime/nDisplay/Source/DisplayCluster/Private/Game/DisplayClusterGameManager.cpp b/Engine/Plugins/Runtime/nDisplay/Source/DisplayCluster/Private/Game/DisplayClusterGameManager.cpp
--- a/Engine/Plugins/Runtime/nDisplay/Source/DisplayCluster/Private/Game/DisplayClusterGameManager.cpp
+++ b/Engine/Plugins/Runtime/nDisplay/Source/DisplayCluster/Private/Game/DisplayClusterGameManager.cpp
@@ -338,7 +338,7 @@ bool FDisplayClusterGameManager::InitializeDisplayClusterActor()
 	APlayerController* pController = UGameplayStatics::GetPlayerController(CurrentWorld, 0);
 	check(pController);
 	
-	VRRootActor = StaticCast<ADisplayClusterPawn*>(pController->GetPawn());
+	VRRootActor = Cast<ADisplayClusterPawn>(pController->GetPawn());
 	if (!VRRootActor)
 	{
 		// Seems the DisplayCluster features has been disabled
diff --git a/Engine/Plugins/Runtime/nDisplay/Source/DisplayCluster/Private/Render/Synchronization/DisplayClusterRenderSyncPolicySoftwareDX11.cpp b/Engine/Plugins/Runtime/nDisplay/Source/DisplayCluster/Private/Render/Synchronization/DisplayClusterRenderSyncPolicySoftwareDX11.cpp
--- a/Engine/Plugins/Runtime/nDisplay/Source/DisplayCluster/Private/Render/Synchronization/DisplayClusterRenderSyncPolicySoftwareDX11.cpp
+++ b/Engine/Plugins/Runtime/nDisplay/Source/DisplayCluster/Private/Render/Synchronization/DisplayClusterRenderSyncPolicySoftwareDX11.cpp
@@ -32,37 +32,43 @@
 static TAutoConsoleVariable<int32> CVarAdvancedSyncEnabled(
 	TEXT("nDisplay.render.softsync.AdvancedSyncEnabled"),
 	0,
-	TEXT("Advanced DWM based render synchronization (0 = disabled)")
+	TEXT("Advanced DWM based render synchronization (0 = disabled)"),
+	ECVF_RenderThreadSafe
 );
 
 static TAutoConsoleVariable<float> CVarVBlankThreshold(
 	TEXT("nDisplay.render.softsync.VBlankThreshold"),
 	0.003f,
-	TEXT("Changes the threshold used to determine whether or not nDisplay SoftSync can make the current frame presentation")
+	TEXT("Changes the threshold used to determine whether or not nDisplay SoftSync can make the current frame presentation"),
+	ECVF_RenderThreadSafe
 );
 
 static TAutoConsoleVariable<float> CVarVBlankThresholdSleepMultiplier(
 	TEXT("nDisplay.render.softsync.VBlankThresholdSleepMultipler"),
 	1.5f,
-	TEXT("Multiplier applied to VBlankThreshold to compute sleep time to skip frame presentation on upcoming VBlank")
+	TEXT("Multiplier applied to VBlankThreshold to compute sleep time to skip frame presentation on upcoming VBlank"),
+	ECVF_RenderThreadSafe
 );
 
 static TAutoConsoleVariable<int32>  CVarVBlankBasisUpdate(
 	TEXT("nDisplay.render.softsync.VBlankBasisUpdate"),
 	1,
-	TEXT("Update VBlank basis periodically to avoid time drifting (0 = disabled)")
+	TEXT("Update VBlank basis periodically to avoid time drifting (0 = disabled)"),
+	ECVF_RenderThreadSafe
 );
 
 static TAutoConsoleVariable<float>  CVarVBlankBasisUpdatePeriod(
 	TEXT("nDisplay.render.softsync.VBlankBasisUpdatePeriod"),
 	120.f,
-	TEXT("VBlank basis update period in seconds")
+	TEXT("VBlank basis update period in seconds"),
+	ECVF_RenderThreadSafe
 );
 
 static TAutoConsoleVariable<int32>  CVarLogDwmStats(
 	TEXT("nDisplay.render.softsync.LogDwmStats"),
 	0,
-	TEXT("Print DWM stats to log (0 = disabled)")
+	TEXT("Print DWM stats to log (0 = disabled)"),
+	ECVF_RenderThreadSafe
 );
 
 
