Integrate misc nDisplay fixes from Dev-Enterprise

- nDisplayLauncher: Removed useless no-val command line arguments
- Fixed Editor crash when translation direction is set in BP_DemoWand_Blueprint
- Use quotation marks for paths
- A lot of nDisplay warings "wand not found"
- nDisplayLauncher passes custom arguments at beginning
- Update blueprints (resubmit + comments)
- Single view for monoscipic devices
- Move swap eye attribute to [cluster_node]
- Removed viewport flip feature
- Disable input on slave nodes
- Fix for GL3/4 monoscopic device. Missed changes for 'Viewport positioning fix'.
- Removed custom swap_int
- Fixed demo blueprints
- Fixed crash in DC components when the DC is not initialized
- Fixed visibility of some inherited members
- Fixed Editor crash on second PIE session start
- Fix for viewport positioning
- BP API docs (comments)
- Turn off some screen effects in the nDisplay template project

diff --git a/Engine/Plugins/Runtime/nDisplay/Source/DisplayCluster/Private/Blueprints/DisplayClusterBlueprintAPIImpl.cpp b/Engine/Plugins/Runtime/nDisplay/Source/DisplayCluster/Private/Blueprints/DisplayClusterBlueprintAPIImpl.cpp
--- a/Engine/Plugins/Runtime/nDisplay/Source/DisplayCluster/Private/Blueprints/DisplayClusterBlueprintAPIImpl.cpp
+++ b/Engine/Plugins/Runtime/nDisplay/Source/DisplayCluster/Private/Blueprints/DisplayClusterBlueprintAPIImpl.cpp
@@ -442,28 +442,6 @@ bool UDisplayClusterBlueprintAPIImpl::ToggleEyesSwap()
 	return false;
 }
 
-void UDisplayClusterBlueprintAPIImpl::SetOutputFlip(bool flipH, bool flipV)
-{
-	IDisplayClusterRenderManager* const Manager = IDisplayCluster::Get().GetRenderMgr();
-	if (Manager)
-	{
-		return Manager->GetStereoDevice()->SetOutputFlip(flipH, flipV);
-	}
-
-	return;
-}
-
-void UDisplayClusterBlueprintAPIImpl::GetOutputFlip(bool& flipH, bool& flipV)
-{
-	IDisplayClusterRenderManager* const Manager = IDisplayCluster::Get().GetRenderMgr();
-	if (Manager)
-	{
-		return Manager->GetStereoDevice()->GetOutputFlip(flipH, flipV);
-	}
-
-	return;
-}
-
 void UDisplayClusterBlueprintAPIImpl::GetCullingDistance(float& NearClipPlane, float& FarClipPlane)
 {
 	IDisplayClusterRenderManager* const Manager = IDisplayCluster::Get().GetRenderMgr();
diff --git a/Engine/Plugins/Runtime/nDisplay/Source/DisplayCluster/Private/Blueprints/DisplayClusterBlueprintAPIImpl.h b/Engine/Plugins/Runtime/nDisplay/Source/DisplayCluster/Private/Blueprints/DisplayClusterBlueprintAPIImpl.h
--- a/Engine/Plugins/Runtime/nDisplay/Source/DisplayCluster/Private/Blueprints/DisplayClusterBlueprintAPIImpl.h
+++ b/Engine/Plugins/Runtime/nDisplay/Source/DisplayCluster/Private/Blueprints/DisplayClusterBlueprintAPIImpl.h
@@ -183,12 +183,6 @@ public:
 	UFUNCTION(BlueprintCallable, meta = (DisplayName = "Toggle eye swap"), Category = "DisplayCluster|Render")
 	virtual bool ToggleEyesSwap() override;
 
-	UFUNCTION(BlueprintCallable, meta = (DisplayName = "Set output flip"), Category = "DisplayCluster|Render")
-	virtual void SetOutputFlip(bool flipH, bool flipV) override;
-
-	UFUNCTION(BlueprintCallable, meta = (DisplayName = "Get output flip"), Category = "DisplayCluster|Render")
-	virtual void GetOutputFlip(bool& flipH, bool& flipV) override;
-
 	UFUNCTION(BlueprintCallable, meta = (DisplayName = "Get near and far clipping distance"), Category = "DisplayCluster|Render")
 	virtual void GetCullingDistance(float& NearClipPlane, float& FarClipPlane) override;
 
diff --git a/Engine/Plugins/Runtime/nDisplay/Source/DisplayCluster/Private/Cluster/Controller/DisplayClusterClusterNodeCtrlBase.cpp b/Engine/Plugins/Runtime/nDisplay/Source/DisplayCluster/Private/Cluster/Controller/DisplayClusterClusterNodeCtrlBase.cpp
--- a/Engine/Plugins/Runtime/nDisplay/Source/DisplayCluster/Private/Cluster/Controller/DisplayClusterClusterNodeCtrlBase.cpp
+++ b/Engine/Plugins/Runtime/nDisplay/Source/DisplayCluster/Private/Cluster/Controller/DisplayClusterClusterNodeCtrlBase.cpp
@@ -50,16 +50,16 @@ bool FDisplayClusterClusterNodeCtrlBase::InitializeStereo()
 	IDisplayClusterStereoDevice* const StereoDevice = GDisplayCluster->GetPrivateRenderMgr()->GetStereoDevice();
 	if (StereoDevice)
 	{
-		
 		FDisplayClusterConfigStereo  StereoCfg  = GDisplayCluster->GetPrivateConfigMgr()->GetConfigStereo();
 		FDisplayClusterConfigGeneral GeneralCfg = GDisplayCluster->GetPrivateConfigMgr()->GetConfigGeneral();
-		//FDisplayClusterConfigRender  RenderCfg  = GDisplayCluster->GetPrivateConfigMgr()->GetConfigRender();
+		FDisplayClusterConfigClusterNode ClusterNodeCfg;
+		GDisplayCluster->GetPrivateConfigMgr()->GetLocalClusterNode(ClusterNodeCfg);
+		
 
 		// Configure the device
 		StereoDevice->SetViewportArea(ViewportCfg.Loc, ViewportCfg.Size);
-		StereoDevice->SetEyesSwap(StereoCfg.EyeSwap);
+		StereoDevice->SetEyesSwap(ClusterNodeCfg.EyeSwap);
 		StereoDevice->SetInterpupillaryDistance(StereoCfg.EyeDist);
-		StereoDevice->SetOutputFlip(ViewportCfg.FlipHorizontal, ViewportCfg.FlipVertical);
 		StereoDevice->SetSwapSyncPolicy((EDisplayClusterSwapSyncPolicy)GeneralCfg.SwapSyncPolicy);
 	}
 	else
diff --git a/Engine/Plugins/Runtime/nDisplay/Source/DisplayCluster/Private/Config/DisplayClusterConfigTypes.cpp b/Engine/Plugins/Runtime/nDisplay/Source/DisplayCluster/Private/Config/DisplayClusterConfigTypes.cpp
--- a/Engine/Plugins/Runtime/nDisplay/Source/DisplayCluster/Private/Config/DisplayClusterConfigTypes.cpp
+++ b/Engine/Plugins/Runtime/nDisplay/Source/DisplayCluster/Private/Config/DisplayClusterConfigTypes.cpp
@@ -12,7 +12,7 @@
 //////////////////////////////////////////////////////////////////////////////////////////////
 FString FDisplayClusterConfigClusterNode::ToString() const
 {
-	return FString::Printf(TEXT("[%s + %s=%s, %s=%s, %s=%s, %s=%s, %s=%s, %s=%d, %s=%d, %s=%s]"),
+	return FString::Printf(TEXT("[%s + %s=%s, %s=%s, %s=%s, %s=%s, %s=%s, %s=%d, %s=%d, %s=%s, %s=%s]"),
 		*FDisplayClusterConfigBase::ToString(),
 		DisplayClusterStrings::cfg::data::Id, *Id, 
 		DisplayClusterStrings::cfg::data::cluster::Addr, *Addr, 
@@ -21,7 +21,8 @@ FString FDisplayClusterConfigClusterNode::ToString() const
 		DisplayClusterStrings::cfg::data::cluster::Viewport, *ViewportId, 
 		DisplayClusterStrings::cfg::data::cluster::PortCS, Port_CS,
 		DisplayClusterStrings::cfg::data::cluster::PortSS, Port_SS,
-		DisplayClusterStrings::cfg::data::cluster::Sound, DisplayClusterHelpers::str::BoolToStr(SoundEnabled));
+		DisplayClusterStrings::cfg::data::cluster::Sound, DisplayClusterHelpers::str::BoolToStr(SoundEnabled),
+		DisplayClusterStrings::cfg::data::cluster::EyeSwap, DisplayClusterHelpers::str::BoolToStr(EyeSwap));
 }
 
 bool FDisplayClusterConfigClusterNode::DeserializeFromString(const FString& line)
@@ -34,6 +35,7 @@ bool FDisplayClusterConfigClusterNode::DeserializeFromString(const FString& line
 	DisplayClusterHelpers::str::ExtractCommandLineValue(line, FString(DisplayClusterStrings::cfg::data::cluster::PortCS),   Port_CS);
 	DisplayClusterHelpers::str::ExtractCommandLineValue(line, FString(DisplayClusterStrings::cfg::data::cluster::PortSS),   Port_SS);
 	DisplayClusterHelpers::str::ExtractCommandLineValue(line, FString(DisplayClusterStrings::cfg::data::cluster::Sound),    SoundEnabled);
+	DisplayClusterHelpers::str::ExtractCommandLineValue(line, FString(DisplayClusterStrings::cfg::data::cluster::EyeSwap),  EyeSwap);
 	return FDisplayClusterConfigBase::DeserializeFromString(line);
 }
 
@@ -43,14 +45,12 @@ bool FDisplayClusterConfigClusterNode::DeserializeFromString(const FString& line
 //////////////////////////////////////////////////////////////////////////////////////////////
 FString FDisplayClusterConfigViewport::ToString() const
 {
-	return FString::Printf(TEXT("[%s + %s=%s, %s=%s, %s=%d, %s=%d, %s=%s, %s=%s]"),
+	return FString::Printf(TEXT("[%s + %s=%s, %s=%s, %s=%d, %s=%d]"),
 		*FDisplayClusterConfigBase::ToString(),
 		DisplayClusterStrings::cfg::data::Id, *Id,
 		DisplayClusterStrings::cfg::data::Loc, *Loc.ToString(),
 		DisplayClusterStrings::cfg::data::viewport::Width,  Size.X,
-		DisplayClusterStrings::cfg::data::viewport::Height, Size.Y,
-		DisplayClusterStrings::cfg::data::viewport::FlipH, DisplayClusterHelpers::str::BoolToStr(FlipHorizontal),
-		DisplayClusterStrings::cfg::data::viewport::FlipV, DisplayClusterHelpers::str::BoolToStr(FlipVertical));
+		DisplayClusterStrings::cfg::data::viewport::Height, Size.Y);
 }
 
 bool FDisplayClusterConfigViewport::DeserializeFromString(const FString& line)
@@ -60,8 +60,6 @@ bool FDisplayClusterConfigViewport::DeserializeFromString(const FString& line)
 	DisplayClusterHelpers::str::ExtractCommandLineValue(line, FString(DisplayClusterStrings::cfg::data::viewport::PosY),   Loc.Y);
 	DisplayClusterHelpers::str::ExtractCommandLineValue(line, FString(DisplayClusterStrings::cfg::data::viewport::Width),  Size.X);
 	DisplayClusterHelpers::str::ExtractCommandLineValue(line, FString(DisplayClusterStrings::cfg::data::viewport::Height), Size.Y);
-	DisplayClusterHelpers::str::ExtractCommandLineValue(line, FString(DisplayClusterStrings::cfg::data::viewport::FlipH),  FlipHorizontal);
-	DisplayClusterHelpers::str::ExtractCommandLineValue(line, FString(DisplayClusterStrings::cfg::data::viewport::FlipV),  FlipVertical);
 	return FDisplayClusterConfigBase::DeserializeFromString(line);
 }
 
@@ -215,16 +213,14 @@ bool FDisplayClusterConfigRender::DeserializeFromString(const FString& line)
 //////////////////////////////////////////////////////////////////////////////////////////////
 FString FDisplayClusterConfigStereo::ToString() const
 {
-	return FString::Printf(TEXT("[%s + %s=%s, %s=%f]"),
+	return FString::Printf(TEXT("[%s + %s=%f]"),
 		*FDisplayClusterConfigBase::ToString(),
-		DisplayClusterStrings::cfg::data::stereo::EyeSwap, DisplayClusterHelpers::str::BoolToStr(EyeSwap),
 		DisplayClusterStrings::cfg::data::stereo::EyeDist, EyeDist);
 }
 
 bool FDisplayClusterConfigStereo::DeserializeFromString(const FString& line)
 {
 	DisplayClusterHelpers::str::ExtractCommandLineValue(line, FString(DisplayClusterStrings::cfg::data::stereo::EyeDist), EyeDist);
-	DisplayClusterHelpers::str::ExtractCommandLineValue(line, FString(DisplayClusterStrings::cfg::data::stereo::EyeSwap), EyeSwap);
 	return FDisplayClusterConfigBase::DeserializeFromString(line);
 }
 
diff --git a/Engine/Plugins/Runtime/nDisplay/Source/DisplayCluster/Private/Config/Parser/DisplayClusterConfigParserDebugAuto.cpp b/Engine/Plugins/Runtime/nDisplay/Source/DisplayCluster/Private/Config/Parser/DisplayClusterConfigParserDebugAuto.cpp
--- a/Engine/Plugins/Runtime/nDisplay/Source/DisplayCluster/Private/Config/Parser/DisplayClusterConfigParserDebugAuto.cpp
+++ b/Engine/Plugins/Runtime/nDisplay/Source/DisplayCluster/Private/Config/Parser/DisplayClusterConfigParserDebugAuto.cpp
@@ -25,6 +25,7 @@ bool FDisplayClusterConfigParserDebugAuto::ParseFile(const FString& path)
 	ClusterNode.ScreenId   = TEXT("screen_stub");;
 	ClusterNode.ViewportId = TEXT("viewport_stub");
 	ClusterNode.SoundEnabled = true;
+	ClusterNode.EyeSwap    = false;
 	AddClusterNode(ClusterNode);
 
 	const float PixelDensity = 0.6f / 1920.f;
@@ -40,8 +41,6 @@ bool FDisplayClusterConfigParserDebugAuto::ParseFile(const FString& path)
 	Viewport.Id   = ClusterNode.ViewportId;
 	Viewport.Loc  = FIntPoint(0, 0);
 	Viewport.Size = FIntPoint(DisplayClusterConstants::misc::DebugAutoResX, DisplayClusterConstants::misc::DebugAutoResY);
-	Viewport.FlipHorizontal = false;
-	Viewport.FlipVertical   = false;
 	AddViewport(Viewport);
 	
 	FDisplayClusterConfigCamera Camera;
diff --git a/Engine/Plugins/Runtime/nDisplay/Source/DisplayCluster/Private/DisplayClusterStrings.h b/Engine/Plugins/Runtime/nDisplay/Source/DisplayCluster/Private/DisplayClusterStrings.h
--- a/Engine/Plugins/Runtime/nDisplay/Source/DisplayCluster/Private/DisplayClusterStrings.h
+++ b/Engine/Plugins/Runtime/nDisplay/Source/DisplayCluster/Private/DisplayClusterStrings.h
@@ -72,6 +72,7 @@ namespace DisplayClusterStrings
 				static constexpr auto PortSS   = TEXT("port_ss");
 				static constexpr auto Master   = TEXT("master");
 				static constexpr auto Sound    = TEXT("sound");
+				static constexpr auto EyeSwap  = TEXT("eye_swap");
 				// + Id
 			}
 
@@ -91,8 +92,6 @@ namespace DisplayClusterStrings
 				static constexpr auto PosY   = TEXT("y");
 				static constexpr auto Width  = TEXT("width");
 				static constexpr auto Height = TEXT("height");
-				static constexpr auto FlipH  = TEXT("flip_h");
-				static constexpr auto FlipV  = TEXT("flip_v");
 				// + Id
 			}
 
@@ -148,7 +147,6 @@ namespace DisplayClusterStrings
 			namespace stereo
 			{
 				static constexpr auto Header  = TEXT("[stereo]");
-				static constexpr auto EyeSwap = TEXT("eye_swap");
 				static constexpr auto EyeDist = TEXT("eye_dist");
 			}
 
@@ -171,7 +169,6 @@ namespace DisplayClusterStrings
 			namespace custom
 			{
 				static constexpr auto Header     = TEXT("[custom]");
-				static constexpr auto SwapInt    = TEXT("swap_int");
 			}
 		}
 	};
diff --git a/Engine/Plugins/Runtime/nDisplay/Source/DisplayCluster/Private/Game/Classes/Scene/DisplayClusterPawn.cpp b/Engine/Plugins/Runtime/nDisplay/Source/DisplayCluster/Private/Game/Classes/Scene/DisplayClusterPawn.cpp
--- a/Engine/Plugins/Runtime/nDisplay/Source/DisplayCluster/Private/Game/Classes/Scene/DisplayClusterPawn.cpp
+++ b/Engine/Plugins/Runtime/nDisplay/Source/DisplayCluster/Private/Game/Classes/Scene/DisplayClusterPawn.cpp
@@ -3,12 +3,15 @@
 #include "DisplayClusterPawn.h"
 
 #include "Engine/CollisionProfile.h"
+#include "Engine/World.h"
 
 #include "Camera/CameraComponent.h"
 #include "Components/SphereComponent.h"
+#include "GameFramework/PlayerController.h"
 
 #include "Cluster/IPDisplayClusterClusterManager.h"
 #include "Game/IPDisplayClusterGameManager.h"
+#include "Kismet/GameplayStatics.h"
 
 #include "DisplayClusterSceneComponentSyncParent.h"
 
@@ -96,6 +99,20 @@ void ADisplayClusterPawn::BeginPlay()
 				UE_LOG(LogDisplayClusterGame, Log, TEXT("Collision offset: %s"), *CollisionOffset.ToString());
 			}
 		}
+		else
+		{
+			// Turn off input processing on slave nodes
+			UWorld* World = GetWorld();
+			if (World)
+			{
+				APlayerController* PlayerController = World->GetFirstPlayerController();
+				if (PlayerController)
+				{
+					UE_LOG(LogDisplayClusterGame, Log, TEXT("Deactivating input on slave node..."));
+					this->DisableInput(PlayerController);
+				}
+			}
+		}
 	}
 }
 
diff --git a/Engine/Plugins/Runtime/nDisplay/Source/DisplayCluster/Private/Game/DisplayClusterGameManager.cpp b/Engine/Plugins/Runtime/nDisplay/Source/DisplayCluster/Private/Game/DisplayClusterGameManager.cpp
--- a/Engine/Plugins/Runtime/nDisplay/Source/DisplayCluster/Private/Game/DisplayClusterGameManager.cpp
+++ b/Engine/Plugins/Runtime/nDisplay/Source/DisplayCluster/Private/Game/DisplayClusterGameManager.cpp
@@ -179,8 +179,6 @@ void FDisplayClusterGameManager::SetActiveCamera(int32 idx)
 		return;
 	}
 
-	check(idx < GDisplayCluster->GetPrivateConfigMgr()->GetCamerasAmount());
-
 	FDisplayClusterConfigCamera cam;
 	if (!GDisplayCluster->GetPrivateConfigMgr()->GetCamera(idx, cam))
 	{
@@ -234,11 +232,16 @@ TArray<UDisplayClusterSceneComponent*> FDisplayClusterGameManager::GetAllNodes()
 USceneComponent* FDisplayClusterGameManager::GetTranslationDirectionComponent() const
 {
 	if (!IsDisplayClusterActive())
+	{
 		return nullptr;
+	}
 
-	FScopeLock lock(&InternalsSyncScope);
-	check(VRRootActor);
+	if (VRRootActor == nullptr)
+	{
+		return nullptr;
+	}
 
+	FScopeLock lock(&InternalsSyncScope);
 	UE_LOG(LogDisplayClusterGame, Verbose, TEXT("GetTranslationDirectionComponent: %s"), (VRRootActor->TranslationDirection ? *VRRootActor->TranslationDirection->GetName() : TEXT("nullptr")));
 	return VRRootActor->TranslationDirection;
 }
@@ -252,9 +255,12 @@ void FDisplayClusterGameManager::SetTranslationDirectionComponent(USceneComponen
 		return;
 	}
 
-	FScopeLock lock(&InternalsSyncScope);
-	check(VRRootActor);
+	if (VRRootActor == nullptr)
+	{
+		return;
+	}
 
+	FScopeLock lock(&InternalsSyncScope);
 	UE_LOG(LogDisplayClusterGame, Log, TEXT("New translation direction component set: %s"), (pComp ? *pComp->GetName() : TEXT("nullptr")));
 	VRRootActor->TranslationDirection = pComp;
 }
@@ -268,9 +274,6 @@ void FDisplayClusterGameManager::SetTranslationDirectionComponent(const FString&
 		return;
 	}
 
-	FScopeLock lock(&InternalsSyncScope);
-	check(VRRootActor);
-
 	UE_LOG(LogDisplayClusterGame, Log, TEXT("New translation direction node id requested: %s"), *id);
 	SetTranslationDirectionComponent(GetNodeById(id));
 }
@@ -282,9 +285,12 @@ USceneComponent* FDisplayClusterGameManager::GetRotateAroundComponent() const
 		return nullptr;
 	}
 
-	FScopeLock lock(&InternalsSyncScope);
-	check(VRRootActor);
+	if (VRRootActor == nullptr)
+	{
+		return nullptr;
+	}
 
+	FScopeLock lock(&InternalsSyncScope);
 	UE_LOG(LogDisplayClusterGame, Verbose, TEXT("GetRotateAroundComponent: %s"), (VRRootActor->RotationAround ? *VRRootActor->RotationAround->GetName() : TEXT("nullptr")));
 	return VRRootActor->RotationAround;
 }
@@ -298,9 +304,12 @@ void FDisplayClusterGameManager::SetRotateAroundComponent(USceneComponent* pComp
 		return;
 	}
 
+	if (VRRootActor == nullptr)
+	{
+		return;
+	}
+
 	FScopeLock lock(&InternalsSyncScope);
-	check(VRRootActor);
-	
 	UE_LOG(LogDisplayClusterGame, Log, TEXT("New rotate around component set: %s"), (pComp ? *pComp->GetName() : TEXT("nullptr")));
 	VRRootActor->RotationAround = pComp;
 }
@@ -314,9 +323,12 @@ void FDisplayClusterGameManager::SetRotateAroundComponent(const FString& id)
 		return;
 	}
 
-	FScopeLock lock(&InternalsSyncScope);
-	check(VRRootActor);
+	if (VRRootActor == nullptr)
+	{
+		return;
+	}
 
+	FScopeLock lock(&InternalsSyncScope);
 	UE_LOG(LogDisplayClusterGame, Log, TEXT("New rotate around node id requested: %s"), *id);
 	VRRootActor->RotationAround = GetNodeById(id);
 }
@@ -407,7 +419,6 @@ bool FDisplayClusterGameManager::CreateScreens()
 	}
 
 	// Check if local screen was found
-	check(ActiveScreenComponent);
 	if (!ActiveScreenComponent)
 	{
 		UE_LOG(LogDisplayClusterGame, Error, TEXT("Local screen not found"));
diff --git a/Engine/Plugins/Runtime/nDisplay/Source/DisplayCluster/Private/Render/Devices/DisplayClusterDeviceBase.cpp b/Engine/Plugins/Runtime/nDisplay/Source/DisplayCluster/Private/Render/Devices/DisplayClusterDeviceBase.cpp
--- a/Engine/Plugins/Runtime/nDisplay/Source/DisplayCluster/Private/Render/Devices/DisplayClusterDeviceBase.cpp
+++ b/Engine/Plugins/Runtime/nDisplay/Source/DisplayCluster/Private/Render/Devices/DisplayClusterDeviceBase.cpp
@@ -39,11 +39,8 @@ bool FDisplayClusterDeviceBase::Initialize()
 		return false;
 	}
 
-	const FDisplayClusterConfigCustom cfgCustom = GDisplayCluster->GetPrivateConfigMgr()->GetConfigCustom();
-
-	if (cfgCustom.Args.Contains(FString(DisplayClusterStrings::cfg::data::custom::SwapInt)))
 	{
-		SwapInterval = FCString::Atoi(*cfgCustom.Args[FString(DisplayClusterStrings::cfg::data::custom::SwapInt)]);
+		return false;
 	}
 
 	UE_LOG(LogDisplayClusterRender, Log, TEXT("Use swap interval: %d"), SwapInterval);
@@ -384,24 +381,6 @@ bool FDisplayClusterDeviceBase::ToggleEyesSwap()
 	return bEyeSwap;
 }
 
-void FDisplayClusterDeviceBase::SetOutputFlip(bool flipH, bool flipV)
-{
-	UE_LOG(LogDisplayClusterRender, Log, TEXT("SetOutputFlip: horizontal=%s vertical=%s"), DisplayClusterHelpers::str::BoolToStr(flipH), DisplayClusterHelpers::str::BoolToStr(flipV));
-	
-	FScopeLock lock(&InternalsSyncScope);
-	bFlipHorizontal = flipH;
-	bFlipVertical = flipV;
-}
-
-void FDisplayClusterDeviceBase::GetOutputFlip(bool& flipH, bool& flipV) const
-{
-	UE_LOG(LogDisplayClusterRender, Verbose, TEXT("GetOutputFlip: horizontal=%s vertical=%s"), DisplayClusterHelpers::str::BoolToStr(bFlipHorizontal), DisplayClusterHelpers::str::BoolToStr(bFlipVertical));
-
-	FScopeLock lock(&InternalsSyncScope);
-	flipH = bFlipHorizontal;
-	flipV = bFlipVertical;
-}
-
 void FDisplayClusterDeviceBase::SetSwapSyncPolicy(EDisplayClusterSwapSyncPolicy policy)
 {
 	UE_LOG(LogDisplayClusterRender, Log, TEXT("Swap sync policy: %d"), (int)policy);
diff --git a/Engine/Plugins/Runtime/nDisplay/Source/DisplayCluster/Private/Render/Devices/DisplayClusterDeviceBase.h b/Engine/Plugins/Runtime/nDisplay/Source/DisplayCluster/Private/Render/Devices/DisplayClusterDeviceBase.h
--- a/Engine/Plugins/Runtime/nDisplay/Source/DisplayCluster/Private/Render/Devices/DisplayClusterDeviceBase.h
+++ b/Engine/Plugins/Runtime/nDisplay/Source/DisplayCluster/Private/Render/Devices/DisplayClusterDeviceBase.h
@@ -136,8 +136,6 @@ public:
 	virtual void SetEyesSwap(bool swap) override;
 	virtual bool GetEyesSwap() const override;
 	virtual bool ToggleEyesSwap() override;
-	virtual void SetOutputFlip(bool flipH, bool flipV) override;
-	virtual void GetOutputFlip(bool& flipH, bool& flipV) const override;
 	virtual void SetSwapSyncPolicy(EDisplayClusterSwapSyncPolicy policy) override;
 	virtual EDisplayClusterSwapSyncPolicy GetSwapSyncPolicy() const override;
 	virtual void GetCullingDistance(float& NearDistance, float& FarDistance) const override;
@@ -185,10 +183,6 @@ protected:
 	// Viewport area settings
 	FDisplayClusterViewportArea ViewportArea;
 
-	// Render config
-	bool bFlipHorizontal = false;
-	bool bFlipVertical = false;
-
 	// Clipping plane
 	float NearClipPlane = GNearClippingPlane;
 	float FarClipPlane = 2000000.f;
diff --git a/Engine/Plugins/Runtime/nDisplay/Source/DisplayCluster/Private/Render/Devices/Monoscopic/DisplayClusterDeviceMonoscopicD3D11.h b/Engine/Plugins/Runtime/nDisplay/Source/DisplayCluster/Private/Render/Devices/Monoscopic/DisplayClusterDeviceMonoscopicD3D11.h
--- a/Engine/Plugins/Runtime/nDisplay/Source/DisplayCluster/Private/Render/Devices/Monoscopic/DisplayClusterDeviceMonoscopicD3D11.h
+++ b/Engine/Plugins/Runtime/nDisplay/Source/DisplayCluster/Private/Render/Devices/Monoscopic/DisplayClusterDeviceMonoscopicD3D11.h
@@ -15,7 +15,17 @@ public:
 	FDisplayClusterDeviceMonoscopicD3D11();
 	virtual ~FDisplayClusterDeviceMonoscopicD3D11();
 
-	virtual bool ShouldUseSeparateRenderTarget() const override { return false; };
+public:
+	virtual bool ShouldUseSeparateRenderTarget() const override
+	{ return false; };
+
+public:
+	//////////////////////////////////////////////////////////////////////////////////////////////
+	// IStereoRendering
+	//////////////////////////////////////////////////////////////////////////////////////////////
+	virtual int32 GetDesiredNumberOfViews(bool bStereoRequested) const override
+	{ return 1; }
+
 protected:
 	virtual bool Present(int32& InOutSyncInterval) override;
 };
diff --git a/Engine/Plugins/Runtime/nDisplay/Source/DisplayCluster/Private/Render/Devices/Monoscopic/DisplayClusterDeviceMonoscopicD3D12.h b/Engine/Plugins/Runtime/nDisplay/Source/DisplayCluster/Private/Render/Devices/Monoscopic/DisplayClusterDeviceMonoscopicD3D12.h
--- a/Engine/Plugins/Runtime/nDisplay/Source/DisplayCluster/Private/Render/Devices/Monoscopic/DisplayClusterDeviceMonoscopicD3D12.h
+++ b/Engine/Plugins/Runtime/nDisplay/Source/DisplayCluster/Private/Render/Devices/Monoscopic/DisplayClusterDeviceMonoscopicD3D12.h
@@ -16,9 +16,17 @@ public:
 	FDisplayClusterDeviceMonoscopicD3D12();
 	virtual ~FDisplayClusterDeviceMonoscopicD3D12();
 
+public:
 	virtual bool ShouldUseSeparateRenderTarget() const override
 	{ return false; }
 
+public:
+	//////////////////////////////////////////////////////////////////////////////////////////////
+	// IStereoRendering
+	//////////////////////////////////////////////////////////////////////////////////////////////
+	virtual int32 GetDesiredNumberOfViews(bool bStereoRequested) const override
+	{ return 1; }
+
 protected:
 	virtual bool Present(int32& InOutSyncInterval) override;
 };
diff --git a/Engine/Plugins/Runtime/nDisplay/Source/DisplayCluster/Private/Render/Devices/Monoscopic/DisplayClusterDeviceMonoscopicOpenGL.cpp b/Engine/Plugins/Runtime/nDisplay/Source/DisplayCluster/Private/Render/Devices/Monoscopic/DisplayClusterDeviceMonoscopicOpenGL.cpp
--- a/Engine/Plugins/Runtime/nDisplay/Source/DisplayCluster/Private/Render/Devices/Monoscopic/DisplayClusterDeviceMonoscopicOpenGL.cpp
+++ b/Engine/Plugins/Runtime/nDisplay/Source/DisplayCluster/Private/Render/Devices/Monoscopic/DisplayClusterDeviceMonoscopicOpenGL.cpp
@@ -5,7 +5,6 @@
 
 #include "Misc/DisplayClusterLog.h"
 
-#include <utility>
 
 
 FDisplayClusterDeviceMonoscopicOpenGL::FDisplayClusterDeviceMonoscopicOpenGL()
@@ -22,20 +21,13 @@ bool FDisplayClusterDeviceMonoscopicOpenGL::Present(int32& InOutSyncInterval)
 {
 	UE_LOG(LogDisplayClusterRender, Verbose, TEXT("FDisplayClusterDeviceQuadBufferStereoOpenGL::Present"));
 
-	const int halfSizeY = BackBuffSize.Y / 2;
-	int dstX1 = 0;
-	int dstX2 = BackBuffSize.X;
+	const int halfSizeX = BackBuffSize.X / 2;
+	const int dstX1 = 0;
+	const int dstX2 = halfSizeX;
 
 	// Convert to left bottom origin and flip Y
-	int dstY1 = ViewportSize.Y;
-	int dstY2 = 0;
-
-	if (bFlipHorizontal)
-		std::swap(dstX1, dstX2);
-
-	if (bFlipVertical)
-		std::swap(dstY1, dstY2);
-
+	const int dstY1 = ViewportSize.Y;
+	const int dstY2 = 0;
 
 	FOpenGLViewport* pOglViewport = static_cast<FOpenGLViewport*>(CurrentViewport->GetViewportRHI().GetReference());
 	check(pOglViewport);
@@ -49,7 +41,7 @@ bool FDisplayClusterDeviceMonoscopicOpenGL::Present(int32& InOutSyncInterval)
 
 	glDrawBuffer(GL_BACK);
 	glBlitFramebuffer(
-		0, 0, BackBuffSize.X, halfSizeY,
+		0, 0, halfSizeX, BackBuffSize.Y,
 		dstX1, dstY1, dstX2, dstY2,
 		GL_COLOR_BUFFER_BIT,
 		GL_NEAREST);
diff --git a/Engine/Plugins/Runtime/nDisplay/Source/DisplayCluster/Private/Render/Devices/Monoscopic/DisplayClusterDeviceMonoscopicOpenGL.h b/Engine/Plugins/Runtime/nDisplay/Source/DisplayCluster/Private/Render/Devices/Monoscopic/DisplayClusterDeviceMonoscopicOpenGL.h
--- a/Engine/Plugins/Runtime/nDisplay/Source/DisplayCluster/Private/Render/Devices/Monoscopic/DisplayClusterDeviceMonoscopicOpenGL.h
+++ b/Engine/Plugins/Runtime/nDisplay/Source/DisplayCluster/Private/Render/Devices/Monoscopic/DisplayClusterDeviceMonoscopicOpenGL.h
@@ -15,6 +15,13 @@ public:
 	FDisplayClusterDeviceMonoscopicOpenGL();
 	virtual ~FDisplayClusterDeviceMonoscopicOpenGL();
 
+public:
+	//////////////////////////////////////////////////////////////////////////////////////////////
+	// IStereoRendering
+	//////////////////////////////////////////////////////////////////////////////////////////////
+	virtual int32 GetDesiredNumberOfViews(bool bStereoRequested) const override
+	{ return 1; }
+
 protected:
 	virtual bool Present(int32& InOutSyncInterval) override;
 };
diff --git a/Engine/Plugins/Runtime/nDisplay/Source/DisplayCluster/Private/Render/Devices/QuadBufferStereo/DisplayClusterDeviceQuadBufferStereoOpenGL.cpp b/Engine/Plugins/Runtime/nDisplay/Source/DisplayCluster/Private/Render/Devices/QuadBufferStereo/DisplayClusterDeviceQuadBufferStereoOpenGL.cpp
--- a/Engine/Plugins/Runtime/nDisplay/Source/DisplayCluster/Private/Render/Devices/QuadBufferStereo/DisplayClusterDeviceQuadBufferStereoOpenGL.cpp
+++ b/Engine/Plugins/Runtime/nDisplay/Source/DisplayCluster/Private/Render/Devices/QuadBufferStereo/DisplayClusterDeviceQuadBufferStereoOpenGL.cpp
@@ -6,8 +6,6 @@
 
 #include "Misc/DisplayClusterLog.h"
 
-#include <utility>
-
 
 FDisplayClusterDeviceQuadBufferStereoOpenGL::FDisplayClusterDeviceQuadBufferStereoOpenGL() :
 	FDisplayClusterDeviceQuadBufferStereoBase()
@@ -53,23 +51,12 @@ bool FDisplayClusterDeviceQuadBufferStereoOpenGL::Present(int32& InOutSyncInterv
 	UE_LOG(LogDisplayClusterRender, Verbose, TEXT("FDisplayClusterDeviceQuadBufferStereoOpenGL::Present"));
 
 	const int halfSizeX = BackBuffSize.X / 2;
-	int dstX1 = 0;
-	int dstX2 = halfSizeX;
+	const int dstX1 = 0;
+	const int dstX2 = halfSizeX;
 
 	// Convert to left bottom origin and flip Y
-	int dstY1 = ViewportSize.Y;
-	int dstY2 = 0;
-
-	if (bFlipHorizontal)
-	{
-		std::swap(dstX1, dstX2);
-	}
-
-	if (bFlipVertical)
-	{
-		std::swap(dstY1, dstY2);
-	}
-
+	const int dstY1 = ViewportSize.Y;
+	const int dstY2 = 0;
 
 	FOpenGLViewport* pOglViewport = static_cast<FOpenGLViewport*>(CurrentViewport->GetViewportRHI().GetReference());
 	check(pOglViewport);
diff --git a/Engine/Plugins/Runtime/nDisplay/Source/DisplayCluster/Public/Blueprints/IDisplayClusterBlueprintAPI.h b/Engine/Plugins/Runtime/nDisplay/Source/DisplayCluster/Public/Blueprints/IDisplayClusterBlueprintAPI.h
--- a/Engine/Plugins/Runtime/nDisplay/Source/DisplayCluster/Public/Blueprints/IDisplayClusterBlueprintAPI.h
+++ b/Engine/Plugins/Runtime/nDisplay/Source/DisplayCluster/Public/Blueprints/IDisplayClusterBlueprintAPI.h
@@ -226,12 +226,6 @@ public:
 	UFUNCTION(BlueprintCallable, meta = (DisplayName = "Toggle eye swap"), Category = "DisplayCluster|Render")
 	virtual bool ToggleEyesSwap() = 0;
 
-	UFUNCTION(BlueprintCallable, meta = (DisplayName = "Set output flip"), Category = "DisplayCluster|Render")
-	virtual void SetOutputFlip(bool flipH, bool flipV) = 0;
-
-	UFUNCTION(BlueprintCallable, meta = (DisplayName = "Get output flip"), Category = "DisplayCluster|Render")
-	virtual void GetOutputFlip(bool& flipH, bool& flipV) = 0;
-
 	/** Return near and far plane clip plane distances. */
 	UFUNCTION(BlueprintCallable, meta = (DisplayName = "Get near and far clipping distance"), Category = "DisplayCluster|Render")
 	virtual void GetCullingDistance(float& NearClipPlane, float& FarClipPlane) = 0;
diff --git a/Engine/Plugins/Runtime/nDisplay/Source/DisplayCluster/Public/Config/DisplayClusterConfigTypes.h b/Engine/Plugins/Runtime/nDisplay/Source/DisplayCluster/Public/Config/DisplayClusterConfigTypes.h
--- a/Engine/Plugins/Runtime/nDisplay/Source/DisplayCluster/Public/Config/DisplayClusterConfigTypes.h
+++ b/Engine/Plugins/Runtime/nDisplay/Source/DisplayCluster/Public/Config/DisplayClusterConfigTypes.h
@@ -40,6 +40,7 @@ struct FDisplayClusterConfigClusterNode : public FDisplayClusterConfigBase
 	int32   Port_CS = -1;
 	int32   Port_SS = -1;
 	bool    SoundEnabled = false;
+	bool  EyeSwap = false;
 
 	virtual FString ToString() const override;
 	virtual bool    DeserializeFromString(const FString& line) override;
@@ -53,8 +54,6 @@ struct FDisplayClusterConfigViewport : public FDisplayClusterConfigBase
 	FString   Id;
 	FIntPoint Loc  = FIntPoint::ZeroValue;
 	FIntPoint Size = FIntPoint::ZeroValue;
-	bool FlipHorizontal = false;
-	bool FlipVertical = false;
 
 	virtual FString ToString() const override;
 	virtual bool    DeserializeFromString(const FString& line) override;
@@ -137,7 +136,6 @@ struct FDisplayClusterConfigRender : public FDisplayClusterConfigBase
 //////////////////////////////////////////////////////////////////////////////////////////////
 struct FDisplayClusterConfigStereo : public FDisplayClusterConfigBase
 {
-	bool  EyeSwap = false;
 	float EyeDist = 0.064f;
 
 	virtual FString ToString() const override;
diff --git a/Engine/Plugins/Runtime/nDisplay/Source/DisplayCluster/Public/Render/IDisplayClusterStereoDevice.h b/Engine/Plugins/Runtime/nDisplay/Source/DisplayCluster/Public/Render/IDisplayClusterStereoDevice.h
--- a/Engine/Plugins/Runtime/nDisplay/Source/DisplayCluster/Public/Render/IDisplayClusterStereoDevice.h
+++ b/Engine/Plugins/Runtime/nDisplay/Source/DisplayCluster/Public/Render/IDisplayClusterStereoDevice.h
@@ -80,22 +80,6 @@ struct IDisplayClusterStereoDevice
 	*/
 	virtual bool ToggleEyesSwap() = 0;
 
-	/**
-	* Configures output flipping
-	*
-	* @param flipH - enable horizontal output flip
-	* @param flipV - enable vertical output flip
-	*/
-	virtual void SetOutputFlip(bool flipH, bool flipV) = 0;
-
-	/**
-	* Returns current output flip settings
-	*
-	* @param flipH - (out) current horizontal output flip
-	* @param flipV - (out) current vertical output flip
-	*/
-	virtual void GetOutputFlip(bool& flipH, bool& flipV) const = 0;
-
 	/**
 	* Set swap synchronization policy
 	*
diff --git a/Engine/Source/Programs/nDisplayLauncher/Runner/Runner.StartApp.cs b/Engine/Source/Programs/nDisplayLauncher/Runner/Runner.StartApp.cs
--- a/Engine/Source/Programs/nDisplayLauncher/Runner/Runner.StartApp.cs
+++ b/Engine/Source/Programs/nDisplayLauncher/Runner/Runner.StartApp.cs
@@ -51,10 +51,17 @@ namespace nDisplayLauncher
 		{
 			string commandCmd;
 
-			// Executable and mandatory arguments
-			commandCmd = string.Format("{0} {1} {2}", CommandStartApp, SelectedApplication, ArgMandatory);
+			// Executable
+			commandCmd = string.Format("{0} \"{1}\"", CommandStartApp, SelectedApplication);
+			// Custom common arguments
+			if (!string.IsNullOrWhiteSpace(CustomCommonParams))
+			{
+				commandCmd = string.Format("{0} {1}", commandCmd, CustomCommonParams.Trim());
+			}
+			// Mandatory arguments
+			commandCmd = string.Format("{0} {1}", commandCmd, ArgMandatory);
 			// Config file
-			commandCmd = string.Format("{0} {1}={2}", commandCmd, ArgConfig, SelectedConfig);
+			commandCmd = string.Format("{0} {1}=\"{2}\"", commandCmd, ArgConfig, SelectedConfig);
 			// Render API and mode
 			commandCmd = string.Format("{0} {1} {2}", commandCmd, SelectedRenderApiParam.Value, SelectedRenderModeParam.Value);
 			// Camera
@@ -73,15 +80,15 @@ namespace nDisplayLauncher
 			{
 				commandCmd = string.Format("{0} {1}", commandCmd, ArgUseAllAvailableCores);
 			}
-			// Custom common arguments
-			if (!string.IsNullOrWhiteSpace(CustomCommonParams))
-			{
-				commandCmd = string.Format("{0} {1}", commandCmd, CustomCommonParams.Trim());
-			}
 			// Fullscreen/windowed
 			commandCmd = string.Format("{0} {1}", commandCmd, Node.isWindowed ? ArgWindowed : ArgFullscreen);
-			// Window location and size
-			commandCmd = string.Format("{0} WinX={1} WinY={2} ResX={3} ResY={4}", commandCmd, Node.winX, Node.winY, Node.resX, Node.resY);
+			// Window mode, location and size
+			commandCmd = string.Format("{0} {1} {2} {3} {4}", 
+				commandCmd, 
+				string.IsNullOrEmpty(Node.winX) ? string.Empty : string.Format("WinX={0}", Node.winX),
+				string.IsNullOrEmpty(Node.winY) ? string.Empty : string.Format("WinY={0}", Node.winY),
+				string.IsNullOrEmpty(Node.resX) ? string.Empty : string.Format("ResX={0}", Node.resX),
+				string.IsNullOrEmpty(Node.resY) ? string.Empty : string.Format("ResY={0}", Node.resY));
 			// Node ID
 			commandCmd = string.Format("{0} {1}={2}", commandCmd, ArgNode, Node.id);
 
diff --git a/Templates/TP_nDisplay/Config/DefaultEngine.ini b/Templates/TP_nDisplay/Config/DefaultEngine.ini
--- a/Templates/TP_nDisplay/Config/DefaultEngine.ini
+++ b/Templates/TP_nDisplay/Config/DefaultEngine.ini
@@ -78,3 +78,7 @@ ReverbPlugin=
 OcclusionPlugin=
 
 
+[/Script/Engine.RendererSettings]
+r.DefaultFeature.Bloom=False
+r.DefaultFeature.AutoExposure=False
+r.DefaultFeature.MotionBlur=False
