nDisplay plugin is still in Beta in prevision of API changes

diff --git a/Engine/Plugins/Runtime/nDisplay/nDisplay.uplugin b/Engine/Plugins/Runtime/nDisplay/nDisplay.uplugin
--- a/Engine/Plugins/Runtime/nDisplay/nDisplay.uplugin
+++ b/Engine/Plugins/Runtime/nDisplay/nDisplay.uplugin
@@ -12,7 +12,7 @@
 	"SupportURL": "",
 	"EnabledByDefault": false,
 	"CanContainContent": true,
-	"IsBetaVersion": false,
+	"IsBetaVersion": true,
 	"Installed": false,
 	"Modules": [
 		{
