Fix nDisplay init

diff --git a/Engine/Plugins/Runtime/nDisplay/Source/DisplayCluster/Private/Render/Devices/DisplayClusterDeviceBase.cpp b/Engine/Plugins/Runtime/nDisplay/Source/DisplayCluster/Private/Render/Devices/DisplayClusterDeviceBase.cpp
--- a/Engine/Plugins/Runtime/nDisplay/Source/DisplayCluster/Private/Render/Devices/DisplayClusterDeviceBase.cpp
+++ b/Engine/Plugins/Runtime/nDisplay/Source/DisplayCluster/Private/Render/Devices/DisplayClusterDeviceBase.cpp
@@ -39,10 +39,6 @@ bool FDisplayClusterDeviceBase::Initialize()
 		return false;
 	}
 
-	{
-		return false;
-	}
-
 	UE_LOG(LogDisplayClusterRender, Log, TEXT("Use swap interval: %d"), SwapInterval);
 
 	return true;
