Fix compile errors with unity build disabled

diff --git a/Engine/Plugins/Runtime/nDisplay/Source/DisplayCluster/Private/Input/Devices/DisplayClusterInputDeviceBase.h b/Engine/Plugins/Runtime/nDisplay/Source/DisplayCluster/Private/Input/Devices/DisplayClusterInputDeviceBase.h
--- a/Engine/Plugins/Runtime/nDisplay/Source/DisplayCluster/Private/Input/Devices/DisplayClusterInputDeviceBase.h
+++ b/Engine/Plugins/Runtime/nDisplay/Source/DisplayCluster/Private/Input/Devices/DisplayClusterInputDeviceBase.h
@@ -4,6 +4,7 @@
 
 #include "IDisplayClusterInputDevice.h"
 #include "DisplayClusterInputDeviceTraits.h"
+#include "Misc/DisplayClusterLog.h"
 
 #include "CoreMinimal.h"
 
diff --git a/Engine/Source/Developer/StandaloneRenderer/Private/Windows/D3D/SlateD3DConstantBuffer.h b/Engine/Source/Developer/StandaloneRenderer/Private/Windows/D3D/SlateD3DConstantBuffer.h
--- a/Engine/Source/Developer/StandaloneRenderer/Private/Windows/D3D/SlateD3DConstantBuffer.h
+++ b/Engine/Source/Developer/StandaloneRenderer/Private/Windows/D3D/SlateD3DConstantBuffer.h
@@ -5,6 +5,7 @@
 #include "CoreMinimal.h"
 #include "Templates/RefCounting.h"
 #include "StandaloneRendererPlatformHeaders.h"
+#include "Windows/D3D/SlateD3DRenderer.h"
 
 template<typename BufferType>
 class FSlateD3DConstantBuffer
