Fix for TEnumAsByte warning.

diff --git a/Engine/Plugins/Runtime/nDisplay/Source/DisplayClusterInput/Private/Blueprints/DisplayClusterInputBlueprintAPIImpl.cpp b/Engine/Plugins/Runtime/nDisplay/Source/DisplayClusterInput/Private/Blueprints/DisplayClusterInputBlueprintAPIImpl.cpp
--- a/Engine/Plugins/Runtime/nDisplay/Source/DisplayClusterInput/Private/Blueprints/DisplayClusterInputBlueprintAPIImpl.cpp
+++ b/Engine/Plugins/Runtime/nDisplay/Source/DisplayClusterInput/Private/Blueprints/DisplayClusterInputBlueprintAPIImpl.cpp
@@ -64,6 +64,6 @@ bool UDisplayClusterInputBlueprintAPIImpl::BindVrpnTracker(const FString& VrpnDe
 {
 	DISPLAY_CLUSTER_FUNC_TRACE(LogDisplayClusterInputBP);
 
-	const FString TargetString = UEnum::GetValueAsString<EControllerHand>(Target);
+	const FString TargetString = UEnum::GetValueAsString(Target);
 	return IDisplayClusterInputModule::Get().BindVrpnChannel(VrpnDeviceId, VrpnChannel, TargetString);
 }
