nDisplay features
- warp & blend support
- in camera pixel
- pfm generation
- mpcdi support

diff --git a/Engine/Plugins/Runtime/WarpUtils/Source/PFMExporter/Private/PFMExporterLog.cpp b/Engine/Plugins/Runtime/WarpUtils/Source/PFMExporter/Private/PFMExporterLog.cpp
new file mode 100644
--- /dev/null
+++ b/Engine/Plugins/Runtime/WarpUtils/Source/PFMExporter/Private/PFMExporterLog.cpp
@@ -0,0 +1,5 @@
+// Copyright 1998-2019 Epic Games, Inc. All Rights Reserved.
+
+#include "PFMExporterLog.h"
+
+DEFINE_LOG_CATEGORY(LogPFMExporter);
diff --git a/Engine/Plugins/Runtime/WarpUtils/Source/WarpUtils/Private/WarpUtilsLog.cpp b/Engine/Plugins/Runtime/WarpUtils/Source/WarpUtils/Private/WarpUtilsLog.cpp
new file mode 100644
--- /dev/null
+++ b/Engine/Plugins/Runtime/WarpUtils/Source/WarpUtils/Private/WarpUtilsLog.cpp
@@ -0,0 +1,6 @@
+// Copyright 1998-2019 Epic Games, Inc. All Rights Reserved.
+
+#include "WarpUtilsLog.h"
+
+// Plugin-wide log categories
+DEFINE_LOG_CATEGORY(LogWarpUtilsBlueprint);
diff --git a/Engine/Plugins/Runtime/WarpUtils/Source/WarpUtils/Private/WarpUtilsLog.h b/Engine/Plugins/Runtime/WarpUtils/Source/WarpUtils/Private/WarpUtilsLog.h
new file mode 100644
--- /dev/null
+++ b/Engine/Plugins/Runtime/WarpUtils/Source/WarpUtils/Private/WarpUtilsLog.h
@@ -0,0 +1,7 @@
+// Copyright 1998-2019 Epic Games, Inc. All Rights Reserved.
+
+#pragma once
+
+#include "CoreMinimal.h"
+
+DECLARE_LOG_CATEGORY_EXTERN(LogWarpUtilsBlueprint, Log, All);
diff --git a/Engine/Plugins/Runtime/WarpUtils/Source/WarpUtils/WarpUtils.Build.cs b/Engine/Plugins/Runtime/WarpUtils/Source/WarpUtils/WarpUtils.Build.cs
new file mode 100644
--- /dev/null
+++ b/Engine/Plugins/Runtime/WarpUtils/Source/WarpUtils/WarpUtils.Build.cs
@@ -0,0 +1,32 @@
+// Copyright 1998-2019 Epic Games, Inc. All Rights Reserved.
+
+using UnrealBuildTool;
+
+public class WarpUtils : ModuleRules
+{
+	public WarpUtils(ReadOnlyTargetRules Target) : base(Target)
+	{
+		PCHUsage = ModuleRules.PCHUsageMode.UseExplicitOrSharedPCHs;
+		
+		PrivateIncludePaths.AddRange(
+			new string[] {
+				"WarpUtils/Private"
+			});
+			
+		
+		PublicDependencyModuleNames.AddRange(
+			new string[]
+			{
+				"Core",
+			});
+		
+		PrivateDependencyModuleNames.AddRange(
+			new string[]
+			{
+				"CoreUObject",
+				"Engine",
+				"Slate",
+				"SlateCore",
+			});
+	}
+}
diff --git a/Engine/Plugins/Runtime/nDisplay/Shaders/Private/MPCDIOverlayShaders.usf b/Engine/Plugins/Runtime/nDisplay/Shaders/Private/MPCDIOverlayShaders.usf
--- a/Engine/Plugins/Runtime/nDisplay/Shaders/Private/MPCDIOverlayShaders.usf
+++ b/Engine/Plugins/Runtime/nDisplay/Shaders/Private/MPCDIOverlayShaders.usf
@@ -1,6 +1,5 @@
 // Copyright 1998-2019 Epic Games, Inc. All Rights Reserved.
-#include "/Engine/Private/GammaCorrectionCommon.ush"
-
+#include "MPCDIUtils.ush"
 
 //Draw rectangle support
 float4 DrawRectanglePosScaleBias;
@@ -52,46 +51,16 @@ struct OutputPS
 
 float3 ApplyPicpTextureGammaCorrection(float3 Color)
 {
-	float3 OutColor = ApplyGammaCorrection(Color, 2.2);
-	       OutColor = LinearToSrgb(OutColor);
+	//float3 OutColor = ApplyGammaCorrection(Color, 2.2);
+	       //OutColor = LinearToSrgb(OutColor);
 
-	return OutColor;
+	return Color;
 }
 
 float4 ApplyPicpTextureGammaCorrection(float4 Color)
 {
-	return float4(ApplyPicpTextureGammaCorrection(Color.xyz), Color.w);
-}
-
-
-float sRGBToLinearMPCDI(float srgb)
-{
-	if (srgb <= 0.04045)
-		return srgb / 12.92;
-	else
-		return pow((srgb + 0.055) / 1.055, 2.4);
-}
-
-float MPCDILinarTosRGB(float lin)
-{
-	if (lin <= 0.0031308)
-		return lin * 12.92;
-	else
-		return 1.055 * pow(lin, 1.0 / 2.4) - 0.055;
-}
-
-float3 sRGBToLinearMPCDI(float3 srgb)
-{
-	return float3(sRGBToLinearMPCDI(srgb.r),
-				  sRGBToLinearMPCDI(srgb.g),
-				  sRGBToLinearMPCDI(srgb.b));
-}
-
-float3 MPCDILinearTosRGB(float3 lin)
-{
-	return float3(MPCDILinarTosRGB(lin.r),
-				  MPCDILinarTosRGB(lin.g),
-				  MPCDILinarTosRGB(lin.b));
+	return Color;
+	//return float4(ApplyPicpTextureGammaCorrection(Color.xyz), Color.w);
 }
 
 /** Used for calculating vertex positions and UVs when drawing with DrawRectangle */
@@ -105,7 +74,7 @@ void DrawRectangle(in float4 InPosition, in float2 InTexCoord, out float4 OutPos
 
 float4 ApplyBlending(float3 color, OutputVS IN, float ColorAlphaValue)
 {
-	color = sRGBToLinearMPCDI(color);
+	color = MPCDI_sRGBToLinear(color);
 
 	// Load alpha
 	float alpha = pow(AlphaMap.Sample(AlphaMapSampler, IN.UV.xy).r, AlphaEmbeddedGamma);
@@ -115,7 +84,7 @@ float4 ApplyBlending(float3 color, OutputVS IN, float ColorAlphaValue)
 
 	color = (color*alpha*(1 - beta)) + beta;
 
-	return float4(MPCDILinearTosRGB(color), ColorAlphaValue);
+	return float4(MPCDI_LinearTosRGB(color), ColorAlphaValue);
 }
 
 
@@ -128,7 +97,7 @@ OutputVS MainVS( InputVS IN)
 	return OUT;
 }
 
-OutputPS PassThrought_PS( OutputVS IN )
+OutputPS Passthrough_PS( OutputVS IN )
 {
 	OutputPS OUT;
 	float4 UV = mul(IN.UV, MultiViewportMatrix);
@@ -149,11 +118,17 @@ OutputPS WarpAndBlend_PS( OutputVS IN )
 	// Load from color map
 	float3 Color = PostprocessInput0.Sample(PostprocessInput0Sampler, UV.xy/UV.w).rgb;
 
+#if RENDER_CASE_APPLYBLENDING
 	// Apply blending
-	OUT.Color = ApplyBlending(Color, IN, 1.0f);
+	OUT.Color = ApplyBlending(Color, IN, 1.f);
+#else
+	// No blending
+	OUT.Color = float4(Color, 1.f);
+#endif
 	return OUT;
 }
 
+
 float Pow2(float X)
 {
 	return X * X;
@@ -165,8 +140,6 @@ float2       CameraSoftEdge;
 float4x4     CameraProjectionMatrix;
 Texture2D    CameraMap;
 SamplerState CameraMapSampler;
-//Texture2D    CameraOverlayMap;
-//SamplerState CameraOverlayMapSampler;
 
 Texture2D    ViewportOverlayMap;
 SamplerState ViewportOverlayMapSampler;
@@ -205,7 +178,6 @@ OutputPS PicpMPCDIWarpAndBlendOverlay_PS(OutputVS IN)
 	float2 CameraUV = CameraUVW.xy / CameraUVW.w;
 	// Load from Cameramap
 	float3 CameraBaseColor        = ApplyPicpTextureGammaCorrection(CameraMap.Sample(CameraMapSampler, CameraUV).rgb);
-	//float4 CameraOverlayColor = CameraOverlayMap.Sample(CameraOverlayMapSampler, CameraUV).rgb;
 
 	float4 CameraColor = float4(CameraBaseColor, 0);
 
@@ -213,8 +185,7 @@ OutputPS PicpMPCDIWarpAndBlendOverlay_PS(OutputVS IN)
 	{//Defined texel:
 		float2 ToEdge = (CameraUV.xy * 2) - 1.0f; // -1..1
 		float Weight = 1 - max(abs(ToEdge.x), abs(ToEdge.y));
-		// Near clip Plane tests
-		//! UV.z <= UV.w &&
+		// Near clip Plane tests		
 		if (Weight >= 0)
 		{
 			float2 ToSoftEdge = saturate((abs(ToEdge) - 1.0f + CameraSoftEdge) / CameraSoftEdge);
@@ -246,9 +217,13 @@ OutputPS PicpMPCDIWarpAndBlendOverlay_PS(OutputVS IN)
 	//!fixme: add LUT here
 
 
-
+#if RENDER_CASE_APPLYBLENDING
 	// Apply final mpcdi color blending
 	OUT.Color = ApplyBlending(OutColor.xyz, IN, OutColor.w);
+#else
+	// No blending
+	OUT.Color = OutColor;
+#endif
 	return OUT;
 }
 
@@ -269,9 +244,13 @@ OutputPS PicpMPCDIWarpAndBlendViewportOverlay_PS(OutputVS IN)
 	//!fixme: add LUT here
 
 
-
+#if RENDER_CASE_APPLYBLENDING
 	// Apply final mpcdi color blending
 	OUT.Color = ApplyBlending(ViewportOverlayColor.xyz, IN, ViewportOverlayColor.w);
+#else
+	// No blending
+	OUT.Color = ViewportOverlayColor;
+#endif
 	return OUT;
 }
 
@@ -324,10 +303,6 @@ OutputPS PicpMPCDIWarpAndBlendOneCameraFull_PS(OutputVS IN)
 	}
 
 	float4 OutColor = ViewportBaseColor;
-
-	//Color correction
-	//OutColor.xyz = ApplyGammaCorrection(sRGBToLinear(OutColor.xyz), 1);
-
 	
 #if USE_OVERLAY_BACK	
 	OutColor.xyz = lerp(OutColor.xyz, ApplyPicpTextureGammaCorrection(ViewportOverlayBackColor.xyz), saturate(1.0f - ViewportOverlayBackColor.w));
@@ -342,9 +317,14 @@ OutputPS PicpMPCDIWarpAndBlendOneCameraFull_PS(OutputVS IN)
 	//!fixme: add LUT here
 
 
-
+#if RENDER_CASE_APPLYBLENDING
 	// Apply final mpcdi color blending
 	OUT.Color = ApplyBlending(OutColor.xyz, IN, OutColor.w);
+#else
+	// No blending
+	OUT.Color = OutColor;
+#endif
+
 	return OUT;
 }
 
diff --git a/Engine/Plugins/Runtime/nDisplay/Shaders/Private/MPCDIShaders.usf b/Engine/Plugins/Runtime/nDisplay/Shaders/Private/MPCDIShaders.usf
--- a/Engine/Plugins/Runtime/nDisplay/Shaders/Private/MPCDIShaders.usf
+++ b/Engine/Plugins/Runtime/nDisplay/Shaders/Private/MPCDIShaders.usf
@@ -1,4 +1,6 @@
 // Copyright 1998-2019 Epic Games, Inc. All Rights Reserved.
+#include "/Engine/Public/Platform.ush"
+#include "MPCDIUtils.ush"
 
 //Draw rectangle support
 float4 DrawRectanglePosScaleBias;
@@ -46,36 +48,6 @@ struct OutputPS
 	float4 Color : SV_Target0;
 };
 
-float sRGBToLinear(float srgb)
-{
-	if (srgb <= 0.04045)
-		return srgb / 12.92;
-	else
-		return pow((srgb + 0.055) / 1.055, 2.4);
-}
-
-float LinarTosRGB(float lin)
-{
-	if (lin <= 0.0031308)
-		return lin * 12.92;
-	else
-		return 1.055 * pow(lin, 1.0 / 2.4) - 0.055;
-}
-
-float3 sRGBToLinear(float3 srgb)
-{
-	return float3(sRGBToLinear(srgb.r), 
-				  sRGBToLinear(srgb.g),
-				  sRGBToLinear(srgb.b));
-}
-
-float3 LinearTosRGB(float3 lin)
-{
-	return float3(LinarTosRGB(lin.r), 
-				  LinarTosRGB(lin.g),
-				  LinarTosRGB(lin.b));
-}
-
 /** Used for calculating vertex positions and UVs when drawing with DrawRectangle */
 void DrawRectangle(in float4 InPosition, in float2 InTexCoord, out float4 OutPosition, out float2 OutTexCoord)
 {
@@ -87,7 +59,7 @@ void DrawRectangle(in float4 InPosition, in float2 InTexCoord, out float4 OutPos
 
 float4 ApplyBlending(float3 color, OutputVS IN)
 {
-	color = sRGBToLinear(color);
+	color = MPCDI_sRGBToLinear(color);
 
 	// Load alpha
 	float alpha = pow(AlphaMap.Sample(AlphaMapSampler, IN.UV.xy).r, AlphaEmbeddedGamma);
@@ -97,7 +69,7 @@ float4 ApplyBlending(float3 color, OutputVS IN)
 
 	color = (color*alpha*(1 - beta)) + beta;
 
-	return float4(LinearTosRGB(color), 1.f);
+	return float4(MPCDI_LinearTosRGB(color), 1.f);
 }
 
 OutputVS MainVS( InputVS IN)
@@ -109,7 +81,7 @@ OutputVS MainVS( InputVS IN)
 	return OUT;
 }
 
-OutputPS PassThrought_PS( OutputVS IN )
+OutputPS Passthrough_PS( OutputVS IN )
 {
 	OutputPS OUT;
 	float4 UV = mul(IN.UV,TextureProjectionMatrix);
@@ -129,6 +101,7 @@ OutputPS ShowWarpTexture_PS(OutputVS IN)
 	OUT.Color = float4(UV, 0, 0.5f);
 	return OUT;
 }
+
 OutputPS WarpAndBlend_PS( OutputVS IN )
 {
 	OutputPS OUT;
@@ -142,8 +115,14 @@ OutputPS WarpAndBlend_PS( OutputVS IN )
 	// Load from color map
 	float3 Color = PostprocessInput0.Sample(PostprocessInput0Sampler, UV.xy/UV.w).rgb;
 
+#if RENDER_CASE_APPLYBLENDING
 	// Apply blending
 	OUT.Color = ApplyBlending(Color, IN);
+#else
+	// No blending
+	OUT.Color = float4(Color,1.f);
+#endif
+
 	return OUT;
 }
 
@@ -155,7 +134,7 @@ OutputPS WarpAndBlendCube_PS(OutputVS IN)
 	float4 WarpedPos = WarpMap.Sample(WarpMapSampler, IN.UV.xy);
 
 	// Construct world space Vector from mesh position to eye position
-	float3 ViewVector = normalize((WarpedPos - EyePosition).xyz);
+	float3 ViewVector = normalize(WarpedPos.xyz - EyePosition);
 
 	// Load from color map
 	float3 Color = SceneCubemap.Sample(SceneCubemapSampler, ViewVector).rgb;
@@ -187,3 +166,17 @@ OutputPS BuildProjectedAABB_PS(OutputVS IN)
 	OUT.Color = float4(UV.xyz, 1.f);		
 	return OUT;
 }
+
+
+OutputVS DirectProjectionVS(InputVS IN)
+{
+	float2 Pos = IN.Position.xy;
+
+	Pos.xy = Pos.xy * 2 - 1.0f;
+	Pos.xy *= float2(1, -1);
+
+	OutputVS Out;
+	Out.Position = float4(Pos, 0, 1);
+	Out.UV = float4(IN.UV, 0.f, 1.f);
+	return Out;
+}
diff --git a/Engine/Plugins/Runtime/nDisplay/Source/DisplayCluster/DisplayCluster.Build.cs b/Engine/Plugins/Runtime/nDisplay/Source/DisplayCluster/DisplayCluster.Build.cs
--- a/Engine/Plugins/Runtime/nDisplay/Source/DisplayCluster/DisplayCluster.Build.cs
+++ b/Engine/Plugins/Runtime/nDisplay/Source/DisplayCluster/DisplayCluster.Build.cs
@@ -12,11 +12,6 @@ public class DisplayCluster : ModuleRules
 		PrivateIncludePaths.AddRange(
 			new string[] {
 				"DisplayCluster/Private",
-				"../../../../../Engine/Source/Runtime/Renderer/Private",
-				"../../../../../Engine/Source/Runtime/Windows/D3D11RHI/Private",
-				"../../../../../Engine/Source/Runtime/Windows/D3D11RHI/Private/Windows",
-				"../../../../../Engine/Source/Runtime/D3D12RHI/Private",
-				"../../../../../Engine/Source/Runtime/D3D12RHI/Private/Windows"
 			});
 
 		PublicDependencyModuleNames.AddRange(
@@ -26,8 +21,6 @@ public class DisplayCluster : ModuleRules
 				"CoreUObject",
 				"Engine",
 				"InputCore",
-				"Json",
-				"JsonUtilities"
 			});
 
 		PrivateDependencyModuleNames.AddRange(
@@ -40,6 +33,8 @@ public class DisplayCluster : ModuleRules
 				"Engine",
 				"HeadMountedDisplay",
 				"InputCore",
+				"Json",
+				"JsonUtilities",
 				"Networking",
 				"OpenGLDrv",
 				"RHI",
@@ -59,7 +54,7 @@ public class DisplayCluster : ModuleRules
 		AddEngineThirdPartyPrivateStaticDependencies(Target, "DX11");
 		AddEngineThirdPartyPrivateStaticDependencies(Target, "DX12");
 		AddEngineThirdPartyPrivateStaticDependencies(Target, "NVAftermath");
-        AddEngineThirdPartyPrivateStaticDependencies(Target, "NVAPI");
+		AddEngineThirdPartyPrivateStaticDependencies(Target, "NVAPI");
 		AddEngineThirdPartyPrivateStaticDependencies(Target, "IntelMetricsDiscovery");
 
 		// 3rd party dependencies
@@ -67,10 +62,9 @@ public class DisplayCluster : ModuleRules
 	}
 
 
-	public bool AddThirdPartyDependencies(ReadOnlyTargetRules ROTargetRules)
+	public void AddThirdPartyDependencies(ReadOnlyTargetRules ROTargetRules)
 	{
-		string ModulePath = Path.GetDirectoryName(UnrealBuildTool.RulesCompiler.GetFileNameFromType(GetType()));
-		string ThirdPartyPath = Path.GetFullPath(Path.Combine(ModulePath, "../../ThirdParty/"));
+		string ThirdPartyPath = Path.GetFullPath(Path.Combine(ModuleDirectory, "../../ThirdParty/"));
 
 		string PathLib = string.Empty;
 		string PathInc = string.Empty;
@@ -81,7 +75,5 @@ public class DisplayCluster : ModuleRules
 		PublicAdditionalLibraries.Add(Path.Combine(PathLib, "vrpn.lib"));
 		PublicAdditionalLibraries.Add(Path.Combine(PathLib, "quat.lib"));
 		PublicIncludePaths.Add(PathInc);
-
-		return true;
 	}
 }
diff --git a/Engine/Plugins/Runtime/nDisplay/Source/DisplayCluster/Private/Blueprints/DisplayClusterBlueprintAPIImpl.cpp b/Engine/Plugins/Runtime/nDisplay/Source/DisplayCluster/Private/Blueprints/DisplayClusterBlueprintAPIImpl.cpp
--- a/Engine/Plugins/Runtime/nDisplay/Source/DisplayCluster/Private/Blueprints/DisplayClusterBlueprintAPIImpl.cpp
+++ b/Engine/Plugins/Runtime/nDisplay/Source/DisplayCluster/Private/Blueprints/DisplayClusterBlueprintAPIImpl.cpp
@@ -12,6 +12,8 @@
 #include "Input/IDisplayClusterInputManager.h"
 #include "Render/IDisplayClusterRenderManager.h"
 
+#include "Config/DisplayClusterConfigTypes.h"
+#include "Misc/DisplayClusterHelpers.h"
 
 
 //////////////////////////////////////////////////////////////////////////////////////////////
@@ -136,7 +138,27 @@ void UDisplayClusterBlueprintAPIImpl::EmitClusterEvent(const FDisplayClusterClus
 //////////////////////////////////////////////////////////////////////////////////////////////
 // Config API
 //////////////////////////////////////////////////////////////////////////////////////////////
+void UDisplayClusterBlueprintAPIImpl::GetViewports(bool IsRTT, TArray<FString>& ViewportIDs, TArray<FString>& ViewportTypes, TArray<FIntPoint>& ViewportLocations, TArray<FIntPoint>& ViewportSizes)
+{
+	DISPLAY_CLUSTER_FUNC_TRACE(LogDisplayClusterBlueprint);
 
+	TArray<FDisplayClusterConfigViewport> SelectedViewports = DisplayClusterHelpers::config::GetLocalViewports().FilterByPredicate([IsRTT](const FDisplayClusterConfigViewport& Item)
+	{
+		return Item.IsRTT == IsRTT;
+	});
+
+	for (const FDisplayClusterConfigViewport& Item : SelectedViewports)
+	{
+		FDisplayClusterConfigProjection CfgProjection;
+		if (IDisplayCluster::Get().GetConfigMgr()->GetProjection(Item.ProjectionId, CfgProjection))
+		{
+			ViewportIDs.Add(Item.Id);
+			ViewportTypes.Add(CfgProjection.Type);
+			ViewportLocations.Add(Item.Loc);
+			ViewportSizes.Add(Item.Size);
+		}
+	}	
+}
 
 
 //////////////////////////////////////////////////////////////////////////////////////////////
@@ -556,12 +578,43 @@ void UDisplayClusterBlueprintAPIImpl::SetViewportCamera(const FString& InCameraI
 	IDisplayClusterRenderManager* const Manager = IDisplayCluster::Get().GetRenderMgr();
 	if (Manager)
 	{
-		return Manager->SetViewportCamera(InCameraId, InViewportId);
+		Manager->SetViewportCamera(InCameraId, InViewportId);
 	}
 
 	return;
 }
 
+void UDisplayClusterBlueprintAPIImpl::SetCustomPostProcessing(const FString& ViewportID, const FPostProcessSettings& PostProcessingSettings)
+{
+	DISPLAY_CLUSTER_FUNC_TRACE(LogDisplayClusterBlueprint);
+
+	IDisplayClusterRenderManager* const Manager = IDisplayCluster::Get().GetRenderMgr();
+	if (Manager)
+	{
+		Manager->SetCustomPostProcessing(ViewportID, PostProcessingSettings);
+	}
+}
+
+bool UDisplayClusterBlueprintAPIImpl::GetViewportRect(const FString& ViewportID, FIntPoint& ViewportLoc, FIntPoint& ViewportSize)
+{
+	DISPLAY_CLUSTER_FUNC_TRACE(LogDisplayClusterBlueprint);
+
+	IDisplayClusterRenderManager* const Manager = IDisplayCluster::Get().GetRenderMgr();
+	if (Manager)
+	{
+		FIntRect ViewportRect;
+		if (Manager->GetViewportRect(ViewportID, ViewportRect))
+		{
+			ViewportLoc  = ViewportRect.Min;
+			ViewportSize = ViewportRect.Size();
+			return true;
+		}
+	}
+
+	return false;
+}
+
+
 float UDisplayClusterBlueprintAPIImpl::GetInterpupillaryDistance(const FString& CameraId)
 {
 	DISPLAY_CLUSTER_FUNC_TRACE(LogDisplayClusterBlueprint);
diff --git a/Engine/Plugins/Runtime/nDisplay/Source/DisplayCluster/Private/Blueprints/DisplayClusterBlueprintAPIImpl.h b/Engine/Plugins/Runtime/nDisplay/Source/DisplayCluster/Private/Blueprints/DisplayClusterBlueprintAPIImpl.h
--- a/Engine/Plugins/Runtime/nDisplay/Source/DisplayCluster/Private/Blueprints/DisplayClusterBlueprintAPIImpl.h
+++ b/Engine/Plugins/Runtime/nDisplay/Source/DisplayCluster/Private/Blueprints/DisplayClusterBlueprintAPIImpl.h
@@ -7,6 +7,7 @@
 #include "DisplayClusterBlueprintAPIImpl.generated.h"
 
 struct FDisplayClusterClusterEvent;
+struct FPostProcessSettings;
 
 
 /**
@@ -68,7 +69,8 @@ public:
 	//////////////////////////////////////////////////////////////////////////////////////////////
 	// Config API
 	//////////////////////////////////////////////////////////////////////////////////////////////
-
+	UFUNCTION(BlueprintCallable, meta = (DisplayName = "Get viewports"), Category = "DisplayCluster|Config")
+	virtual void GetViewports(bool IsRTT, TArray<FString>& ViewportIDs, TArray<FString>& ViewportTypes, TArray<FIntPoint>& ViewportLocations, TArray<FIntPoint>& ViewportSizes) override;
 
 public:
 	//////////////////////////////////////////////////////////////////////////////////////////////
@@ -189,6 +191,12 @@ public:
 	//////////////////////////////////////////////////////////////////////////////////////////////
 	UFUNCTION(BlueprintCallable, meta = (DisplayName = "Set viewport camera"), Category = "DisplayCluster|Render")
 	virtual void SetViewportCamera(const FString& InCameraId, const FString& InViewportId) override;
+	
+	UFUNCTION(BlueprintCallable, meta = (DisplayName = "Set custom post processing settings for viewport"), Category = "DisplayCluster|Render")
+	virtual void SetCustomPostProcessing(const FString& ViewportID, const FPostProcessSettings& PostProcessingSettings) override;
+
+	UFUNCTION(BlueprintCallable, meta = (DisplayName = "Get viewport rectangle"), Category = "DisplayCluster|Render")
+	virtual bool GetViewportRect(const FString& ViewportID, FIntPoint& ViewportLoc, FIntPoint& ViewportSize) override;
 
 public:
 	//////////////////////////////////////////////////////////////////////////////////////////////
diff --git a/Engine/Plugins/Runtime/nDisplay/Source/DisplayCluster/Private/Config/Checker/DisplayClusterConfigChecker.cpp b/Engine/Plugins/Runtime/nDisplay/Source/DisplayCluster/Private/Config/Checker/DisplayClusterConfigChecker.cpp
--- a/Engine/Plugins/Runtime/nDisplay/Source/DisplayCluster/Private/Config/Checker/DisplayClusterConfigChecker.cpp
+++ b/Engine/Plugins/Runtime/nDisplay/Source/DisplayCluster/Private/Config/Checker/DisplayClusterConfigChecker.cpp
@@ -31,6 +31,10 @@ void FDisplayClusterConfigChecker::AddViewport(const FDisplayClusterConfigViewpo
 {
 }
 
+void FDisplayClusterConfigChecker::AddPostprocess(const FDisplayClusterConfigPostprocess& InCfgPostprocess)
+{
+}
+
 void FDisplayClusterConfigChecker::AddCamera(const FDisplayClusterConfigCamera& InCfgCamera)
 {
 }
diff --git a/Engine/Plugins/Runtime/nDisplay/Source/DisplayCluster/Private/Config/Checker/DisplayClusterConfigChecker.h b/Engine/Plugins/Runtime/nDisplay/Source/DisplayCluster/Private/Config/Checker/DisplayClusterConfigChecker.h
--- a/Engine/Plugins/Runtime/nDisplay/Source/DisplayCluster/Private/Config/Checker/DisplayClusterConfigChecker.h
+++ b/Engine/Plugins/Runtime/nDisplay/Source/DisplayCluster/Private/Config/Checker/DisplayClusterConfigChecker.h
@@ -22,6 +22,7 @@ public:
 	virtual void AddClusterNode (const FDisplayClusterConfigClusterNode& InCfgCNode) override;
 	virtual void AddScreen      (const FDisplayClusterConfigScreen& InCfgScreen) override;
 	virtual void AddViewport    (const FDisplayClusterConfigViewport& InCfgViewport) override;
+	virtual void AddPostprocess (const FDisplayClusterConfigPostprocess& InCfgPostprocess) override;
 	virtual void AddCamera      (const FDisplayClusterConfigCamera& InCfgCamera) override;
 	virtual void AddSceneNode   (const FDisplayClusterConfigSceneNode& InCfgSNode) override;
 	virtual void AddGeneral     (const FDisplayClusterConfigGeneral& InCfgGeneral) override;
diff --git a/Engine/Plugins/Runtime/nDisplay/Source/DisplayCluster/Private/Config/DisplayClusterConfigManager.cpp b/Engine/Plugins/Runtime/nDisplay/Source/DisplayCluster/Private/Config/DisplayClusterConfigManager.cpp
--- a/Engine/Plugins/Runtime/nDisplay/Source/DisplayCluster/Private/Config/DisplayClusterConfigManager.cpp
+++ b/Engine/Plugins/Runtime/nDisplay/Source/DisplayCluster/Private/Config/DisplayClusterConfigManager.cpp
@@ -195,7 +195,7 @@ TArray<FDisplayClusterConfigViewport> FDisplayClusterConfigManager::GetViewports
 
 int32 FDisplayClusterConfigManager::GetViewportsAmount() const
 {
-	return static_cast<uint32>(CfgViewports.Num());
+	return CfgViewports.Num();
 }
 
 bool FDisplayClusterConfigManager::GetViewport(int32 idx, FDisplayClusterConfigViewport& viewport) const
@@ -209,6 +209,28 @@ bool FDisplayClusterConfigManager::GetViewport(const FString& id, FDisplayCluste
 }
 
 
+// Postprocess
+TArray<FDisplayClusterConfigPostprocess> FDisplayClusterConfigManager::GetPostprocess() const
+{
+	return CfgPostprocess;
+}
+
+int32 FDisplayClusterConfigManager::GetPostprocessAmount() const
+{
+	return CfgPostprocess.Num();
+}
+
+bool FDisplayClusterConfigManager::GetPostprocess(int32 idx, FDisplayClusterConfigPostprocess& postprocess) const
+{
+	return GetItem(CfgPostprocess, idx, postprocess, FString("GetPostprocess"));
+}
+
+bool FDisplayClusterConfigManager::GetPostprocess(const FString& id, FDisplayClusterConfigPostprocess& postprocess) const
+{
+	return GetItem(CfgPostprocess, id, postprocess, FString("GetPostprocess"));
+}
+
+
 // Scene nodes
 TArray<FDisplayClusterConfigSceneNode> FDisplayClusterConfigManager::GetSceneNodes() const
 {
@@ -217,7 +239,7 @@ TArray<FDisplayClusterConfigSceneNode> FDisplayClusterConfigManager::GetSceneNod
 
 int32 FDisplayClusterConfigManager::GetSceneNodesAmount() const
 {
-	return static_cast<uint32>(CfgSceneNodes.Num());
+	return CfgSceneNodes.Num();
 }
 
 bool FDisplayClusterConfigManager::GetSceneNode(int32 idx, FDisplayClusterConfigSceneNode& actor) const
@@ -311,6 +333,13 @@ void FDisplayClusterConfigManager::AddViewport(const FDisplayClusterConfigViewpo
 	CfgViewports.Add(InCfgViewport);
 }
 
+void FDisplayClusterConfigManager::AddPostprocess(const FDisplayClusterConfigPostprocess& InCfgPostprocess)
+{
+	DISPLAY_CLUSTER_FUNC_TRACE(LogDisplayClusterConfig);
+	UE_LOG(LogDisplayClusterConfig, Log, TEXT("Found postprocess: %s"), *InCfgPostprocess.ToString());
+	CfgPostprocess.Add(InCfgPostprocess);
+}
+
 void FDisplayClusterConfigManager::AddCamera(const FDisplayClusterConfigCamera& InCfgCamera)
 {
 	DISPLAY_CLUSTER_FUNC_TRACE(LogDisplayClusterConfig);
@@ -525,3 +554,29 @@ bool FDisplayClusterConfigManager::GetItem(const TArray<DataType>& container, co
 	item = *pFound;
 	return true;
 }
+
+FString FDisplayClusterConfigManager::GetFullPathToFile(const FString& FileName) const
+{
+	if (!FPaths::FileExists(FileName))
+	{
+		TArray<FString> OrderedBaseDirs;
+
+		//Add ordered search base dirs
+		OrderedBaseDirs.Add(FPaths::GetPath(ConfigPath));
+		OrderedBaseDirs.Add(FPaths::RootDir());
+
+		// Process base dirs in order:
+		for (auto It : OrderedBaseDirs)
+		{
+			FString FullPath = FPaths::ConvertRelativePathToFull(It, FileName);
+			if (FPaths::FileExists(FullPath))
+			{
+				return FullPath;
+			}
+		}
+
+		//@Handle error, file not found
+	}
+
+	return FileName;
+}
\ No newline at end of file
diff --git a/Engine/Plugins/Runtime/nDisplay/Source/DisplayCluster/Private/Config/DisplayClusterConfigManager.h b/Engine/Plugins/Runtime/nDisplay/Source/DisplayCluster/Private/Config/DisplayClusterConfigManager.h
--- a/Engine/Plugins/Runtime/nDisplay/Source/DisplayCluster/Private/Config/DisplayClusterConfigManager.h
+++ b/Engine/Plugins/Runtime/nDisplay/Source/DisplayCluster/Private/Config/DisplayClusterConfigManager.h
@@ -63,6 +63,11 @@ public:
 	virtual bool GetViewport(int32 idx, FDisplayClusterConfigViewport& viewport) const override;
 	virtual bool GetViewport(const FString& id, FDisplayClusterConfigViewport& viewport) const override;
 
+	virtual TArray<FDisplayClusterConfigPostprocess> GetPostprocess() const override;
+	virtual int32 GetPostprocessAmount() const override;
+	virtual bool GetPostprocess(int32 idx, FDisplayClusterConfigPostprocess& postprocess) const override;
+	virtual bool GetPostprocess(const FString& id, FDisplayClusterConfigPostprocess& postprocess) const override;
+
 	virtual TArray<FDisplayClusterConfigSceneNode> GetSceneNodes() const override;
 	virtual int32 GetSceneNodesAmount() const override;
 	virtual bool GetSceneNode(int32 idx, FDisplayClusterConfigSceneNode& actor) const override;
@@ -97,6 +102,8 @@ public:
 	virtual FDisplayClusterConfigCustom  GetConfigCustom() const override
 	{ return CfgCustom; }
 
+	virtual FString GetFullPathToFile(const FString& FileName) const override;
+
 public:
 	//////////////////////////////////////////////////////////////////////////////////////////////
 	// IPDisplayClusterConfigManager
@@ -115,6 +122,7 @@ public:
 	virtual void AddWindow(const FDisplayClusterConfigWindow& InCfgWindow) override;
 	virtual void AddScreen(const FDisplayClusterConfigScreen& InCfgScreen) override;
 	virtual void AddViewport(const FDisplayClusterConfigViewport& InCfgViewport) override;
+	virtual void AddPostprocess(const FDisplayClusterConfigPostprocess& InCfgPostprocess) override;
 	virtual void AddCamera(const FDisplayClusterConfigCamera& InCfgCamera) override;
 	virtual void AddSceneNode(const FDisplayClusterConfigSceneNode& InCfgSNode)  override;
 	virtual void AddGeneral(const FDisplayClusterConfigGeneral& InCfgGeneral)  override;
@@ -156,6 +164,7 @@ private:
 	TArray<FDisplayClusterConfigWindow>      CfgWindows;
 	TArray<FDisplayClusterConfigScreen>      CfgScreens;
 	TArray<FDisplayClusterConfigViewport>    CfgViewports;
+	TArray<FDisplayClusterConfigPostprocess> CfgPostprocess;
 	TArray<FDisplayClusterConfigCamera>      CfgCameras;
 	TArray<FDisplayClusterConfigSceneNode>   CfgSceneNodes;
 	TArray<FDisplayClusterConfigInput>       CfgInputDevices;
diff --git a/Engine/Plugins/Runtime/nDisplay/Source/DisplayCluster/Private/Config/DisplayClusterConfigTypes.cpp b/Engine/Plugins/Runtime/nDisplay/Source/DisplayCluster/Private/Config/DisplayClusterConfigTypes.cpp
--- a/Engine/Plugins/Runtime/nDisplay/Source/DisplayCluster/Private/Config/DisplayClusterConfigTypes.cpp
+++ b/Engine/Plugins/Runtime/nDisplay/Source/DisplayCluster/Private/Config/DisplayClusterConfigTypes.cpp
@@ -65,27 +65,29 @@ bool FDisplayClusterConfigClusterNode::DeserializeFromString(const FString& line
 //////////////////////////////////////////////////////////////////////////////////////////////
 FString FDisplayClusterConfigWindow::ToString() const
 {
-	return FString::Printf(TEXT("[%s + %s=%s, %s=%s, %s=%s, %s=%d, %s=%d, %s=%d, %s=%d]"),
+	return FString::Printf(TEXT("[%s + %s=%s, %s=%s, %s=%s, %s=%s, %s=%d, %s=%d, %s=%d, %s=%d]"),
 		*FDisplayClusterConfigBase::ToString(),
-		DisplayClusterStrings::cfg::data::Id,                 *Id,
-		DisplayClusterStrings::cfg::data::window::Viewports,  *DisplayClusterHelpers::str::ArrayToStr(ViewportIds),
-		DisplayClusterStrings::cfg::data::window::Fullscreen, DisplayClusterHelpers::str::BoolToStr(IsFullscreen),
-		DisplayClusterStrings::cfg::data::window::WinX,       WinX,
-		DisplayClusterStrings::cfg::data::window::WinY,       WinY,
-		DisplayClusterStrings::cfg::data::window::ResX,       ResX,
-		DisplayClusterStrings::cfg::data::window::ResY,       ResY
+		DisplayClusterStrings::cfg::data::Id,                  *Id,
+		DisplayClusterStrings::cfg::data::window::Viewports,   *DisplayClusterHelpers::str::ArrayToStr(ViewportIds),
+		DisplayClusterStrings::cfg::data::window::Postprocess, *DisplayClusterHelpers::str::ArrayToStr(PostprocessIds),
+		DisplayClusterStrings::cfg::data::window::Fullscreen,  DisplayClusterHelpers::str::BoolToStr(IsFullscreen),
+		DisplayClusterStrings::cfg::data::window::WinX,        WinX,
+		DisplayClusterStrings::cfg::data::window::WinY,        WinY,
+		DisplayClusterStrings::cfg::data::window::ResX,        ResX,
+		DisplayClusterStrings::cfg::data::window::ResY,        ResY
 	);
 }
 
 bool FDisplayClusterConfigWindow::DeserializeFromString(const FString& line)
 {
-	DisplayClusterHelpers::str::ExtractValue(line, FString(DisplayClusterStrings::cfg::data::Id),                 Id);
-	DisplayClusterHelpers::str::ExtractArray(line, FString(DisplayClusterStrings::cfg::data::window::Viewports), FString(DisplayClusterStrings::strArrayValSeparator), ViewportIds);
-	DisplayClusterHelpers::str::ExtractValue(line, FString(DisplayClusterStrings::cfg::data::window::Fullscreen), IsFullscreen);
-	DisplayClusterHelpers::str::ExtractValue(line, FString(DisplayClusterStrings::cfg::data::window::WinX),       WinX);
-	DisplayClusterHelpers::str::ExtractValue(line, FString(DisplayClusterStrings::cfg::data::window::WinY),       WinY);
-	DisplayClusterHelpers::str::ExtractValue(line, FString(DisplayClusterStrings::cfg::data::window::ResX),       ResX);
-	DisplayClusterHelpers::str::ExtractValue(line, FString(DisplayClusterStrings::cfg::data::window::ResY),       ResY);
+	DisplayClusterHelpers::str::ExtractValue(line, FString(DisplayClusterStrings::cfg::data::Id),                  Id);
+	DisplayClusterHelpers::str::ExtractArray(line, FString(DisplayClusterStrings::cfg::data::window::Viewports),   FString(DisplayClusterStrings::strArrayValSeparator), ViewportIds);
+	DisplayClusterHelpers::str::ExtractArray(line, FString(DisplayClusterStrings::cfg::data::window::Postprocess), FString(DisplayClusterStrings::strArrayValSeparator), PostprocessIds);
+	DisplayClusterHelpers::str::ExtractValue(line, FString(DisplayClusterStrings::cfg::data::window::Fullscreen),  IsFullscreen);
+	DisplayClusterHelpers::str::ExtractValue(line, FString(DisplayClusterStrings::cfg::data::window::WinX),        WinX);
+	DisplayClusterHelpers::str::ExtractValue(line, FString(DisplayClusterStrings::cfg::data::window::WinY),        WinY);
+	DisplayClusterHelpers::str::ExtractValue(line, FString(DisplayClusterStrings::cfg::data::window::ResX),        ResX);
+	DisplayClusterHelpers::str::ExtractValue(line, FString(DisplayClusterStrings::cfg::data::window::ResY),        ResY);
 
 	return FDisplayClusterConfigBase::DeserializeFromString(line);
 }
@@ -96,7 +98,7 @@ bool FDisplayClusterConfigWindow::DeserializeFromString(const FString& line)
 //////////////////////////////////////////////////////////////////////////////////////////////
 FString FDisplayClusterConfigViewport::ToString() const
 {
-	return FString::Printf(TEXT("[%s + %s=%s, %s=%s, %s=%s, %s=%d, %s=%d, %s=%d, %s=%d]"),
+	return FString::Printf(TEXT("[%s + %s=%s, %s=%s, %s=%s, %s=%d, %s=%d, %s=%d, %s=%d, %s=%s]"),
 		*FDisplayClusterConfigBase::ToString(),
 		DisplayClusterStrings::cfg::data::Id, *Id,
 		DisplayClusterStrings::cfg::data::viewport::Projection, *ProjectionId,
@@ -104,7 +106,8 @@ FString FDisplayClusterConfigViewport::ToString() const
 		DisplayClusterStrings::cfg::data::viewport::PosX,       Loc.X,
 		DisplayClusterStrings::cfg::data::viewport::PosY,       Loc.Y,
 		DisplayClusterStrings::cfg::data::viewport::Width,      Size.X,
-		DisplayClusterStrings::cfg::data::viewport::Height,     Size.Y
+		DisplayClusterStrings::cfg::data::viewport::Height,     Size.Y,
+		DisplayClusterStrings::cfg::data::viewport::RTT,        DisplayClusterHelpers::str::BoolToStr(IsRTT)
 	);
 }
 
@@ -117,11 +120,35 @@ bool FDisplayClusterConfigViewport::DeserializeFromString(const FString& line)
 	DisplayClusterHelpers::str::ExtractValue(line, FString(DisplayClusterStrings::cfg::data::viewport::PosY),       Loc.Y);
 	DisplayClusterHelpers::str::ExtractValue(line, FString(DisplayClusterStrings::cfg::data::viewport::Width),      Size.X);
 	DisplayClusterHelpers::str::ExtractValue(line, FString(DisplayClusterStrings::cfg::data::viewport::Height),     Size.Y);
+	DisplayClusterHelpers::str::ExtractValue(line, FString(DisplayClusterStrings::cfg::data::viewport::RTT),        IsRTT);
+
 
 	return FDisplayClusterConfigBase::DeserializeFromString(line);
 }
 
 
+//////////////////////////////////////////////////////////////////////////////////////////////
+// FDisplayClusterConfigPostprocess
+//////////////////////////////////////////////////////////////////////////////////////////////
+FString FDisplayClusterConfigPostprocess::ToString() const
+{
+	return FString::Printf(TEXT("[%s + %s=%s, %s=%s, %s]"),
+		*FDisplayClusterConfigBase::ToString(),
+		DisplayClusterStrings::cfg::data::Id, *Id,
+		DisplayClusterStrings::cfg::data::postprocess::PostprocessId, *PostprocessId,
+		*ConfigLine
+	);
+}
+
+bool FDisplayClusterConfigPostprocess::DeserializeFromString(const FString& line)
+{
+	DisplayClusterHelpers::str::ExtractValue(line, FString(DisplayClusterStrings::cfg::data::Id), Id);
+	DisplayClusterHelpers::str::ExtractValue(line, FString(DisplayClusterStrings::cfg::data::postprocess::PostprocessId), PostprocessId);
+	ConfigLine = line; //Save unparsed args for custom pp parsers
+
+	return FDisplayClusterConfigBase::DeserializeFromString(line);
+}
+
 //////////////////////////////////////////////////////////////////////////////////////////////
 // FDisplayClusterConfigSceneNode
 //////////////////////////////////////////////////////////////////////////////////////////////
diff --git a/Engine/Plugins/Runtime/nDisplay/Source/DisplayCluster/Private/Config/Parser/DisplayClusterConfigParser.cpp b/Engine/Plugins/Runtime/nDisplay/Source/DisplayCluster/Private/Config/Parser/DisplayClusterConfigParser.cpp
--- a/Engine/Plugins/Runtime/nDisplay/Source/DisplayCluster/Private/Config/Parser/DisplayClusterConfigParser.cpp
+++ b/Engine/Plugins/Runtime/nDisplay/Source/DisplayCluster/Private/Config/Parser/DisplayClusterConfigParser.cpp
@@ -49,6 +49,11 @@ void FDisplayClusterConfigParser::AddViewport(const FDisplayClusterConfigViewpor
 	ConfigParserListener->AddViewport(InCfgViewport);
 }
 
+void FDisplayClusterConfigParser::AddPostprocess(const FDisplayClusterConfigPostprocess& InCfgPostprocess)
+{
+	ConfigParserListener->AddPostprocess(InCfgPostprocess);
+}
+
 void FDisplayClusterConfigParser::AddCamera(const FDisplayClusterConfigCamera& InCfgCamera)
 {
 	ConfigParserListener->AddCamera(InCfgCamera);
diff --git a/Engine/Plugins/Runtime/nDisplay/Source/DisplayCluster/Private/Config/Parser/DisplayClusterConfigParser.h b/Engine/Plugins/Runtime/nDisplay/Source/DisplayCluster/Private/Config/Parser/DisplayClusterConfigParser.h
--- a/Engine/Plugins/Runtime/nDisplay/Source/DisplayCluster/Private/Config/Parser/DisplayClusterConfigParser.h
+++ b/Engine/Plugins/Runtime/nDisplay/Source/DisplayCluster/Private/Config/Parser/DisplayClusterConfigParser.h
@@ -25,22 +25,23 @@ protected:
 	//////////////////////////////////////////////////////////////////////////////////////////////
 	// IDisplayClusterConfigParserListener
 	//////////////////////////////////////////////////////////////////////////////////////////////
-	virtual void AddInfo(const FDisplayClusterConfigInfo& InCfgInfo)                override final;
-	virtual void AddClusterNode(const FDisplayClusterConfigClusterNode& InCfgCNode) override final;
-	virtual void AddWindow(const FDisplayClusterConfigWindow& InCfgWindow)          override final;
-	virtual void AddScreen(const FDisplayClusterConfigScreen& InCfgScreen)          override final;
-	virtual void AddViewport(const FDisplayClusterConfigViewport& InCfgViewport)    override final;
-	virtual void AddCamera(const FDisplayClusterConfigCamera& InCfgCamera)          override final;
-	virtual void AddSceneNode(const FDisplayClusterConfigSceneNode& InCfgSNode)     override final;
-	virtual void AddGeneral(const FDisplayClusterConfigGeneral& InCfgGeneral)       override final;
-	virtual void AddRender(const FDisplayClusterConfigRender& InCfgRender)          override final;
-	virtual void AddStereo(const FDisplayClusterConfigStereo& InCfgStereo)          override final;
-	virtual void AddDebug(const FDisplayClusterConfigDebug& InCfgDebug)             override final;
-	virtual void AddNetwork(const FDisplayClusterConfigNetwork& InCfgNetwork)       override final;
-	virtual void AddInput(const FDisplayClusterConfigInput& InCfgInput)             override final;
-	virtual void AddInputSetup(const FDisplayClusterConfigInputSetup& InCfgInputSetup)  override final;
-	virtual void AddCustom(const FDisplayClusterConfigCustom& InCfgCustom)          override final;
-	virtual void AddProjection(const FDisplayClusterConfigProjection& InCfgProjection) override final;
+	virtual void AddInfo(const FDisplayClusterConfigInfo& InCfgInfo)                         override final;
+	virtual void AddClusterNode(const FDisplayClusterConfigClusterNode& InCfgCNode)          override final;
+	virtual void AddWindow(const FDisplayClusterConfigWindow& InCfgWindow)                   override final;
+	virtual void AddScreen(const FDisplayClusterConfigScreen& InCfgScreen)                   override final;
+	virtual void AddViewport(const FDisplayClusterConfigViewport& InCfgViewport)             override final;
+	virtual void AddPostprocess(const FDisplayClusterConfigPostprocess& InCfgPostprocess)    override final;
+	virtual void AddCamera(const FDisplayClusterConfigCamera& InCfgCamera)                   override final;
+	virtual void AddSceneNode(const FDisplayClusterConfigSceneNode& InCfgSNode)              override final;
+	virtual void AddGeneral(const FDisplayClusterConfigGeneral& InCfgGeneral)                override final;
+	virtual void AddRender(const FDisplayClusterConfigRender& InCfgRender)                   override final;
+	virtual void AddStereo(const FDisplayClusterConfigStereo& InCfgStereo)                   override final;
+	virtual void AddDebug(const FDisplayClusterConfigDebug& InCfgDebug)                      override final;
+	virtual void AddNetwork(const FDisplayClusterConfigNetwork& InCfgNetwork)                override final;
+	virtual void AddInput(const FDisplayClusterConfigInput& InCfgInput)                      override final;
+	virtual void AddInputSetup(const FDisplayClusterConfigInputSetup& InCfgInputSetup)       override final;
+	virtual void AddCustom(const FDisplayClusterConfigCustom& InCfgCustom)                   override final;
+	virtual void AddProjection(const FDisplayClusterConfigProjection& InCfgProjection)       override final;
 
 private:
 	IDisplayClusterConfigParserListener* const ConfigParserListener;
diff --git a/Engine/Plugins/Runtime/nDisplay/Source/DisplayCluster/Private/Config/Parser/DisplayClusterConfigParserText.cpp b/Engine/Plugins/Runtime/nDisplay/Source/DisplayCluster/Private/Config/Parser/DisplayClusterConfigParserText.cpp
--- a/Engine/Plugins/Runtime/nDisplay/Source/DisplayCluster/Private/Config/Parser/DisplayClusterConfigParserText.cpp
+++ b/Engine/Plugins/Runtime/nDisplay/Source/DisplayCluster/Private/Config/Parser/DisplayClusterConfigParserText.cpp
@@ -67,6 +67,10 @@ void FDisplayClusterConfigParserText::ParseLine(const FString& line)
 	{
 		AddViewport(impl_parse<FDisplayClusterConfigViewport>(line));
 	}
+	else if (line.StartsWith(FString(DisplayClusterStrings::cfg::data::postprocess::Header)))
+	{
+		AddPostprocess(impl_parse<FDisplayClusterConfigPostprocess>(line));
+	}
 	else if (line.StartsWith(FString(DisplayClusterStrings::cfg::data::camera::Header)))
 	{
 		AddCamera(impl_parse<FDisplayClusterConfigCamera>(line));
diff --git a/Engine/Plugins/Runtime/nDisplay/Source/DisplayCluster/Private/Config/Parser/IDisplayClusterConfigParserListener.h b/Engine/Plugins/Runtime/nDisplay/Source/DisplayCluster/Private/Config/Parser/IDisplayClusterConfigParserListener.h
--- a/Engine/Plugins/Runtime/nDisplay/Source/DisplayCluster/Private/Config/Parser/IDisplayClusterConfigParserListener.h
+++ b/Engine/Plugins/Runtime/nDisplay/Source/DisplayCluster/Private/Config/Parser/IDisplayClusterConfigParserListener.h
@@ -19,6 +19,7 @@ public:
 	virtual void AddWindow(const FDisplayClusterConfigWindow& InCfgWindow) = 0;
 	virtual void AddScreen(const FDisplayClusterConfigScreen& InCfgScreen) = 0;
 	virtual void AddViewport(const FDisplayClusterConfigViewport& InCfgViewport) = 0;
+	virtual void AddPostprocess(const FDisplayClusterConfigPostprocess& InCfgPostprocess) = 0;
 	virtual void AddCamera(const FDisplayClusterConfigCamera& InCfgCamera) = 0;
 	virtual void AddSceneNode(const FDisplayClusterConfigSceneNode& InCfgSNode) = 0;
 	virtual void AddGeneral(const FDisplayClusterConfigGeneral& InCfgGeneral) = 0;
diff --git a/Engine/Plugins/Runtime/nDisplay/Source/DisplayCluster/Private/Network/DisplayClusterServer.cpp b/Engine/Plugins/Runtime/nDisplay/Source/DisplayCluster/Private/Network/DisplayClusterServer.cpp
--- a/Engine/Plugins/Runtime/nDisplay/Source/DisplayCluster/Private/Network/DisplayClusterServer.cpp
+++ b/Engine/Plugins/Runtime/nDisplay/Source/DisplayCluster/Private/Network/DisplayClusterServer.cpp
@@ -22,17 +22,11 @@ FDisplayClusterServer::~FDisplayClusterServer()
 {
 	// Call from child .dtor
 	Shutdown();
-
-	// Now we can safely free all memory from our sessions. Need to be refactored in future.
-	for (auto* Session : Sessions)
-	{
-		delete Session;
-	}
 }
 
 bool FDisplayClusterServer::Start()
 {
-	FScopeLock lock(&InternalsSyncScope);
+	FScopeLock lock(&InternalsCritSec);
 
 	if (bIsRunning == true)
 	{
@@ -53,9 +47,9 @@ bool FDisplayClusterServer::Start()
 
 void FDisplayClusterServer::Shutdown()
 {
-	FScopeLock lock(&InternalsSyncScope);
+	FScopeLock lock(&InternalsCritSec);
 
-	if (bIsRunning == false)
+	if (!bIsRunning)
 	{
 		return;
 	}
@@ -70,15 +64,15 @@ void FDisplayClusterServer::Shutdown()
 	bIsRunning = false;
 }
 
-bool FDisplayClusterServer::IsRunning()
+bool FDisplayClusterServer::IsRunning() const
 {
-	FScopeLock lock(&InternalsSyncScope);
+	FScopeLock lock(&InternalsCritSec);
 	return bIsRunning;
 }
 
 bool FDisplayClusterServer::ConnectionHandler(FSocket* InSock, const FIPv4Endpoint& InEP)
 {
-	FScopeLock lock(&InternalsSyncScope);
+	FScopeLock lock(&InternalsCritSec);
 	check(InSock);
 
 	if (IsRunning() && IsConnectionAllowed(InSock, InEP))
@@ -86,12 +80,12 @@ bool FDisplayClusterServer::ConnectionHandler(FSocket* InSock, const FIPv4Endpoi
 		InSock->SetLinger(false, 0);
 		InSock->SetNonBlocking(false);
 
-		int32 newSize = static_cast<int32>(DisplayClusterConstants::net::SocketBufferSize);
-		int32 setSize;
-		InSock->SetReceiveBufferSize(newSize, setSize);
-		InSock->SetSendBufferSize(newSize, setSize);
+		const int32 RequestedBufferSize = static_cast<int32>(DisplayClusterConstants::net::SocketBufferSize);
+		int32 FinalBufferSize;
+		InSock->SetReceiveBufferSize(RequestedBufferSize, FinalBufferSize);
+		InSock->SetSendBufferSize(RequestedBufferSize, FinalBufferSize);
 
-		FDisplayClusterSessionBase* Session = CreateSession(InSock, InEP);
+		TSharedPtr<FDisplayClusterSessionBase> Session = CreateSession(InSock, InEP);
 		Session->StartSession();
 		Sessions.Add(Session);
 
diff --git a/Engine/Plugins/Runtime/nDisplay/Source/DisplayCluster/Private/Network/DisplayClusterServer.h b/Engine/Plugins/Runtime/nDisplay/Source/DisplayCluster/Private/Network/DisplayClusterServer.h
--- a/Engine/Plugins/Runtime/nDisplay/Source/DisplayCluster/Private/Network/DisplayClusterServer.h
+++ b/Engine/Plugins/Runtime/nDisplay/Source/DisplayCluster/Private/Network/DisplayClusterServer.h
@@ -28,7 +28,7 @@ public:
 	virtual void Shutdown();
 
 	// Returns current server state
-	bool IsRunning();
+	bool IsRunning() const;
 
 	// Server name
 	inline const FString& GetName() const
@@ -55,7 +55,7 @@ protected:
 	{ return true; }
 
 	// Allow to specify custom session class
-	virtual FDisplayClusterSessionBase* CreateSession(FSocket* InSock, const FIPv4Endpoint& InEP) = 0;
+	virtual TSharedPtr<FDisplayClusterSessionBase> CreateSession(FSocket* InSock, const FIPv4Endpoint& InEP) = 0;
 
 private:
 	// Handles incoming connections
@@ -71,10 +71,11 @@ private:
 	bool bIsRunning = false;
 	// Socket listener
 	FDisplayClusterTcpListener Listener;
-	// Sync access
-	FCriticalSection InternalsSyncScope;
 
 	// Active sessions
-	TArray<FDisplayClusterSessionBase*> Sessions;
+	TArray<TSharedPtr<FDisplayClusterSessionBase>> Sessions;
+
+	// Sync access
+	mutable FCriticalSection InternalsCritSec;
 };
 
diff --git a/Engine/Plugins/Runtime/nDisplay/Source/DisplayCluster/Private/Network/DisplayClusterTcpListener.cpp b/Engine/Plugins/Runtime/nDisplay/Source/DisplayCluster/Private/Network/DisplayClusterTcpListener.cpp
--- a/Engine/Plugins/Runtime/nDisplay/Source/DisplayCluster/Private/Network/DisplayClusterTcpListener.cpp
+++ b/Engine/Plugins/Runtime/nDisplay/Source/DisplayCluster/Private/Network/DisplayClusterTcpListener.cpp
@@ -27,7 +27,7 @@ FDisplayClusterTcpListener::~FDisplayClusterTcpListener()
 
 bool FDisplayClusterTcpListener::StartListening(const FString& InAddr, const int32 InPort)
 {
-	FScopeLock lock(&InternalsSyncScope);
+	FScopeLock lock(&InternalsCritSec);
 
 	if (bIsListening == true)
 	{
@@ -45,7 +45,7 @@ bool FDisplayClusterTcpListener::StartListening(const FString& InAddr, const int
 
 bool FDisplayClusterTcpListener::StartListening(const FIPv4Endpoint& InEP)
 {
-	FScopeLock lock(&InternalsSyncScope);
+	FScopeLock lock(&InternalsCritSec);
 
 	if (bIsListening == true)
 	{
@@ -56,11 +56,11 @@ bool FDisplayClusterTcpListener::StartListening(const FIPv4Endpoint& InEP)
 	Endpoint = InEP;
 
 	// Create listening thread
-	ThreadObj = FRunnableThread::Create(this, *(Name + FString("_thread")), 1 * 1024, TPri_Normal);
+	ThreadObj.Reset(FRunnableThread::Create(this, *(Name + FString("_thread")), 128 * 1024, TPri_Normal));
 	ensure(ThreadObj);
 
 	// Update state
-	bIsListening = true;
+	bIsListening = ThreadObj.IsValid();
 	
 	return bIsListening;
 }
@@ -68,7 +68,7 @@ bool FDisplayClusterTcpListener::StartListening(const FIPv4Endpoint& InEP)
 
 void FDisplayClusterTcpListener::StopListening()
 {
-	FScopeLock lock(&InternalsSyncScope);
+	FScopeLock lock(&InternalsCritSec);
 
 	if (bIsListening == false)
 	{
@@ -82,8 +82,7 @@ void FDisplayClusterTcpListener::StopListening()
 	if (ThreadObj)
 	{
 		ThreadObj->WaitForCompletion();
-		delete ThreadObj;
-		ThreadObj = nullptr;
+		ThreadObj.Reset();
 	}
 }
 
@@ -96,15 +95,8 @@ bool FDisplayClusterTcpListener::Init()
 {
 	// Create socket
 	SocketObj = FTcpSocketBuilder(*Name).AsBlocking().BoundToEndpoint(Endpoint).Listening(128);
-	if (!SocketObj)
-	{
-		// Just exit. No need to perform some notification from this thread to the cluster manager to notify
-		// about this fail. Just kill the application. 
-		FDisplayClusterAppExit::ExitApplication(FDisplayClusterAppExit::ExitType::KillImmediately, FString("Couldn't start listener socket"));
-		return false;
-	}
 
-	return true;
+	return SocketObj != nullptr;
 }
 
 uint32 FDisplayClusterTcpListener::Run()
@@ -113,23 +105,28 @@ uint32 FDisplayClusterTcpListener::Run()
 
 	if (SocketObj)
 	{
-		while (FSocket* pSock = SocketObj->Accept(*RemoteAddress, TEXT("FDisplayClusterTcpListener client")))
+		while (FSocket* NewSock = SocketObj->Accept(*RemoteAddress, TEXT("FDisplayClusterTcpListener client")))
 		{
 			UE_LOG(LogDisplayClusterNetwork, Log, TEXT("New incoming connection: %s"), *RemoteAddress->ToString(true));
 
-			if (OnConnectionAcceptedDelegate.IsBound())
+			if (NewSock)
 			{
-				if (!OnConnectionAcceptedDelegate.Execute(pSock, FIPv4Endpoint(RemoteAddress)))
+				// Ask a server implementation if it confirms this new incoming connection
+				if (OnConnectionAcceptedDelegate.IsBound())
 				{
-					pSock->Close();
-					ISocketSubsystem::Get(PLATFORM_SOCKETSUBSYSTEM)->DestroySocket(pSock);
+					// If no, close the socket
+					if (!OnConnectionAcceptedDelegate.Execute(NewSock, FIPv4Endpoint(RemoteAddress)))
+					{
+						NewSock->Close();
+						ISocketSubsystem::Get(PLATFORM_SOCKETSUBSYSTEM)->DestroySocket(NewSock);
+					}
 				}
 			}
 		}
 	}
 	else
 	{
-		UE_LOG(LogDisplayClusterNetwork, Error, TEXT("Socket is not initialized"));
+		UE_LOG(LogDisplayClusterNetwork, Error, TEXT("Socket %s is not initialized"), *Name);
 		return 0;
 	}
 
diff --git a/Engine/Plugins/Runtime/nDisplay/Source/DisplayCluster/Private/Network/DisplayClusterTcpListener.h b/Engine/Plugins/Runtime/nDisplay/Source/DisplayCluster/Private/Network/DisplayClusterTcpListener.h
--- a/Engine/Plugins/Runtime/nDisplay/Source/DisplayCluster/Private/Network/DisplayClusterTcpListener.h
+++ b/Engine/Plugins/Runtime/nDisplay/Source/DisplayCluster/Private/Network/DisplayClusterTcpListener.h
@@ -51,9 +51,9 @@ private:
 	// Listening endpoint
 	FIPv4Endpoint Endpoint;
 	// Holds the thread object
-	FRunnableThread* ThreadObj;
+	TUniquePtr<FRunnableThread> ThreadObj;
 	// Sync access
-	FCriticalSection InternalsSyncScope;
+	FCriticalSection InternalsCritSec;
 	// Listening state
 	bool bIsListening = false;
 
diff --git a/Engine/Plugins/Runtime/nDisplay/Source/DisplayCluster/Private/Network/Service/ClusterEvents/DisplayClusterClusterEventsService.cpp b/Engine/Plugins/Runtime/nDisplay/Source/DisplayCluster/Private/Network/Service/ClusterEvents/DisplayClusterClusterEventsService.cpp
--- a/Engine/Plugins/Runtime/nDisplay/Source/DisplayCluster/Private/Network/Service/ClusterEvents/DisplayClusterClusterEventsService.cpp
+++ b/Engine/Plugins/Runtime/nDisplay/Source/DisplayCluster/Private/Network/Service/ClusterEvents/DisplayClusterClusterEventsService.cpp
@@ -43,9 +43,9 @@ void FDisplayClusterClusterEventsService::Shutdown()
 	return FDisplayClusterServer::Shutdown();
 }
 
-FDisplayClusterSessionBase* FDisplayClusterClusterEventsService::CreateSession(FSocket* InSocket, const FIPv4Endpoint& InEP)
+TSharedPtr<FDisplayClusterSessionBase> FDisplayClusterClusterEventsService::CreateSession(FSocket* InSocket, const FIPv4Endpoint& InEP)
 {
-	return new FDisplayClusterSessionExternal(InSocket, this, GetName() + FString("_session_external") + InEP.ToString());
+	return TSharedPtr<FDisplayClusterSessionBase>(new FDisplayClusterSessionExternal(InSocket, this, GetName() + FString("_session_external") + InEP.ToString()));
 }
 
 //////////////////////////////////////////////////////////////////////////////////////////////
diff --git a/Engine/Plugins/Runtime/nDisplay/Source/DisplayCluster/Private/Network/Service/ClusterEvents/DisplayClusterClusterEventsService.h b/Engine/Plugins/Runtime/nDisplay/Source/DisplayCluster/Private/Network/Service/ClusterEvents/DisplayClusterClusterEventsService.h
--- a/Engine/Plugins/Runtime/nDisplay/Source/DisplayCluster/Private/Network/Service/ClusterEvents/DisplayClusterClusterEventsService.h
+++ b/Engine/Plugins/Runtime/nDisplay/Source/DisplayCluster/Private/Network/Service/ClusterEvents/DisplayClusterClusterEventsService.h
@@ -35,7 +35,7 @@ public:
 	virtual void Shutdown() override;
 
 protected:
-	virtual FDisplayClusterSessionBase* CreateSession(FSocket* InSocket, const FIPv4Endpoint& InEP) override;
+	virtual TSharedPtr<FDisplayClusterSessionBase> CreateSession(FSocket* InSocket, const FIPv4Endpoint& InEP) override;
 	virtual bool IsConnectionAllowed(FSocket* InSocket, const FIPv4Endpoint& InEP)
 	{ return true; }
 
diff --git a/Engine/Plugins/Runtime/nDisplay/Source/DisplayCluster/Private/Network/Service/ClusterSync/DisplayClusterClusterSyncService.cpp b/Engine/Plugins/Runtime/nDisplay/Source/DisplayCluster/Private/Network/Service/ClusterSync/DisplayClusterClusterSyncService.cpp
--- a/Engine/Plugins/Runtime/nDisplay/Source/DisplayCluster/Private/Network/Service/ClusterSync/DisplayClusterClusterSyncService.cpp
+++ b/Engine/Plugins/Runtime/nDisplay/Source/DisplayCluster/Private/Network/Service/ClusterSync/DisplayClusterClusterSyncService.cpp
@@ -51,9 +51,9 @@ void FDisplayClusterClusterSyncService::Shutdown()
 	return FDisplayClusterServer::Shutdown();
 }
 
-FDisplayClusterSessionBase* FDisplayClusterClusterSyncService::CreateSession(FSocket* InSocket, const FIPv4Endpoint& InEP)
+TSharedPtr<FDisplayClusterSessionBase> FDisplayClusterClusterSyncService::CreateSession(FSocket* InSocket, const FIPv4Endpoint& InEP)
 {
-	return new FDisplayClusterSessionInternal(InSocket, this, GetName() + FString("_session_") + InEP.ToString());
+	return TSharedPtr<FDisplayClusterSessionBase>(new FDisplayClusterSessionInternal(InSocket, this, GetName() + FString("_session_") + InEP.ToString()));
 }
 
 
@@ -73,8 +73,8 @@ void FDisplayClusterClusterSyncService::NotifySessionClose(FDisplayClusterSessio
 	BarrierFrameEnd.Deactivate();
 	BarrierTickEnd.Deactivate();
 
-	FDisplayClusterAppExit::ExitApplication(FDisplayClusterAppExit::ExitType::NormalSoft, GetName() + FString(" - Connection interrupted. Application exit requested."));
 	FDisplayClusterService::NotifySessionClose(InSession);
+	FDisplayClusterAppExit::ExitApplication(FDisplayClusterAppExit::ExitType::NormalSoft, GetName() + FString(" - Connection interrupted. Application exit requested."));
 }
 
 TSharedPtr<FDisplayClusterMessage> FDisplayClusterClusterSyncService::ProcessMessage(const TSharedPtr<FDisplayClusterMessage>& Request)
diff --git a/Engine/Plugins/Runtime/nDisplay/Source/DisplayCluster/Private/Network/Service/ClusterSync/DisplayClusterClusterSyncService.h b/Engine/Plugins/Runtime/nDisplay/Source/DisplayCluster/Private/Network/Service/ClusterSync/DisplayClusterClusterSyncService.h
--- a/Engine/Plugins/Runtime/nDisplay/Source/DisplayCluster/Private/Network/Service/ClusterSync/DisplayClusterClusterSyncService.h
+++ b/Engine/Plugins/Runtime/nDisplay/Source/DisplayCluster/Private/Network/Service/ClusterSync/DisplayClusterClusterSyncService.h
@@ -25,7 +25,7 @@ public:
 	void Shutdown() override;
 
 protected:
-	virtual FDisplayClusterSessionBase* CreateSession(FSocket* InSocket, const FIPv4Endpoint& InEP) override;
+	virtual TSharedPtr<FDisplayClusterSessionBase> CreateSession(FSocket* InSocket, const FIPv4Endpoint& InEP) override;
 
 protected:
 	//////////////////////////////////////////////////////////////////////////////////////////////
diff --git a/Engine/Plugins/Runtime/nDisplay/Source/DisplayCluster/Private/Network/Service/SwapSync/DisplayClusterSwapSyncService.cpp b/Engine/Plugins/Runtime/nDisplay/Source/DisplayCluster/Private/Network/Service/SwapSync/DisplayClusterSwapSyncService.cpp
--- a/Engine/Plugins/Runtime/nDisplay/Source/DisplayCluster/Private/Network/Service/SwapSync/DisplayClusterSwapSyncService.cpp
+++ b/Engine/Plugins/Runtime/nDisplay/Source/DisplayCluster/Private/Network/Service/SwapSync/DisplayClusterSwapSyncService.cpp
@@ -13,9 +13,9 @@
 #include "DisplayClusterLog.h"
 
 
-FDisplayClusterSwapSyncService::FDisplayClusterSwapSyncService(const FString& InAddr, const int32 InPort) :
-	FDisplayClusterService(FString("SRV_SS"), InAddr, InPort),
-	BarrierSwap(GDisplayCluster->GetPrivateClusterMgr()->GetNodesAmount(), FString("SwapSync_barrier"), GDisplayCluster->GetConfigMgr()->GetConfigNetwork().BarrierWaitTimeout)
+FDisplayClusterSwapSyncService::FDisplayClusterSwapSyncService(const FString& InAddr, const int32 InPort)
+	: FDisplayClusterService(FString("SRV_SS"), InAddr, InPort)
+	, BarrierSwap(GDisplayCluster->GetPrivateClusterMgr()->GetNodesAmount(), FString("SwapSync_barrier"), GDisplayCluster->GetConfigMgr()->GetConfigNetwork().BarrierWaitTimeout)
 {
 }
 
@@ -39,9 +39,9 @@ void FDisplayClusterSwapSyncService::Shutdown()
 	return FDisplayClusterServer::Shutdown();
 }
 
-FDisplayClusterSessionBase* FDisplayClusterSwapSyncService::CreateSession(FSocket* InSocket, const FIPv4Endpoint& InEP)
+TSharedPtr<FDisplayClusterSessionBase> FDisplayClusterSwapSyncService::CreateSession(FSocket* InSocket, const FIPv4Endpoint& InEP)
 {
-	return new FDisplayClusterSessionInternal(InSocket, this, GetName() + FString("_session_") + InEP.ToString());
+	return TSharedPtr<FDisplayClusterSessionBase>(new FDisplayClusterSessionInternal(InSocket, this, GetName() + FString("_session_") + InEP.ToString()));
 }
 
 
@@ -58,8 +58,8 @@ void FDisplayClusterSwapSyncService::NotifySessionClose(FDisplayClusterSessionBa
 	// Unblock waiting threads to allow current Tick() finish
 	BarrierSwap.Deactivate();
 
-	FDisplayClusterAppExit::ExitApplication(FDisplayClusterAppExit::ExitType::NormalSoft, GetName() + FString(" - Connection interrupted. Application exit requested."));
 	FDisplayClusterService::NotifySessionClose(InSession);
+	FDisplayClusterAppExit::ExitApplication(FDisplayClusterAppExit::ExitType::NormalSoft, GetName() + FString(" - Connection interrupted. Application exit requested."));
 }
 
 TSharedPtr<FDisplayClusterMessage> FDisplayClusterSwapSyncService::ProcessMessage(const TSharedPtr<FDisplayClusterMessage>& Request)
diff --git a/Engine/Plugins/Runtime/nDisplay/Source/DisplayCluster/Private/Network/Session/DisplayClusterSessionBase.cpp b/Engine/Plugins/Runtime/nDisplay/Source/DisplayCluster/Private/Network/Session/DisplayClusterSessionBase.cpp
--- a/Engine/Plugins/Runtime/nDisplay/Source/DisplayCluster/Private/Network/Session/DisplayClusterSessionBase.cpp
+++ b/Engine/Plugins/Runtime/nDisplay/Source/DisplayCluster/Private/Network/Session/DisplayClusterSessionBase.cpp
@@ -21,16 +21,14 @@ FDisplayClusterSessionBase::FDisplayClusterSessionBase(FSocket* InSocket, IDispl
 FDisplayClusterSessionBase::~FDisplayClusterSessionBase()
 {
 	UE_LOG(LogDisplayClusterNetwork, VeryVerbose, TEXT("Session %s .dtor"), *Name);
-
-	ThreadObj->WaitForCompletion();
-	delete ThreadObj;
+	Stop();
 }
 
 void FDisplayClusterSessionBase::StartSession()
 {
 	Listener->NotifySessionOpen(this);
 
-	ThreadObj = FRunnableThread::Create(this, *(Name + FString("_thread")), 128 * 1024, TPri_Normal, FPlatformAffinity::GetPoolThreadMask());
+	ThreadObj.Reset(FRunnableThread::Create(this, *(Name + FString("_thread")), 1024 * 1024, TPri_Normal, FPlatformAffinity::GetPoolThreadMask()));
 	ensure(ThreadObj);
 
 	UE_LOG(LogDisplayClusterNetwork, Log, TEXT("Session %s started"), *Name);
@@ -39,5 +37,5 @@ void FDisplayClusterSessionBase::StartSession()
 void FDisplayClusterSessionBase::Stop()
 {
 	GetSocket()->Close();
-	GetListener()->NotifySessionClose(this);
+	ThreadObj->WaitForCompletion();
 }
diff --git a/Engine/Plugins/Runtime/nDisplay/Source/DisplayCluster/Private/Network/Session/DisplayClusterSessionBase.h b/Engine/Plugins/Runtime/nDisplay/Source/DisplayCluster/Private/Network/Session/DisplayClusterSessionBase.h
--- a/Engine/Plugins/Runtime/nDisplay/Source/DisplayCluster/Private/Network/Session/DisplayClusterSessionBase.h
+++ b/Engine/Plugins/Runtime/nDisplay/Source/DisplayCluster/Private/Network/Session/DisplayClusterSessionBase.h
@@ -35,5 +35,5 @@ protected:
 private:
 	const FString        Name;
 	IDisplayClusterSessionListener* Listener = nullptr;
-	FRunnableThread* ThreadObj = nullptr;
+	TUniquePtr<FRunnableThread> ThreadObj;
 };
diff --git a/Engine/Plugins/Runtime/nDisplay/Source/DisplayCluster/Private/Network/Session/DisplayClusterSessionExternal.cpp b/Engine/Plugins/Runtime/nDisplay/Source/DisplayCluster/Private/Network/Session/DisplayClusterSessionExternal.cpp
--- a/Engine/Plugins/Runtime/nDisplay/Source/DisplayCluster/Private/Network/Session/DisplayClusterSessionExternal.cpp
+++ b/Engine/Plugins/Runtime/nDisplay/Source/DisplayCluster/Private/Network/Session/DisplayClusterSessionExternal.cpp
@@ -39,7 +39,7 @@ uint32 FDisplayClusterSessionExternal::Run()
 		}
 	}
 
-	Stop();
+	GetListener()->NotifySessionClose(this);
 
 	UE_LOG(LogDisplayClusterNetwork, Log, TEXT("Session thread %s finished"), *GetName());
 	return 0;
diff --git a/Engine/Plugins/Runtime/nDisplay/Source/DisplayCluster/Private/Network/Session/DisplayClusterSessionInternal.cpp b/Engine/Plugins/Runtime/nDisplay/Source/DisplayCluster/Private/Network/Session/DisplayClusterSessionInternal.cpp
--- a/Engine/Plugins/Runtime/nDisplay/Source/DisplayCluster/Private/Network/Session/DisplayClusterSessionInternal.cpp
+++ b/Engine/Plugins/Runtime/nDisplay/Source/DisplayCluster/Private/Network/Session/DisplayClusterSessionInternal.cpp
@@ -44,7 +44,7 @@ uint32 FDisplayClusterSessionInternal::Run()
 		}
 	}
 
-	Stop();
+	GetListener()->NotifySessionClose(this);
 
 	UE_LOG(LogDisplayClusterNetwork, Log, TEXT("Session thread %s finished"), *GetName());
 	return 0;
diff --git a/Engine/Plugins/Runtime/nDisplay/Source/DisplayCluster/Private/Render/Device/DisplayClusterDeviceBase.cpp b/Engine/Plugins/Runtime/nDisplay/Source/DisplayCluster/Private/Render/Device/DisplayClusterDeviceBase.cpp
--- a/Engine/Plugins/Runtime/nDisplay/Source/DisplayCluster/Private/Render/Device/DisplayClusterDeviceBase.cpp
+++ b/Engine/Plugins/Runtime/nDisplay/Source/DisplayCluster/Private/Render/Device/DisplayClusterDeviceBase.cpp
@@ -21,6 +21,7 @@
 #include "Render/PostProcess/IDisplayClusterPostProcess.h"
 #include "Render/Projection/IDisplayClusterProjectionPolicy.h"
 #include "Render/Projection/IDisplayClusterProjectionPolicyFactory.h"
+#include "Render/Synchronization/IDisplayClusterRenderSyncPolicy.h"
 
 #include "DisplayClusterGlobals.h"
 #include "DisplayClusterLog.h"
@@ -40,7 +41,10 @@ static TAutoConsoleVariable<int32>  CVarVSyncInterval(
 static TAutoConsoleVariable<int32> CVarWarpBlendEnabled(
 	TEXT("nDisplay.render.WarpBlendEnabled"),
 	1,
-	TEXT("Warp & Blend status (0 = disabled)\n"),
+	TEXT("Warp & Blend status\n")
+	TEXT("0 : disabled\n")
+	TEXT("1 : enabled\n")
+	,
 	ECVF_RenderThreadSafe
 );
 
@@ -139,6 +143,19 @@ bool FDisplayClusterDeviceBase::Initialize()
 		return false;
 	}
 
+	// Forward cfg line to postprocess:
+	// Get list of local postprocess (assigned to this cluster node)
+	TArray<FDisplayClusterConfigPostprocess> LocalPostprocess = DisplayClusterHelpers::config::GetLocalPostprocess();
+	TMap<FString, IPDisplayClusterRenderManager::FDisplayClusterPPInfo> Postprocess = RenderMgr->GetRegisteredPostprocessOperations();
+	// Initialize all local Postprocess
+	for (const FDisplayClusterConfigPostprocess& CfgPostprocess : LocalPostprocess)
+	{
+		if (Postprocess.Contains(CfgPostprocess.PostprocessId))
+		{
+			Postprocess[CfgPostprocess.PostprocessId].Operation->InitializePostProcess(CfgPostprocess.ConfigLine);
+		}
+	}
+
 	return true;
 }
 
@@ -186,6 +203,44 @@ void FDisplayClusterDeviceBase::SetViewportCamera(const FString& InCameraId /* =
 	UE_LOG(LogDisplayClusterRender, Warning, TEXT("Couldn't assign '%s' camera. Viewport '%s' not found"), *InCameraId, *InViewportId);
 }
 
+void FDisplayClusterDeviceBase::SetCustomPostProcessing(const FString& ViewportID, const FPostProcessSettings& PostProcessingSettings)
+{	
+	// find viewport index 
+	int ViewportIdx = -1;
+	for(int i = 0; i < RenderViewports.Num(); i++)
+	{
+		if (RenderViewports[i].GetId() == ViewportID)
+		{
+			ViewportIdx = i;
+			break;
+		}
+	}
+
+	// check if post processing settings assigned then override or add
+	if (ViewportIdx != -1)
+	{
+		ViewportFinalPostProcessingSettingsOverride.Remove(ViewportIdx);
+		ViewportFinalPostProcessingSettingsOverride.Emplace(ViewportIdx, PostProcessingSettings);
+	}
+}
+
+bool FDisplayClusterDeviceBase::GetViewportRect(const FString& InViewportID, FIntRect& Rect)
+{
+	FDisplayClusterRenderViewport* DesiredViewport = RenderViewports.FindByPredicate([InViewportID](const FDisplayClusterRenderViewport& ItemViewport)
+	{
+		return InViewportID.Equals(ItemViewport.GetId(), ESearchCase::IgnoreCase);
+	});
+
+	if (!DesiredViewport)
+	{
+		return false;
+	}
+
+	Rect = DesiredViewport->GetArea();
+
+	return true;
+}
+
 //////////////////////////////////////////////////////////////////////////////////////////////
 // IStereoRendering
 //////////////////////////////////////////////////////////////////////////////////////////////
@@ -461,6 +516,12 @@ void FDisplayClusterDeviceBase::CalculateRenderTargetSize(const class FViewport&
 	InOutSizeX = Viewport.GetSizeXY().X;
 	InOutSizeY = Viewport.GetSizeXY().Y;
 
+	for (const FDisplayClusterRenderViewport& Item : RenderViewports)
+	{
+		InOutSizeX = FMath::Max(InOutSizeX, (uint32)Item.GetArea().Max.X);
+		InOutSizeY = FMath::Max(InOutSizeY, (uint32)Item.GetArea().Max.Y);
+	}
+
 	UE_LOG(LogDisplayClusterRender, Verbose, TEXT("Render target size: [%d x %d]"), InOutSizeX, InOutSizeY);
 
 	check(InOutSizeX > 0 && InOutSizeY > 0);
@@ -478,7 +539,7 @@ bool FDisplayClusterDeviceBase::NeedReAllocateViewportRenderTarget(const class F
 	// Get desired RT size
 	uint32 newSizeX = 0;
 	uint32 newSizeY = 0;
-	CalculateRenderTargetSize(Viewport, newSizeX, newSizeY);
+	CalculateRenderTargetSize(Viewport, newSizeX, newSizeY);	
 
 	// Here we conclude if need to re-allocate
 	const bool Result = (newSizeX != rtSize.X || newSizeY != rtSize.Y);
@@ -512,12 +573,22 @@ void FDisplayClusterDeviceBase::OnBackBufferResize()
 bool FDisplayClusterDeviceBase::Present(int32& InOutSyncInterval)
 {
 	DISPLAY_CLUSTER_FUNC_TRACE(LogDisplayClusterRender);
-	UE_LOG(LogDisplayClusterRender, Warning, TEXT("Present - default handler implementation. Check stereo device instantiation."));
 
-	// Default behavior
-	// Return false to force clean screen. This will indicate that something is going wrong
-	// or particular stereo device hasn't been implemented appropriately yet.
-	return false;
+	// Update sync value with nDisplay value
+	InOutSyncInterval = GetSwapInt();
+
+	// Get sync policy instance
+	TSharedPtr<IDisplayClusterRenderSyncPolicy> SyncPolicy = GDisplayCluster->GetRenderMgr()->GetCurrentSynchronizationPolicy();
+	if (SyncPolicy.IsValid())
+	{
+		// False results means we don't need to present current frame, the sync object already presented it
+		if (!SyncPolicy->SynchronizeClusterRendering(InOutSyncInterval))
+		{
+			return false;
+		}
+	}
+
+	return true;
 }
 
 
@@ -629,7 +700,7 @@ uint32 FDisplayClusterDeviceBase::GetSwapInt() const
 	return (SyncInterval);
 }
 
-void FDisplayClusterDeviceBase::AddViewport(const FString& InViewportId, const FIntPoint& InViewportLocation, const FIntPoint& InViewportSize, TSharedPtr<IDisplayClusterProjectionPolicy> InProjPolicy, const FString& InCameraId)
+void FDisplayClusterDeviceBase::AddViewport(const FString& InViewportId, const FIntPoint& InViewportLocation, const FIntPoint& InViewportSize, TSharedPtr<IDisplayClusterProjectionPolicy> InProjPolicy, const FString& InCameraId, bool IsRTT /*= false*/)
 {
 	DISPLAY_CLUSTER_FUNC_TRACE(LogDisplayClusterRender);
 
@@ -652,13 +723,21 @@ void FDisplayClusterDeviceBase::AddViewport(const FString& InViewportId, const F
 		return;
 	}
 
-	if (InProjPolicy->HandleAddViewport(InViewportSize, ViewsAmountPerViewport))
+	if (!(IsRTT && bViewportRttAdded))
 	{
-		UE_LOG(LogDisplayClusterRender, Log, TEXT("A corresponded projection policy object has initialized the viewport '%s'"), *InViewportId);
-		FIntRect ViewportArea = FIntRect(InViewportLocation, InViewportLocation + InViewportSize);
-		FDisplayClusterRenderViewport NewViewport(InViewportId, ViewportArea, InProjPolicy, EDisplayClusterEyeType::COUNT, InCameraId);
-		RenderViewports.Add(NewViewport);
-		return;
+		if (InProjPolicy->HandleAddViewport(InViewportSize, ViewsAmountPerViewport))
+		{
+			UE_LOG(LogDisplayClusterRender, Log, TEXT("A corresponded projection policy object has initialized the viewport '%s'"), *InViewportId);
+
+			FIntRect ViewportArea = FIntRect(InViewportLocation, InViewportLocation + InViewportSize);
+			FDisplayClusterRenderViewport NewViewport(InViewportId, ViewportArea, InProjPolicy, EDisplayClusterEyeType::COUNT, InCameraId, IsRTT);
+			RenderViewports.Add(NewViewport);
+
+			if (IsRTT)
+			{
+				bViewportRttAdded = true;
+			}
+		}
 	}
 }
 
@@ -689,3 +768,14 @@ void FDisplayClusterDeviceBase::CopyTextureToBackBuffer_RenderThread(FRHICommand
 
 	RHICmdList.CopyToResolveTarget(SrcTexture, BackBuffer, copyParams);
 }
+
+void FDisplayClusterDeviceBase::UpdatePostProcessSettings(struct FPostProcessSettings* FinalPostProcessingSettings, const enum EStereoscopicPass StereoPassType)
+{
+	int ViewportNumber = DecodeViewportIndex(StereoPassType);
+	auto CustomPostProcessAvailable = ViewportFinalPostProcessingSettingsOverride.Find(ViewportNumber);
+	if (CustomPostProcessAvailable)
+	{
+		*FinalPostProcessingSettings = ViewportFinalPostProcessingSettingsOverride[ViewportNumber];
+		//ViewportFinalPostProcessingSettingsOverride.Remove(ViewportNumber);
+	}		
+}
\ No newline at end of file
diff --git a/Engine/Plugins/Runtime/nDisplay/Source/DisplayCluster/Private/Render/Device/DisplayClusterDeviceBase.h b/Engine/Plugins/Runtime/nDisplay/Source/DisplayCluster/Private/Render/Device/DisplayClusterDeviceBase.h
--- a/Engine/Plugins/Runtime/nDisplay/Source/DisplayCluster/Private/Render/Device/DisplayClusterDeviceBase.h
+++ b/Engine/Plugins/Runtime/nDisplay/Source/DisplayCluster/Private/Render/Device/DisplayClusterDeviceBase.h
@@ -38,6 +38,8 @@ public:
 	virtual bool Initialize() override;
 	virtual void InitializeWorldContent(UWorld* InWorld) override;
 	virtual void SetViewportCamera(const FString& InCameraId = FString(), const FString& InViewportId = FString()) override;
+	virtual void SetCustomPostProcessing(const FString& ViewportID, const FPostProcessSettings& PostProcessingSettings) override;
+	virtual bool GetViewportRect(const FString& InViewportID, FIntRect& Rect) override;
 
 public:
 	//////////////////////////////////////////////////////////////////////////////////////////////
@@ -120,9 +122,11 @@ protected:
 	uint32 GetSwapInt() const;
 
 	// Adds a new viewport with specified parameters and projection policy object
-	void AddViewport(const FString& InViewportId, const FIntPoint& InViewportLocation, const FIntPoint& InViewportSize, TSharedPtr<IDisplayClusterProjectionPolicy> InProjPolicy, const FString& InCameraId);
+	void AddViewport(const FString& InViewportId, const FIntPoint& InViewportLocation, const FIntPoint& InViewportSize, TSharedPtr<IDisplayClusterProjectionPolicy> InProjPolicy, const FString& InCameraId, bool IsRTT = false);
 	// Performs copying of render target data to the back buffer
 	virtual void CopyTextureToBackBuffer_RenderThread(FRHICommandListImmediate& RHICmdList, FRHITexture2D* BackBuffer, FRHITexture2D* SrcTexture, FVector2D WindowSize) const;
+	// Checks if custom post processing settings is assigned for specific viewport and assign them to be used
+	virtual void UpdatePostProcessSettings(struct FPostProcessSettings* FinalPostProcessingSettings, const enum EStereoscopicPass StereoPassType) override;
 
 protected:
 	// Viewports
@@ -131,7 +135,12 @@ protected:
 	uint32 ViewsAmountPerViewport = 0;
 	// UE4 main viewport
 	FViewport* MainViewport = nullptr;
+	// custom post processing settings
+	TMap<int, FPostProcessSettings> ViewportFinalPostProcessingSettingsOverride;	 
 
 	// Data access synchronization
 	mutable FCriticalSection InternalsSyncScope;
+
+	// Temporary: don't allow to add more than 1 RTT viewport
+	bool bViewportRttAdded = false;
 };
diff --git a/Engine/Plugins/Runtime/nDisplay/Source/DisplayCluster/Private/Render/Device/DisplayClusterOpenGL.h b/Engine/Plugins/Runtime/nDisplay/Source/DisplayCluster/Private/Render/Device/DisplayClusterOpenGL.h
new file mode 100644
--- /dev/null
+++ b/Engine/Plugins/Runtime/nDisplay/Source/DisplayCluster/Private/Render/Device/DisplayClusterOpenGL.h
@@ -0,0 +1,44 @@
+// Copyright 1998-2019 Epic Games, Inc. All Rights Reserved.
+
+#pragma once
+
+#include "OpenGLDrv.h"
+#include "OpenGLResources.h"
+
+
+#if PLATFORM_WINDOWS
+// This is redeclaration of WINDOWS specific FPlatformOpenGLContext
+// which is declared in private OpenGLWindows.cpp file.
+//@note: Keep it synced with original type (Engine\Source\Runtime\OpenGLDrv\Private\Windows\OpenGLWindows.cpp)
+struct FPlatformOpenGLContext
+{
+	HWND WindowHandle;
+	HDC DeviceContext;
+	HGLRC OpenGLContext;
+	bool bReleaseWindowOnDestroy;
+	int32 SyncInterval;
+	GLuint	ViewportFramebuffer;
+	GLuint	VertexArrayObject;	// one has to be generated and set for each context (OpenGL 3.2 Core requirements)
+	GLuint	BackBufferResource;
+	GLenum	BackBufferTarget;
+};
+#endif
+
+
+#if PLATFORM_LINUX
+// This is redeclaration of LINUX specific FPlatformOpenGLContext
+// which is declared in private OpenGLWindows.cpp file.
+//@note: Keep it synced with original type (Engine\Source\Runtime\OpenGLDrv\Private\Linux\OpenGLLinux.cpp)
+struct FPlatformOpenGLContext
+{
+	SDL_HWindow    hWnd;
+	SDL_HGLContext hGLContext; // this is a (void*) pointer
+
+	bool bReleaseWindowOnDestroy;
+	int32 SyncInterval;
+	GLuint	ViewportFramebuffer;
+	GLuint	VertexArrayObject;	// one has to be generated and set for each context (OpenGL 3.2 Core requirements)
+};
+
+//@note: Place here any Linux targeted device implementations
+#endif
diff --git a/Engine/Plugins/Runtime/nDisplay/Source/DisplayCluster/Private/Render/Device/DisplayClusterRenderViewport.h b/Engine/Plugins/Runtime/nDisplay/Source/DisplayCluster/Private/Render/Device/DisplayClusterRenderViewport.h
--- a/Engine/Plugins/Runtime/nDisplay/Source/DisplayCluster/Private/Render/Device/DisplayClusterRenderViewport.h
+++ b/Engine/Plugins/Runtime/nDisplay/Source/DisplayCluster/Private/Render/Device/DisplayClusterRenderViewport.h
@@ -14,10 +14,16 @@ class FDisplayClusterRenderViewport
 {
 public:
 	FDisplayClusterRenderViewport(const FString& ViewportId, const FIntRect& ViewportArea, TSharedPtr<IDisplayClusterProjectionPolicy> ProjectionPolicy, uint8 ContextsAmount, const FString& InCameraId)
+		: FDisplayClusterRenderViewport(ViewportId, ViewportArea, ProjectionPolicy, ContextsAmount, InCameraId, false)
+	{
+	}
+
+	FDisplayClusterRenderViewport(const FString& ViewportId, const FIntRect& ViewportArea, TSharedPtr<IDisplayClusterProjectionPolicy> ProjectionPolicy, uint8 ContextsAmount, const FString& InCameraId, bool IsRTT)
 		: Id(ViewportId)
 		, CameraId(InCameraId)
 		, Area(ViewportArea)
 		, Policy(ProjectionPolicy)
+		, bRTT(IsRTT)
 	{
 		check(ProjectionPolicy.IsValid());
 		Contexts.AddDefaulted(ContextsAmount);
@@ -48,6 +54,9 @@ public:
 	const FIntRect& GetArea() const
 	{ return Area; }
 
+	bool IsRTT() const
+	{ return bRTT; }
+
 	FDisplayClusterRenderViewContext& GetContext(uint8 ContextNum)
 	{
 		check(ContextNum < Contexts.Num());
@@ -71,4 +80,6 @@ private:
 	TSharedPtr<IDisplayClusterProjectionPolicy> Policy;
 	// Viewport contexts (left/center/right eyes)
 	TArray<FDisplayClusterRenderViewContext> Contexts;
+	// Is RTT viewport
+	bool bRTT;
 };
diff --git a/Engine/Plugins/Runtime/nDisplay/Source/DisplayCluster/Private/Render/Device/Monoscopic/DisplayClusterDeviceMonoscopicDX11.cpp b/Engine/Plugins/Runtime/nDisplay/Source/DisplayCluster/Private/Render/Device/Monoscopic/DisplayClusterDeviceMonoscopicDX11.cpp
--- a/Engine/Plugins/Runtime/nDisplay/Source/DisplayCluster/Private/Render/Device/Monoscopic/DisplayClusterDeviceMonoscopicDX11.cpp
+++ b/Engine/Plugins/Runtime/nDisplay/Source/DisplayCluster/Private/Render/Device/Monoscopic/DisplayClusterDeviceMonoscopicDX11.cpp
@@ -1,12 +1,6 @@
 // Copyright 1998-2019 Epic Games, Inc. All Rights Reserved.
 
 #include "Render/Device/Monoscopic/DisplayClusterDeviceMonoscopicDX11.h"
-
-#include "Render/IDisplayClusterRenderManager.h"
-#include "Render/Device/DisplayClusterRenderViewport.h"
-#include "Render/Synchronization/IDisplayClusterRenderSyncPolicy.h"
-
-#include "DisplayClusterGlobals.h"
 #include "DisplayClusterLog.h"
 
 
@@ -24,29 +18,12 @@ bool FDisplayClusterDeviceMonoscopicDX11::Present(int32& InOutSyncInterval)
 {
 	DISPLAY_CLUSTER_FUNC_TRACE(LogDisplayClusterRender);
 
-	FD3D11Viewport* Viewport = static_cast<FD3D11Viewport*>(MainViewport->GetViewportRHI().GetReference());
-	check(Viewport);
-
-#if !WITH_EDITOR
-	// Issue frame event
-	Viewport->IssueFrameEvent();
-	// Wait until GPU finish last frame commands
-	Viewport->WaitForFrameEventCompletion();
-#endif
-
-	// Update sync value with nDisplay value
-	InOutSyncInterval = GetSwapInt();
-
-	TSharedPtr<IDisplayClusterRenderSyncPolicy> SyncPolicy = GDisplayCluster->GetRenderMgr()->GetCurrentSynchronizationPolicy();
-	if (SyncPolicy.IsValid())
+	// Perform abstract synchronization on a higher level
+	if (!FDisplayClusterDeviceBase::Present(InOutSyncInterval))
 	{
-		// False results means we don't need to present current frame, the sync object already presented it
-		if (!SyncPolicy->SynchronizeClusterRendering(InOutSyncInterval))
-		{
-			return false;
-		}
+		return false;
 	}
 
-	// Sync policy haven't presented current frame so we do it ourselves
-	return PresentImpl(Viewport, InOutSyncInterval);
+	// A sync policy hasn't presented current frame so we do it ourselves
+	return PresentImpl(MainViewport->GetViewportRHI(), InOutSyncInterval);
 }
diff --git a/Engine/Plugins/Runtime/nDisplay/Source/DisplayCluster/Private/Render/Device/Monoscopic/DisplayClusterDeviceMonoscopicDX11.h b/Engine/Plugins/Runtime/nDisplay/Source/DisplayCluster/Private/Render/Device/Monoscopic/DisplayClusterDeviceMonoscopicDX11.h
--- a/Engine/Plugins/Runtime/nDisplay/Source/DisplayCluster/Private/Render/Device/Monoscopic/DisplayClusterDeviceMonoscopicDX11.h
+++ b/Engine/Plugins/Runtime/nDisplay/Source/DisplayCluster/Private/Render/Device/Monoscopic/DisplayClusterDeviceMonoscopicDX11.h
@@ -3,7 +3,6 @@
 #pragma once
 
 #include "Render/Device/DisplayClusterDeviceMonoscopicBase.h"
-#include "Render/Device/DisplayClusterDeviceInternals.h"
 #include "Render/Presentation/DisplayClusterDevicePresentationDX11.h"
 
 
diff --git a/Engine/Plugins/Runtime/nDisplay/Source/DisplayCluster/Private/Render/Device/Monoscopic/DisplayClusterDeviceMonoscopicDX12.cpp b/Engine/Plugins/Runtime/nDisplay/Source/DisplayCluster/Private/Render/Device/Monoscopic/DisplayClusterDeviceMonoscopicDX12.cpp
--- a/Engine/Plugins/Runtime/nDisplay/Source/DisplayCluster/Private/Render/Device/Monoscopic/DisplayClusterDeviceMonoscopicDX12.cpp
+++ b/Engine/Plugins/Runtime/nDisplay/Source/DisplayCluster/Private/Render/Device/Monoscopic/DisplayClusterDeviceMonoscopicDX12.cpp
@@ -1,16 +1,8 @@
 // Copyright 1998-2019 Epic Games, Inc. All Rights Reserved.
 
 #include "Render/Device/Monoscopic/DisplayClusterDeviceMonoscopicDX12.h"
-
-#include "Render/Device/DisplayClusterRenderViewport.h"
-#include "Render/Synchronization/IDisplayClusterRenderSyncPolicy.h"
-
-#include "DisplayClusterGlobals.h"
 #include "DisplayClusterLog.h"
 
-#include "D3D12Viewport.h"
-//#include "D3D12Resources.h"
-
 
 FDisplayClusterDeviceMonoscopicDX12::FDisplayClusterDeviceMonoscopicDX12()
 {
@@ -26,29 +18,12 @@ bool FDisplayClusterDeviceMonoscopicDX12::Present(int32& InOutSyncInterval)
 {
 	DISPLAY_CLUSTER_FUNC_TRACE(LogDisplayClusterRender);
 
-	FD3D12Viewport* Viewport = static_cast<FD3D12Viewport*>(MainViewport->GetViewportRHI().GetReference());
-	check(Viewport);
-
-#if !WITH_EDITOR
-	// Issue frame event
-	Viewport->IssueFrameEvent();
-	// Wait until GPU finish last frame commands
-	Viewport->WaitForFrameEventCompletion();
-#endif
-
-	// Update sync value with nDisplay value
-	InOutSyncInterval = GetSwapInt();
-
-	TSharedPtr<IDisplayClusterRenderSyncPolicy> SyncPolicy = GDisplayCluster->GetRenderMgr()->GetCurrentSynchronizationPolicy();
-	if (SyncPolicy.IsValid())
+	// Perform abstract synchronization on a higher level
+	if (!FDisplayClusterDeviceBase::Present(InOutSyncInterval))
 	{
-		// False results means we don't need to present current frame, the sync object already presented it
-		if (!SyncPolicy->SynchronizeClusterRendering(InOutSyncInterval))
-		{
-			return false;
-		}
+		return false;
 	}
 
-	// Sync policy haven't presented current frame so we do it ourselves
-	return PresentImpl(Viewport, InOutSyncInterval);
+	// A sync policy hasn't presented current frame so we do it ourselves
+	return PresentImpl(MainViewport->GetViewportRHI(), InOutSyncInterval);
 }
diff --git a/Engine/Plugins/Runtime/nDisplay/Source/DisplayCluster/Private/Render/Device/Monoscopic/DisplayClusterDeviceMonoscopicDX12.h b/Engine/Plugins/Runtime/nDisplay/Source/DisplayCluster/Private/Render/Device/Monoscopic/DisplayClusterDeviceMonoscopicDX12.h
--- a/Engine/Plugins/Runtime/nDisplay/Source/DisplayCluster/Private/Render/Device/Monoscopic/DisplayClusterDeviceMonoscopicDX12.h
+++ b/Engine/Plugins/Runtime/nDisplay/Source/DisplayCluster/Private/Render/Device/Monoscopic/DisplayClusterDeviceMonoscopicDX12.h
@@ -4,7 +4,6 @@
 
 #include "CoreMinimal.h"
 #include "Render/Device/DisplayClusterDeviceMonoscopicBase.h"
-#include "Render/Device/DisplayClusterDeviceInternals.h"
 #include "Render/Presentation/DisplayClusterDevicePresentationDX12.h"
 
 
diff --git a/Engine/Plugins/Runtime/nDisplay/Source/DisplayCluster/Private/Render/Device/Monoscopic/DisplayClusterDeviceMonoscopicOpenGL.cpp b/Engine/Plugins/Runtime/nDisplay/Source/DisplayCluster/Private/Render/Device/Monoscopic/DisplayClusterDeviceMonoscopicOpenGL.cpp
--- a/Engine/Plugins/Runtime/nDisplay/Source/DisplayCluster/Private/Render/Device/Monoscopic/DisplayClusterDeviceMonoscopicOpenGL.cpp
+++ b/Engine/Plugins/Runtime/nDisplay/Source/DisplayCluster/Private/Render/Device/Monoscopic/DisplayClusterDeviceMonoscopicOpenGL.cpp
@@ -1,13 +1,9 @@
 // Copyright 1998-2019 Epic Games, Inc. All Rights Reserved.
 
 #include "Render/Device/Monoscopic/DisplayClusterDeviceMonoscopicOpenGL.h"
-#include "Render/Device/DisplayClusterDeviceInternals.h"
 
-#include "Render/IDisplayClusterRenderManager.h"
-#include "Render/Synchronization/IDisplayClusterRenderSyncPolicy.h"
-
-#include "DisplayClusterGlobals.h"
 #include "DisplayClusterLog.h"
+#include "Render/Device/DisplayClusterOpenGL.h"
 
 
 FDisplayClusterDeviceMonoscopicOpenGL::FDisplayClusterDeviceMonoscopicOpenGL()
@@ -60,17 +56,10 @@ bool FDisplayClusterDeviceMonoscopicOpenGL::Present(int32& InOutSyncInterval)
 	pOglViewport->WaitForFrameEventCompletion();
 #endif
 
-	// Update sync value with nDisplay value
-	InOutSyncInterval = GetSwapInt();
-
-	TSharedPtr<IDisplayClusterRenderSyncPolicy> SyncPolicy = GDisplayCluster->GetRenderMgr()->GetCurrentSynchronizationPolicy();
-	if (SyncPolicy.IsValid())
+	// Perform abstract synchronization on a higher level
+	if (!FDisplayClusterDeviceBase::Present(InOutSyncInterval))
 	{
-		// False results means we don't need to present current frame, the sync object already presented it
-		if (!SyncPolicy->SynchronizeClusterRendering(InOutSyncInterval))
-		{
-			return false;
-		}
+		return false;
 	}
 
 	::SwapBuffers(pOglViewport->GetGLContext()->DeviceContext);
diff --git a/Engine/Plugins/Runtime/nDisplay/Source/DisplayCluster/Private/Render/Device/QuadBufferStereo/DisplayClusterDeviceQuadBufferStereoDX11.cpp b/Engine/Plugins/Runtime/nDisplay/Source/DisplayCluster/Private/Render/Device/QuadBufferStereo/DisplayClusterDeviceQuadBufferStereoDX11.cpp
--- a/Engine/Plugins/Runtime/nDisplay/Source/DisplayCluster/Private/Render/Device/QuadBufferStereo/DisplayClusterDeviceQuadBufferStereoDX11.cpp
+++ b/Engine/Plugins/Runtime/nDisplay/Source/DisplayCluster/Private/Render/Device/QuadBufferStereo/DisplayClusterDeviceQuadBufferStereoDX11.cpp
@@ -1,11 +1,6 @@
 // Copyright 1998-2019 Epic Games, Inc. All Rights Reserved.
 
 #include "Render/Device/QuadBufferStereo/DisplayClusterDeviceQuadBufferStereoDX11.h"
-#include "Render/Device/DisplayClusterDeviceInternals.h"
-#include "Render/IDisplayClusterRenderManager.h"
-#include "Render/Synchronization/IDisplayClusterRenderSyncPolicy.h"
-
-#include "DisplayClusterGlobals.h"
 #include "DisplayClusterLog.h"
 
 
@@ -23,29 +18,12 @@ bool FDisplayClusterDeviceQuadBufferStereoDX11::Present(int32& InOutSyncInterval
 {
 	DISPLAY_CLUSTER_FUNC_TRACE(LogDisplayClusterRender);
 
-	FD3D11Viewport* Viewport = static_cast<FD3D11Viewport*>(MainViewport->GetViewportRHI().GetReference());
-	check(Viewport);
-
-#if !WITH_EDITOR
-	// Issue frame event
-	Viewport->IssueFrameEvent();
-	// Wait until GPU finish last frame commands
-	Viewport->WaitForFrameEventCompletion();
-#endif
-
-	// Update sync value with nDisplay value
-	InOutSyncInterval = GetSwapInt();
-
-	TSharedPtr<IDisplayClusterRenderSyncPolicy> SyncPolicy = GDisplayCluster->GetRenderMgr()->GetCurrentSynchronizationPolicy();
-	if (SyncPolicy.IsValid())
+	// Perform abstract synchronization on a higher level
+	if (!FDisplayClusterDeviceBase::Present(InOutSyncInterval))
 	{
-		// False results means we don't need to present current frame, the sync object already presented it
-		if (!SyncPolicy->SynchronizeClusterRendering(InOutSyncInterval))
-		{
-			return false;
-		}
+		return false;
 	}
 
-	// Sync policy haven't presented current frame so we do it ourselves
-	return PresentImpl(Viewport, InOutSyncInterval);
+	// A sync policy hasn't presented current frame so we do it ourselves
+	return PresentImpl(MainViewport->GetViewportRHI(), InOutSyncInterval);
 }
diff --git a/Engine/Plugins/Runtime/nDisplay/Source/DisplayCluster/Private/Render/Device/QuadBufferStereo/DisplayClusterDeviceQuadBufferStereoDX12.cpp b/Engine/Plugins/Runtime/nDisplay/Source/DisplayCluster/Private/Render/Device/QuadBufferStereo/DisplayClusterDeviceQuadBufferStereoDX12.cpp
--- a/Engine/Plugins/Runtime/nDisplay/Source/DisplayCluster/Private/Render/Device/QuadBufferStereo/DisplayClusterDeviceQuadBufferStereoDX12.cpp
+++ b/Engine/Plugins/Runtime/nDisplay/Source/DisplayCluster/Private/Render/Device/QuadBufferStereo/DisplayClusterDeviceQuadBufferStereoDX12.cpp
@@ -1,16 +1,8 @@
 // Copyright 1998-2019 Epic Games, Inc. All Rights Reserved.
 
 #include "Render/Device/QuadBufferStereo/DisplayClusterDeviceQuadBufferStereoDX12.h"
-#include "Render/Device/DisplayClusterDeviceInternals.h"
-
-#include "Render/IDisplayClusterRenderManager.h"
-#include "Render/Synchronization/IDisplayClusterRenderSyncPolicy.h"
-
-#include "DisplayClusterGlobals.h"
 #include "DisplayClusterLog.h"
 
-#include "D3D12Viewport.h"
-
 
 FDisplayClusterDeviceQuadBufferStereoDX12::FDisplayClusterDeviceQuadBufferStereoDX12()
 {
@@ -26,29 +18,12 @@ bool FDisplayClusterDeviceQuadBufferStereoDX12::Present(int32& InOutSyncInterval
 {
 	DISPLAY_CLUSTER_FUNC_TRACE(LogDisplayClusterRender);
 
-	FD3D12Viewport* Viewport = static_cast<FD3D12Viewport*>(MainViewport->GetViewportRHI().GetReference());
-	check(Viewport);
-
-#if !WITH_EDITOR
-	// Issue frame event
-	Viewport->IssueFrameEvent();
-	// Wait until GPU finish last frame commands
-	Viewport->WaitForFrameEventCompletion();
-#endif
-
-	// Update sync value with nDisplay value
-	InOutSyncInterval = GetSwapInt();
-
-	TSharedPtr<IDisplayClusterRenderSyncPolicy> SyncPolicy = GDisplayCluster->GetRenderMgr()->GetCurrentSynchronizationPolicy();
-	if (SyncPolicy.IsValid())
+	// Perform abstract synchronization on a higher level
+	if (!FDisplayClusterDeviceBase::Present(InOutSyncInterval))
 	{
-		// False results means we don't need to present current frame, the sync object already presented it
-		if (!SyncPolicy->SynchronizeClusterRendering(InOutSyncInterval))
-		{
-			return false;
-		}
+		return false;
 	}
 
-	// Sync policy haven't presented current frame so we do it ourselves
-	return PresentImpl(Viewport, InOutSyncInterval);
+	// A sync policy hasn't presented current frame so we do it ourselves
+	return PresentImpl(MainViewport->GetViewportRHI(), InOutSyncInterval);
 }
diff --git a/Engine/Plugins/Runtime/nDisplay/Source/DisplayCluster/Private/Render/Device/QuadBufferStereo/DisplayClusterDeviceQuadBufferStereoOpenGL.cpp b/Engine/Plugins/Runtime/nDisplay/Source/DisplayCluster/Private/Render/Device/QuadBufferStereo/DisplayClusterDeviceQuadBufferStereoOpenGL.cpp
--- a/Engine/Plugins/Runtime/nDisplay/Source/DisplayCluster/Private/Render/Device/QuadBufferStereo/DisplayClusterDeviceQuadBufferStereoOpenGL.cpp
+++ b/Engine/Plugins/Runtime/nDisplay/Source/DisplayCluster/Private/Render/Device/QuadBufferStereo/DisplayClusterDeviceQuadBufferStereoOpenGL.cpp
@@ -1,13 +1,9 @@
 // Copyright 1998-2019 Epic Games, Inc. All Rights Reserved.
 
 #include "Render/Device/QuadBufferStereo/DisplayClusterDeviceQuadBufferStereoOpenGL.h"
-#include "Render/Device/DisplayClusterDeviceInternals.h"
 
-#include "Render/IDisplayClusterRenderManager.h"
-#include "Render/Synchronization/IDisplayClusterRenderSyncPolicy.h"
-
-#include "DisplayClusterGlobals.h"
 #include "DisplayClusterLog.h"
+#include "Render/Device/DisplayClusterOpenGL.h"
 
 
 FDisplayClusterDeviceQuadBufferStereoOpenGL::FDisplayClusterDeviceQuadBufferStereoOpenGL()
@@ -76,17 +72,10 @@ bool FDisplayClusterDeviceQuadBufferStereoOpenGL::Present(int32& InOutSyncInterv
 	pOglViewport->WaitForFrameEventCompletion();
 #endif
 
-	// Update sync value with nDisplay value
-	InOutSyncInterval = GetSwapInt();
-
-	TSharedPtr<IDisplayClusterRenderSyncPolicy> SyncPolicy = GDisplayCluster->GetRenderMgr()->GetCurrentSynchronizationPolicy();
-	if (SyncPolicy.IsValid())
+	// Perform abstract synchronization on a higher level
+	if (!FDisplayClusterDeviceBase::Present(InOutSyncInterval))
 	{
-		// False results means we don't need to present current frame, the sync object already presented it
-		if (!SyncPolicy->SynchronizeClusterRendering(InOutSyncInterval))
-		{
-			return false;
-		}
+		return false;
 	}
 
 	::SwapBuffers(pOglViewport->GetGLContext()->DeviceContext);
diff --git a/Engine/Plugins/Runtime/nDisplay/Source/DisplayCluster/Private/Render/Device/QuadBufferStereo/DisplayClusterDeviceQuadBufferStereoOpenGL.h b/Engine/Plugins/Runtime/nDisplay/Source/DisplayCluster/Private/Render/Device/QuadBufferStereo/DisplayClusterDeviceQuadBufferStereoOpenGL.h
--- a/Engine/Plugins/Runtime/nDisplay/Source/DisplayCluster/Private/Render/Device/QuadBufferStereo/DisplayClusterDeviceQuadBufferStereoOpenGL.h
+++ b/Engine/Plugins/Runtime/nDisplay/Source/DisplayCluster/Private/Render/Device/QuadBufferStereo/DisplayClusterDeviceQuadBufferStereoOpenGL.h
@@ -3,7 +3,8 @@
 #pragma once
 
 #include "Render/Device/QuadBufferStereo/DisplayClusterDeviceQuadBufferStereoBase.h"
-#include "Render/Device/DisplayClusterDeviceInternals.h"
+
+class FOpenGLViewport;
 
 
 /**
diff --git a/Engine/Plugins/Runtime/nDisplay/Source/DisplayCluster/Private/Render/Device/SideBySide/DisplayClusterDeviceSideBySideDX11.cpp b/Engine/Plugins/Runtime/nDisplay/Source/DisplayCluster/Private/Render/Device/SideBySide/DisplayClusterDeviceSideBySideDX11.cpp
--- a/Engine/Plugins/Runtime/nDisplay/Source/DisplayCluster/Private/Render/Device/SideBySide/DisplayClusterDeviceSideBySideDX11.cpp
+++ b/Engine/Plugins/Runtime/nDisplay/Source/DisplayCluster/Private/Render/Device/SideBySide/DisplayClusterDeviceSideBySideDX11.cpp
@@ -1,10 +1,6 @@
 // Copyright 1998-2019 Epic Games, Inc. All Rights Reserved.
 
 #include "Render/Device/SidebySide/DisplayClusterDeviceSideBySideDX11.h"
-#include "Render/IDisplayClusterRenderManager.h"
-#include "Render/Synchronization/IDisplayClusterRenderSyncPolicy.h"
-
-#include "DisplayClusterGlobals.h"
 #include "DisplayClusterLog.h"
 
 
@@ -23,29 +19,12 @@ bool FDisplayClusterDeviceSideBySideDX11::Present(int32& InOutSyncInterval)
 {
 	DISPLAY_CLUSTER_FUNC_TRACE(LogDisplayClusterRender);
 
-	FD3D11Viewport* Viewport = static_cast<FD3D11Viewport*>(MainViewport->GetViewportRHI().GetReference());
-	check(Viewport);
-
-#if !WITH_EDITOR
-	// Issue frame event
-	Viewport->IssueFrameEvent();
-	// Wait until GPU finish last frame commands
-	Viewport->WaitForFrameEventCompletion();
-#endif
-
-	// Update sync value with nDisplay value
-	InOutSyncInterval = GetSwapInt();
-
-	TSharedPtr<IDisplayClusterRenderSyncPolicy> SyncPolicy = GDisplayCluster->GetRenderMgr()->GetCurrentSynchronizationPolicy();
-	if (SyncPolicy.IsValid())
+	// Perform abstract synchronization on a higher level
+	if (!FDisplayClusterDeviceBase::Present(InOutSyncInterval))
 	{
-		// False results means we don't need to present current frame, the sync object already presented it
-		if (!SyncPolicy->SynchronizeClusterRendering(InOutSyncInterval))
-		{
-			return false;
-		}
+		return false;
 	}
 
-	// Sync policy haven't presented current frame so we do it ourselves
-	return PresentImpl(Viewport, InOutSyncInterval);
+	// A sync policy hasn't presented current frame so we do it ourselves
+	return PresentImpl(MainViewport->GetViewportRHI(), InOutSyncInterval);
 }
diff --git a/Engine/Plugins/Runtime/nDisplay/Source/DisplayCluster/Private/Render/Device/SideBySide/DisplayClusterDeviceSideBySideDX12.cpp b/Engine/Plugins/Runtime/nDisplay/Source/DisplayCluster/Private/Render/Device/SideBySide/DisplayClusterDeviceSideBySideDX12.cpp
--- a/Engine/Plugins/Runtime/nDisplay/Source/DisplayCluster/Private/Render/Device/SideBySide/DisplayClusterDeviceSideBySideDX12.cpp
+++ b/Engine/Plugins/Runtime/nDisplay/Source/DisplayCluster/Private/Render/Device/SideBySide/DisplayClusterDeviceSideBySideDX12.cpp
@@ -1,15 +1,8 @@
 // Copyright 1998-2019 Epic Games, Inc. All Rights Reserved.
 
 #include "Render/Device/SidebySide/DisplayClusterDeviceSideBySideDX12.h"
-
-#include "Render/IDisplayClusterRenderManager.h"
-#include "Render/Synchronization/IDisplayClusterRenderSyncPolicy.h"
-
-#include "DisplayClusterGlobals.h"
 #include "DisplayClusterLog.h"
 
-#include "D3D12Viewport.h"
-
 
 FDisplayClusterDeviceSideBySideDX12::FDisplayClusterDeviceSideBySideDX12()
 {
@@ -26,29 +19,12 @@ bool FDisplayClusterDeviceSideBySideDX12::Present(int32& InOutSyncInterval)
 {
 	DISPLAY_CLUSTER_FUNC_TRACE(LogDisplayClusterRender);
 
-	FD3D12Viewport* Viewport = static_cast<FD3D12Viewport*>(MainViewport->GetViewportRHI().GetReference());
-	check(Viewport);
-
-#if !WITH_EDITOR
-	// Issue frame event
-	Viewport->IssueFrameEvent();
-	// Wait until GPU finish last frame commands
-	Viewport->WaitForFrameEventCompletion();
-#endif
-
-	// Update sync value with nDisplay value
-	InOutSyncInterval = GetSwapInt();
-
-	TSharedPtr<IDisplayClusterRenderSyncPolicy> SyncPolicy = GDisplayCluster->GetRenderMgr()->GetCurrentSynchronizationPolicy();
-	if (SyncPolicy.IsValid())
+	// Perform abstract synchronization on a higher level
+	if (!FDisplayClusterDeviceBase::Present(InOutSyncInterval))
 	{
-		// False results means we don't need to present current frame, the sync object already presented it
-		if (!SyncPolicy->SynchronizeClusterRendering(InOutSyncInterval))
-		{
-			return false;
-		}
+		return false;
 	}
 
-	// Sync policy haven't presented current frame so we do it ourselves
-	return PresentImpl(Viewport, InOutSyncInterval);
+	// A sync policy hasn't presented current frame so we do it ourselves
+	return PresentImpl(MainViewport->GetViewportRHI(), InOutSyncInterval);
 }
diff --git a/Engine/Plugins/Runtime/nDisplay/Source/DisplayCluster/Private/Render/Device/TopBottom/DisplayClusterDeviceTopBottomDX11.cpp b/Engine/Plugins/Runtime/nDisplay/Source/DisplayCluster/Private/Render/Device/TopBottom/DisplayClusterDeviceTopBottomDX11.cpp
--- a/Engine/Plugins/Runtime/nDisplay/Source/DisplayCluster/Private/Render/Device/TopBottom/DisplayClusterDeviceTopBottomDX11.cpp
+++ b/Engine/Plugins/Runtime/nDisplay/Source/DisplayCluster/Private/Render/Device/TopBottom/DisplayClusterDeviceTopBottomDX11.cpp
@@ -1,11 +1,6 @@
 // Copyright 1998-2019 Epic Games, Inc. All Rights Reserved.
 
 #include "Render/Device/TopBottom/DisplayClusterDeviceTopBottomDX11.h"
-
-#include "Render/IDisplayClusterRenderManager.h"
-#include "Render/Synchronization/IDisplayClusterRenderSyncPolicy.h"
-
-#include "DisplayClusterGlobals.h"
 #include "DisplayClusterLog.h"
 
 
@@ -24,29 +19,12 @@ bool FDisplayClusterDeviceTopBottomDX11::Present(int32& InOutSyncInterval)
 {
 	DISPLAY_CLUSTER_FUNC_TRACE(LogDisplayClusterRender);
 
-	FD3D11Viewport* Viewport = static_cast<FD3D11Viewport*>(MainViewport->GetViewportRHI().GetReference());
-	check(Viewport);
-
-#if !WITH_EDITOR
-	// Issue frame event
-	Viewport->IssueFrameEvent();
-	// Wait until GPU finish last frame commands
-	Viewport->WaitForFrameEventCompletion();
-#endif
-
-	// Update sync value with nDisplay value
-	InOutSyncInterval = GetSwapInt();
-
-	TSharedPtr<IDisplayClusterRenderSyncPolicy> SyncPolicy = GDisplayCluster->GetRenderMgr()->GetCurrentSynchronizationPolicy();
-	if (SyncPolicy.IsValid())
+	// Perform abstract synchronization on a higher level
+	if (!FDisplayClusterDeviceBase::Present(InOutSyncInterval))
 	{
-		// False results means we don't need to present current frame, the sync object already presented it
-		if (!SyncPolicy->SynchronizeClusterRendering(InOutSyncInterval))
-		{
-			return false;
-		}
+		return false;
 	}
 
-	// Sync policy haven't presented current frame so we do it ourselves
-	return PresentImpl(Viewport, InOutSyncInterval);
+	// A sync policy hasn't presented current frame so we do it ourselves
+	return PresentImpl(MainViewport->GetViewportRHI(), InOutSyncInterval);
 }
diff --git a/Engine/Plugins/Runtime/nDisplay/Source/DisplayCluster/Private/Render/Device/TopBottom/DisplayClusterDeviceTopBottomDX12.cpp b/Engine/Plugins/Runtime/nDisplay/Source/DisplayCluster/Private/Render/Device/TopBottom/DisplayClusterDeviceTopBottomDX12.cpp
--- a/Engine/Plugins/Runtime/nDisplay/Source/DisplayCluster/Private/Render/Device/TopBottom/DisplayClusterDeviceTopBottomDX12.cpp
+++ b/Engine/Plugins/Runtime/nDisplay/Source/DisplayCluster/Private/Render/Device/TopBottom/DisplayClusterDeviceTopBottomDX12.cpp
@@ -1,15 +1,8 @@
 // Copyright 1998-2019 Epic Games, Inc. All Rights Reserved.
 
 #include "Render/Device/TopBottom/DisplayClusterDeviceTopBottomDX12.h"
-
-#include "Render/IDisplayClusterRenderManager.h"
-#include "Render/Synchronization/IDisplayClusterRenderSyncPolicy.h"
-
-#include "DisplayClusterGlobals.h"
 #include "DisplayClusterLog.h"
 
-#include "D3D12Viewport.h"
-
 
 FDisplayClusterDeviceTopBottomDX12::FDisplayClusterDeviceTopBottomDX12()
 {
@@ -26,29 +19,12 @@ bool FDisplayClusterDeviceTopBottomDX12::Present(int32& InOutSyncInterval)
 {
 	DISPLAY_CLUSTER_FUNC_TRACE(LogDisplayClusterRender);
 
-	FD3D12Viewport* Viewport = static_cast<FD3D12Viewport*>(MainViewport->GetViewportRHI().GetReference());
-	check(Viewport);
-
-#if !WITH_EDITOR
-	// Issue frame event
-	Viewport->IssueFrameEvent();
-	// Wait until GPU finish last frame commands
-	Viewport->WaitForFrameEventCompletion();
-#endif
-
-	// Update sync value with nDisplay value
-	InOutSyncInterval = GetSwapInt();
-
-	TSharedPtr<IDisplayClusterRenderSyncPolicy> SyncPolicy = GDisplayCluster->GetRenderMgr()->GetCurrentSynchronizationPolicy();
-	if (SyncPolicy.IsValid())
+	// Perform abstract synchronization on a higher level
+	if (!FDisplayClusterDeviceBase::Present(InOutSyncInterval))
 	{
-		// False results means we don't need to present current frame, the sync object already presented it
-		if (!SyncPolicy->SynchronizeClusterRendering(InOutSyncInterval))
-		{
-			return false;
-		}
+		return false;
 	}
 
-	// Sync policy haven't presented current frame so we do it ourselves
-	return PresentImpl(Viewport, InOutSyncInterval);
+	// A sync policy hasn't presented current frame so we do it ourselves
+	return PresentImpl(MainViewport->GetViewportRHI(), InOutSyncInterval);
 }
