nDisplay features

diff --git a/Engine/Plugins/Runtime/nDisplay/Source/DisplayCluster/Private/Render/DisplayClusterRenderManager.cpp b/Engine/Plugins/Runtime/nDisplay/Source/DisplayCluster/Private/Render/DisplayClusterRenderManager.cpp
--- a/Engine/Plugins/Runtime/nDisplay/Source/DisplayCluster/Private/Render/DisplayClusterRenderManager.cpp
+++ b/Engine/Plugins/Runtime/nDisplay/Source/DisplayCluster/Private/Render/DisplayClusterRenderManager.cpp
@@ -456,6 +456,26 @@ void FDisplayClusterRenderManager::SetViewportCamera(const FString& InCameraId /
 	}
 }
 
+bool FDisplayClusterRenderManager::GetViewportRect(const FString& InViewportID, FIntRect& Rect)
+{
+	if (!RenderDevice.IsValid())
+	{
+		return false;
+	}
+
+	return RenderDevice->GetViewportRect(InViewportID, Rect);
+}
+
+void FDisplayClusterRenderManager::SetCustomPostProcessing(const FString& ViewportID, const FPostProcessSettings& PostProcessingSettings)
+{
+	if (!RenderDevice.IsValid())
+	{
+		return;
+	}
+
+	RenderDevice->SetCustomPostProcessing(ViewportID, PostProcessingSettings);
+}
+
 void FDisplayClusterRenderManager::SetInterpupillaryDistance(const FString& CameraId, float EyeDistance)
 {
 	DISPLAY_CLUSTER_FUNC_TRACE(LogDisplayClusterRender);
diff --git a/Engine/Plugins/Runtime/nDisplay/Source/DisplayCluster/Private/Render/DisplayClusterRenderManager.h b/Engine/Plugins/Runtime/nDisplay/Source/DisplayCluster/Private/Render/DisplayClusterRenderManager.h
--- a/Engine/Plugins/Runtime/nDisplay/Source/DisplayCluster/Private/Render/DisplayClusterRenderManager.h
+++ b/Engine/Plugins/Runtime/nDisplay/Source/DisplayCluster/Private/Render/DisplayClusterRenderManager.h
@@ -58,9 +58,15 @@ public:
 	virtual bool RegisterPostprocessOperation(const FString& InName, IPDisplayClusterRenderManager::FDisplayClusterPPInfo& InPPInfo) override;
 	virtual bool UnregisterPostprocessOperation(const FString& InName) override;
 	virtual TMap<FString, IPDisplayClusterRenderManager::FDisplayClusterPPInfo> GetRegisteredPostprocessOperations() const override;
+	// Custom Rendering Post-process
+	virtual void SetCustomPostProcessing(const FString& ViewportID, const FPostProcessSettings& PostProcessingSettings) override;
+
 	// Camera
 	virtual void SetViewportCamera(const FString& InCameraId = FString(), const FString& InViewportId = FString()) override;
 
+	// Viewports
+	virtual bool GetViewportRect(const FString& InViewportID, FIntRect& Rect) override;
+
 	// Camera API
 	virtual float GetInterpupillaryDistance(const FString& CameraId) const override;
 	virtual void  SetInterpupillaryDistance(const FString& CameraId, float EyeDistance) override;
diff --git a/Engine/Plugins/Runtime/nDisplay/Source/DisplayCluster/Private/Render/Presentation/DisplayClusterDevicePresentationDX11.cpp b/Engine/Plugins/Runtime/nDisplay/Source/DisplayCluster/Private/Render/Presentation/DisplayClusterDevicePresentationDX11.cpp
new file mode 100644
--- /dev/null
+++ b/Engine/Plugins/Runtime/nDisplay/Source/DisplayCluster/Private/Render/Presentation/DisplayClusterDevicePresentationDX11.cpp
@@ -0,0 +1,34 @@
+// Copyright 1998-2019 Epic Games, Inc. All Rights Reserved.
+
+#include "Render/Presentation/DisplayClusterDevicePresentationDX11.h"
+
+#include "DisplayClusterLog.h"
+
+#include "RHIResources.h"
+
+#include "Windows/AllowWindowsPlatformTypes.h"
+#include "DirectX/Include/DXGI.h"
+#include "Windows/HideWindowsPlatformTypes.h"
+
+
+bool FDisplayClusterDevicePresentationDX11::PresentImpl(FRHIViewport* Viewport, const int32 InSyncInterval)
+{
+	if (!Viewport)
+	{
+		check(Viewport);
+		UE_LOG(LogDisplayClusterRender, Warning, TEXT("Viewport is null"));
+		return true;
+	}
+
+	IDXGISwapChain* const SwapChain = static_cast<IDXGISwapChain*>(Viewport->GetNativeSwapChain());
+	if (!SwapChain)
+	{
+		check(SwapChain);
+		UE_LOG(LogDisplayClusterRender, Warning, TEXT("SwapChain is null"));
+		return true;
+	}
+
+	SwapChain->Present(InSyncInterval, 0);
+
+	return false;
+}
diff --git a/Engine/Plugins/Runtime/nDisplay/Source/DisplayCluster/Private/Render/Presentation/DisplayClusterDevicePresentationDX11.h b/Engine/Plugins/Runtime/nDisplay/Source/DisplayCluster/Private/Render/Presentation/DisplayClusterDevicePresentationDX11.h
--- a/Engine/Plugins/Runtime/nDisplay/Source/DisplayCluster/Private/Render/Presentation/DisplayClusterDevicePresentationDX11.h
+++ b/Engine/Plugins/Runtime/nDisplay/Source/DisplayCluster/Private/Render/Presentation/DisplayClusterDevicePresentationDX11.h
@@ -2,11 +2,9 @@
 
 #pragma once
 
-#include "Render/device/DisplayClusterDeviceInternals.h"
-#include "DisplayClusterLog.h"
+#include "CoreMinimal.h"
 
-//#include "UnrealClient.h"
-#include "D3D11Viewport.h"
+class FRHIViewport;
 
 
 /**
@@ -19,23 +17,5 @@ public:
 	virtual ~FDisplayClusterDevicePresentationDX11() = default;
 
 public:
-	bool PresentImpl(FRHIViewport* Viewport, const int32 InSyncInterval)
-	{
-		check(Viewport);
-		if (Viewport == nullptr)
-		{
-			UE_LOG(LogDisplayClusterRender, Warning, TEXT("Viewport is nullptr. Use native present."));
-			return true;
-		}
-
-		FD3D11Viewport* D3D11Viewport = static_cast<FD3D11Viewport*>(Viewport);
-		check(D3D11Viewport);
-
-		IDXGISwapChain* SwapChain = D3D11Viewport->GetSwapChain();
-		check(SwapChain);
-
-		SwapChain->Present(InSyncInterval, 0);
-
-		return false;
-	}
+	bool PresentImpl(FRHIViewport* Viewport, const int32 InSyncInterval);
 };
diff --git a/Engine/Plugins/Runtime/nDisplay/Source/DisplayCluster/Private/Render/Presentation/DisplayClusterDevicePresentationDX12.cpp b/Engine/Plugins/Runtime/nDisplay/Source/DisplayCluster/Private/Render/Presentation/DisplayClusterDevicePresentationDX12.cpp
new file mode 100644
--- /dev/null
+++ b/Engine/Plugins/Runtime/nDisplay/Source/DisplayCluster/Private/Render/Presentation/DisplayClusterDevicePresentationDX12.cpp
@@ -0,0 +1,34 @@
+// Copyright 1998-2019 Epic Games, Inc. All Rights Reserved.
+
+#include "Render/Presentation/DisplayClusterDevicePresentationDX12.h"
+
+#include "DisplayClusterLog.h"
+
+#include "RHIResources.h"
+
+#include "Windows/AllowWindowsPlatformTypes.h"
+#include "DirectX/Include/DXGI.h"
+#include "Windows/HideWindowsPlatformTypes.h"
+
+
+bool FDisplayClusterDevicePresentationDX12::PresentImpl(FRHIViewport* Viewport, const int32 InSyncInterval)
+{
+	if (!Viewport)
+	{
+		check(Viewport);
+		UE_LOG(LogDisplayClusterRender, Warning, TEXT("Viewport is null"));
+		return true;
+	}
+
+	IDXGISwapChain* const SwapChain = static_cast<IDXGISwapChain*>(Viewport->GetNativeSwapChain());
+	if (!SwapChain)
+	{
+		check(SwapChain);
+		UE_LOG(LogDisplayClusterRender, Warning, TEXT("SwapChain is null"));
+		return true;
+	}
+
+	SwapChain->Present(InSyncInterval, 0);
+
+	return false;
+}
diff --git a/Engine/Plugins/Runtime/nDisplay/Source/DisplayCluster/Private/Render/Presentation/DisplayClusterDevicePresentationDX12.h b/Engine/Plugins/Runtime/nDisplay/Source/DisplayCluster/Private/Render/Presentation/DisplayClusterDevicePresentationDX12.h
--- a/Engine/Plugins/Runtime/nDisplay/Source/DisplayCluster/Private/Render/Presentation/DisplayClusterDevicePresentationDX12.h
+++ b/Engine/Plugins/Runtime/nDisplay/Source/DisplayCluster/Private/Render/Presentation/DisplayClusterDevicePresentationDX12.h
@@ -2,10 +2,9 @@
 
 #pragma once
 
-#include "Render/device/DisplayClusterDeviceInternals.h"
-#include "DisplayClusterLog.h"
+#include "CoreMinimal.h"
 
-#include "D3D12Viewport.h"
+class FRHIViewport;
 
 
 /**
@@ -18,23 +17,5 @@ public:
 	virtual ~FDisplayClusterDevicePresentationDX12() = default;
 
 public:
-	bool PresentImpl(FRHIViewport* Viewport, const int32 InSyncInterval)
-	{
-		check(Viewport);
-		if (Viewport == nullptr)
-		{
-			UE_LOG(LogDisplayClusterRender, Warning, TEXT("Viewport is nullptr. Use native present."));
-			return true;
-		}
-
-		FD3D12Viewport* D3D12Viewport = static_cast<FD3D12Viewport*>(Viewport);
-		check(D3D12Viewport);
-
-		IDXGISwapChain* SwapChain = D3D12Viewport->GetSwapChain();
-		check(SwapChain);
-
-		SwapChain->Present(InSyncInterval, 0);
-
-		return false;
-	}
+	bool PresentImpl(FRHIViewport* Viewport, const int32 InSyncInterval);
 };
diff --git a/Engine/Plugins/Runtime/nDisplay/Source/DisplayCluster/Private/Render/Synchronization/DisplayClusterRenderSyncPolicyBase.cpp b/Engine/Plugins/Runtime/nDisplay/Source/DisplayCluster/Private/Render/Synchronization/DisplayClusterRenderSyncPolicyBase.cpp
--- a/Engine/Plugins/Runtime/nDisplay/Source/DisplayCluster/Private/Render/Synchronization/DisplayClusterRenderSyncPolicyBase.cpp
+++ b/Engine/Plugins/Runtime/nDisplay/Source/DisplayCluster/Private/Render/Synchronization/DisplayClusterRenderSyncPolicyBase.cpp
@@ -20,6 +20,8 @@ FDisplayClusterRenderSyncPolicyBase::~FDisplayClusterRenderSyncPolicyBase()
 
 void FDisplayClusterRenderSyncPolicyBase::SyncBarrierRenderThread()
 {
+	check(IsInRenderingThread());
+
 	DISPLAY_CLUSTER_FUNC_TRACE(LogDisplayClusterRender);
 
 	if (GDisplayCluster->GetOperationMode() == EDisplayClusterOperationMode::Disabled)
diff --git a/Engine/Plugins/Runtime/nDisplay/Source/DisplayCluster/Private/Render/Synchronization/DisplayClusterRenderSyncPolicyNone.cpp b/Engine/Plugins/Runtime/nDisplay/Source/DisplayCluster/Private/Render/Synchronization/DisplayClusterRenderSyncPolicyNone.cpp
--- a/Engine/Plugins/Runtime/nDisplay/Source/DisplayCluster/Private/Render/Synchronization/DisplayClusterRenderSyncPolicyNone.cpp
+++ b/Engine/Plugins/Runtime/nDisplay/Source/DisplayCluster/Private/Render/Synchronization/DisplayClusterRenderSyncPolicyNone.cpp
@@ -14,6 +14,8 @@ FDisplayClusterRenderSyncPolicyNone::~FDisplayClusterRenderSyncPolicyNone()
 
 bool FDisplayClusterRenderSyncPolicyNone::SynchronizeClusterRendering(int32& InOutSyncInterval)
 {
+	check(IsInRenderingThread());
+
 	// Override sync interval with 0 to show a frame ASAP. We don't care about tearing in this policy.
 	InOutSyncInterval = 0;
 	// Tell a caller that he still needs to present a frame
diff --git a/Engine/Plugins/Runtime/nDisplay/Source/DisplayCluster/Private/Render/Synchronization/DisplayClusterRenderSyncPolicyNvidiaOpenGL.cpp b/Engine/Plugins/Runtime/nDisplay/Source/DisplayCluster/Private/Render/Synchronization/DisplayClusterRenderSyncPolicyNvidiaOpenGL.cpp
--- a/Engine/Plugins/Runtime/nDisplay/Source/DisplayCluster/Private/Render/Synchronization/DisplayClusterRenderSyncPolicyNvidiaOpenGL.cpp
+++ b/Engine/Plugins/Runtime/nDisplay/Source/DisplayCluster/Private/Render/Synchronization/DisplayClusterRenderSyncPolicyNvidiaOpenGL.cpp
@@ -1,49 +1,90 @@
 // Copyright 1998-2019 Epic Games, Inc. All Rights Reserved.
 
 #include "Render/Synchronization/DisplayClusterRenderSyncPolicyNvidiaOpenGL.h"
-#include "Render/Device/DisplayClusterDeviceInternals.h"
 
 #include "DisplayClusterLog.h"
+#include "Render/Device/DisplayClusterOpenGL.h"
 
+#include "Engine/Engine.h"
 #include "Engine/GameViewportClient.h"
-#include "Engine/GameEngine.h"
+#include "UnrealClient.h"
 
 
-FDisplayClusterRenderSyncPolicyNvidiaOpenGL::FDisplayClusterRenderSyncPolicyNvidiaOpenGL()
+namespace
 {
-}
+	PFNWGLSWAPINTERVALEXTPROC      DisplayCluster_wglSwapIntervalEXT_ProcAddress = nullptr;
+	PFNWGLJOINSWAPGROUPNVPROC      DisplayCluster_wglJoinSwapGroupNV_ProcAddress = nullptr;
+	PFNWGLBINDSWAPBARRIERNVPROC    DisplayCluster_wglBindSwapBarrierNV_ProcAddress = nullptr;
+	PFNWGLQUERYSWAPGROUPNVPROC     DisplayCluster_wglQuerySwapGroupNV_ProcAddress = nullptr;
+	PFNWGLQUERYMAXSWAPGROUPSNVPROC DisplayCluster_wglQueryMaxSwapGroupsNV_ProcAddress = nullptr;
+	PFNWGLQUERYFRAMECOUNTNVPROC    DisplayCluster_wglQueryFrameCountNV_ProcAddress = nullptr;
+	PFNWGLRESETFRAMECOUNTNVPROC    DisplayCluster_wglResetFrameCountNV_ProcAddress = nullptr;
 
-FDisplayClusterRenderSyncPolicyNvidiaOpenGL::~FDisplayClusterRenderSyncPolicyNvidiaOpenGL()
-{
+
+	// Copy/pasted from OpenGLDrv.cpp
+	static void DisplayClusterGetExtensionsString(FString& ExtensionsString)
+	{
+		GLint ExtensionCount = 0;
+		ExtensionsString = TEXT("");
+		if (FOpenGL::SupportsIndexedExtensions())
+		{
+			glGetIntegerv(GL_NUM_EXTENSIONS, &ExtensionCount);
+			for (int32 ExtensionIndex = 0; ExtensionIndex < ExtensionCount; ++ExtensionIndex)
+			{
+				const ANSICHAR* ExtensionString = FOpenGL::GetStringIndexed(GL_EXTENSIONS, ExtensionIndex);
+
+				ExtensionsString += TEXT(" ");
+				ExtensionsString += ANSI_TO_TCHAR(ExtensionString);
+			}
+		}
+		else
+		{
+			const ANSICHAR* GlGetStringOutput = (const ANSICHAR*)glGetString(GL_EXTENSIONS);
+			if (GlGetStringOutput)
+			{
+				ExtensionsString += GlGetStringOutput;
+				ExtensionsString += TEXT(" ");
+			}
+		}
+	}
 }
 
-bool FDisplayClusterRenderSyncPolicyNvidiaOpenGL::SynchronizeClusterRendering(int32& InOutSyncInterval)
+
+FDisplayClusterRenderSyncPolicyNvidiaOpenGL::FDisplayClusterRenderSyncPolicyNvidiaOpenGL()
 {
-	//@todo implement hardware based solution
-	SyncBarrierRenderThread();
-	// Tell a caller that he still needs to present a frame
-	return true;
+	InitializeOpenGLCapabilities();
 }
 
-#if 0
-bool FDisplayClusterRenderSyncPolicyNvidiaOpenGL::Initialize()
+FDisplayClusterRenderSyncPolicyNvidiaOpenGL::~FDisplayClusterRenderSyncPolicyNvidiaOpenGL()
 {
-#if PLATFORM_WINDOWS
-	DisplayClusterInitCapabilitiesForGL();
-#endif
-
-	return true;
 }
 
-bool FDisplayClusterRenderSyncPolicyNvidiaOpenGL::Present(int32& InOutSyncInterval)
+bool FDisplayClusterRenderSyncPolicyNvidiaOpenGL::SynchronizeClusterRendering(int32& InOutSyncInterval)
 {
-	if (bNvSwapInitialized == false)
+	check(IsInRenderingThread());
+
+	if (bNVAPIInitialized == false)
 	{
 		// Use render barrier to guaranty that all nv barriers are initialized simultaneously
 		SyncBarrierRenderThread();
-		bNvSwapInitialized = InitializeNvidiaSwapLock();
+		bNVAPIInitialized = InitializeNvidiaSwapLock();
+	}
+
+	if (!bNVAPIInitialized)
+	{
+		return true;
+	}
+
+	if(!(GEngine && GEngine->GameViewport && GEngine->GameViewport->Viewport))
+	{
+		return false;
 	}
 
+	FOpenGLViewport* pOglViewport = static_cast<FOpenGLViewport*>(GEngine->GameViewport->Viewport->GetViewportRHI().GetReference());
+	check(pOglViewport);
+	FPlatformOpenGLContext* const pContext = pOglViewport->GetGLContext();
+	check(pContext && pContext->DeviceContext);
+
 	if (OpenGLContext)
 	{
 		UpdateSwapInterval(InOutSyncInterval);
@@ -54,10 +95,9 @@ bool FDisplayClusterRenderSyncPolicyNvidiaOpenGL::Present(int32& InOutSyncInterv
 	return true;
 }
 
-#if PLATFORM_WINDOWS
 bool FDisplayClusterRenderSyncPolicyNvidiaOpenGL::InitializeNvidiaSwapLock()
 {
-	if (bNvSwapInitialized)
+	if (bNVAPIInitialized)
 	{
 		return true;
 	}
@@ -68,13 +108,12 @@ bool FDisplayClusterRenderSyncPolicyNvidiaOpenGL::InitializeNvidiaSwapLock()
 		return false;
 	}
 
-	if (!GEngine || !GEngine->GameViewport || 
-		!GEngine->GameViewport->Viewport || 
-		!GEngine->GameViewport->Viewport->GetViewportRHI().IsValid())
+	if (!GEngine || !GEngine->GameViewport || !GEngine->GameViewport->Viewport || !GEngine->GameViewport->Viewport->GetViewportRHI().IsValid())
 	{
 		UE_LOG(LogDisplayClusterRender, Warning, TEXT("Game viewport hasn't been initialized yet"));
 		return false;
 	}
+
 	FViewport* MainViewport = GEngine->GameViewport->Viewport;
 	OpenGLViewport = static_cast<FOpenGLViewport*>(MainViewport->GetViewportRHI().GetReference());
 	check(OpenGLViewport);
@@ -120,20 +159,41 @@ bool FDisplayClusterRenderSyncPolicyNvidiaOpenGL::InitializeNvidiaSwapLock()
 
 	return true;
 }
-#elif PLATFORM_LINUX
-bool FDisplayClusterRenderSyncPolicyNvidiaOpenGL::InitializeNvidiaSwapLock()
+
+void FDisplayClusterRenderSyncPolicyNvidiaOpenGL::InitializeOpenGLCapabilities()
 {
-	//@todo: Implementation for Linux
-	return false;
-}
-#endif
+	bool bWindowsSwapControlExtensionPresent = false;
+	{
+		FString ExtensionsString;
+		DisplayClusterGetExtensionsString(ExtensionsString);
+
+		if (ExtensionsString.Contains(TEXT("WGL_EXT_swap_control")))
+		{
+			bWindowsSwapControlExtensionPresent = true;
+		}
+	}
 
+#pragma warning(push)
+#pragma warning(disable:4191)
+	if (bWindowsSwapControlExtensionPresent)
+	{
+		DisplayCluster_wglSwapIntervalEXT_ProcAddress = (PFNWGLSWAPINTERVALEXTPROC)wglGetProcAddress("wglSwapIntervalEXT");
+	}
+
+	DisplayCluster_wglJoinSwapGroupNV_ProcAddress = (PFNWGLJOINSWAPGROUPNVPROC)wglGetProcAddress("wglJoinSwapGroupNV");
+	DisplayCluster_wglBindSwapBarrierNV_ProcAddress = (PFNWGLBINDSWAPBARRIERNVPROC)wglGetProcAddress("wglBindSwapBarrierNV");
+	DisplayCluster_wglQuerySwapGroupNV_ProcAddress = (PFNWGLQUERYSWAPGROUPNVPROC)wglGetProcAddress("wglQuerySwapGroupNV");
+	DisplayCluster_wglQueryMaxSwapGroupsNV_ProcAddress = (PFNWGLQUERYMAXSWAPGROUPSNVPROC)wglGetProcAddress("wglQueryMaxSwapGroupsNV");
+	DisplayCluster_wglQueryFrameCountNV_ProcAddress = (PFNWGLQUERYFRAMECOUNTNVPROC)wglGetProcAddress("wglQueryFrameCountNV");
+	DisplayCluster_wglResetFrameCountNV_ProcAddress = (PFNWGLRESETFRAMECOUNTNVPROC)wglGetProcAddress("wglResetFrameCountNV");
+#pragma warning(pop)
+}
 
 void FDisplayClusterRenderSyncPolicyNvidiaOpenGL::UpdateSwapInterval(int32 InSyncInterval) const
 {
-#if PLATFORM_WINDOWS
 	/*
 	https://www.opengl.org/registry/specs/EXT/wgl_swap_control.txt
+
 	wglSwapIntervalEXT specifies the minimum number of video frame periods
 	per buffer swap for the window associated with the current context.
 	The interval takes effect when SwapBuffers or wglSwapLayerBuffer
@@ -163,9 +223,4 @@ void FDisplayClusterRenderSyncPolicyNvidiaOpenGL::UpdateSwapInterval(int32 InSyn
 	// Perform that each frame
 	if (!DisplayCluster_wglSwapIntervalEXT_ProcAddress(InSyncInterval))
 		UE_LOG(LogDisplayClusterRender, Warning, TEXT("Couldn't set swap interval: %d"), InSyncInterval);
-
-#elif
-	//@todo: Implementation for Linux
-#endif
 }
-#endif
diff --git a/Engine/Plugins/Runtime/nDisplay/Source/DisplayCluster/Private/Render/Synchronization/DisplayClusterRenderSyncPolicyNvidiaOpenGL.h b/Engine/Plugins/Runtime/nDisplay/Source/DisplayCluster/Private/Render/Synchronization/DisplayClusterRenderSyncPolicyNvidiaOpenGL.h
--- a/Engine/Plugins/Runtime/nDisplay/Source/DisplayCluster/Private/Render/Synchronization/DisplayClusterRenderSyncPolicyNvidiaOpenGL.h
+++ b/Engine/Plugins/Runtime/nDisplay/Source/DisplayCluster/Private/Render/Synchronization/DisplayClusterRenderSyncPolicyNvidiaOpenGL.h
@@ -23,14 +23,13 @@ public:
 	//////////////////////////////////////////////////////////////////////////////////////////////
 	virtual bool SynchronizeClusterRendering(int32& InOutSyncInterval) override;
 
-#if 0
-protected:
+private:
 	// Joins swap groups and binds to a swap barrier
 	bool InitializeNvidiaSwapLock();
+	void InitializeOpenGLCapabilities();
 	void UpdateSwapInterval(int32 InSyncInterval) const;
 
 private:
 	FOpenGLViewport* OpenGLViewport = nullptr;
 	FPlatformOpenGLContext* OpenGLContext = nullptr;
-#endif
 };
diff --git a/Engine/Plugins/Runtime/nDisplay/Source/DisplayCluster/Private/Render/Synchronization/DisplayClusterRenderSyncPolicySoftwareDX11.cpp b/Engine/Plugins/Runtime/nDisplay/Source/DisplayCluster/Private/Render/Synchronization/DisplayClusterRenderSyncPolicySoftwareDX11.cpp
--- a/Engine/Plugins/Runtime/nDisplay/Source/DisplayCluster/Private/Render/Synchronization/DisplayClusterRenderSyncPolicySoftwareDX11.cpp
+++ b/Engine/Plugins/Runtime/nDisplay/Source/DisplayCluster/Private/Render/Synchronization/DisplayClusterRenderSyncPolicySoftwareDX11.cpp
@@ -4,28 +4,20 @@
 
 #include "DisplayClusterLog.h"
 
+#include "Engine/Engine.h"
 #include "Engine/GameViewportClient.h"
-#include "Engine/GameEngine.h"
 
-#include "Stats/Stats2.h"
-
-#include "DisplayClusterUtils/DisplayClusterTypesConverter.h"
+#include "Windows/AllowWindowsPlatformTypes.h"
+#include "DirectX/Include/DXGI.h"
+#include "dwmapi.h"
+#include "Windows/HideWindowsPlatformTypes.h"
 
-#include "Render/Device/IDisplayClusterRenderDevice.h"
-#include "Render/Device/DisplayClusterDeviceInternals.h"
+#include "D3D11RHI/Private/Windows/D3D11RHIBasePrivate.h"
 
-#include "d3d11.h"
-#include "dxgi1_2.h"
+#include "D3D11State.h"
+#include "D3D11Resources.h"
 #include "D3D11Viewport.h"
 
-#include "RenderResource.h"
-#include "RenderUtils.h"
-
-#include "UObject/Stack.h"
-
-#include <exception>
-#include "dwmapi.h"
-
 #include "NVIDIA/nvapi/nvapi.h"
 
 
@@ -143,6 +135,23 @@ FDisplayClusterRenderSyncPolicySoftwareDX11::~FDisplayClusterRenderSyncPolicySof
 
 bool FDisplayClusterRenderSyncPolicySoftwareDX11::SynchronizeClusterRendering(int32& InOutSyncInterval)
 {
+	check(IsInRenderingThread());
+
+	if(!(GEngine && GEngine->GameViewport && GEngine->GameViewport->Viewport))
+	{
+		return false;
+	}
+
+	FD3D11Viewport* const Viewport = static_cast<FD3D11Viewport*>(GEngine->GameViewport->Viewport->GetViewportRHI().GetReference());
+	check(Viewport);
+
+#if !WITH_EDITOR
+	// Issue frame event
+	Viewport->IssueFrameEvent();
+	// Wait until GPU finish last frame commands
+	Viewport->WaitForFrameEventCompletion();
+#endif
+
 	if (!bUseAdvancedSynchronization)
 	{
 		// Sync by barrier only
@@ -173,15 +182,7 @@ bool FDisplayClusterRenderSyncPolicySoftwareDX11::SynchronizeClusterRendering(in
 
 	++FrameCounter;
 
-	FD3D11Viewport* const Viewport = static_cast<FD3D11Viewport*>(GEngine->GameViewport->Viewport->GetViewportRHI().GetReference());
-	check(Viewport);
-	if (!Viewport)
-	{
-		UE_LOG(LogDisplayClusterRender, Warning, TEXT("Couldn't get a viewport. Native present will be performed."));
-		return true;
-	}
-
-	IDXGISwapChain* const SwapChain = (IDXGISwapChain*)Viewport->GetSwapChain();
+	IDXGISwapChain* const SwapChain = (IDXGISwapChain*)GEngine->GameViewport->Viewport->GetViewportRHI()->GetNativeSwapChain();
 	check(SwapChain);
 	if (!SwapChain)
 	{
@@ -199,9 +200,8 @@ bool FDisplayClusterRenderSyncPolicySoftwareDX11::SynchronizeClusterRendering(in
 	}
 
 	// Get dynamic values from console variables
-	const float VBlankThreshold = CVarVBlankThreshold.GetValueOnGameThread();
-	const float VBlankThresholdSleepMultiplier = CVarVBlankThresholdSleepMultiplier.GetValueOnGameThread();
-
+	const float VBlankThreshold = CVarVBlankThreshold.GetValueOnRenderThread();
+	const float VBlankThresholdSleepMultiplier = CVarVBlankThresholdSleepMultiplier.GetValueOnRenderThread();
 
 	if (!bTimersInitialized)
 	{
@@ -271,12 +271,12 @@ bool FDisplayClusterRenderSyncPolicySoftwareDX11::SynchronizeClusterRendering(in
 		QUICK_SCOPE_CYCLE_COUNTER(STAT_nDisplayPresent_SoftSync_SwapChainPresent);
 
 		// Get dynamic values from console variables
-		const int32 VBlankBasisUpdate = CVarVBlankBasisUpdate.GetValueOnGameThread();
-		const float VBlankBasisUpdatePeriod = CVarVBlankBasisUpdatePeriod.GetValueOnGameThread();
+		const int32 VBlankBasisUpdate = CVarVBlankBasisUpdate.GetValueOnRenderThread();
+		const float VBlankBasisUpdatePeriod = CVarVBlankBasisUpdatePeriod.GetValueOnRenderThread();
 
 		// Regardless of where we are, it's safe to present a frame now. If we need to update the VBlank basis,
 		// we have to wait for a VBlank and store the time. We don't want to miss a frame presentation so we
-		// present it with swap interval 0 right after VBlank sygnal.
+		// present it with swap interval 0 right after VBlank signal.
 		int SyncIntervalToUse = 1;
 		if ((VBlankBasisUpdate > 0) && (B2A - VBlankBasis) > VBlankBasisUpdatePeriod)
 		{
@@ -291,7 +291,7 @@ bool FDisplayClusterRenderSyncPolicySoftwareDX11::SynchronizeClusterRendering(in
 
 	UE_LOG(LogDisplayClusterRender, Verbose, TEXT("##SYNC_LOG - %d:%lf:%lf:%lf:%lf:%lf:%lf:%lf:%lf:%lf"), FrameCounter, B1B, B1A, TToB, SB, SA, B2B, B2A, PB, PA);
 
-	const bool LogDwmStats = (CVarLogDwmStats.GetValueOnGameThread() != 0);
+	const bool LogDwmStats = (CVarLogDwmStats.GetValueOnRenderThread() != 0);
 	if (LogDwmStats)
 	{
 		PrintDwmStats(FrameCounter);
diff --git a/Engine/Plugins/Runtime/nDisplay/Source/DisplayCluster/Private/Render/Synchronization/DisplayClusterRenderSyncPolicySoftwareDX12.cpp b/Engine/Plugins/Runtime/nDisplay/Source/DisplayCluster/Private/Render/Synchronization/DisplayClusterRenderSyncPolicySoftwareDX12.cpp
--- a/Engine/Plugins/Runtime/nDisplay/Source/DisplayCluster/Private/Render/Synchronization/DisplayClusterRenderSyncPolicySoftwareDX12.cpp
+++ b/Engine/Plugins/Runtime/nDisplay/Source/DisplayCluster/Private/Render/Synchronization/DisplayClusterRenderSyncPolicySoftwareDX12.cpp
@@ -2,6 +2,16 @@
 
 #include "Render/Synchronization/DisplayClusterRenderSyncPolicySoftwareDX12.h"
 
+#include "Engine/Engine.h"
+#include "Engine/GameViewportClient.h"
+
+// The following #undef is a temporary hacky fix to avoid GetRenderTargetFormat redefinition
+// error. Both D3D11Viewport.h and D3D12Viewport.h public headers have the same function that
+// leads to build error in Unity builds. 
+#define GetRenderTargetFormat GetRenderTargetFormat_D3D12
+#include "D3D12RHI/Private/D3D12RHIPrivate.h"
+#include "D3D12Viewport.h"
+#undef GetRenderTargetFormat
 
 
 FDisplayClusterRenderSyncPolicySoftwareDX12::FDisplayClusterRenderSyncPolicySoftwareDX12()
@@ -14,6 +24,23 @@ FDisplayClusterRenderSyncPolicySoftwareDX12::~FDisplayClusterRenderSyncPolicySof
 
 bool FDisplayClusterRenderSyncPolicySoftwareDX12::SynchronizeClusterRendering(int32& InOutSyncInterval)
 {
+	check(IsInRenderingThread());
+
+	if(!(GEngine && GEngine->GameViewport && GEngine->GameViewport->Viewport))
+	{
+		return false;
+	}
+
+	FD3D12Viewport* const Viewport = static_cast<FD3D12Viewport*>(GEngine->GameViewport->Viewport->GetViewportRHI().GetReference());
+	check(Viewport);
+
+#if !WITH_EDITOR
+	// Issue frame event
+	Viewport->IssueFrameEvent();
+	// Wait until GPU finish last frame commands
+	Viewport->WaitForFrameEventCompletion();
+#endif
+
 	//@todo DWM based synchronization
 
 	// As a temporary solution, synchronize render threads on a barrier only
diff --git a/Engine/Plugins/Runtime/nDisplay/Source/DisplayCluster/Public/Blueprints/IDisplayClusterBlueprintAPI.h b/Engine/Plugins/Runtime/nDisplay/Source/DisplayCluster/Public/Blueprints/IDisplayClusterBlueprintAPI.h
--- a/Engine/Plugins/Runtime/nDisplay/Source/DisplayCluster/Public/Blueprints/IDisplayClusterBlueprintAPI.h
+++ b/Engine/Plugins/Runtime/nDisplay/Source/DisplayCluster/Public/Blueprints/IDisplayClusterBlueprintAPI.h
@@ -9,8 +9,13 @@
 #include "Cluster/DisplayClusterClusterEvent.h"
 #include "Cluster/IDisplayClusterClusterEventListener.h"
 
+#include "Config/DisplayClusterConfigTypes.h"
+
 #include "IDisplayClusterBlueprintAPI.generated.h"
 
+struct FPostProcessSettings;
+
+
 
 UINTERFACE(meta = (CannotImplementInterfaceInBlueprint))
 class DISPLAYCLUSTER_API UDisplayClusterBlueprintAPI : public UInterface
@@ -82,6 +87,8 @@ public:
 	//////////////////////////////////////////////////////////////////////////////////////////////
 	// Config API
 	//////////////////////////////////////////////////////////////////////////////////////////////
+	UFUNCTION(BlueprintCallable, meta = (DisplayName = "Get viewports"), Category = "DisplayCluster|Config")
+	virtual void GetViewports(bool IsRTT, TArray<FString>& ViewportIDs, TArray<FString>& ViewportTypes, TArray<FIntPoint>& ViewportLocations, TArray<FIntPoint>& ViewportSizes) = 0;
 
 
 public:
@@ -236,6 +243,12 @@ public:
 	UFUNCTION(BlueprintCallable, meta = (DisplayName = "Set viewport camera"), Category = "DisplayCluster|Render")
 	virtual void SetViewportCamera(const FString& InCameraId, const FString& InViewportId) = 0;
 
+	UFUNCTION(BlueprintCallable, meta = (DisplayName = "Set custom post processing settings for viewport"), Category = "DisplayCluster|Render")
+	virtual void SetCustomPostProcessing(const FString& ViewportID, const FPostProcessSettings& PostProcessingSettings) = 0;
+
+	UFUNCTION(BlueprintCallable, meta = (DisplayName = "Get viewport rectangle"), Category = "DisplayCluster|Render")
+	virtual bool GetViewportRect(const FString& ViewportID, FIntPoint& ViewportLoc, FIntPoint& ViewportSize) = 0;
+
 public:
 	//////////////////////////////////////////////////////////////////////////////////////////////
 	// Render/Camera API
diff --git a/Engine/Plugins/Runtime/nDisplay/Source/DisplayCluster/Public/Config/DisplayClusterConfigTypes.h b/Engine/Plugins/Runtime/nDisplay/Source/DisplayCluster/Public/Config/DisplayClusterConfigTypes.h
--- a/Engine/Plugins/Runtime/nDisplay/Source/DisplayCluster/Public/Config/DisplayClusterConfigTypes.h
+++ b/Engine/Plugins/Runtime/nDisplay/Source/DisplayCluster/Public/Config/DisplayClusterConfigTypes.h
@@ -63,6 +63,7 @@ struct DISPLAYCLUSTER_API FDisplayClusterConfigWindow : public FDisplayClusterCo
 {
 	FString Id;
 	TArray<FString> ViewportIds;
+	TArray<FString> PostprocessIds;
 	bool IsFullscreen = false;
 	int32 WinX = 0;
 	int32 WinY = 0;
@@ -83,6 +84,20 @@ struct DISPLAYCLUSTER_API FDisplayClusterConfigViewport : public FDisplayCluster
 	FString CameraId;
 	FIntPoint Loc  = FIntPoint::ZeroValue;
 	FIntPoint Size = FIntPoint::ZeroValue;
+	bool IsRTT = false;
+
+	virtual FString ToString() const override;
+	virtual bool    DeserializeFromString(const FString& line) override;
+};
+
+//////////////////////////////////////////////////////////////////////////////////////////////
+// Postprocess configuration
+//////////////////////////////////////////////////////////////////////////////////////////////
+struct DISPLAYCLUSTER_API FDisplayClusterConfigPostprocess : public FDisplayClusterConfigBase
+{
+	FString Id;
+	FString PostprocessId;
+	FString ConfigLine;
 
 	virtual FString ToString() const override;
 	virtual bool    DeserializeFromString(const FString& line) override;
diff --git a/Engine/Plugins/Runtime/nDisplay/Source/DisplayCluster/Public/Config/IDisplayClusterConfigManager.h b/Engine/Plugins/Runtime/nDisplay/Source/DisplayCluster/Public/Config/IDisplayClusterConfigManager.h
--- a/Engine/Plugins/Runtime/nDisplay/Source/DisplayCluster/Public/Config/IDisplayClusterConfigManager.h
+++ b/Engine/Plugins/Runtime/nDisplay/Source/DisplayCluster/Public/Config/IDisplayClusterConfigManager.h
@@ -40,6 +40,11 @@ public:
 	virtual bool GetViewport(int32 idx, FDisplayClusterConfigViewport& viewport) const = 0;
 	virtual bool GetViewport(const FString& id, FDisplayClusterConfigViewport& viewport) const = 0;
 
+	virtual int32 GetPostprocessAmount() const = 0;
+	virtual TArray<FDisplayClusterConfigPostprocess> GetPostprocess() const = 0;
+	virtual bool GetPostprocess(int32 idx, FDisplayClusterConfigPostprocess& postprocess) const = 0;
+	virtual bool GetPostprocess(const FString& id, FDisplayClusterConfigPostprocess& postprocess) const = 0;
+
 	virtual int32 GetSceneNodesAmount() const = 0;
 	virtual TArray<FDisplayClusterConfigSceneNode> GetSceneNodes() const = 0;
 	virtual bool GetSceneNode(int32 idx, FDisplayClusterConfigSceneNode& snode) const = 0;
@@ -62,4 +67,6 @@ public:
 	virtual FDisplayClusterConfigNetwork GetConfigNetwork() const = 0;
 	virtual FDisplayClusterConfigDebug   GetConfigDebug()   const = 0;
 	virtual FDisplayClusterConfigCustom  GetConfigCustom()  const = 0;
+
+	virtual FString GetFullPathToFile(const FString& FileName) const = 0;
 };
diff --git a/Engine/Plugins/Runtime/nDisplay/Source/DisplayCluster/Public/DisplayClusterPawn.h b/Engine/Plugins/Runtime/nDisplay/Source/DisplayCluster/Public/DisplayClusterPawn.h
--- a/Engine/Plugins/Runtime/nDisplay/Source/DisplayCluster/Public/DisplayClusterPawn.h
+++ b/Engine/Plugins/Runtime/nDisplay/Source/DisplayCluster/Public/DisplayClusterPawn.h
@@ -19,9 +19,11 @@ UCLASS()
 class DISPLAYCLUSTER_API ADisplayClusterPawn
 	: public APawn
 {
-	GENERATED_UCLASS_BODY()
+	GENERATED_BODY()
 
 public:
+	ADisplayClusterPawn(const FObjectInitializer& ObjectInitializer);
+
 	inline USphereComponent* GetCollisionComponent() const
 	{ return CollisionComponent; }
 
@@ -50,7 +52,7 @@ public:
 
 protected:
 	/** Camera component */
-	UPROPERTY(VisibleAnywhere, Category = "DisplayCluster")
+	UPROPERTY(VisibleAnywhere, BlueprintReadOnly, Category = "DisplayCluster")
 	UCameraComponent* CameraComponent;
 
 	/** Collision component */
diff --git a/Engine/Plugins/Runtime/nDisplay/Source/DisplayCluster/Public/IDisplayCluster.h b/Engine/Plugins/Runtime/nDisplay/Source/DisplayCluster/Public/IDisplayCluster.h
--- a/Engine/Plugins/Runtime/nDisplay/Source/DisplayCluster/Public/IDisplayCluster.h
+++ b/Engine/Plugins/Runtime/nDisplay/Source/DisplayCluster/Public/IDisplayCluster.h
@@ -35,7 +35,8 @@ public:
 	*/
 	static inline IDisplayCluster& Get()
 	{
-		return FModuleManager::LoadModuleChecked<IDisplayCluster>(IDisplayCluster::ModuleName);
+		//return FModuleManager::LoadModuleChecked<IDisplayCluster>(IDisplayCluster::ModuleName);
+		return FModuleManager::GetModuleChecked<IDisplayCluster>(IDisplayCluster::ModuleName);
 	}
 
 	/**
diff --git a/Engine/Plugins/Runtime/nDisplay/Source/DisplayCluster/Public/Render/Device/IDisplayClusterRenderDevice.h b/Engine/Plugins/Runtime/nDisplay/Source/DisplayCluster/Public/Render/Device/IDisplayClusterRenderDevice.h
--- a/Engine/Plugins/Runtime/nDisplay/Source/DisplayCluster/Public/Render/Device/IDisplayClusterRenderDevice.h
+++ b/Engine/Plugins/Runtime/nDisplay/Source/DisplayCluster/Public/Render/Device/IDisplayClusterRenderDevice.h
@@ -43,4 +43,8 @@ public:
 	* @param InViewportId - ID of a viewport to assign the camera (all viewports if empty)
 	*/
 	virtual void SetViewportCamera(const FString& InCameraId = FString(), const FString& InViewportId = FString()) = 0;
+
+	virtual void SetCustomPostProcessing(const FString& ViewportID, const FPostProcessSettings& PostProcessingSettings) = 0;
+
+	virtual bool GetViewportRect(const FString& InViewportID, FIntRect& Rect) = 0;
 };
diff --git a/Engine/Plugins/Runtime/nDisplay/Source/DisplayCluster/Public/Render/IDisplayClusterRenderManager.h b/Engine/Plugins/Runtime/nDisplay/Source/DisplayCluster/Public/Render/IDisplayClusterRenderManager.h
--- a/Engine/Plugins/Runtime/nDisplay/Source/DisplayCluster/Public/Render/IDisplayClusterRenderManager.h
+++ b/Engine/Plugins/Runtime/nDisplay/Source/DisplayCluster/Public/Render/IDisplayClusterRenderManager.h
@@ -145,6 +145,8 @@ public:
 	*/
 	virtual TMap<FString, FDisplayClusterPPInfo> GetRegisteredPostprocessOperations() const = 0;
 
+	virtual void SetCustomPostProcessing(const FString& ViewportID, const FPostProcessSettings& PostProcessingSettings) = 0;
+
 	/**
 	* Assigns camera to a specified viewport. If InViewportId is empty, all viewports will be assigned to a new camera. Empty camera ID means default active camera.
 	*
@@ -153,6 +155,8 @@ public:
 	*/
 	virtual void SetViewportCamera(const FString& InCameraId = FString(), const FString& InViewportId = FString()) = 0;
 
+	virtual bool GetViewportRect(const FString& InViewportID, FIntRect& Rect) = 0;
+
 
 
 	/**
diff --git a/Engine/Plugins/Runtime/nDisplay/Source/DisplayCluster/Public/Render/PostProcess/IDisplayClusterPostProcess.h b/Engine/Plugins/Runtime/nDisplay/Source/DisplayCluster/Public/Render/PostProcess/IDisplayClusterPostProcess.h
--- a/Engine/Plugins/Runtime/nDisplay/Source/DisplayCluster/Public/Render/PostProcess/IDisplayClusterPostProcess.h
+++ b/Engine/Plugins/Runtime/nDisplay/Source/DisplayCluster/Public/Render/PostProcess/IDisplayClusterPostProcess.h
@@ -19,6 +19,14 @@ public:
 	{ }
 
 public:
+	/**
+	* Game thread call. Initialize postprocess from config line
+	*
+	* @param CfgLine - Configuration line for this postprocess
+	*/
+	virtual void InitializePostProcess(const FString& CfgLine)
+	{ }
+
 	/**
 	* Returns if an interface implementation processes each view region before warp&blend
 	*
diff --git a/Engine/Plugins/Runtime/nDisplay/Source/DisplayClusterEditor/DisplayClusterEditor.Build.cs b/Engine/Plugins/Runtime/nDisplay/Source/DisplayClusterEditor/DisplayClusterEditor.Build.cs
--- a/Engine/Plugins/Runtime/nDisplay/Source/DisplayClusterEditor/DisplayClusterEditor.Build.cs
+++ b/Engine/Plugins/Runtime/nDisplay/Source/DisplayClusterEditor/DisplayClusterEditor.Build.cs
@@ -10,17 +10,9 @@ public class DisplayClusterEditor : ModuleRules
 		PrivateDependencyModuleNames.AddRange( new string[] {
 			"Core",
 			"CoreUObject",
+			"DisplayCluster",
 			"Engine",
 			"UnrealEd"
 		});
-
-		PrivateDependencyModuleNames.AddRange( new string[] {
-			"DisplayCluster"
-		});
-
-        PrivateIncludePathModuleNames.AddRange( new string[] {
-			"Settings",
-			"DisplayCluster"
-		});
 	}
 }
diff --git a/Engine/Plugins/Runtime/nDisplay/Source/DisplayClusterEditor/Public/DisplayClusterEditorSettings.h b/Engine/Plugins/Runtime/nDisplay/Source/DisplayClusterEditor/Public/DisplayClusterEditorSettings.h
--- a/Engine/Plugins/Runtime/nDisplay/Source/DisplayClusterEditor/Public/DisplayClusterEditorSettings.h
+++ b/Engine/Plugins/Runtime/nDisplay/Source/DisplayClusterEditor/Public/DisplayClusterEditorSettings.h
@@ -15,7 +15,9 @@ class DISPLAYCLUSTEREDITOR_API UDisplayClusterEditorSettings : public UObject
 {
 	GENERATED_UCLASS_BODY()
 
-	UPROPERTY(config, EditAnywhere, Category = Main)
+	UPROPERTY(config, EditAnywhere, Category = Main, meta = (
+		ToolTip = "When enabled changes GameEngine to DisplayClusterGameEngine",
+		ConfigRestartRequired = true))
 	bool bEnabled;
 
 public:
diff --git a/Engine/Plugins/Runtime/nDisplay/Source/DisplayClusterInput/DisplayClusterInput.Build.cs b/Engine/Plugins/Runtime/nDisplay/Source/DisplayClusterInput/DisplayClusterInput.Build.cs
--- a/Engine/Plugins/Runtime/nDisplay/Source/DisplayClusterInput/DisplayClusterInput.Build.cs
+++ b/Engine/Plugins/Runtime/nDisplay/Source/DisplayClusterInput/DisplayClusterInput.Build.cs
@@ -10,9 +10,9 @@ public class DisplayClusterInput : ModuleRules
 		PrivateIncludePathModuleNames.AddRange(
 			new string[]
 			{
-				"InputDevice",          // For IInputDevice.h
-				"HeadMountedDisplay",   // For IMotionController.h
-				"DisplayCluster"        // For IDisplayCluster
+				"InputDevice",
+				"HeadMountedDisplay",
+				"DisplayCluster"
 			}
 		);
 
@@ -21,7 +21,6 @@ public class DisplayClusterInput : ModuleRules
 			{
 				"Core",
 				"CoreUObject",
-				"DisplayCluster",
 				"Engine",
 				"InputCore",
 				"InputDevice"
@@ -43,12 +42,8 @@ public class DisplayClusterInput : ModuleRules
 
 		PrivateIncludePaths.AddRange(
 			new string[] {
-				// Relative to Engine\Plugins\Runtime\nDisplay\Source\
+				"DisplayClusterInput/Private",
 				"DisplayCluster/Private",
-				"DisplayClusterInput/Private/Controller",
-				"DisplayClusterInput/Private/State",
-				"../../../../Source/Runtime/Renderer/Private",
-				"../../../../Source/Runtime/Engine/Classes/Components",
 			}
 		);
 	}
diff --git a/Engine/Plugins/Runtime/nDisplay/Source/DisplayClusterInput/Private/Controller/DisplayClusterInputControllerBase.h b/Engine/Plugins/Runtime/nDisplay/Source/DisplayClusterInput/Private/Controller/DisplayClusterInputControllerBase.h
--- a/Engine/Plugins/Runtime/nDisplay/Source/DisplayClusterInput/Private/Controller/DisplayClusterInputControllerBase.h
+++ b/Engine/Plugins/Runtime/nDisplay/Source/DisplayClusterInput/Private/Controller/DisplayClusterInputControllerBase.h
@@ -9,10 +9,11 @@
 #include "Input/Devices/DisplayClusterInputDeviceTraits.h"
 
 #include "InputCoreTypes.h"
-#include "DisplayClusterInputStateAnalog.h"
-#include "DisplayClusterInputStateButton.h"
-#include "DisplayClusterInputStateKeyboard.h"
-#include "DisplayClusterInputStateTracker.h"
+
+#include "State/DisplayClusterInputStateAnalog.h"
+#include "State/DisplayClusterInputStateButton.h"
+#include "State/DisplayClusterInputStateKeyboard.h"
+#include "State/DisplayClusterInputStateTracker.h"
 
 #include "Misc/DisplayClusterInputLog.h"
 
diff --git a/Engine/Plugins/Runtime/nDisplay/Source/DisplayClusterPostprocess/DisplayClusterPostprocess.Build.cs b/Engine/Plugins/Runtime/nDisplay/Source/DisplayClusterPostprocess/DisplayClusterPostprocess.Build.cs
new file mode 100644
--- /dev/null
+++ b/Engine/Plugins/Runtime/nDisplay/Source/DisplayClusterPostprocess/DisplayClusterPostprocess.Build.cs
@@ -0,0 +1,38 @@
+// Copyright 1998-2019 Epic Games, Inc. All Rights Reserved.
+
+using UnrealBuildTool;
+using System.IO;
+
+public class DisplayClusterPostprocess : ModuleRules
+{
+	public DisplayClusterPostprocess(ReadOnlyTargetRules ROTargetRules) : base(ROTargetRules)
+	{
+		PrivateIncludePaths.AddRange(
+			new string[]
+			{
+				"DisplayClusterPostprocess/Private",
+			}
+		);
+
+		PublicDependencyModuleNames.AddRange(
+			new string[]
+			{
+				"Core",
+				"Engine",
+			});
+
+		PrivateDependencyModuleNames.AddRange(
+			new string[]
+			{
+				"DisplayCluster",
+				"OutputRemap",
+				"RHI",
+			}
+		);
+
+		if (Target.bBuildEditor == true)
+		{
+			PrivateDependencyModuleNames.Add("UnrealEd");
+		}
+	}
+}
diff --git a/Engine/Plugins/Runtime/nDisplay/Source/DisplayClusterPostprocess/Private/DisplayClusterPostprocessHelpers.h b/Engine/Plugins/Runtime/nDisplay/Source/DisplayClusterPostprocess/Private/DisplayClusterPostprocessHelpers.h
new file mode 100644
--- /dev/null
+++ b/Engine/Plugins/Runtime/nDisplay/Source/DisplayClusterPostprocess/Private/DisplayClusterPostprocessHelpers.h
@@ -0,0 +1,12 @@
+// Copyright 1998-2019 Epic Games, Inc. All Rights Reserved.
+
+#pragma once
+
+#include "CoreMinimal.h"
+
+#include "DisplayClusterUtils/DisplayClusterCommonHelpers.h"
+
+
+namespace DisplayClusterHelpers
+{
+};
diff --git a/Engine/Plugins/Runtime/nDisplay/Source/DisplayClusterPostprocess/Private/DisplayClusterPostprocessLog.cpp b/Engine/Plugins/Runtime/nDisplay/Source/DisplayClusterPostprocess/Private/DisplayClusterPostprocessLog.cpp
new file mode 100644
--- /dev/null
+++ b/Engine/Plugins/Runtime/nDisplay/Source/DisplayClusterPostprocess/Private/DisplayClusterPostprocessLog.cpp
@@ -0,0 +1,7 @@
+// Copyright 1998-2019 Epic Games, Inc. All Rights Reserved.
+
+#include "DisplayClusterPostprocessLog.h"
+
+
+DEFINE_LOG_CATEGORY(LogDisplayClusterPostprocess);
+DEFINE_LOG_CATEGORY(LogDisplayClusterPostprocessOutputRemap);
diff --git a/Engine/Plugins/Runtime/nDisplay/Source/DisplayClusterPostprocess/Private/DisplayClusterPostprocessLog.h b/Engine/Plugins/Runtime/nDisplay/Source/DisplayClusterPostprocess/Private/DisplayClusterPostprocessLog.h
new file mode 100644
--- /dev/null
+++ b/Engine/Plugins/Runtime/nDisplay/Source/DisplayClusterPostprocess/Private/DisplayClusterPostprocessLog.h
@@ -0,0 +1,13 @@
+// Copyright 1998-2019 Epic Games, Inc. All Rights Reserved.
+
+#pragma once
+
+#include "CoreMinimal.h"
+
+#if UE_BUILD_SHIPPING
+DECLARE_LOG_CATEGORY_EXTERN(LogDisplayClusterPostprocess,            Warning, Warning);
+DECLARE_LOG_CATEGORY_EXTERN(LogDisplayClusterPostprocessOutputRemap, Warning, Warning);
+#else
+DECLARE_LOG_CATEGORY_EXTERN(LogDisplayClusterPostprocess,            Log, All);
+DECLARE_LOG_CATEGORY_EXTERN(LogDisplayClusterPostprocessOutputRemap, Log, All);
+#endif
diff --git a/Engine/Plugins/Runtime/nDisplay/Source/DisplayClusterPostprocess/Private/DisplayClusterPostprocessModule.cpp b/Engine/Plugins/Runtime/nDisplay/Source/DisplayClusterPostprocess/Private/DisplayClusterPostprocessModule.cpp
new file mode 100644
--- /dev/null
+++ b/Engine/Plugins/Runtime/nDisplay/Source/DisplayClusterPostprocess/Private/DisplayClusterPostprocessModule.cpp
@@ -0,0 +1,76 @@
+// Copyright 1998-2019 Epic Games, Inc. All Rights Reserved.
+
+#include "DisplayClusterPostprocessModule.h"
+
+#include "DisplayClusterPostprocessLog.h"
+#include "DisplayClusterPostprocessStrings.h"
+
+#include "PostProcess/DisplayClusterPostprocessOutputRemap.h"
+
+#include "IDisplayCluster.h"
+#include "Render/IDisplayClusterRenderManager.h"
+
+
+FDisplayClusterPostprocessModule::FDisplayClusterPostprocessModule()
+{
+	TSharedPtr<IDisplayClusterPostProcess> Postprocess;
+
+	// Output Remap
+	Postprocess = MakeShareable(new FDisplayClusterPostprocessOutputRemap);
+	PostprocessAssets.Emplace(DisplayClusterStrings::postprocess::OutputRemap, Postprocess);
+
+
+	
+	UE_LOG(LogDisplayClusterPostprocess, Log, TEXT("Postprocess module has been instantiated"));
+}
+
+FDisplayClusterPostprocessModule::~FDisplayClusterPostprocessModule()
+{
+	UE_LOG(LogDisplayClusterPostprocess, Log, TEXT("Postprocess module has been destroyed"));
+}
+
+
+//////////////////////////////////////////////////////////////////////////////////////////////
+// IModuleInterface
+//////////////////////////////////////////////////////////////////////////////////////////////
+void FDisplayClusterPostprocessModule::StartupModule()
+{
+	UE_LOG(LogDisplayClusterPostprocess, Log, TEXT("Postprocess module startup"));
+
+	IDisplayClusterRenderManager* RenderMgr = IDisplayCluster::Get().GetRenderMgr();
+	if (RenderMgr)
+	{
+		for (auto it = PostprocessAssets.CreateIterator(); it; ++it)
+		{
+			UE_LOG(LogDisplayClusterPostprocess, Log, TEXT("Registering <%s> projection policy factory..."), *it->Key);
+
+			if (!RenderMgr->RegisterPostprocessOperation(it->Key, it->Value))
+			{
+				UE_LOG(LogDisplayClusterPostprocess, Warning, TEXT("Couldn't register <%s> projection policy factory"), *it->Key);
+			}
+		}
+	}
+
+	UE_LOG(LogDisplayClusterPostprocess, Log, TEXT("Postprocess module has started"));
+}
+
+void FDisplayClusterPostprocessModule::ShutdownModule()
+{
+	UE_LOG(LogDisplayClusterPostprocess, Log, TEXT("Postprocess module shutdown"));
+
+	IDisplayClusterRenderManager* RenderMgr = IDisplayCluster::Get().GetRenderMgr();
+	if (RenderMgr)
+	{
+		for (auto it = PostprocessAssets.CreateConstIterator(); it; ++it)
+		{
+			UE_LOG(LogDisplayClusterPostprocess, Log, TEXT("Un-registering <%s> projection factory..."), *it->Key);
+
+			if (!RenderMgr->UnregisterPostprocessOperation(it->Key))
+			{
+				UE_LOG(LogDisplayClusterPostprocess, Warning, TEXT("An error occurred during un-registering the <%s> projection factory"), *it->Key);
+			}
+		}
+	}
+}
+
+IMPLEMENT_MODULE(FDisplayClusterPostprocessModule, DisplayClusterPostprocess)
diff --git a/Engine/Plugins/Runtime/nDisplay/Source/DisplayClusterPostprocess/Private/DisplayClusterPostprocessModule.h b/Engine/Plugins/Runtime/nDisplay/Source/DisplayClusterPostprocess/Private/DisplayClusterPostprocessModule.h
new file mode 100644
--- /dev/null
+++ b/Engine/Plugins/Runtime/nDisplay/Source/DisplayClusterPostprocess/Private/DisplayClusterPostprocessModule.h
@@ -0,0 +1,28 @@
+// Copyright 1998-2019 Epic Games, Inc. All Rights Reserved.
+
+#pragma once
+
+#include "IDisplayClusterPostProcess.h"
+#include "Render/PostProcess/IDisplayClusterPostProcess.h"
+
+class IDisplayClusterProjectionPolicyFactory;
+
+
+class FDisplayClusterPostprocessModule
+	: public IDisplayClusterPostprocess
+{
+public:
+	FDisplayClusterPostprocessModule();
+	virtual ~FDisplayClusterPostprocessModule();
+
+public:
+	//////////////////////////////////////////////////////////////////////////////////////////////
+	// IModuleInterface
+	//////////////////////////////////////////////////////////////////////////////////////////////
+	virtual void StartupModule() override;
+	virtual void ShutdownModule() override;
+
+private:
+	// Available postprocess policy
+	TMap<FString, TSharedPtr<IDisplayClusterPostProcess>> PostprocessAssets;
+};
diff --git a/Engine/Plugins/Runtime/nDisplay/Source/DisplayClusterPostprocess/Private/DisplayClusterPostprocessStrings.h b/Engine/Plugins/Runtime/nDisplay/Source/DisplayClusterPostprocess/Private/DisplayClusterPostprocessStrings.h
new file mode 100644
--- /dev/null
+++ b/Engine/Plugins/Runtime/nDisplay/Source/DisplayClusterPostprocess/Private/DisplayClusterPostprocessStrings.h
@@ -0,0 +1,28 @@
+// Copyright 1998-2019 Epic Games, Inc. All Rights Reserved.
+
+#pragma once
+
+#include "CoreMinimal.h"
+#include "DisplayClusterUtils/DisplayClusterCommonStrings.h"
+
+
+namespace DisplayClusterStrings
+{
+	namespace cfg
+	{
+		namespace data
+		{
+			namespace postprocess
+			{
+				namespace output_remap
+				{
+					static constexpr auto File   = TEXT("file");
+				}
+			}
+		}
+	}
+	namespace postprocess
+	{
+		static constexpr auto OutputRemap = TEXT("OutputRemap");
+	}
+};
diff --git a/Engine/Plugins/Runtime/nDisplay/Source/DisplayClusterPostprocess/Private/PostProcess/DisplayClusterPostprocessOutputRemap.cpp b/Engine/Plugins/Runtime/nDisplay/Source/DisplayClusterPostprocess/Private/PostProcess/DisplayClusterPostprocessOutputRemap.cpp
new file mode 100644
--- /dev/null
+++ b/Engine/Plugins/Runtime/nDisplay/Source/DisplayClusterPostprocess/Private/PostProcess/DisplayClusterPostprocessOutputRemap.cpp
@@ -0,0 +1,100 @@
+// Copyright 1998-2019 Epic Games, Inc. All Rights Reserved.
+
+#include "PostProcess/DisplayClusterPostprocessOutputRemap.h"
+
+#include "DisplayClusterPostprocessHelpers.h"
+#include "DisplayClusterPostprocessLog.h"
+#include "DisplayClusterPostprocessStrings.h"
+
+#include "Engine/RendererSettings.h"
+#include "Misc/Paths.h"
+
+
+FDisplayClusterPostprocessOutputRemap::FDisplayClusterPostprocessOutputRemap()
+	: OutputRemapAPI(IOutputRemap::Get())
+	, MeshRef(-1)
+{
+
+}
+
+FDisplayClusterPostprocessOutputRemap::~FDisplayClusterPostprocessOutputRemap()
+{
+}
+
+
+//////////////////////////////////////////////////////////////////////////////////////////////
+// IDisplayClusterPostProcess
+//////////////////////////////////////////////////////////////////////////////////////////////
+bool FDisplayClusterPostprocessOutputRemap::IsPostProcessRenderTargetAfterWarpBlendRequired()
+{
+	return MeshRef>=0;
+}
+
+void FDisplayClusterPostprocessOutputRemap::InitializePostProcess(const FString& CfgLine)
+{
+	// PFM file (optional)
+	FString ExtMeshFile;
+	if (DisplayClusterHelpers::str::ExtractValue(CfgLine, DisplayClusterStrings::cfg::data::postprocess::output_remap::File, ExtMeshFile))
+	{
+		UE_LOG(LogDisplayClusterPostprocessOutputRemap, Log, TEXT("Found Argument '%s'='%s'"), DisplayClusterStrings::cfg::data::postprocess::output_remap::File, *ExtMeshFile);
+	}
+
+	// Load ext mesh:
+	uint32 OutMeshRef;
+	if (!OutputRemapAPI.Load(ExtMeshFile, OutMeshRef))
+	{
+		UE_LOG(LogDisplayClusterPostprocessOutputRemap, Error, TEXT("Failed to load ext mesh from file '%s'"), *ExtMeshFile);
+		MeshRef = -1;
+	}
+	else
+	{
+		MeshRef = (int)OutMeshRef;
+	}
+}
+
+void FDisplayClusterPostprocessOutputRemap::PerformPostProcessRenderTargetAfterWarpBlend_RenderThread(FRHICommandListImmediate& RHICmdList, FRHITexture2D* InOutTexture) const
+{
+	check(MeshRef >= 0);
+	FIntPoint ScreenSize = InOutTexture->GetSizeXY();
+
+	InitializeResources_RenderThread(ScreenSize);
+
+	if (!RTTexture || !RTTexture->IsValid())
+	{
+		UE_LOG(LogDisplayClusterPostprocessOutputRemap, Error, TEXT("RT not allocated"));
+		return;
+	}
+
+	uint32 InMeshRef = (uint32)MeshRef;
+	if (OutputRemapAPI.ApplyOutputRemap_RenderThread(RHICmdList, InOutTexture, RTTexture, InMeshRef))
+	{
+		//Copy remapped result back
+		RHICmdList.CopyToResolveTarget(RTTexture, InOutTexture, FResolveParams());
+	}
+}
+
+//////////////////////////////////////////////////////////////////////////////////////////////
+// FDisplayClusterPostprocessOutputRemap
+//////////////////////////////////////////////////////////////////////////////////////////////
+
+bool FDisplayClusterPostprocessOutputRemap::InitializeResources_RenderThread(const FIntPoint& ScreenSize) const
+{
+	check(IsInRenderingThread());
+
+	if (!bIsRenderResourcesInitialized)
+	{
+		FScopeLock lock(&RenderingResourcesInitializationCS);
+		if (!bIsRenderResourcesInitialized)
+		{
+			static const TConsoleVariableData<int32>* CVarDefaultBackBufferPixelFormat = IConsoleManager::Get().FindTConsoleVariableDataInt(TEXT("r.DefaultBackBufferPixelFormat"));
+			static const EPixelFormat SceneTargetFormat = EDefaultBackBufferPixelFormat::Convert2PixelFormat(EDefaultBackBufferPixelFormat::FromInt(CVarDefaultBackBufferPixelFormat->GetValueOnRenderThread()));
+
+			FRHIResourceCreateInfo CreateInfo;
+			FTexture2DRHIRef DummyTexRef;
+			RHICreateTargetableShaderResource2D(ScreenSize.X, ScreenSize.Y, SceneTargetFormat, 1, TexCreate_None, TexCreate_RenderTargetable, false, CreateInfo, RTTexture, DummyTexRef);
+
+			bIsRenderResourcesInitialized = true;
+		}
+	}
+	return true;
+}
diff --git a/Engine/Plugins/Runtime/nDisplay/Source/DisplayClusterPostprocess/Private/PostProcess/DisplayClusterPostprocessOutputRemap.h b/Engine/Plugins/Runtime/nDisplay/Source/DisplayClusterPostprocess/Private/PostProcess/DisplayClusterPostprocessOutputRemap.h
new file mode 100644
--- /dev/null
+++ b/Engine/Plugins/Runtime/nDisplay/Source/DisplayClusterPostprocess/Private/PostProcess/DisplayClusterPostprocessOutputRemap.h
@@ -0,0 +1,47 @@
+// Copyright 1998-2019 Epic Games, Inc. All Rights Reserved.
+
+#pragma once
+
+#include "Render/PostProcess/IDisplayClusterPostProcess.h"
+
+#include "IOutputRemap.h"
+
+#include "RHI.h"
+#include "RHICommandList.h"
+#include "RHIResources.h"
+#include "RHIUtilities.h"
+
+/**
+ * MPCDI projection policy
+ */
+class FDisplayClusterPostprocessOutputRemap
+	: public IDisplayClusterPostProcess
+{
+public:
+	FDisplayClusterPostprocessOutputRemap();
+	virtual ~FDisplayClusterPostprocessOutputRemap();
+
+public:
+	//////////////////////////////////////////////////////////////////////////////////////////////
+	// FDisplayClusterPostprocessOutputRemap
+	//////////////////////////////////////////////////////////////////////////////////////////////
+	virtual bool IsPostProcessRenderTargetAfterWarpBlendRequired() override;
+
+	virtual void PerformPostProcessRenderTargetAfterWarpBlend_RenderThread(FRHICommandListImmediate& RHICmdList, FRHITexture2D* InOutTexture) const override;
+
+	virtual void InitializePostProcess(const FString& CfgLine) override;
+
+protected:
+	bool InitializeResources_RenderThread(const FIntPoint& ScreenSize) const;
+
+private:
+	IOutputRemap& OutputRemapAPI;
+
+	int MeshRef;
+
+	//@todo: remove, use ext paired texture
+	mutable FTexture2DRHIRef RTTexture;
+
+	mutable bool bIsRenderResourcesInitialized = false;
+	mutable FCriticalSection RenderingResourcesInitializationCS;
+};
diff --git a/Engine/Plugins/Runtime/nDisplay/Source/DisplayClusterPostprocess/Public/IDisplayClusterPostprocess.h b/Engine/Plugins/Runtime/nDisplay/Source/DisplayClusterPostprocess/Public/IDisplayClusterPostprocess.h
new file mode 100644
--- /dev/null
+++ b/Engine/Plugins/Runtime/nDisplay/Source/DisplayClusterPostprocess/Public/IDisplayClusterPostprocess.h
@@ -0,0 +1,40 @@
+// Copyright 1998-2019 Epic Games, Inc. All Rights Reserved.
+
+#pragma once
+
+#include "Modules/ModuleManager.h"
+#include "Modules/ModuleInterface.h"
+
+class IDisplayClusterPostprocess : public IModuleInterface
+{
+public:
+	static constexpr auto ModuleName = TEXT("DisplayClusterPostprocess");
+
+public:
+	virtual ~IDisplayClusterPostprocess()
+	{ 
+	}
+
+public:
+	/**
+	 * Singleton-like access to this module's interface.  This is just for convenience!
+	 * Beware of calling this during the shutdown phase, though.  Your module might have been unloaded already.
+	 *
+	 * @return Returns singleton instance, loading the module on demand if needed
+	 */
+	static inline IDisplayClusterPostprocess& Get()
+	{
+		return FModuleManager::LoadModuleChecked<IDisplayClusterPostprocess>(IDisplayClusterPostprocess::ModuleName);
+	}
+
+	/**
+	 * Checks to see if this module is loaded and ready.  It is only valid to call Get() if IsAvailable() returns true.
+	 *
+	 * @return True if the module is loaded and ready to use
+	 */
+	static inline bool IsAvailable()
+	{
+		return FModuleManager::Get().IsModuleLoaded(IDisplayClusterPostprocess::ModuleName);
+	}
+		
+};
diff --git a/Engine/Plugins/Runtime/nDisplay/Source/DisplayClusterProjection/DisplayClusterProjection.Build.cs b/Engine/Plugins/Runtime/nDisplay/Source/DisplayClusterProjection/DisplayClusterProjection.Build.cs
--- a/Engine/Plugins/Runtime/nDisplay/Source/DisplayClusterProjection/DisplayClusterProjection.Build.cs
+++ b/Engine/Plugins/Runtime/nDisplay/Source/DisplayClusterProjection/DisplayClusterProjection.Build.cs
@@ -11,11 +11,8 @@ public class DisplayClusterProjection : ModuleRules
 			new string[]
 			{
 				"DisplayClusterProjection/Private",
-                "../../../../Source/Runtime/Renderer/Private",
-                "../../../../Source/Runtime/Engine/Classes/Components",
-                "../../../../../Engine/Source/Runtime/Windows/D3D11RHI/Private",
-                "../../../../../Engine/Source/Runtime/Windows/D3D11RHI/Private/Windows",
-
+				"../../../../../Engine/Source/Runtime/Windows/D3D11RHI/Private",
+				"../../../../../Engine/Source/Runtime/Windows/D3D11RHI/Private/Windows",
 			}
 		);
 
@@ -24,13 +21,7 @@ public class DisplayClusterProjection : ModuleRules
 			{
 				"Core",
 				"CoreUObject",
-                "D3D11RHI",
-				"DisplayCluster",
 				"Engine",
-				"HeadMountedDisplay",
-				"MPCDI",
-				"RenderCore",
-				"RHI",
 			});
 
 		PrivateDependencyModuleNames.AddRange(
@@ -38,8 +29,16 @@ public class DisplayClusterProjection : ModuleRules
 			{
 				"ApplicationCore",
 				"Core",
+				"Composure",
 				"CoreUObject",
+				"D3D11RHI",
+				"DisplayCluster",
 				"Engine",
+				"HeadMountedDisplay",
+				"MPCDI",
+				"Projects",
+				"RenderCore",
+				"RHI",
 			}
 		);
 
@@ -48,25 +47,30 @@ public class DisplayClusterProjection : ModuleRules
 			PrivateDependencyModuleNames.Add("UnrealEd");
 		}
 
-        AddEngineThirdPartyPrivateStaticDependencies(Target, "DX11");
-        AddEngineThirdPartyPrivateStaticDependencies(Target, "NVAftermath");
-        AddEngineThirdPartyPrivateStaticDependencies(Target, "IntelMetricsDiscovery");
+		AddEngineThirdPartyPrivateStaticDependencies(Target, "DX11");
+		AddEngineThirdPartyPrivateStaticDependencies(Target, "NVAftermath");
+		AddEngineThirdPartyPrivateStaticDependencies(Target, "IntelMetricsDiscovery");
 
-        AddThirdPartyDependencies(ROTargetRules);
-    }
+		AddThirdPartyDependencies(ROTargetRules);
+	}
 
 
-    public bool AddThirdPartyDependencies(ReadOnlyTargetRules ROTargetRules)
-    {
-		string ModulePath = Path.GetDirectoryName(UnrealBuildTool.RulesCompiler.GetFileNameFromType(GetType()));
-		string ThirdPartyPath = Path.GetFullPath(Path.Combine(ModulePath, "../../ThirdParty/"));
+	public void AddThirdPartyDependencies(ReadOnlyTargetRules ROTargetRules)
+	{
+		string ThirdPartyPath = Path.GetFullPath(Path.Combine(ModuleDirectory, "../../ThirdParty/"));
 
 		string PathInc = string.Empty;
+		string PathDll = string.Empty;
 
 		// EasyBlend
 		PathInc = Path.Combine(ThirdPartyPath, "EasyBlend", "Include");
+		PathDll = Path.Combine(ThirdPartyPath, "EasyBlend", "DLL");
+		
 		PublicIncludePaths.Add(PathInc);
-			   
-        return true;
+
+		if (Target.Platform == UnrealTargetPlatform.Win64)
+		{
+			RuntimeDependencies.Add(Path.Combine(PathDll, "mplEasyBlendSDKDX1164.dll"));
+		}
 	}
 }
diff --git a/Engine/Plugins/Runtime/nDisplay/Source/DisplayClusterProjection/Private/Blueprints/DisplayClusterProjectionBlueprintAPIImpl.cpp b/Engine/Plugins/Runtime/nDisplay/Source/DisplayClusterProjection/Private/Blueprints/DisplayClusterProjectionBlueprintAPIImpl.cpp
new file mode 100644
--- /dev/null
+++ b/Engine/Plugins/Runtime/nDisplay/Source/DisplayClusterProjection/Private/Blueprints/DisplayClusterProjectionBlueprintAPIImpl.cpp
@@ -0,0 +1,41 @@
+// Copyright 1998-2019 Epic Games, Inc. All Rights Reserved.
+
+#include "Blueprints/DisplayClusterProjectionBlueprintAPIImpl.h"
+
+#include "IDisplayClusterProjection.h"
+#include "DisplayClusterProjectionLog.h"
+#include "DisplayClusterProjectionStrings.h"
+
+#include "Policy/Camera/DisplayClusterProjectionCameraPolicy.h"
+#include "Policy/Camera/DisplayClusterProjectionCameraPolicyFactory.h"
+
+#include "Camera/CameraComponent.h"
+
+
+
+//////////////////////////////////////////////////////////////////////////////////////////////
+// Policy: CAMERA
+//////////////////////////////////////////////////////////////////////////////////////////////
+void UDisplayClusterProjectionBlueprintAPIImpl::CameraPolicySetCamera(const FString& ViewportId, UCameraComponent* NewCamera)
+{
+	IDisplayClusterProjection& Module = IDisplayClusterProjection::Get();
+	
+	TSharedPtr<IDisplayClusterProjectionPolicyFactory> Factory = Module.GetProjectionFactory(DisplayClusterStrings::projection::Camera);
+	if (Factory.IsValid())
+	{
+		TSharedPtr<FDisplayClusterProjectionCameraPolicyFactory> CameraFactory = StaticCastSharedPtr<FDisplayClusterProjectionCameraPolicyFactory>(Factory);
+		if (CameraFactory.IsValid())
+		{
+			TSharedPtr<IDisplayClusterProjectionPolicy> PolicyInstance = CameraFactory->GetPolicyInstance(ViewportId);
+			if (PolicyInstance.IsValid())
+			{
+				TSharedPtr<FDisplayClusterProjectionCameraPolicy> CameraPolicyInstance = StaticCastSharedPtr<FDisplayClusterProjectionCameraPolicy>(PolicyInstance);
+				if (CameraPolicyInstance)
+				{
+					CameraPolicyInstance->SetCamera(NewCamera);
+				}
+			}
+		}
+	}
+
+}
diff --git a/Engine/Plugins/Runtime/nDisplay/Source/DisplayClusterProjection/Private/Blueprints/DisplayClusterProjectionBlueprintAPIImpl.h b/Engine/Plugins/Runtime/nDisplay/Source/DisplayClusterProjection/Private/Blueprints/DisplayClusterProjectionBlueprintAPIImpl.h
new file mode 100644
--- /dev/null
+++ b/Engine/Plugins/Runtime/nDisplay/Source/DisplayClusterProjection/Private/Blueprints/DisplayClusterProjectionBlueprintAPIImpl.h
@@ -0,0 +1,27 @@
+// Copyright 1998-2019 Epic Games, Inc. All Rights Reserved.
+
+#pragma once
+
+#include "CoreMinimal.h"
+#include "Blueprints/IDisplayClusterProjectionBlueprintAPI.h"
+#include "DisplayClusterProjectionBlueprintAPIImpl.generated.h"
+
+
+/**
+ * Blueprint API interface implementation
+ */
+UCLASS()
+class DISPLAYCLUSTERPROJECTION_API UDisplayClusterProjectionBlueprintAPIImpl
+	: public UObject
+	, public IDisplayClusterProjectionBlueprintAPI
+{
+	GENERATED_BODY()
+
+public:
+	//////////////////////////////////////////////////////////////////////////////////////////////
+	// Policy: CAMERA
+	//////////////////////////////////////////////////////////////////////////////////////////////
+	/** Sets active camera component for camera policy */
+	UFUNCTION(BlueprintCallable, meta = (DisplayName = "Set camera"), Category = "DisplayClusterProjection|Camera")
+	virtual void CameraPolicySetCamera(const FString& ViewportId, UCameraComponent* NewCamera) override;
+};
diff --git a/Engine/Plugins/Runtime/nDisplay/Source/DisplayClusterProjection/Private/Blueprints/DisplayClusterProjectionBlueprintLib.cpp b/Engine/Plugins/Runtime/nDisplay/Source/DisplayClusterProjection/Private/Blueprints/DisplayClusterProjectionBlueprintLib.cpp
new file mode 100644
--- /dev/null
+++ b/Engine/Plugins/Runtime/nDisplay/Source/DisplayClusterProjection/Private/Blueprints/DisplayClusterProjectionBlueprintLib.cpp
@@ -0,0 +1,17 @@
+// Copyright 1998-2019 Epic Games, Inc. All Rights Reserved.
+
+#include "Blueprints/DisplayClusterProjectionBlueprintLib.h"
+#include "Blueprints/DisplayClusterProjectionBlueprintAPIImpl.h"
+#include "UObject/Package.h"
+
+
+UDisplayClusterProjectionBlueprintLib::UDisplayClusterProjectionBlueprintLib(class FObjectInitializer const & ObjectInitializer)
+	: Super(ObjectInitializer)
+{
+}
+
+void UDisplayClusterProjectionBlueprintLib::GetAPI(TScriptInterface<IDisplayClusterProjectionBlueprintAPI>& OutAPI)
+{
+	static UDisplayClusterProjectionBlueprintAPIImpl* Obj = NewObject<UDisplayClusterProjectionBlueprintAPIImpl>(GetTransientPackage(), NAME_None, RF_MarkAsRootSet);
+	OutAPI = Obj;
+}
diff --git a/Engine/Plugins/Runtime/nDisplay/Source/DisplayClusterProjection/Private/DisplayClusterProjectionLog.cpp b/Engine/Plugins/Runtime/nDisplay/Source/DisplayClusterProjection/Private/DisplayClusterProjectionLog.cpp
--- a/Engine/Plugins/Runtime/nDisplay/Source/DisplayClusterProjection/Private/DisplayClusterProjectionLog.cpp
+++ b/Engine/Plugins/Runtime/nDisplay/Source/DisplayClusterProjection/Private/DisplayClusterProjectionLog.cpp
@@ -4,6 +4,7 @@
 
 
 DEFINE_LOG_CATEGORY(LogDisplayClusterProjection);
+DEFINE_LOG_CATEGORY(LogDisplayClusterProjectionCamera);
 DEFINE_LOG_CATEGORY(LogDisplayClusterProjectionSimple);
 DEFINE_LOG_CATEGORY(LogDisplayClusterProjectionMPCDI);
 DEFINE_LOG_CATEGORY(LogDisplayClusterProjectionEasyBlend);
diff --git a/Engine/Plugins/Runtime/nDisplay/Source/DisplayClusterProjection/Private/DisplayClusterProjectionLog.h b/Engine/Plugins/Runtime/nDisplay/Source/DisplayClusterProjection/Private/DisplayClusterProjectionLog.h
--- a/Engine/Plugins/Runtime/nDisplay/Source/DisplayClusterProjection/Private/DisplayClusterProjectionLog.h
+++ b/Engine/Plugins/Runtime/nDisplay/Source/DisplayClusterProjection/Private/DisplayClusterProjectionLog.h
@@ -6,11 +6,13 @@
 
 #if UE_BUILD_SHIPPING
 DECLARE_LOG_CATEGORY_EXTERN(LogDisplayClusterProjection,          Warning, Warning);
+DECLARE_LOG_CATEGORY_EXTERN(LogDisplayClusterProjectionCamera,    Warning, Warning);
 DECLARE_LOG_CATEGORY_EXTERN(LogDisplayClusterProjectionSimple,    Warning, Warning);
 DECLARE_LOG_CATEGORY_EXTERN(LogDisplayClusterProjectionMPCDI,     Warning, Warning);
 DECLARE_LOG_CATEGORY_EXTERN(LogDisplayClusterProjectionEasyBlend, Warning, Warning);
 #else
 DECLARE_LOG_CATEGORY_EXTERN(LogDisplayClusterProjection,          Log, All);
+DECLARE_LOG_CATEGORY_EXTERN(LogDisplayClusterProjectionCamera,    Log, All);
 DECLARE_LOG_CATEGORY_EXTERN(LogDisplayClusterProjectionSimple,    Log, All);
 DECLARE_LOG_CATEGORY_EXTERN(LogDisplayClusterProjectionMPCDI,     Log, All);
 DECLARE_LOG_CATEGORY_EXTERN(LogDisplayClusterProjectionEasyBlend, Log, All);
diff --git a/Engine/Plugins/Runtime/nDisplay/Source/DisplayClusterProjection/Private/DisplayClusterProjectionModule.cpp b/Engine/Plugins/Runtime/nDisplay/Source/DisplayClusterProjection/Private/DisplayClusterProjectionModule.cpp
--- a/Engine/Plugins/Runtime/nDisplay/Source/DisplayClusterProjection/Private/DisplayClusterProjectionModule.cpp
+++ b/Engine/Plugins/Runtime/nDisplay/Source/DisplayClusterProjection/Private/DisplayClusterProjectionModule.cpp
@@ -4,6 +4,7 @@
 
 #include "DisplayClusterProjectionLog.h"
 
+#include "Policy/Camera/DisplayClusterProjectionCameraPolicyFactory.h"
 #include "Policy/EasyBlend/DisplayClusterProjectionEasyBlendPolicyFactory.h"
 #include "Policy/Simple/DisplayClusterProjectionSimplePolicyFactory.h"
 #include "Policy/MPCDI/DisplayClusterProjectionMPCDIPolicyFactory.h"
@@ -16,6 +17,10 @@ FDisplayClusterProjectionModule::FDisplayClusterProjectionModule()
 {
 	TSharedPtr<IDisplayClusterProjectionPolicyFactory> Factory;
 
+	// Camera projection
+	Factory = MakeShareable(new FDisplayClusterProjectionCameraPolicyFactory);
+	ProjectionPolicyFactories.Emplace(DisplayClusterStrings::projection::Camera, Factory);
+
 	// Simple projection
 	Factory = MakeShareable(new FDisplayClusterProjectionSimplePolicyFactory);
 	ProjectionPolicyFactories.Emplace(DisplayClusterStrings::projection::Simple, Factory);
diff --git a/Engine/Plugins/Runtime/nDisplay/Source/DisplayClusterProjection/Private/DisplayClusterProjectionStrings.h b/Engine/Plugins/Runtime/nDisplay/Source/DisplayClusterProjection/Private/DisplayClusterProjectionStrings.h
--- a/Engine/Plugins/Runtime/nDisplay/Source/DisplayClusterProjection/Private/DisplayClusterProjectionStrings.h
+++ b/Engine/Plugins/Runtime/nDisplay/Source/DisplayClusterProjection/Private/DisplayClusterProjectionStrings.h
@@ -19,14 +19,6 @@ namespace DisplayClusterStrings
 					static constexpr auto Screen = TEXT("screen");
 				}
 
-				namespace mpcdi
-				{
-					static constexpr auto File   = TEXT("file");
-					static constexpr auto Buffer = TEXT("buffer");
-					static constexpr auto Region = TEXT("region");
-					static constexpr auto Origin = TEXT("origin");
-				}
-
 				namespace easyblend
 				{
 					static constexpr auto File   = TEXT("file");
@@ -39,6 +31,7 @@ namespace DisplayClusterStrings
 
 	namespace projection
 	{
+		static constexpr auto Camera    = TEXT("camera");
 		static constexpr auto Simple    = TEXT("simple");
 		static constexpr auto MPCDI     = TEXT("mpcdi");
 		static constexpr auto EasyBlend = TEXT("easyblend");
diff --git a/Engine/Plugins/Runtime/nDisplay/Source/DisplayClusterProjection/Private/Policy/Camera/DisplayClusterProjectionCameraPolicy.cpp b/Engine/Plugins/Runtime/nDisplay/Source/DisplayClusterProjection/Private/Policy/Camera/DisplayClusterProjectionCameraPolicy.cpp
new file mode 100644
--- /dev/null
+++ b/Engine/Plugins/Runtime/nDisplay/Source/DisplayClusterProjection/Private/Policy/Camera/DisplayClusterProjectionCameraPolicy.cpp
@@ -0,0 +1,110 @@
+// Copyright 1998-2019 Epic Games, Inc. All Rights Reserved.
+
+#include "Policy/Camera/DisplayClusterProjectionCameraPolicy.h"
+
+#include "DisplayClusterProjectionLog.h"
+#include "DisplayClusterProjectionStrings.h"
+
+#include "DisplayClusterUtils/DisplayClusterCommonHelpers.h"
+
+#include "IDisplayCluster.h"
+#include "Config/IDisplayClusterConfigManager.h"
+#include "Game/IDisplayClusterGameManager.h"
+#include "DisplayClusterCameraComponent.h"
+
+#include "DisplayClusterPawn.h"
+#include "DisplayClusterScreenComponent.h"
+
+#include "Camera/CameraComponent.h"
+#include "ComposurePostMoves.h"
+
+
+
+FDisplayClusterProjectionCameraPolicy::FDisplayClusterProjectionCameraPolicy(const FString& ViewportId)
+	: FDisplayClusterProjectionPolicyBase(ViewportId)
+{
+}
+
+FDisplayClusterProjectionCameraPolicy::~FDisplayClusterProjectionCameraPolicy()
+{
+}
+
+
+//////////////////////////////////////////////////////////////////////////////////////////////
+// IDisplayClusterProjectionPolicy
+//////////////////////////////////////////////////////////////////////////////////////////////
+void FDisplayClusterProjectionCameraPolicy::StartScene(UWorld* World)
+{
+	check(IsInGameThread());
+
+	if (!AssignedCamera)
+	{
+		ADisplayClusterPawn* Pawn = IDisplayCluster::Get().GetGameMgr()->GetRoot();
+		if (Pawn)
+		{
+			AssignedCamera = Pawn->GetCameraComponent();
+		}
+	}
+}
+
+void FDisplayClusterProjectionCameraPolicy::EndScene()
+{
+	check(IsInGameThread());
+}
+
+bool FDisplayClusterProjectionCameraPolicy::HandleAddViewport(const FIntPoint& InViewportSize, const uint32 InViewsAmount)
+{
+	check(IsInGameThread());
+
+	UE_LOG(LogDisplayClusterProjectionCamera, Log, TEXT("Initializing internals for the viewport '%s'"), *GetViewportId());
+	
+	return true;
+}
+
+void FDisplayClusterProjectionCameraPolicy::HandleRemoveViewport()
+{
+	check(IsInGameThread());
+	UE_LOG(LogDisplayClusterProjectionCamera, Log, TEXT("Removing viewport '%s'"), *GetViewportId());
+}
+
+bool FDisplayClusterProjectionCameraPolicy::CalculateView(const uint32 ViewIdx, FVector& InOutViewLocation, FRotator& InOutViewRotation, const FVector& ViewOffset, const float WorldToMeters, const float NCP, const float FCP)
+{
+	check(IsInGameThread());
+
+	InOutViewLocation = (AssignedCamera ? AssignedCamera->GetComponentLocation() : FVector::ZeroVector);
+	InOutViewRotation = (AssignedCamera ? AssignedCamera->GetComponentRotation() : FRotator::ZeroRotator);
+
+	return true;
+}
+
+bool FDisplayClusterProjectionCameraPolicy::GetProjectionMatrix(const uint32 ViewIdx, FMatrix& OutPrjMatrix)
+{
+	check(IsInGameThread());
+
+	if (!AssignedCamera)
+	{
+		return false;
+	}
+
+	FComposurePostMoveSettings ComposureSettings;
+	OutPrjMatrix = ComposureSettings.GetProjectionMatrix(AssignedCamera->FieldOfView, AssignedCamera->AspectRatio);
+
+	return true;
+}
+
+
+//////////////////////////////////////////////////////////////////////////////////////////////
+// FDisplayClusterProjectionCameraPolicy
+//////////////////////////////////////////////////////////////////////////////////////////////
+void FDisplayClusterProjectionCameraPolicy::SetCamera(UCameraComponent* NewCamera)
+{
+	if (NewCamera)
+	{
+		UE_LOG(LogDisplayClusterProjectionCamera, Verbose, TEXT("New camera set: %s"), *NewCamera->GetFullName());
+		AssignedCamera = NewCamera;
+	}
+	else
+	{
+		UE_LOG(LogDisplayClusterProjectionCamera, Warning, TEXT("Trying to set nullptr camera pointer"));
+	}
+}
diff --git a/Engine/Plugins/Runtime/nDisplay/Source/DisplayClusterProjection/Private/Policy/Camera/DisplayClusterProjectionCameraPolicy.h b/Engine/Plugins/Runtime/nDisplay/Source/DisplayClusterProjection/Private/Policy/Camera/DisplayClusterProjectionCameraPolicy.h
new file mode 100644
--- /dev/null
+++ b/Engine/Plugins/Runtime/nDisplay/Source/DisplayClusterProjection/Private/Policy/Camera/DisplayClusterProjectionCameraPolicy.h
@@ -0,0 +1,43 @@
+// Copyright 1998-2019 Epic Games, Inc. All Rights Reserved.
+
+#pragma once
+
+#include "Policy/DisplayClusterProjectionPolicyBase.h"
+#include "Config/DisplayClusterConfigTypes.h"
+
+
+class UCameraComponent;
+
+
+/**
+ * Implements math behind the native camera projection policy (use symmetric frustum of a camera)
+ */
+class FDisplayClusterProjectionCameraPolicy
+	: public FDisplayClusterProjectionPolicyBase
+{
+public:
+	FDisplayClusterProjectionCameraPolicy(const FString& ViewportId);
+	virtual ~FDisplayClusterProjectionCameraPolicy();
+
+public:
+	//////////////////////////////////////////////////////////////////////////////////////////////
+	// IDisplayClusterProjectionPolicy
+	//////////////////////////////////////////////////////////////////////////////////////////////
+	virtual void StartScene(UWorld* World) override;
+	virtual void EndScene() override;
+	virtual bool HandleAddViewport(const FIntPoint& InViewportSize, const uint32 InViewsAmount) override;
+	virtual void HandleRemoveViewport() override;
+
+	virtual bool CalculateView(const uint32 ViewIdx, FVector& InOutViewLocation, FRotator& InOutViewRotation, const FVector& ViewOffset, const float WorldToMeters, const float NCP, const float FCP) override;
+	virtual bool GetProjectionMatrix(const uint32 ViewIdx, FMatrix& OutPrjMatrix) override;
+
+	virtual bool IsWarpBlendSupported() override
+	{ return false; }
+
+public:
+	void SetCamera(UCameraComponent* NewCamera);
+
+private:
+	// Camera to use for rendering
+	UCameraComponent* AssignedCamera = nullptr;
+};
diff --git a/Engine/Plugins/Runtime/nDisplay/Source/DisplayClusterProjection/Private/Policy/Camera/DisplayClusterProjectionCameraPolicyFactory.cpp b/Engine/Plugins/Runtime/nDisplay/Source/DisplayClusterProjection/Private/Policy/Camera/DisplayClusterProjectionCameraPolicyFactory.cpp
new file mode 100644
--- /dev/null
+++ b/Engine/Plugins/Runtime/nDisplay/Source/DisplayClusterProjection/Private/Policy/Camera/DisplayClusterProjectionCameraPolicyFactory.cpp
@@ -0,0 +1,37 @@
+// Copyright 1998-2019 Epic Games, Inc. All Rights Reserved.
+
+#include "Policy/Camera/DisplayClusterProjectionCameraPolicyFactory.h"
+#include "Policy/Camera/DisplayClusterProjectionCameraPolicy.h"
+
+#include "DisplayClusterProjectionLog.h"
+
+
+FDisplayClusterProjectionCameraPolicyFactory::FDisplayClusterProjectionCameraPolicyFactory()
+{
+}
+
+FDisplayClusterProjectionCameraPolicyFactory ::~FDisplayClusterProjectionCameraPolicyFactory()
+{
+}
+
+TSharedPtr<IDisplayClusterProjectionPolicy> FDisplayClusterProjectionCameraPolicyFactory::GetPolicyInstance(const FString& ViewportId)
+{
+	if (PolicyInstances.Contains(ViewportId))
+	{
+		return PolicyInstances[ViewportId];
+	}
+
+	return nullptr;
+}
+
+
+//////////////////////////////////////////////////////////////////////////////////////////////
+// IDisplayClusterProjectionPolicyFactory
+//////////////////////////////////////////////////////////////////////////////////////////////
+TSharedPtr<IDisplayClusterProjectionPolicy> FDisplayClusterProjectionCameraPolicyFactory::Create(const FString& PolicyType, const FString& RHIName, const FString& ViewportId)
+{
+	UE_LOG(LogDisplayClusterProjectionCamera, Log, TEXT("Instantiating projection policy <%s>..."), *PolicyType);
+	TSharedPtr<IDisplayClusterProjectionPolicy> NewPolicy = MakeShareable(new FDisplayClusterProjectionCameraPolicy(ViewportId));
+	PolicyInstances.Emplace(ViewportId, NewPolicy);
+	return NewPolicy;
+}
diff --git a/Engine/Plugins/Runtime/nDisplay/Source/DisplayClusterProjection/Private/Policy/Camera/DisplayClusterProjectionCameraPolicyFactory.h b/Engine/Plugins/Runtime/nDisplay/Source/DisplayClusterProjection/Private/Policy/Camera/DisplayClusterProjectionCameraPolicyFactory.h
new file mode 100644
--- /dev/null
+++ b/Engine/Plugins/Runtime/nDisplay/Source/DisplayClusterProjection/Private/Policy/Camera/DisplayClusterProjectionCameraPolicyFactory.h
@@ -0,0 +1,29 @@
+// Copyright 1998-2019 Epic Games, Inc. All Rights Reserved.
+
+#pragma once
+
+#include "Render/Projection/IDisplayClusterProjectionPolicyFactory.h"
+#include "DisplayClusterProjectionStrings.h"
+
+/**
+ * Implements projection policy factory for the 'camera' policy
+ */
+class FDisplayClusterProjectionCameraPolicyFactory
+	: public IDisplayClusterProjectionPolicyFactory
+{
+public:
+	FDisplayClusterProjectionCameraPolicyFactory();
+	virtual ~FDisplayClusterProjectionCameraPolicyFactory();
+
+public:
+	//////////////////////////////////////////////////////////////////////////////////////////////
+	// IDisplayClusterProjectionPolicyFactory
+	//////////////////////////////////////////////////////////////////////////////////////////////
+	virtual TSharedPtr<IDisplayClusterProjectionPolicy> Create(const FString& PolicyType, const FString& RHIName, const FString& ViewportId) override;
+
+public:
+	TSharedPtr<IDisplayClusterProjectionPolicy> GetPolicyInstance(const FString& ViewportId);
+
+private:
+	TMap<FString, TSharedPtr<IDisplayClusterProjectionPolicy>> PolicyInstances;
+};
diff --git a/Engine/Plugins/Runtime/nDisplay/Source/DisplayClusterProjection/Private/Policy/EasyBlend/DX11/DisplayClusterProjectionEasyBlendLibraryDX11.cpp b/Engine/Plugins/Runtime/nDisplay/Source/DisplayClusterProjection/Private/Policy/EasyBlend/DX11/DisplayClusterProjectionEasyBlendLibraryDX11.cpp
--- a/Engine/Plugins/Runtime/nDisplay/Source/DisplayClusterProjection/Private/Policy/EasyBlend/DX11/DisplayClusterProjectionEasyBlendLibraryDX11.cpp
+++ b/Engine/Plugins/Runtime/nDisplay/Source/DisplayClusterProjection/Private/Policy/EasyBlend/DX11/DisplayClusterProjectionEasyBlendLibraryDX11.cpp
@@ -3,7 +3,10 @@
 #include "Policy/EasyBlend/DX11/DisplayClusterProjectionEasyBlendLibraryDX11.h"
 #include "DisplayClusterProjectionLog.h"
 
+#include "IDisplayClusterProjection.h"
+
 #include "HAL/PlatformProcess.h"
+#include "Interfaces/IPluginManager.h"
 
 #if PLATFORM_WINDOWS
 #include "EasyBlendSDKDXApi.h"
@@ -56,10 +59,12 @@ bool DisplayClusterProjectionEasyBlendLibraryDX11::Initialize()
 		if (!DllHandle)
 		{
 #ifdef PLATFORM_WINDOWS
-			// nDisplay supports Win64 only
-			const FString LibName = TEXT("mplEasyBlendSDKDX1164.dll");
+			const FString LibName   = TEXT("mplEasyBlendSDKDX1164.dll");
+			const FString PluginDir = IPluginManager::Get().FindPlugin(TEXT("nDisplay"))->GetBaseDir();
+			const FString DllPath   = FPaths::Combine(PluginDir, TEXT("ThirdParty/EasyBlend/DLL"), LibName);
+			
 			// Try to load DLL
-			DllHandle = FPlatformProcess::GetDllHandle(*LibName);
+			DllHandle = FPlatformProcess::GetDllHandle(*DllPath);
 
 			if (DllHandle)
 			{
@@ -108,7 +113,7 @@ bool DisplayClusterProjectionEasyBlendLibraryDX11::Initialize()
 			}
 			else
 			{
-				UE_LOG(LogDisplayClusterProjectionEasyBlend, Error, TEXT("Couldn't initialize EasyBlend API. No <%> library found."), *LibName);
+				UE_LOG(LogDisplayClusterProjectionEasyBlend, Error, TEXT("Couldn't initialize EasyBlend API. No <%s> library found."), *DllPath);
 				return false;
 			}
 #endif
diff --git a/Engine/Plugins/Runtime/nDisplay/Source/DisplayClusterProjection/Private/Policy/EasyBlend/DX11/DisplayClusterProjectionEasyBlendLibraryDX11.h b/Engine/Plugins/Runtime/nDisplay/Source/DisplayClusterProjection/Private/Policy/EasyBlend/DX11/DisplayClusterProjectionEasyBlendLibraryDX11.h
--- a/Engine/Plugins/Runtime/nDisplay/Source/DisplayClusterProjection/Private/Policy/EasyBlend/DX11/DisplayClusterProjectionEasyBlendLibraryDX11.h
+++ b/Engine/Plugins/Runtime/nDisplay/Source/DisplayClusterProjection/Private/Policy/EasyBlend/DX11/DisplayClusterProjectionEasyBlendLibraryDX11.h
@@ -3,7 +3,14 @@
 #pragma once
 
 #include "CoreMinimal.h"
-#include "D3D11RHIPrivate.h"
+
+#include "RHIResources.h"
+
+#include "Windows/AllowWindowsPlatformTypes.h"
+#include "DirectX/Include/DXGI.h"
+#include "Windows/HideWindowsPlatformTypes.h"
+
+#include "D3D11RHI/Private/Windows/D3D11RHIBasePrivate.h"
 
 #include "EasyBlendSDKDXVer.h"
 #include "EasyBlendSDKDXErrors.h"
diff --git a/Engine/Plugins/Runtime/nDisplay/Source/DisplayClusterProjection/Private/Policy/EasyBlend/DX11/DisplayClusterProjectionEasyBlendViewAdapterDX11.cpp b/Engine/Plugins/Runtime/nDisplay/Source/DisplayClusterProjection/Private/Policy/EasyBlend/DX11/DisplayClusterProjectionEasyBlendViewAdapterDX11.cpp
--- a/Engine/Plugins/Runtime/nDisplay/Source/DisplayClusterProjection/Private/Policy/EasyBlend/DX11/DisplayClusterProjectionEasyBlendViewAdapterDX11.cpp
+++ b/Engine/Plugins/Runtime/nDisplay/Source/DisplayClusterProjection/Private/Policy/EasyBlend/DX11/DisplayClusterProjectionEasyBlendViewAdapterDX11.cpp
@@ -9,14 +9,13 @@
 #include "RHIResources.h"
 #include "RHIUtilities.h"
 
-#include "XRThreadUtils.h"
+#include "Windows/D3D11RHI/Private/D3D11RHIPrivate.h"
 
 #include "Engine/GameViewportClient.h"
 #include "Engine/Engine.h"
 #include "Engine/RendererSettings.h"
 #include "Misc/Paths.h"
 #include "UnrealClient.h"
-#include "D3D11RHIPrivate.h"
 
 
 FDisplayClusterProjectionEasyBlendViewAdapterDX11::FDisplayClusterProjectionEasyBlendViewAdapterDX11(const FDisplayClusterProjectionEasyBlendViewAdapterBase::FInitParams& InitParams)
@@ -102,8 +101,11 @@ bool FDisplayClusterProjectionEasyBlendViewAdapterDX11::CalculateView(const uint
 	FVector EasyBlendEyeLocation = Game2EasyBlend.TransformPosition(InOutViewLocation);
 
 	// Update EasyBlend state
-	check(DisplayClusterProjectionEasyBlendLibraryDX11::EasyBlendSetEyepointFunc);
-	DisplayClusterProjectionEasyBlendLibraryDX11::EasyBlendSetEyepointFunc(Views[ViewIdx].EasyBlendMeshData.Get(), EasyBlendEyeLocation.X, EasyBlendEyeLocation.Y, EasyBlendEyeLocation.Z);
+	{
+		FScopeLock lock(&DllAccessCS);
+		check(DisplayClusterProjectionEasyBlendLibraryDX11::EasyBlendSetEyepointFunc);
+		DisplayClusterProjectionEasyBlendLibraryDX11::EasyBlendSetEyepointFunc(Views[ViewIdx].EasyBlendMeshData.Get(), EasyBlendEyeLocation.X, EasyBlendEyeLocation.Y, EasyBlendEyeLocation.Z);
+	}
 
 	// Get rotation from the frustum that has been updated already
 	InOutViewRotation.Pitch = float(-Views[ViewIdx].EasyBlendMeshData->Frustum.ViewAngleC - 90);
@@ -187,16 +189,20 @@ bool FDisplayClusterProjectionEasyBlendViewAdapterDX11::ApplyWarpBlend_RenderThr
 	ID3D11Texture2D * SrcTextureD3D11 = static_cast<ID3D11Texture2D*>(SrcTextureRHI->GetResource());
 
 	// Setup In/Out EasyBlend textures
-	check(DisplayClusterProjectionEasyBlendLibraryDX11::EasyBlendSetInputTexture2DFunc);
-	EasyBlendSDKDXError EasyBlendSDKDXError1 = DisplayClusterProjectionEasyBlendLibraryDX11::EasyBlendSetInputTexture2DFunc(Views[ViewIdx].EasyBlendMeshData.Get(), SrcTextureD3D11);
+	{
+		FScopeLock lock(&DllAccessCS);
 
-	check(DisplayClusterProjectionEasyBlendLibraryDX11::EasyBlendSetOutputTexture2DFunc);
-	EasyBlendSDKDXError EasyBlendSDKDXError2 = DisplayClusterProjectionEasyBlendLibraryDX11::EasyBlendSetOutputTexture2DFunc(Views[ViewIdx].EasyBlendMeshData.Get(), DstTextureD3D11);
+		check(DisplayClusterProjectionEasyBlendLibraryDX11::EasyBlendSetInputTexture2DFunc);
+		const EasyBlendSDKDXError EasyBlendSDKDXError1 = DisplayClusterProjectionEasyBlendLibraryDX11::EasyBlendSetInputTexture2DFunc(Views[ViewIdx].EasyBlendMeshData.Get(), SrcTextureD3D11);
 
-	if (!(EasyBlendSDKDX_SUCCEEDED(EasyBlendSDKDXError1) && EasyBlendSDKDX_SUCCEEDED(EasyBlendSDKDXError2)))
-	{
-		UE_LOG(LogDisplayClusterProjectionEasyBlend, Error, TEXT("Coulnd't configure in/out textures"));
-		return false;
+		check(DisplayClusterProjectionEasyBlendLibraryDX11::EasyBlendSetOutputTexture2DFunc);
+		const EasyBlendSDKDXError EasyBlendSDKDXError2 = DisplayClusterProjectionEasyBlendLibraryDX11::EasyBlendSetOutputTexture2DFunc(Views[ViewIdx].EasyBlendMeshData.Get(), DstTextureD3D11);
+
+		if (!(EasyBlendSDKDX_SUCCEEDED(EasyBlendSDKDXError1) && EasyBlendSDKDX_SUCCEEDED(EasyBlendSDKDXError2)))
+		{
+			UE_LOG(LogDisplayClusterProjectionEasyBlend, Error, TEXT("Coulnd't configure in/out textures"));
+			return false;
+		}
 	}
 
 	D3D11_VIEWPORT RenderViewportData;
@@ -218,6 +224,8 @@ bool FDisplayClusterProjectionEasyBlendViewAdapterDX11::ApplyWarpBlend_RenderThr
 	DeviceContext->Flush();
 
 	{
+		FScopeLock lock(&DllAccessCS);
+
 		// Perform warp&blend by the EasyBlend
 		check(DisplayClusterProjectionEasyBlendLibraryDX11::EasyBlendDXRenderFunc);
 		EasyBlendSDKDXError EasyBlendSDKDXError = DisplayClusterProjectionEasyBlendLibraryDX11::EasyBlendDXRenderFunc(
diff --git a/Engine/Plugins/Runtime/nDisplay/Source/DisplayClusterProjection/Private/Policy/EasyBlend/DX11/DisplayClusterProjectionEasyBlendViewAdapterDX11.h b/Engine/Plugins/Runtime/nDisplay/Source/DisplayClusterProjection/Private/Policy/EasyBlend/DX11/DisplayClusterProjectionEasyBlendViewAdapterDX11.h
--- a/Engine/Plugins/Runtime/nDisplay/Source/DisplayClusterProjection/Private/Policy/EasyBlend/DX11/DisplayClusterProjectionEasyBlendViewAdapterDX11.h
+++ b/Engine/Plugins/Runtime/nDisplay/Source/DisplayClusterProjection/Private/Policy/EasyBlend/DX11/DisplayClusterProjectionEasyBlendViewAdapterDX11.h
@@ -55,4 +55,5 @@ private:
 
 	bool bIsRenderResourcesInitialized = false;
 	FCriticalSection RenderingResourcesInitializationCS;
+	FCriticalSection DllAccessCS;
 };
diff --git a/Engine/Plugins/Runtime/nDisplay/Source/DisplayClusterProjection/Private/Policy/EasyBlend/DisplayClusterProjectionEasyBlendPolicyBase.cpp b/Engine/Plugins/Runtime/nDisplay/Source/DisplayClusterProjection/Private/Policy/EasyBlend/DisplayClusterProjectionEasyBlendPolicyBase.cpp
--- a/Engine/Plugins/Runtime/nDisplay/Source/DisplayClusterProjection/Private/Policy/EasyBlend/DisplayClusterProjectionEasyBlendPolicyBase.cpp
+++ b/Engine/Plugins/Runtime/nDisplay/Source/DisplayClusterProjection/Private/Policy/EasyBlend/DisplayClusterProjectionEasyBlendPolicyBase.cpp
@@ -57,9 +57,16 @@ bool FDisplayClusterProjectionEasyBlendPolicyBase::HandleAddViewport(const FIntP
 		return false;
 	}
 
+	FString FullFilePath = DisplayClusterHelpers::config::GetFullPath(File);
+	if (!FPaths::FileExists(FullFilePath))
+	{
+		//! Handle error: EasyBlend configuration file not found
+		return false;
+	}
+
 	// Create and store nDisplay-to-EasyBlend viewport adapter
 	TSharedPtr<FDisplayClusterProjectionEasyBlendViewAdapterBase> NewViewAdapter = CreateViewAdapter(FDisplayClusterProjectionEasyBlendViewAdapterBase::FInitParams{ ViewportSize, ViewsAmount });
-	if (!NewViewAdapter->Initialize(File))
+	if (!NewViewAdapter->Initialize(FullFilePath))
 	{
 		UE_LOG(LogDisplayClusterProjectionEasyBlend, Error, TEXT("Couldn't initialize EasyBlend viewport adapter"));
 		return false;
@@ -98,7 +105,7 @@ bool FDisplayClusterProjectionEasyBlendPolicyBase::CalculateView(const uint32 Vi
 	// Calculate view location in origin space
 	const FVector LocalViewLoc = World2LocalTransform.InverseTransformPosition(InOutViewLocation);
 
-	// Prepare non-const variables to get output(
+	// Prepare non-const variables to get output
 	FVector  EasyBlendViewLocation = LocalViewLoc * EasyBlendScale;
 	FRotator EasyBlendViewRotation = FRotator::ZeroRotator;
 	if (!ViewAdapter->CalculateView(ViewIdx, EasyBlendViewLocation, EasyBlendViewRotation, ViewOffset, WorldScale, NCP, FCP))
diff --git a/Engine/Plugins/Runtime/nDisplay/Source/DisplayClusterProjection/Private/Policy/MPCDI/DisplayClusterProjectionMPCDIPolicy.cpp b/Engine/Plugins/Runtime/nDisplay/Source/DisplayClusterProjection/Private/Policy/MPCDI/DisplayClusterProjectionMPCDIPolicy.cpp
--- a/Engine/Plugins/Runtime/nDisplay/Source/DisplayClusterProjection/Private/Policy/MPCDI/DisplayClusterProjectionMPCDIPolicy.cpp
+++ b/Engine/Plugins/Runtime/nDisplay/Source/DisplayClusterProjection/Private/Policy/MPCDI/DisplayClusterProjectionMPCDIPolicy.cpp
@@ -9,8 +9,6 @@
 #include "Config/DisplayClusterConfigTypes.h"
 #include "Game/IDisplayClusterGameManager.h"
 
-#include "XRThreadUtils.h"
-
 #include "Engine/RendererSettings.h"
 #include "Misc/Paths.h"
 
@@ -18,6 +16,7 @@
 FDisplayClusterProjectionMPCDIPolicy::FDisplayClusterProjectionMPCDIPolicy(const FString& ViewportId)
 	: FDisplayClusterProjectionPolicyBase(ViewportId)
 	, MPCDIAPI(IMPCDI::Get())
+	, bIsRenderResourcesInitialized(false)
 {
 }
 
@@ -52,39 +51,30 @@ bool FDisplayClusterProjectionMPCDIPolicy::HandleAddViewport(const FIntPoint& In
 	check(IsInGameThread());
 	check(InViewsAmount > 0);
 
-	FString File;
-	FString Buffer;
-	FString Region;
-
-	// Read MPCDI config data from nDisplay config file
-	if (!ReadConfigData(GetViewportId(), File, Buffer, Region, OriginCompId))
-	{
-		UE_LOG(LogDisplayClusterProjectionMPCDI, Error, TEXT("Couldn't read MPCDI configuration from the config file"));
-		return false;
-	}
-
-	// Check if MPCDI file exists
-	if (!FPaths::FileExists(File))
-	{
-		UE_LOG(LogDisplayClusterProjectionMPCDI, Warning, TEXT("File not found: %s"), *File);
-		return false;
-	}
+	IMPCDI::ConfigParser CfgData;
+	{//Load config:
+		FDisplayClusterConfigProjection CfgProjection;
+		if (!DisplayClusterHelpers::config::GetViewportProjection(GetViewportId(), CfgProjection))
+		{
+			UE_LOG(LogDisplayClusterProjectionMPCDI, Error, TEXT("Couldn't obtain projection info for '%s' viewport"), *GetViewportId());
+			return false;
+		}
 
-	// Load MPCDI file
-	if (!MPCDIAPI.Load(File))
-	{
-		UE_LOG(LogDisplayClusterProjectionMPCDI, Warning, TEXT("Couldn't load MPCDI file: %s"), *File);
-		return false;
+		if (!MPCDIAPI.LoadConfig(CfgProjection.Params, CfgData))
+		{
+			UE_LOG(LogDisplayClusterProjectionMPCDI, Error, TEXT("Couldn't read MPCDI configuration from the config file"));
+			return false;
+		}
 	}
 
-	// Store MPCDI region locator for this viewport
-	if (!MPCDIAPI.GetRegionLocator(File, Buffer, Region, WarpRef))
+	// Load MPCDI config
+	if (!MPCDIAPI.Load(CfgData, WarpRef))
 	{
-		UE_LOG(LogDisplayClusterProjectionMPCDI, Warning, TEXT("Couldn't get region locator for <buf %s, reg %s> in file: %s # %s "), *Buffer, *Region, *File);
+		UE_LOG(LogDisplayClusterProjectionMPCDI, Warning, TEXT("Couldn't load MPCDI config: %s"), *CfgData.ConfigLineStr);
 		return false;
 	}
 
-	UE_LOG(LogDisplayClusterProjectionMPCDI, Log, TEXT("MPCDI policy has been initialized [%s:%s in %s]"), *Buffer, *Region, *File);
+	UE_LOG(LogDisplayClusterProjectionMPCDI, Log, TEXT("MPCDI policy has been initialized [%s:%s in %s]"), *CfgData.BufferId, *CfgData.RegionId, *CfgData.MPCDIFileName);
 
 	// Finally, initialize internal views data container
 	Views.AddDefaulted(InViewsAmount);
@@ -208,49 +198,6 @@ void FDisplayClusterProjectionMPCDIPolicy::ApplyWarpBlend_RenderThread(const uin
 //////////////////////////////////////////////////////////////////////////////////////////////
 // FDisplayClusterProjectionMPCDIPolicy
 //////////////////////////////////////////////////////////////////////////////////////////////
-bool FDisplayClusterProjectionMPCDIPolicy::ReadConfigData(const FString& InViewportId, FString& OutFile, FString& OutBuffer, FString& OutRegion, FString& OutOrigin)
-{
-	// Get projection settings of the specified viewport
-	FDisplayClusterConfigProjection CfgProjection;
-	if (!DisplayClusterHelpers::config::GetViewportProjection(InViewportId, CfgProjection))
-	{
-		UE_LOG(LogDisplayClusterProjectionMPCDI, Error, TEXT("Couldn't obtain projection info for '%s' viewport"), *InViewportId);
-		return false;
-	}
-
-	// MPCDI file
-	if (!DisplayClusterHelpers::str::ExtractValue(CfgProjection.Params, DisplayClusterStrings::cfg::data::projection::mpcdi::File, OutFile))
-	{
-		UE_LOG(LogDisplayClusterProjectionMPCDI, Error, TEXT("Argument '%s' not found in the config file"), DisplayClusterStrings::cfg::data::projection::mpcdi::File);
-		return false;
-	}
-
-	// Buffer
-	if (!DisplayClusterHelpers::str::ExtractValue(CfgProjection.Params, DisplayClusterStrings::cfg::data::projection::mpcdi::Buffer, OutBuffer))
-	{
-		UE_LOG(LogDisplayClusterProjectionMPCDI, Error, TEXT("Argument '%s' not found in the config file"), DisplayClusterStrings::cfg::data::projection::mpcdi::Buffer);
-		return false;
-	}
-
-	// Region
-	if (!DisplayClusterHelpers::str::ExtractValue(CfgProjection.Params, DisplayClusterStrings::cfg::data::projection::mpcdi::Region, OutRegion))
-	{
-		UE_LOG(LogDisplayClusterProjectionMPCDI, Error, TEXT("Argument '%s' not found in the config file"), DisplayClusterStrings::cfg::data::projection::mpcdi::Region);
-		return false;
-	}
-
-	// Origin node (optional)
-	if (DisplayClusterHelpers::str::ExtractValue(CfgProjection.Params, DisplayClusterStrings::cfg::data::projection::mpcdi::Origin, OutOrigin))
-	{
-		UE_LOG(LogDisplayClusterProjectionMPCDI, Log, TEXT("Found origin node for %s:%s - %s"), *OutBuffer, *OutRegion, *OutOrigin);
-	}
-	else
-	{
-		UE_LOG(LogDisplayClusterProjectionMPCDI, Log, TEXT("No origin node found for %s:%s. VR root will be used as default."), *OutBuffer, *OutRegion);
-	}
-
-	return true;
-}
 
 bool FDisplayClusterProjectionMPCDIPolicy::InitializeResources_RenderThread()
 {
@@ -271,8 +218,10 @@ bool FDisplayClusterProjectionMPCDIPolicy::InitializeResources_RenderThread()
 				FTexture2DRHIRef DummyTexRef;
 				RHICreateTargetableShaderResource2D(ViewportSize.X, ViewportSize.Y, SceneTargetFormat, 1, TexCreate_None, TexCreate_RenderTargetable, false, CreateInfo, It.RTTexture, DummyTexRef);
 			}
+
+			bIsRenderResourcesInitialized = true;
 		}
 	}
 
-	return true;
+	return bIsRenderResourcesInitialized;
 }
diff --git a/Engine/Plugins/Runtime/nDisplay/Source/DisplayClusterProjection/Private/Policy/MPCDI/DisplayClusterProjectionMPCDIPolicy.h b/Engine/Plugins/Runtime/nDisplay/Source/DisplayClusterProjection/Private/Policy/MPCDI/DisplayClusterProjectionMPCDIPolicy.h
--- a/Engine/Plugins/Runtime/nDisplay/Source/DisplayClusterProjection/Private/Policy/MPCDI/DisplayClusterProjectionMPCDIPolicy.h
+++ b/Engine/Plugins/Runtime/nDisplay/Source/DisplayClusterProjection/Private/Policy/MPCDI/DisplayClusterProjectionMPCDIPolicy.h
@@ -39,8 +39,7 @@ public:
 	virtual bool IsWarpBlendSupported() override;
 	virtual void ApplyWarpBlend_RenderThread(const uint32 ViewIdx, FRHICommandListImmediate& RHICmdList, FRHITexture2D* SrcTexture, const FIntRect& ViewportRect) override;
 
-protected:
-	bool ReadConfigData(const FString& ViewportId, FString& OutFile, FString& OutBuffer, FString& OutRegion, FString& OutOrigin);
+protected:	
 	bool InitializeResources_RenderThread();
 
 private:
@@ -58,6 +57,6 @@ private:
 
 	TArray<FViewData> Views;
 
-	bool bIsRenderResourcesInitialized = false;
+	mutable bool bIsRenderResourcesInitialized;
 	FCriticalSection RenderingResourcesInitializationCS;
 };
diff --git a/Engine/Plugins/Runtime/nDisplay/Source/DisplayClusterProjection/Public/Blueprints/DisplayClusterProjectionBlueprintLib.h b/Engine/Plugins/Runtime/nDisplay/Source/DisplayClusterProjection/Public/Blueprints/DisplayClusterProjectionBlueprintLib.h
new file mode 100644
--- /dev/null
+++ b/Engine/Plugins/Runtime/nDisplay/Source/DisplayClusterProjection/Public/Blueprints/DisplayClusterProjectionBlueprintLib.h
@@ -0,0 +1,24 @@
+// Copyright 1998-2019 Epic Games, Inc. All Rights Reserved.
+
+#pragma once
+
+#include "CoreMinimal.h"
+#include "Kismet/BlueprintFunctionLibrary.h"
+#include "Blueprints/IDisplayClusterProjectionBlueprintAPI.h"
+#include "DisplayClusterProjectionBlueprintLib.generated.h"
+
+
+/**
+ * Blueprint API function library
+ */
+UCLASS()
+class UDisplayClusterProjectionBlueprintLib
+	: public UBlueprintFunctionLibrary
+{
+	GENERATED_UCLASS_BODY()
+
+public:
+	/** Return Display Cluster API interface. */
+	UFUNCTION(BlueprintPure, meta = (DisplayName = "DisplayClusterProjection Module API"), Category = "nDisplay")
+	static void GetAPI(TScriptInterface<IDisplayClusterProjectionBlueprintAPI>& OutAPI);
+};
diff --git a/Engine/Plugins/Runtime/nDisplay/Source/DisplayClusterUtils/DisplayClusterCommonHelpers.h b/Engine/Plugins/Runtime/nDisplay/Source/DisplayClusterUtils/DisplayClusterCommonHelpers.h
--- a/Engine/Plugins/Runtime/nDisplay/Source/DisplayClusterUtils/DisplayClusterCommonHelpers.h
+++ b/Engine/Plugins/Runtime/nDisplay/Source/DisplayClusterUtils/DisplayClusterCommonHelpers.h
@@ -4,7 +4,6 @@
 
 #include "CoreMinimal.h"
 #include "EngineUtils.h"
-//#include "HAL/RunnableThread.h"
 
 #include "IDisplayCluster.h"
 #include "Cluster/IDisplayClusterClusterManager.h"
@@ -296,7 +295,7 @@ namespace DisplayClusterHelpers
 
 		static UDisplayClusterCameraComponent* GetCamera(const FString& CameraId)
 		{
-			const IDisplayClusterGameManager* const GameMgr = IDisplayCluster::Get().GetGameMgr();
+			static const IDisplayClusterGameManager* const GameMgr = IDisplayCluster::Get().GetGameMgr();
 			if (GameMgr)
 			{
 				return GameMgr->GetCameraById(CameraId);
@@ -344,7 +343,7 @@ namespace DisplayClusterHelpers
 				return false;
 			}
 
-			const IDisplayClusterConfigManager* const ConfigMgr = IDisplayCluster::Get().GetConfigMgr();
+			static const IDisplayClusterConfigManager* const ConfigMgr = IDisplayCluster::Get().GetConfigMgr();
 			if (!ConfigMgr)
 			{
 				return false;
@@ -363,7 +362,7 @@ namespace DisplayClusterHelpers
 				return LocalViewports;
 			}
 
-			const IDisplayClusterConfigManager* const ConfigMgr = IDisplayCluster::Get().GetConfigMgr();
+			static const IDisplayClusterConfigManager* const ConfigMgr = IDisplayCluster::Get().GetConfigMgr();
 			if (!ConfigMgr)
 			{
 				return LocalViewports;
@@ -380,11 +379,38 @@ namespace DisplayClusterHelpers
 			return LocalViewports;
 		}
 
+		static TArray<FDisplayClusterConfigPostprocess> GetLocalPostprocess()
+		{
+			TArray<FDisplayClusterConfigPostprocess> LocalPostprocess;
+
+			FDisplayClusterConfigWindow LocalWindow;
+			if (!DisplayClusterHelpers::config::GetLocalWindow(LocalWindow))
+			{
+				return LocalPostprocess;
+			}
+
+			static const IDisplayClusterConfigManager* const ConfigMgr = IDisplayCluster::Get().GetConfigMgr();
+			if (!ConfigMgr)
+			{
+				return LocalPostprocess;
+			}
+
+			LocalPostprocess = ConfigMgr->GetPostprocess().FilterByPredicate([&LocalWindow](const FDisplayClusterConfigPostprocess& ItemPostprocess)
+			{
+				return LocalWindow.PostprocessIds.ContainsByPredicate([ItemPostprocess](const FString& ItemId)
+				{
+					return ItemPostprocess.Id.Compare(ItemId, ESearchCase::IgnoreCase) == 0;
+				});
+			});
+
+			return LocalPostprocess;
+		}
+
 		static TArray<FDisplayClusterConfigProjection> GetLocalProjections()
 		{
 			TArray<FDisplayClusterConfigProjection> LocalProjections;
 
-			const IDisplayClusterConfigManager* const ConfigMgr = IDisplayCluster::Get().GetConfigMgr();
+			static const IDisplayClusterConfigManager* const ConfigMgr = IDisplayCluster::Get().GetConfigMgr();
 			if (!ConfigMgr)
 			{
 				return LocalProjections;
@@ -405,7 +431,7 @@ namespace DisplayClusterHelpers
 
 		static bool GetViewportProjection(const FString& ViewportId, FDisplayClusterConfigProjection& ViewportProjection)
 		{
-			const IDisplayClusterConfigManager* const ConfigMgr = IDisplayCluster::Get().GetConfigMgr();
+			static const IDisplayClusterConfigManager* const ConfigMgr = IDisplayCluster::Get().GetConfigMgr();
 			if (ConfigMgr)
 			{
 				FDisplayClusterConfigViewport CfgViewport;
@@ -421,12 +447,23 @@ namespace DisplayClusterHelpers
 			return false;
 		}
 
+		static FString GetFullPath(const FString& LocalPath)
+		{
+			static const IDisplayClusterConfigManager* const ConfigMgr = IDisplayCluster::Get().GetConfigMgr();
+			if (ConfigMgr)
+			{
+				return ConfigMgr->GetFullPathToFile(LocalPath);
+			}
+
+			return LocalPath;
+		}
+
 #if 0
 		static TArray<FDisplayClusterConfigScreen> GetLocalScreens()
 		{
 			TArray<FDisplayClusterConfigScreen> LocalScreens;
 
-			const IDisplayClusterConfigManager* const ConfigMgr = IDisplayCluster::Get().GetConfigMgr();
+			static const IDisplayClusterConfigManager* const ConfigMgr = IDisplayCluster::Get().GetConfigMgr();
 			if (!ConfigMgr)
 			{
 				return LocalScreens;
