nDisplay features

diff --git a/Engine/Plugins/Runtime/nDisplay/Source/DisplayClusterUtils/DisplayClusterCommonStrings.h b/Engine/Plugins/Runtime/nDisplay/Source/DisplayClusterUtils/DisplayClusterCommonStrings.h
--- a/Engine/Plugins/Runtime/nDisplay/Source/DisplayClusterUtils/DisplayClusterCommonStrings.h
+++ b/Engine/Plugins/Runtime/nDisplay/Source/DisplayClusterUtils/DisplayClusterCommonStrings.h
@@ -66,13 +66,14 @@ namespace DisplayClusterStrings
 			// Window tokens
 			namespace window
 			{
-				static constexpr auto Header     = TEXT("[window]");
-				static constexpr auto Viewports  = TEXT("viewports");
-				static constexpr auto Fullscreen = TEXT("fullscreen");
-				static constexpr auto WinX       = TEXT("WinX");
-				static constexpr auto WinY       = TEXT("WinY");
-				static constexpr auto ResX       = TEXT("ResX");
-				static constexpr auto ResY       = TEXT("ResY");
+				static constexpr auto Header       = TEXT("[window]");
+				static constexpr auto Viewports    = TEXT("viewports");
+				static constexpr auto Postprocess  = TEXT("postprocess");
+				static constexpr auto Fullscreen   = TEXT("fullscreen");
+				static constexpr auto WinX         = TEXT("WinX");
+				static constexpr auto WinY         = TEXT("WinY");
+				static constexpr auto ResX         = TEXT("ResX");
+				static constexpr auto ResY         = TEXT("ResY");
 				// + Id
 			}
 
@@ -94,6 +95,15 @@ namespace DisplayClusterStrings
 				static constexpr auto PosY       = TEXT("y");
 				static constexpr auto Width      = TEXT("width");
 				static constexpr auto Height     = TEXT("height");
+				static constexpr auto RTT        = TEXT("rtt");
+				// + Id
+			}
+
+			// Posprocess tokens
+			namespace postprocess
+			{
+				static constexpr auto Header        = TEXT("[postprocess]");
+				static constexpr auto PostprocessId = TEXT("type");
 				// + Id
 			}
 
diff --git a/Engine/Plugins/Runtime/nDisplay/Source/MPCDI/MPCDI.Build.cs b/Engine/Plugins/Runtime/nDisplay/Source/MPCDI/MPCDI.Build.cs
--- a/Engine/Plugins/Runtime/nDisplay/Source/MPCDI/MPCDI.Build.cs
+++ b/Engine/Plugins/Runtime/nDisplay/Source/MPCDI/MPCDI.Build.cs
@@ -15,48 +15,33 @@ public class MPCDI : ModuleRules
 				"Core",
 				"CoreUObject",
 				"Engine",
-				"Projects"
-			});
+				"Projects",
+            });
 
 		PrivateDependencyModuleNames.AddRange(
 			new string[]
 			{
-				"Core",
-				"CoreUObject",
-				"Engine",
+				"DisplayCluster",
 				"RenderCore",
 				"RHI",
 				"UtilityShaders"
 			}
 		);
 
-		PrivateIncludePaths.AddRange(
-			new string[] {
-				"../../../../../Engine/Source/Runtime/Windows/D3D11RHI/Private",
-				"../../../../../Engine/Source/Runtime/Windows/D3D11RHI/Private/Windows",
-				"../../../../../Engine/Source/Runtime/D3D12RHI/Private",
-				"../../../../../Engine/Source/Runtime/D3D12RHI/Private/Windows"
-			});
-
 		if (Target.bBuildEditor == true)
 		{
 			PrivateDependencyModuleNames.Add("UnrealEd");
 		}
 
-		PublicAdditionalLibraries.Add("opengl32.lib");
-		AddEngineThirdPartyPrivateStaticDependencies(Target, "OpenGL");
-		AddEngineThirdPartyPrivateStaticDependencies(Target, "DX11");
-		AddEngineThirdPartyPrivateStaticDependencies(Target, "DX12");
 		AddEngineThirdPartyPrivateStaticDependencies(Target, "UElibPNG");
 		AddEngineThirdPartyPrivateStaticDependencies(Target, "zlib");
 
 		AddThirdPartyDependencies(ROTargetRules);
 	}
 
-	public bool AddThirdPartyDependencies(ReadOnlyTargetRules ROTargetRules)
+	public void AddThirdPartyDependencies(ReadOnlyTargetRules ROTargetRules)
 	{
-		string ModulePath = Path.GetDirectoryName(UnrealBuildTool.RulesCompiler.GetFileNameFromType(GetType()));
-		string ThirdPartyPath = Path.GetFullPath(Path.Combine(ModulePath, "../../ThirdParty/"));
+		string ThirdPartyPath = Path.GetFullPath(Path.Combine(ModuleDirectory, "../../ThirdParty/"));
 
 		string PathLib = string.Empty;
 		string PathInc = string.Empty;
@@ -76,7 +61,5 @@ public class MPCDI : ModuleRules
 		PublicIncludePaths.Add(Path.Combine(PathInc, "Creators"));
 		PublicIncludePaths.Add(Path.Combine(PathInc, "IO"));
 		PublicIncludePaths.Add(Path.Combine(PathInc, "Utils"));
-
-		return true;
 	}
 }
diff --git a/Engine/Plugins/Runtime/nDisplay/Source/MPCDI/Private/Blueprints/MPCDIBlueprintAPIImpl.cpp b/Engine/Plugins/Runtime/nDisplay/Source/MPCDI/Private/Blueprints/MPCDIBlueprintAPIImpl.cpp
--- a/Engine/Plugins/Runtime/nDisplay/Source/MPCDI/Private/Blueprints/MPCDIBlueprintAPIImpl.cpp
+++ b/Engine/Plugins/Runtime/nDisplay/Source/MPCDI/Private/Blueprints/MPCDIBlueprintAPIImpl.cpp
@@ -7,26 +7,90 @@
 #include "UObject/Package.h"
 
 #include "IMPCDI.h"
+#include "MPCDILog.h"
 #include "MPCDIRegion.h"
 #include "MPCDIData.h"
 #include "MPCDIWarpTexture.h"
 
+#include "MPCDIStrings.h"
 
-void UMPCDIAPIImpl::GetMPCDIMeshData(const FString& MPCDIFile, const FString& BufferName, const FString& RegionName, struct FMPCDIGeometryExportData& MeshData)
+bool UMPCDIAPIImpl::GetMPCDIMeshData(const FString& MPCDIFile, const FString& BufferName, const FString& RegionName, struct FMPCDIGeometryExportData& MeshData)
 {	
-	IMPCDI& mpcdimodule = IMPCDI::Get();
-	IMPCDI::FRegionLocator regionLocator;
-	bool res = mpcdimodule.GetRegionLocator(MPCDIFile, BufferName, RegionName, regionLocator);
+	IMPCDI& MpcdiModule = IMPCDI::Get();
 
-	if (res)
+	if (!MpcdiModule.Load(MPCDIFile))
+	{
+		//! err
+		return false;
+	}
+
+	IMPCDI::FRegionLocator RegionLocator;
+	bool bResult = MpcdiModule.GetRegionLocator(MPCDIFile, BufferName, RegionName, RegionLocator);
+	if (bResult)
+	{
+		IMPCDI::FShaderInputData ShaderInputData;
+		ShaderInputData.RegionLocator = RegionLocator;
+		TSharedPtr<FMPCDIData> MpcdiData = MpcdiModule.GetMPCDIData(ShaderInputData);
+
+		MPCDI::FMPCDIRegion* Region = MpcdiData.Get()->GetRegion(RegionLocator);
+		MPCDI::FMPCDIWarpTexture* WarpTexture = &Region->WarpMap;
+		WarpTexture->ExportMeshData(MeshData);
+		return true;
+	}
+	return false;
+}
+
+bool UMPCDIAPIImpl::GetPFMMeshData(const FString& LocalPFMFile, FMPCDIGeometryExportData& MeshData, float PFMScale, bool bIsMPCDIAxis)
+{
+	FString DefID(DisplayClusterStrings::cfg::data::mpcdi::PFMFileDefaultID);
+	
+	IMPCDI::FRegionLocator RegionLocator;
+
+	IMPCDI& MpcdiModule = IMPCDI::Get();
+
+	if (MpcdiModule.CreateCustomRegion(DefID, DefID, LocalPFMFile, RegionLocator))
+	{
+		MpcdiModule.SetMPCDIProfileType(RegionLocator, IMPCDI::EMPCDIProfileType::mpcdi_A3D);		
+	}
+
+	//! todo: Force reload
+	if (!MpcdiModule.LoadPFM(RegionLocator, LocalPFMFile, PFMScale, !bIsMPCDIAxis))
+	{
+		UE_LOG(LogMPCDI, Error, TEXT("Couldn't get load ext pfm from file: %s"), *LocalPFMFile);
+		return false;
+	}
+
+	// Get region warp mesh data:
+	IMPCDI::FShaderInputData ShaderInputData;
+	ShaderInputData.RegionLocator = RegionLocator;
+	TSharedPtr<FMPCDIData> MpcdiData = MpcdiModule.GetMPCDIData(ShaderInputData);
+
+	MPCDI::FMPCDIRegion* Region = MpcdiData.Get()->GetRegion(RegionLocator);
+	MPCDI::FMPCDIWarpTexture* WarpTexture = &Region->WarpMap;
+	WarpTexture->ExportMeshData(MeshData);
+	return true;
+}
+
+void UMPCDIAPIImpl::ReloadChangedExternalFiles()
+{
+	IMPCDI::Get().ReloadAll();
+}
+
+void UMPCDIAPIImpl::SetMPCDIMeshData(const FString& MPCDIFile, const FString& BufferName, const FString& RegionName, const struct FMPCDIGeometryImportData& MeshData)
+{
+	IMPCDI& MpcdiModule = IMPCDI::Get();
+	IMPCDI::FRegionLocator RegionLocator;
+	bool bResult = MpcdiModule.GetRegionLocator(MPCDIFile, BufferName, RegionName, RegionLocator);
+
+	if (bResult)
 	{
 		IMPCDI::FShaderInputData ShaderInputData;
-		ShaderInputData.RegionLocator = regionLocator;
-		TSharedPtr<FMPCDIData> mpcdiData = mpcdimodule.GetMPCDIData(ShaderInputData);
+		ShaderInputData.RegionLocator = RegionLocator;
+		TSharedPtr<FMPCDIData> MpcdiData = MpcdiModule.GetMPCDIData(ShaderInputData);
+
+		MPCDI::FMPCDIRegion* Region = MpcdiData.Get()->GetRegion(RegionLocator);
+		MPCDI::FMPCDIWarpTexture* WarpTexture = &Region->WarpMap;
 
-		MPCDI::FMPCDIRegion* region = mpcdiData.Get()->GetRegion(regionLocator);
-		MPCDI::FMPCDIWarpTexture* warpTexture = &region->WarpMap;
-		
-		warpTexture->ExportMeshData(MeshData);
+		WarpTexture->ImportMeshData(MeshData);
 	}
 }
diff --git a/Engine/Plugins/Runtime/nDisplay/Source/MPCDI/Private/Blueprints/MPCDIBlueprintAPIImpl.h b/Engine/Plugins/Runtime/nDisplay/Source/MPCDI/Private/Blueprints/MPCDIBlueprintAPIImpl.h
--- a/Engine/Plugins/Runtime/nDisplay/Source/MPCDI/Private/Blueprints/MPCDIBlueprintAPIImpl.h
+++ b/Engine/Plugins/Runtime/nDisplay/Source/MPCDI/Private/Blueprints/MPCDIBlueprintAPIImpl.h
@@ -21,10 +21,32 @@ class UMPCDIAPIImpl
 
 public:
 	/**
-	* Binds multiple device channels to multiple keys
+	* Return MPCDI mesh data
 	*
-	* @return true if success
 	*/
 	UFUNCTION(BlueprintCallable, meta = (DisplayName = "Get MPCDI Mesh Data"), Category = "MPCDI")
-	virtual void GetMPCDIMeshData(const FString& MPCDIFile, const FString& BufferName, const FString& RegionName, FMPCDIGeometryExportData& MeshData) override;
+	virtual bool GetMPCDIMeshData(const FString& MPCDIFile, const FString& BufferName, const FString& RegionName, FMPCDIGeometryExportData& MeshData) override;
+
+	/**
+	* Return PFM mesh data
+	*
+	*/
+	UFUNCTION(BlueprintCallable, meta = (DisplayName = "Get PFM Mesh Data"), Category = "MPCDI")
+	virtual bool GetPFMMeshData(const FString& PFMFile, FMPCDIGeometryExportData& MeshData, float PFMScale = 1, bool bIsMPCDIAxis = true) override;
+
+	/**
+	* Set MPCDI mesh data at runtime
+	*
+	*/
+	UFUNCTION(BlueprintCallable, meta = (DisplayName = "Set MPCDI Mesh Data"), Category = "MPCDI")
+	virtual void SetMPCDIMeshData(const FString& MPCDIFile, const FString& BufferName, const FString& RegionName, const FMPCDIGeometryImportData& MeshData) override;
+
+	/**
+	* Reload changed external files
+	*
+	*/
+	UFUNCTION(BlueprintCallable, meta = (DisplayName = "Reload Changed External Files"), Category = "MPCDI")
+	virtual void ReloadChangedExternalFiles() override;
+
+
 };
\ No newline at end of file
diff --git a/Engine/Plugins/Runtime/nDisplay/Source/MPCDI/Private/Blueprints/MPCDIContainers.cpp b/Engine/Plugins/Runtime/nDisplay/Source/MPCDI/Private/Blueprints/MPCDIContainers.cpp
--- a/Engine/Plugins/Runtime/nDisplay/Source/MPCDI/Private/Blueprints/MPCDIContainers.cpp
+++ b/Engine/Plugins/Runtime/nDisplay/Source/MPCDI/Private/Blueprints/MPCDIContainers.cpp
@@ -10,12 +10,12 @@ void FMPCDIGeometryExportData::PostAddFace(int f0, int f1, int f2)
 	Triangles.Add(f2);
 
 	//Update normal
-	const FVector dir1 = Vertices[f1] - Vertices[f0];
-	const FVector dir2 = Vertices[f2] - Vertices[f0];
+	const FVector FaceDir1 = Vertices[f1] - Vertices[f0];
+	const FVector FaceDir2 = Vertices[f2] - Vertices[f0];
 
-	const FVector N = FVector::CrossProduct(dir1, dir2);
+	const FVector FaceNornal = FVector::CrossProduct(FaceDir1, FaceDir2);
 
-	Normal[f0] = N;
-	Normal[f1] = N;
-	Normal[f2] = N;
+	Normal[f0] = FaceNornal;
+	Normal[f1] = FaceNornal;
+	Normal[f2] = FaceNornal;
 }
diff --git a/Engine/Plugins/Runtime/nDisplay/Source/MPCDI/Private/MPCDIBlendTexture.cpp b/Engine/Plugins/Runtime/nDisplay/Source/MPCDI/Private/MPCDIBlendTexture.cpp
--- a/Engine/Plugins/Runtime/nDisplay/Source/MPCDI/Private/MPCDIBlendTexture.cpp
+++ b/Engine/Plugins/Runtime/nDisplay/Source/MPCDI/Private/MPCDIBlendTexture.cpp
@@ -40,10 +40,17 @@ namespace MPCDI
 
 	void FMPCDIBlendTexture::LoadCustomMap(EPixelFormat InPixelFormat, int InWidth, int InHeight, int BufferSize, void *InTextureData)
 	{
+		ReleaseTextureData();
+
 		void *TextureData = FMemory::Malloc(BufferSize);
 		memcpy(TextureData, InTextureData, BufferSize);
 
 		SetTextureData(TextureData, InWidth, InHeight, InPixelFormat);
+
+		if (IsInitialized()) 
+		{
+			BeginUpdateResourceRHI(this);
+		}
 		BeginInitResource(this);
 	}
 
diff --git a/Engine/Plugins/Runtime/nDisplay/Source/MPCDI/Private/MPCDIData.cpp b/Engine/Plugins/Runtime/nDisplay/Source/MPCDI/Private/MPCDIData.cpp
--- a/Engine/Plugins/Runtime/nDisplay/Source/MPCDI/Private/MPCDIData.cpp
+++ b/Engine/Plugins/Runtime/nDisplay/Source/MPCDI/Private/MPCDIData.cpp
@@ -9,6 +9,8 @@
 #include "Misc/ConfigCacheIni.h"
 #include "Misc/Paths.h"
 
+#include "MPCDIHelpers.h"
+
 THIRD_PARTY_INCLUDES_START
 
 #include "mpcdiProfile.h"
@@ -23,9 +25,26 @@ THIRD_PARTY_INCLUDES_START
 
 THIRD_PARTY_INCLUDES_END
 
+
+bool FMPCDIData::FMPCDIBuffer::Initialize(const FString& BufferName)
+{
+	ID = BufferName;
+	AABBox = FBox(FVector(FLT_MAX, FLT_MAX, FLT_MAX), FVector(-FLT_MAX, -FLT_MAX, -FLT_MAX));
+	return true;
+}
+void FMPCDIData::FMPCDIBuffer::AddRegion(MPCDI::FMPCDIRegion* MPCDIRegionPtr)
+{
+	if (MPCDIRegionPtr)
+	{
+		MPCDIRegionPtr->WarpMap.AppendAABB(AABBox);
+		Regions.Add(MPCDIRegionPtr);
+	}
+}
+
+
 FMPCDIData::FMPCDIData()
-	: Profile(EMPCDIProfileType::Invalid)
-	, UseColorAdjustment(false)
+	: Profile(IMPCDI::EMPCDIProfileType::Invalid)	
+	, bForceExtFilesReload(false)
 {
 }
 
@@ -34,6 +53,28 @@ FMPCDIData::~FMPCDIData()
 	CleanupMPCDIData();
 }
 
+void FMPCDIData::ReloadAll()
+{
+	FScopeLock lock(&ExtFilesReloadCS);
+	bForceExtFilesReload = true;
+}
+
+void FMPCDIData::ReloadChangedExternalFiles_RenderThread()
+{
+	FScopeLock lock(&ExtFilesReloadCS);
+
+	// Reload Ext pfm+png files
+	for (auto& It : Buffers)
+	{
+		for (auto& Region : It->Regions)
+		{
+			Region->ReloadExternalFiles_RenderThread(bForceExtFilesReload);
+		}
+	}
+
+	bForceExtFilesReload = false;
+}
+
 void FMPCDIData::CleanupMPCDIData()
 {
 	for (auto& ItBuffer : Buffers)
@@ -54,32 +95,90 @@ void FMPCDIData::CleanupMPCDIData()
 	}
 }
 
-bool FMPCDIData::Load(const FString& MPCIDIFile, bool useCubeMap /* = false */, int cubeMapResolution/* = 2048 */, bool ColorAdjustment /* = false */)
+bool FMPCDIData::LoadFromFile(const FString& MPCIDIFile)
 {
-	FilePath = MPCIDIFile;
+	bool success = Load(MPCIDIFile);
+	return success;
+}
 
-	//CubeMapResolution = cubeMapResolution;
-	UseColorAdjustment = ColorAdjustment;
+bool FMPCDIData::Initialize(const FString& MPCIDIFile, IMPCDI::EMPCDIProfileType MPCDIProfileType)
+{
+	LocalMPCIDIFile = MPCIDIFile;
+	Profile = MPCDIProfileType;
 
-	bool success = Load();
-	//UseCubeMap = useCubeMap && Profile == EMPCDIProfileType::mpcdi_A3D;
+	TextureMatrix = FMatrix::Identity;
+	if (Profile == IMPCDI::EMPCDIProfileType::mpcdi_2D || Profile == IMPCDI::EMPCDIProfileType::mpcdi_3D || Profile == IMPCDI::EMPCDIProfileType::mpcdi_SL)
+	{
+		// Fetching from a 2D source
+		TextureMatrix.M[0][0] = 1.f;
+		TextureMatrix.M[1][1] = -1.f;
+		TextureMatrix.M[3][0] = 0.f;
+		TextureMatrix.M[3][1] = 1.f;
+	}
+	else if (Profile == IMPCDI::EMPCDIProfileType::mpcdi_A3D)
+	{
+		// Fetching from a 3D source
+		TextureMatrix.M[0][0] = 0.5f;
+		TextureMatrix.M[1][1] = -0.5f;
+		TextureMatrix.M[3][0] = 0.5f;
+		TextureMatrix.M[3][1] = 0.5f;
+	}
 
-	return success;
+	return MPCDIProfileType != IMPCDI::EMPCDIProfileType::Invalid;
 }
 
+bool FMPCDIData::AddRegion(const FString& BufferName, const FString& RegionName, IMPCDI::FRegionLocator& OutRegionLocator)
+{
+	if (!FindRegion(BufferName, RegionName, OutRegionLocator))
+	{
+		if (OutRegionLocator.BufferIndex < 0)
+		{
+			//Create new buffer
+			FMPCDIBuffer* buf = new FMPCDIData::FMPCDIBuffer();
+			buf->Initialize(BufferName);
+			Buffers.Add(buf);
+			FindRegion(BufferName, RegionName, OutRegionLocator);
+		}
 
-bool FMPCDIData::Load()
+		if (OutRegionLocator.RegionIndex < 0)
+		{
+			if (OutRegionLocator.BufferIndex < Buffers.Num())
+			{
+				FMPCDIBuffer* Buffer = Buffers[OutRegionLocator.BufferIndex];
+				if (Buffer)
+				{
+					MPCDI::FMPCDIRegion* NewRegionPtr = new MPCDI::FMPCDIRegion(*RegionName, 1920, 1080);
+					Buffer->AddRegion(NewRegionPtr);
+					return FindRegion(BufferName, RegionName, OutRegionLocator);
+				}
+			}
+		}
+	}
+
+	//Region already exist
+	return true;
+}
+
+bool FMPCDIData::Load(const FString& MPCIDIFile)
 {
+
+	FString MPCIDIFileFullPath = DisplayClusterHelpers::config::GetFullPath(MPCIDIFile);
+	if (!FPaths::FileExists(MPCIDIFileFullPath))
+	{
+		//! Handle error: mpcdi file not found
+		return false;
+	}
+
 	bool success = false;
 
 	mpcdi::Profile *profile = mpcdi::Profile::CreateProfile();
 	mpcdi::Reader *reader = mpcdi::Reader::CreateReader();
 	std::string version = reader->GetSupportedVersions();
 
-	mpcdi::MPCDI_Error mpcdi_err = reader->Read(TCHAR_TO_ANSI(*FilePath), profile);
+	mpcdi::MPCDI_Error mpcdi_err = reader->Read(TCHAR_TO_ANSI(*MPCIDIFileFullPath), profile);
 	delete reader;
 
-	UE_LOG(LogMPCDI, Log, TEXT("Loading MPCDI file %s."), *FilePath);
+	UE_LOG(LogMPCDI, Log, TEXT("Loading MPCDI file %s."), *MPCIDIFileFullPath);
 
 	if (MPCDI_FAILED(mpcdi_err))
 	{
@@ -89,98 +188,74 @@ bool FMPCDIData::Load()
 	else
 	{
 		// Compute profile
+		IMPCDI::EMPCDIProfileType MPCDIProfileType = IMPCDI::EMPCDIProfileType::Invalid;
 		switch (profile->GetProfileType())
 		{
 		case mpcdi::ProfileType2d:
-			Profile = EMPCDIProfileType::mpcdi_2D;
+			MPCDIProfileType = IMPCDI::EMPCDIProfileType::mpcdi_2D;
 			break;
 		case mpcdi::ProfileType3d:
-			Profile = EMPCDIProfileType::mpcdi_3D;
+			MPCDIProfileType = IMPCDI::EMPCDIProfileType::mpcdi_3D;
 			break;
 		case mpcdi::ProfileTypea3:
-			Profile = EMPCDIProfileType::mpcdi_A3D;
+			MPCDIProfileType = IMPCDI::EMPCDIProfileType::mpcdi_A3D;
 			break;
 		case mpcdi::ProfileTypesl:
-			Profile = EMPCDIProfileType::mpcdi_SL;
+			MPCDIProfileType = IMPCDI::EMPCDIProfileType::mpcdi_SL;
 			break;
 		};
 
+		if (!Initialize(MPCIDIFile, MPCDIProfileType))
+		{
+			return false;
+		}
+
 		// Store version
 		Version = version.c_str();
 		UE_LOG(LogMPCDI, Verbose, TEXT("Version: %s"), *Version);
 
-		TextureMatrix = FMatrix::Identity;
-		if (Profile == EMPCDIProfileType::mpcdi_2D || Profile == EMPCDIProfileType::mpcdi_3D || Profile == EMPCDIProfileType::mpcdi_SL)
-		{
-			// Fetching from a 2D source
-			TextureMatrix.M[0][0] = 1.f;
-			TextureMatrix.M[1][1] = -1.f;
-			TextureMatrix.M[3][0] = 0.f;
-			TextureMatrix.M[3][1] = 1.f;
-		}
-		else if (Profile == EMPCDIProfileType::mpcdi_A3D)
-		{
-			// Fetching from a 3D source
-			TextureMatrix.M[0][0] = 0.5f;
-			TextureMatrix.M[1][1] = -0.5f;
-			TextureMatrix.M[3][0] = 0.5f;
-			TextureMatrix.M[3][1] = 0.5f;
-		}
+		
 
-		// Fill buffer information
-		int bufferIndex = 0;
+		// Fill buffer information		
 		for (mpcdi::Display::BufferIterator itBuffer = profile->GetDisplay()->GetBufferBegin(); itBuffer != profile->GetDisplay()->GetBufferEnd(); ++itBuffer)
 		{
 			Buffers.Add(new FMPCDIBuffer());
 			FMPCDIBuffer &buffer = *(Buffers.Last());
-
 			mpcdi::Buffer *mpcdiBuffer = itBuffer->second;
-
-			buffer.ID = FString(mpcdiBuffer->GetId().c_str());
-			buffer.AABBox = FBox(FVector(FLT_MAX, FLT_MAX, FLT_MAX), FVector(-FLT_MAX, -FLT_MAX, -FLT_MAX));
-
-			UE_LOG(LogMPCDI, Verbose, TEXT("Buffer %d - %s"), bufferIndex, ANSI_TO_TCHAR(mpcdiBuffer->GetId().c_str()));
+			buffer.Initialize(FString(mpcdiBuffer->GetId().c_str()));
 
 			// Fill-in region information
-			int regionIndex = 0;
-
-			//RegionAllocator::Reset();
 			for (mpcdi::Buffer::RegionIterator it = mpcdiBuffer->GetRegionBegin(); it != mpcdiBuffer->GetRegionEnd(); ++it)
 			{
 				mpcdi::Region *mpcdiRegion = it->second;
 				if (mpcdiRegion)
 				{
 					MPCDI::FMPCDIRegion* NewRegionPtr = new MPCDI::FMPCDIRegion();
-					NewRegionPtr->Load(mpcdiRegion, GetProfileType());
-					NewRegionPtr->WarpMap.AppendAABB(buffer.AABBox);
-					buffer.Regions.Add(NewRegionPtr);
+					if (NewRegionPtr->Load(mpcdiRegion, GetProfileType()))
+					{
+						buffer.AddRegion(NewRegionPtr);
+					}
+					else
+					{
+						//@todo handle error
+						UE_LOG(LogMPCDI, Error, TEXT("Can't load mpcdi region %s"), ANSI_TO_TCHAR(mpcdiRegion->GetId().c_str()));
+					}					
 				}
-				++regionIndex;
 			}// end region loop
-
-			++bufferIndex;
-
 		}//end buffer loop
-
 		success = true;
 	}
 
 	delete profile;
-
 	return success;
 }
 
 void FMPCDIData::AddReferencedObjects(FReferenceCollector& Collector)
 {
-
 }
 
-bool FMPCDIData::IsCubeMapEnabled() const
-{
-	return false;
-}
-
-/*
+#if 0
+//@todo Unsupported now
 bool FMPCDIData::ComputeFrustum_SL(const IMPCDI::FRegionLocator& RegionLocator, IMPCDI::FFrustum &OutFrustum, float WorldScale, float ZNear, float ZFar) const
 {
 	const FMPCDIData::FMPCDIBuffer *MPCDIBuffer = Buffers[RegionLocator.BufferIndex];
@@ -212,7 +287,7 @@ bool FMPCDIData::ComputeFrustum_SL(const IMPCDI::FRegionLocator& RegionLocator,
 	OutFrustum.UVMatrix = FMatrix::Identity;
 	return true;
 }
-*/
+#endif
 
 bool FMPCDIData::ComputeFrustum(const IMPCDI::FRegionLocator& RegionLocator, IMPCDI::FFrustum& Frustum, float WorldScale, float ZNear, float ZFar) const
 {
@@ -227,7 +302,7 @@ bool FMPCDIData::ComputeFrustum(const IMPCDI::FRegionLocator& RegionLocator, IMP
 	{
 		switch (GetProfileType())
 		{
-		case EMPCDIProfileType::mpcdi_A3D:
+		case IMPCDI::EMPCDIProfileType::mpcdi_A3D:
 			return GetRegion(RegionLocator)->WarpMap.GetFrustum_A3D(Frustum, WorldScale, ZNear, ZFar);
 			break;
 #if 0
@@ -244,6 +319,5 @@ bool FMPCDIData::ComputeFrustum(const IMPCDI::FRegionLocator& RegionLocator, IMP
 			break;
 		}
 	}
-
 	return false;
 }
diff --git a/Engine/Plugins/Runtime/nDisplay/Source/MPCDI/Private/MPCDIData.h b/Engine/Plugins/Runtime/nDisplay/Source/MPCDI/Private/MPCDIData.h
--- a/Engine/Plugins/Runtime/nDisplay/Source/MPCDI/Private/MPCDIData.h
+++ b/Engine/Plugins/Runtime/nDisplay/Source/MPCDI/Private/MPCDIData.h
@@ -4,7 +4,6 @@
 
 #include "IMPCDI.h"
 
-#include "MPCDITypes.h"
 #include "MPCDIRegion.h"
 
 #include "RendererInterface.h"
@@ -12,22 +11,8 @@
 #include "RenderResource.h"
 #include "SceneTypes.h"
 
-
-// Engine-side interface to FMPCDIData. We can only call into interfaces from the Engine module since Renderer is intentionally not a dependency of it.
-class IMPCDIData
-{
-public:
-	virtual ~IMPCDIData() = 0
-	{ }
-
-public:
-	virtual bool IsCubeMapEnabled() const = 0;
-};
-
-
 class FMPCDIData
 	: public FGCObject
-	, public IMPCDIData
 {
 public:
 	FMPCDIData();
@@ -41,17 +26,18 @@ public:
 			Regions.Empty();
 		}
 
-		FString ID;
-
-		TArray<MPCDI::FMPCDIRegion*> Regions;
-
-		FBox AABBox;
+		FString                      ID;       // Unique buffer name
+		TArray<MPCDI::FMPCDIRegion*> Regions;  // Asset of warp regions
+		FBox                         AABBox;
 
 		// Shader Lamp view and projection matrices
 		FVector Origin;
 		FMatrix ProjectionMatrix;
 		FMatrix Camera2World;
 
+		bool Initialize(const FString& BufferName);
+		void AddRegion(MPCDI::FMPCDIRegion* MPCDIRegionPtr);
+
 		bool FindRegion(const FString& RegionName, IMPCDI::FRegionLocator& OutRegionLocator) const
 		{
 			for (int RegionIndex = 0; RegionIndex < Regions.Num(); ++RegionIndex)
@@ -67,12 +53,17 @@ public:
 	};
 
 public:
-	virtual bool IsCubeMapEnabled() const override;
+	bool Initialize(const FString& MPCIDIFile, IMPCDI::EMPCDIProfileType MPCDIProfileType);
+	bool LoadFromFile(const FString& MPCIDIFile);
 
-public:
-	bool Load(const FString& MPCIDIFile, bool useCubeMap = false, int cubeMapResolution = 2048, bool mpcdiColorAdjustment = false);
 	void CleanupMPCDIData();
 
+
+	void ReloadAll();
+	void ReloadChangedExternalFiles_RenderThread();
+
+	bool AddRegion(const FString& BufferName, const FString& RegionName, IMPCDI::FRegionLocator& OutRegionLocator);
+
 	MPCDI::FMPCDIRegion* GetRegion(const IMPCDI::FRegionLocator& RegionLocator) const
 	{
 		if (RegionLocator.BufferIndex < Buffers.Num())
@@ -102,33 +93,45 @@ public:
 	}
 
 	inline const FMatrix &GetTextureMatrix() const
-	{ return TextureMatrix; }
+	{ 
+		return TextureMatrix; 
+	}
 
 	inline bool IsValid() const
-	{ return Profile != EMPCDIProfileType::Invalid; }
+	{ 
+		return Profile != IMPCDI::EMPCDIProfileType::Invalid; 
+	}
 
-	inline EMPCDIProfileType GetProfileType() const
-	{ return Profile; }
-	
-	bool ColorAdjustmentsEnabled() const
-	{ return IsValid() && UseColorAdjustment; }
+	inline IMPCDI::EMPCDIProfileType GetProfileType() const
+	{ 
+		return Profile; 
+	}
 
-	inline const FString& GetFilePath() const
-	{ return FilePath; }
+	inline const FString& GetLocalMPCIDIFile() const
+	{ 
+		return LocalMPCIDIFile;
+	}
 
 	bool ComputeFrustum(const IMPCDI::FRegionLocator& RegionLocator, IMPCDI::FFrustum &Frustum, float worldScale, float ZNear, float ZFar) const;
 
 private:
-	bool Load();
+#if 0
+	//@todo Unsupported now
+	bool ComputeFrustum_SL(const IMPCDI::FRegionLocator& RegionLocator, IMPCDI::FFrustum &OutFrustum, float WorldScale, float ZNear, float ZFar) const;
+#endif
+
+	bool Load(const FString& MPCDIFile);
 	virtual void AddReferencedObjects(FReferenceCollector& Collector) override;
 
 private:
 	TArray<FMPCDIBuffer*> Buffers;
-	EMPCDIProfileType     Profile;
+	IMPCDI::EMPCDIProfileType     Profile;
 	FString Version;
 
 	FMatrix TextureMatrix;
 
-	FString FilePath;
-	bool    UseColorAdjustment;
+	FString    LocalMPCIDIFile;
+
+	mutable bool       bForceExtFilesReload;
+	FCriticalSection   ExtFilesReloadCS;
 };
diff --git a/Engine/Plugins/Runtime/nDisplay/Source/MPCDI/Private/MPCDIHelpers.h b/Engine/Plugins/Runtime/nDisplay/Source/MPCDI/Private/MPCDIHelpers.h
new file mode 100644
--- /dev/null
+++ b/Engine/Plugins/Runtime/nDisplay/Source/MPCDI/Private/MPCDIHelpers.h
@@ -0,0 +1,12 @@
+// Copyright 1998-2019 Epic Games, Inc. All Rights Reserved.
+
+#pragma once
+
+#include "CoreMinimal.h"
+
+#include "DisplayClusterUtils/DisplayClusterCommonHelpers.h"
+
+
+namespace DisplayClusterHelpers
+{
+};
diff --git a/Engine/Plugins/Runtime/nDisplay/Source/MPCDI/Private/MPCDIModule.cpp b/Engine/Plugins/Runtime/nDisplay/Source/MPCDI/Private/MPCDIModule.cpp
--- a/Engine/Plugins/Runtime/nDisplay/Source/MPCDI/Private/MPCDIModule.cpp
+++ b/Engine/Plugins/Runtime/nDisplay/Source/MPCDI/Private/MPCDIModule.cpp
@@ -9,8 +9,15 @@
 #include "Misc/Paths.h"
 #include "ShaderCore.h"
 
+#include "Misc/FileHelper.h"
+
 #include "RHICommandList.h"
 
+#include "MPCDIHelpers.h"
+#include "MPCDILog.h"
+#include "MPCDIStrings.h"
+
+
 
 //////////////////////////////////////////////////////////////////////////////////////////////
 // IModuleInterface
@@ -30,23 +37,19 @@ void FMPCDIModule::ShutdownModule()
 //////////////////////////////////////////////////////////////////////////////////////////////
 // IMPCDI
 //////////////////////////////////////////////////////////////////////////////////////////////
-bool FMPCDIModule::Load(const FString& MPCDIFile)
+bool FMPCDIModule::Load(const FString& LocalMPCDIFile)
 {
 	FScopeLock lock(&DataGuard);
 
-	if (IsLoaded(MPCDIFile))
+	if (IsLoaded(LocalMPCDIFile))
 	{
 		return true;
 	}
 
-	if (!FPaths::FileExists(MPCDIFile))
-	{
-		return false;
-	}
-
-	TSharedPtr<FMPCDIData> DataItem(new FMPCDIData);
-	if (!DataItem->Load(MPCDIFile))
+	TSharedPtr<FMPCDIData> DataItem = MakeShareable(new FMPCDIData);
+	if (!DataItem->LoadFromFile(LocalMPCDIFile))
 	{
+		//! Handle error
 		return false;
 	}
 
@@ -56,56 +59,51 @@ bool FMPCDIModule::Load(const FString& MPCDIFile)
 	return true;
 }
 
-bool FMPCDIModule::IsLoaded(const FString& MPCDIFile)
+bool FMPCDIModule::IsLoaded(const FString& LocalMPCDIFile)
 {
 	FScopeLock lock(&DataGuard);
-
-	return MPCDIData.ContainsByPredicate([MPCDIFile](const TSharedPtr<FMPCDIData>& DataItem)
+	return MPCDIData.ContainsByPredicate([LocalMPCDIFile](const TSharedPtr<FMPCDIData>& DataItem)
 	{
-		return FPaths::IsSamePath(MPCDIFile, DataItem->GetFilePath());
+		return FPaths::IsSamePath(LocalMPCDIFile, DataItem->GetLocalMPCIDIFile());
 	});
 }
 
-bool FMPCDIModule::GetRegionLocator(const FString& MPCDIFile, const FString& BufferName, const FString& RegionName, IMPCDI::FRegionLocator& OutRegionLocator)
+bool FMPCDIModule::GetRegionLocator(const FString& LocalMPCDIFile, const FString& BufferName, const FString& RegionName, IMPCDI::FRegionLocator& OutRegionLocator)
 {
 	FScopeLock lock(&DataGuard);
-
+	
 	// Reset to defaults
 	OutRegionLocator = IMPCDI::FRegionLocator();
-
-	if (!IsLoaded(MPCDIFile))
-	{
-		if (!Load(MPCDIFile))
-		{
-			//@todo logs
-			return false;
-		}
-	}
-
-	TSharedPtr<FMPCDIData> DataItem;
+	
+	FMPCDIData* DataItemPtr = nullptr;
 	IMPCDI::FRegionLocator TmpRegionLocator;
 
 	// Find the file index
 	for (int FileIndex = 0; FileIndex < MPCDIData.Num(); ++FileIndex)
 	{
-		if (FPaths::IsSamePath(MPCDIData[FileIndex]->GetFilePath(), MPCDIFile))
+		if (FPaths::IsSamePath(MPCDIData[FileIndex]->GetLocalMPCIDIFile(), LocalMPCDIFile))
 		{
-			DataItem = MPCDIData[FileIndex];
+			DataItemPtr = MPCDIData[FileIndex].Get();
 			TmpRegionLocator.FileIndex = FileIndex;
 			break;
 		}
 	}
 
+	if (!DataItemPtr)
+	{
+		//! handle error: file not loaded
+		return false;
+	}
+
 	// Try to find the requested region
-	if (!DataItem->FindRegion(BufferName, RegionName, TmpRegionLocator))
+	if (!DataItemPtr->FindRegion(BufferName, RegionName, TmpRegionLocator))
 	{
-		//@todo logs
+		//! Handle error: BufferName + RegionName not defined for LocalMPCDIFile
 		return false;
 	}
 
-	// Export data
+	// Return handler to warp data region
 	OutRegionLocator = TmpRegionLocator;
-
 	return true;
 }
 
@@ -144,6 +142,9 @@ TSharedPtr<FMPCDIData> FMPCDIModule::GetMPCDIData(IMPCDI::FShaderInputData& Shad
 
 bool FMPCDIModule::ApplyWarpBlend(FRHICommandListImmediate& RHICmdList, IMPCDI::FTextureWarpData& TextureWarpData, IMPCDI::FShaderInputData& ShaderInputData)
 {
+	//Support runtime PFM reload
+	ReloadAll_RenderThread();
+
 	FScopeLock lock(&DataGuard);
 
 	if (MPCDIData.Num() > ShaderInputData.RegionLocator.FileIndex)
@@ -158,6 +159,7 @@ bool FMPCDIModule::ApplyWarpBlend(FRHICommandListImmediate& RHICmdList, IMPCDI::
 	return false;
 }
 
+
 void FMPCDIModule::ReleaseMPCDIData()
 {
 	FScopeLock lock(&DataGuard);
@@ -165,5 +167,348 @@ void FMPCDIModule::ReleaseMPCDIData()
 	MPCDIData.Empty();
 }
 
+
+//@ Support ext PFM & PNG files
+bool FMPCDIModule::CreateCustomRegion(const FString& LocalMPCDIFile, const FString& BufferName, const FString& RegionName, IMPCDI::FRegionLocator& OutRegionLocator)
+{
+	FScopeLock lock(&DataGuard);
+	if (GetRegionLocator(LocalMPCDIFile, BufferName, RegionName, OutRegionLocator))
+	{
+		return true;
+	}
+
+	if (OutRegionLocator.FileIndex < 0)
+	{
+		//Create new file:
+		TSharedPtr<FMPCDIData> DataItem = MakeShareable(new FMPCDIData);
+		DataItem->Initialize(LocalMPCDIFile, EMPCDIProfileType::mpcdi_A3D); // Use a3d as default profile
+		// Load and cache the file
+		MPCDIData.Add(DataItem);
+		OutRegionLocator.FileIndex = MPCDIData.Num() - 1; // Update locator
+	}
+
+	return MPCDIData[OutRegionLocator.FileIndex]->AddRegion(BufferName, RegionName, OutRegionLocator);
+}
+
+bool FMPCDIModule::SetMPCDIProfileType(const IMPCDI::FRegionLocator& InRegionLocator, const EMPCDIProfileType ProfileType)
+{
+	FScopeLock lock(&DataGuard);
+
+	if (InRegionLocator.FileIndex >= 0)
+	{
+		FMPCDIData& Dst = *MPCDIData[InRegionLocator.FileIndex];
+		FString LocalMPCIDIFile = Dst.GetLocalMPCIDIFile();
+		Dst.Initialize(LocalMPCIDIFile, ProfileType);
+		return true;
+	}
+	return false;
+}
+
+bool FMPCDIModule::LoadPFM(const IMPCDI::FRegionLocator& InRegionLocator, const FString& LocalPFMFile, const float PFMFileScale, bool bIsUnrealGameSpace)
+{
+	FScopeLock lock(&DataGuard);
+
+	if (InRegionLocator.RegionIndex >= 0)
+	{
+		if (MPCDIData.Num() > InRegionLocator.FileIndex)
+		{
+			FMPCDIData& Dst = *MPCDIData[InRegionLocator.FileIndex];
+			MPCDI::FMPCDIRegion* DstRegion = Dst.GetRegion(InRegionLocator);
+			if (DstRegion)
+			{
+				return DstRegion->LoadExtPFMFile(LocalPFMFile, Dst.GetProfileType(), PFMFileScale, bIsUnrealGameSpace);
+			}
+		}
+	}
+	return false;
+}
+
+bool FMPCDIModule::LoadPFMGeometry(const IMPCDI::FRegionLocator& InRegionLocator, const TArray<FVector>& PFMPoints, int DimW, int DimH, const float PFMScale, bool bIsUnrealGameSpace)
+{
+	FScopeLock lock(&DataGuard);
+
+	if (InRegionLocator.RegionIndex >= 0)
+	{
+		if (MPCDIData.Num() > InRegionLocator.FileIndex)
+		{
+			FMPCDIData& Dst = *MPCDIData[InRegionLocator.FileIndex];
+			MPCDI::FMPCDIRegion* DstRegion = Dst.GetRegion(InRegionLocator);
+			if (DstRegion)
+			{
+				return DstRegion->LoadExtGeometry(PFMPoints, DimW, DimH, Dst.GetProfileType(), PFMScale, bIsUnrealGameSpace);
+			}
+		}
+	}
+	return false;
+}
+
+bool FMPCDIModule::LoadAlphaMap(const IMPCDI::FRegionLocator& InRegionLocator, const FString& LocalPNGFile, float GammaValue)
+{
+	FScopeLock lock(&DataGuard);
+
+	if (InRegionLocator.RegionIndex >= 0)
+	{
+		if (MPCDIData.Num() > InRegionLocator.FileIndex)
+		{
+			FMPCDIData& Dst = *MPCDIData[InRegionLocator.FileIndex];
+			MPCDI::FMPCDIRegion* DstRegion = Dst.GetRegion(InRegionLocator);
+			if (DstRegion)
+			{				
+				return DstRegion->LoadExtAlphaMap(LocalPNGFile,GammaValue);
+			}
+		}
+	}
+	return false;
+}
+
+bool FMPCDIModule::LoadBetaMap(const IMPCDI::FRegionLocator& InRegionLocator, const FString& LocalPNGFile)
+{
+	FScopeLock lock(&DataGuard);
+
+	if (InRegionLocator.RegionIndex >= 0)
+	{
+		if (MPCDIData.Num() > InRegionLocator.FileIndex)
+		{
+			FMPCDIData& Dst = *MPCDIData[InRegionLocator.FileIndex];
+			MPCDI::FMPCDIRegion* DstRegion = Dst.GetRegion(InRegionLocator);
+			if (DstRegion)
+			{				
+				return DstRegion->LoadExtBetaMap(LocalPNGFile);
+			}
+		}
+	}
+	return false;
+}
+
+bool FMPCDIModule::Load(const ConfigParser& CfgData, IMPCDI::FRegionLocator& OutRegionLocator)
+{
+	bool bResult = false;
+	if (CfgData.IsExtConfig())
+	{
+		if (CreateCustomRegion(CfgData.MPCDIFileName, CfgData.BufferId, CfgData.RegionId, OutRegionLocator))
+		{
+			SetMPCDIProfileType(OutRegionLocator, CfgData.MPCDIType);
+			bResult = true;
+		}
+	}
+	else
+	{// Load from MPCDI file:
+		// Check if MPCDI file exists
+		if (!FPaths::FileExists(CfgData.MPCDIFileName))
+		{
+			UE_LOG(LogMPCDI, Warning, TEXT("File not found: %s"), *CfgData.MPCDIFileName);
+			return false;
+		}
+
+		// Load MPCDI file
+		if (!Load(CfgData.MPCDIFileName))
+		{
+			UE_LOG(LogMPCDI, Warning, TEXT("Couldn't load MPCDI file: %s"), *CfgData.MPCDIFileName);
+			return false;
+		}
+
+		// Store MPCDI region locator for this viewport
+		if (!GetRegionLocator(CfgData.MPCDIFileName, CfgData.BufferId, CfgData.RegionId, OutRegionLocator))
+		{
+			UE_LOG(LogMPCDI, Warning, TEXT("Couldn't get region locator for <buf %s, reg %s> in file: %s"), *CfgData.BufferId, *CfgData.RegionId, *CfgData.MPCDIFileName);
+			return false;
+		}
+
+		bResult = true;
+	}
+
+	{// Support overrides:
+		// PFM
+		if (!CfgData.PFMFile.IsEmpty())
+		{
+			if (!LoadPFM(OutRegionLocator, CfgData.PFMFile, CfgData.PFMFileScale, CfgData.bIsUnrealGameSpace))
+			{
+				UE_LOG(LogMPCDI, Warning, TEXT("Couldn't get load ext pfm for <buf %s, reg %s> from file: %s"), *CfgData.BufferId, *CfgData.RegionId, *CfgData.PFMFile);
+			}
+		}
+		// AlphaMAP
+		if (!CfgData.AlphaFile.IsEmpty())
+		{
+			if (!LoadAlphaMap(OutRegionLocator, CfgData.AlphaFile, CfgData.AlphaGamma))
+			{
+				UE_LOG(LogMPCDI, Warning, TEXT("Couldn't get load ext alphamap for <buf %s, reg %s> from file: %s"), *CfgData.BufferId, *CfgData.RegionId, *CfgData.AlphaFile);
+			}
+		}	
+		// BetaMAP
+		if (!CfgData.BetaFile.IsEmpty())
+		{
+			if (!LoadBetaMap(OutRegionLocator, CfgData.BetaFile))
+			{
+				UE_LOG(LogMPCDI, Warning, TEXT("Couldn't get load ext betamap for <buf %s, reg %s> from file: %s"), *CfgData.BufferId, *CfgData.RegionId, *CfgData.BetaFile);
+			}
+		}
+	}
+
+	return bResult;
+}
+
+void FMPCDIModule::ReloadAll()
+{
+	for (auto& It : MPCDIData)
+	{
+		It->ReloadAll();
+	}
+}
+
+void FMPCDIModule::ReloadAll_RenderThread()
+{
+	FScopeLock lock(&DataGuard);
+
+	check(IsInRenderingThread());
+
+	for (auto& It : MPCDIData)
+	{
+		It->ReloadChangedExternalFiles_RenderThread();
+	}
+}
+
+
+//@todo: Move config parsing and logic to policy mpcdi+picp
+// virtual bool LoadConfig(const FString& InConfigLineStr, ConfigParser& OutCfgData) override;
+// virtual bool Load(const ConfigParser& CfgData, IMPCDI::FRegionLocator& OutRegionLocator) override;
+
+bool FMPCDIModule::LoadConfig(const FString& ConfigLineStr, ConfigParser& OutConfig)
+{
+	OutConfig.ConfigLineStr = ConfigLineStr;
+
+	// PFM file (optional)
+	FString LocalPFMFile;
+	if (DisplayClusterHelpers::str::ExtractValue(ConfigLineStr, DisplayClusterStrings::cfg::data::mpcdi::FilePFM, LocalPFMFile))
+	{
+		UE_LOG(LogMPCDI, Log, TEXT("Found Argument '%s'='%s'"), DisplayClusterStrings::cfg::data::mpcdi::FilePFM,*LocalPFMFile);
+		OutConfig.PFMFile = LocalPFMFile;
+	}
+
+	// Buffer
+	if (!DisplayClusterHelpers::str::ExtractValue(ConfigLineStr, DisplayClusterStrings::cfg::data::mpcdi::Buffer, OutConfig.BufferId))
+	{
+		if (OutConfig.PFMFile.IsEmpty())
+		{
+			UE_LOG(LogMPCDI, Error, TEXT("Argument '%s' not found in the config file"), DisplayClusterStrings::cfg::data::mpcdi::Buffer);
+			return false;
+		}
+	}
+
+	// Region
+	if (!DisplayClusterHelpers::str::ExtractValue(ConfigLineStr, DisplayClusterStrings::cfg::data::mpcdi::Region, OutConfig.RegionId))
+	{
+		if (OutConfig.PFMFile.IsEmpty())
+		{
+			UE_LOG(LogMPCDI, Error, TEXT("Argument '%s' not found in the config file"), DisplayClusterStrings::cfg::data::mpcdi::Region);
+			return false;
+		}
+	}
+
+	// Filename
+	FString LocalMPCDIFileName;
+	if (DisplayClusterHelpers::str::ExtractValue(ConfigLineStr, DisplayClusterStrings::cfg::data::mpcdi::File, LocalMPCDIFileName))
+	{
+		UE_LOG(LogMPCDI, Log, TEXT("Found mpcdi file name for %s:%s - %s"), *OutConfig.BufferId, *OutConfig.RegionId, *LocalMPCDIFileName);
+		OutConfig.MPCDIFileName = LocalMPCDIFileName;
+	}
+
+	// Do Autofill for ext config:
+	if (!OutConfig.PFMFile.IsEmpty())
+	{
+		if (OutConfig.RegionId.IsEmpty())
+		{
+			OutConfig.RegionId = OutConfig.PFMFile;
+		}
+
+		if (OutConfig.BufferId.IsEmpty())
+		{
+			OutConfig.BufferId = DisplayClusterStrings::cfg::data::mpcdi::PFMFileDefaultID;
+		}
+
+		if (OutConfig.MPCDIFileName.IsEmpty())
+		{
+			OutConfig.MPCDIFileName = DisplayClusterStrings::cfg::data::mpcdi::PFMFileDefaultID;
+		}
+	}
+
+
+	// Origin node (optional)
+	if (DisplayClusterHelpers::str::ExtractValue(ConfigLineStr, DisplayClusterStrings::cfg::data::mpcdi::Origin, OutConfig.OriginType))
+	{
+		UE_LOG(LogMPCDI, Log, TEXT("Found origin node for %s:%s - %s"), *OutConfig.BufferId, *OutConfig.RegionId, *OutConfig.OriginType);
+	}
+	else
+	{
+		UE_LOG(LogMPCDI, Log, TEXT("No origin node found for %s:%s. VR root will be used as default."), *OutConfig.BufferId, *OutConfig.RegionId);
+	}
+
+	{// Try get custom setup:		
+		// MPCDIType (optional)
+		FString MPCDITypeStr;
+		if (!DisplayClusterHelpers::str::ExtractValue(ConfigLineStr, DisplayClusterStrings::cfg::data::mpcdi::MPCDIType, MPCDITypeStr))
+		{
+			OutConfig.MPCDIType = IMPCDI::EMPCDIProfileType::mpcdi_A3D;
+		}
+		else
+		{
+			OutConfig.MPCDIType = IMPCDI::EMPCDIProfileType::Invalid;
+
+			static const TArray<FString> strEnum({"2d","3d","a3d","sl"});
+			for (int i = 0; i < strEnum.Num(); ++i)
+			{
+				if (!MPCDITypeStr.Compare(strEnum[i], ESearchCase::IgnoreCase))
+				{
+					OutConfig.MPCDIType = (EMPCDIProfileType)i;
+					break;
+				}
+			}
+
+			if (OutConfig.MPCDIType == IMPCDI::EMPCDIProfileType::Invalid)
+			{
+				UE_LOG(LogMPCDI, Error, TEXT("Argument '%s' has unknown value '%s'"), DisplayClusterStrings::cfg::data::mpcdi::MPCDIType, *MPCDITypeStr);
+				return false;
+			}
+		}
+		
+		// Default is UE scale, cm
+		OutConfig.PFMFileScale = 1;
+		if (DisplayClusterHelpers::str::ExtractValue(ConfigLineStr, DisplayClusterStrings::cfg::data::mpcdi::WorldScale, OutConfig.PFMFileScale))
+		{
+			UE_LOG(LogMPCDI, Log, TEXT("Found WorldScale value for %s:%s - %.f"), *OutConfig.BufferId, *OutConfig.RegionId, OutConfig.PFMFileScale);
+		}
+		OutConfig.bIsUnrealGameSpace = false;
+		if (DisplayClusterHelpers::str::ExtractValue(ConfigLineStr, DisplayClusterStrings::cfg::data::mpcdi::UseUnrealAxis, OutConfig.bIsUnrealGameSpace))
+		{
+			UE_LOG(LogMPCDI, Log, TEXT("Found bIsUnrealGameSpace value for %s:%s - %s"), *OutConfig.BufferId, *OutConfig.RegionId, OutConfig.bIsUnrealGameSpace?"true":"false");
+		}
+
+
+		// AlphaFile file (optional)
+		FString LocalAlphaFile;
+		if (DisplayClusterHelpers::str::ExtractValue(ConfigLineStr, DisplayClusterStrings::cfg::data::mpcdi::FileAlpha, LocalAlphaFile))
+		{
+			UE_LOG(LogMPCDI, Log, TEXT("Found external AlphaMap file for %s:%s - %s"), *OutConfig.BufferId, *OutConfig.RegionId, *LocalAlphaFile);
+			OutConfig.AlphaFile = LocalAlphaFile;
+		}
+		OutConfig.AlphaGamma = 1;
+		if (DisplayClusterHelpers::str::ExtractValue(ConfigLineStr, DisplayClusterStrings::cfg::data::mpcdi::AlphaGamma, OutConfig.AlphaGamma))
+		{
+			UE_LOG(LogMPCDI, Log, TEXT("Found AlphaGamma value for %s:%s - %.f"), *OutConfig.BufferId, *OutConfig.RegionId, OutConfig.AlphaGamma);
+		}
+
+		// BetaFile file (optional)
+		FString LocalBetaFile;
+		if (DisplayClusterHelpers::str::ExtractValue(ConfigLineStr, DisplayClusterStrings::cfg::data::mpcdi::FileBeta, LocalBetaFile))
+		{
+			UE_LOG(LogMPCDI, Log, TEXT("Found external BetaMap file for %s:%s - %s"), *OutConfig.BufferId, *OutConfig.RegionId, *LocalBetaFile);
+			OutConfig.BetaFile = LocalBetaFile;
+		}
+		
+	}
+
+	return true;
+}
+
+
 IMPLEMENT_MODULE(FMPCDIModule, MPCDI);
 
diff --git a/Engine/Plugins/Runtime/nDisplay/Source/MPCDI/Private/MPCDIModule.h b/Engine/Plugins/Runtime/nDisplay/Source/MPCDI/Private/MPCDIModule.h
--- a/Engine/Plugins/Runtime/nDisplay/Source/MPCDI/Private/MPCDIModule.h
+++ b/Engine/Plugins/Runtime/nDisplay/Source/MPCDI/Private/MPCDIModule.h
@@ -27,8 +27,23 @@ public:
 	virtual bool ApplyWarpBlend(FRHICommandListImmediate& RHICmdList, IMPCDI::FTextureWarpData& TextureWarpData, IMPCDI::FShaderInputData& ShaderInputData) override;
 	virtual TSharedPtr<FMPCDIData> GetMPCDIData(IMPCDI::FShaderInputData& ShaderInputData) override;
 
+	virtual bool CreateCustomRegion(const FString& MPCDIFile, const FString& BufferName, const FString& RegionName, IMPCDI::FRegionLocator& OutRegionLocator) override;
+	virtual bool SetMPCDIProfileType(const IMPCDI::FRegionLocator& InRegionLocator, const EMPCDIProfileType ProfileType) override;
+	virtual bool LoadPFM(const IMPCDI::FRegionLocator& InRegionLocator, const FString& PFMFile, const float PFMScale, bool bIsUnrealGameSpace = false) override;
+	virtual bool LoadAlphaMap(const IMPCDI::FRegionLocator& InRegionLocator, const FString& PNGFile, float GammaValue) override;
+	virtual bool LoadBetaMap(const IMPCDI::FRegionLocator& InRegionLocator, const FString& PNGFile) override;
+
+	virtual bool LoadPFMGeometry(const IMPCDI::FRegionLocator& InRegionLocator, const TArray<FVector>& PFMPoints, int DimW, int DimH, const float WorldScale, bool bIsUnrealGameSpace = false) override;
+
+	virtual bool LoadConfig(const FString& InConfigLineStr, ConfigParser& OutCfgData) override;
+	virtual bool Load(const ConfigParser& CfgData, IMPCDI::FRegionLocator& OutRegionLocator) override;
+
+	virtual void ReloadAll() override;
+	virtual void ReloadAll_RenderThread()override;
+
 private:
 	void ReleaseMPCDIData();
+	
 
 private:
 	FCriticalSection DataGuard;
diff --git a/Engine/Plugins/Runtime/nDisplay/Source/MPCDI/Private/MPCDIRegion.cpp b/Engine/Plugins/Runtime/nDisplay/Source/MPCDI/Private/MPCDIRegion.cpp
--- a/Engine/Plugins/Runtime/nDisplay/Source/MPCDI/Private/MPCDIRegion.cpp
+++ b/Engine/Plugins/Runtime/nDisplay/Source/MPCDI/Private/MPCDIRegion.cpp
@@ -2,6 +2,7 @@
 
 #include "MPCDIRegion.h"
 #include "MPCDIBlendTexture.h"
+#include "MPCDIHelpers.h"
 
 THIRD_PARTY_INCLUDES_START
 
@@ -15,11 +16,74 @@ THIRD_PARTY_INCLUDES_START
 #include "mpcdiDistortionMap.h"
 #include "mpcdiGeometryWarpFile.h"
 
+
+#include "IO/mpcdiPfmIO.h"
+#include "mpcdiPNGReadWrite.h"
+
 THIRD_PARTY_INCLUDES_END
 
+#include "Misc/FileHelper.h"
+
+static TAutoConsoleVariable<int32> CVarMPCDIReloadChangedExtFiles(
+	TEXT("nDisplay.reload.pfm"),
+	(int)0, // Disabled by default
+	TEXT("Changed external files reload period for PFM:\n")
+	TEXT("0 : disable\n"),
+	ECVF_RenderThreadSafe
+);
+
+
 namespace MPCDI
 {
-	void FMPCDIRegion::Load(mpcdi::Region* mpcdiRegion, const EMPCDIProfileType ProfileType)
+	void FMPCDIRegion::FExternalFileReloader::Initialize(const FString& FullPathFileName)
+	{
+		bIsEnabled = false;
+		FileName = FullPathFileName;
+		FrameIndex = 0;
+	}
+	void FMPCDIRegion::FExternalFileReloader::SyncDateTime()
+	{
+		DateTime = IFileManager::Get().GetAccessTimeStamp(*FileName);
+		bIsEnabled = true;
+		FrameIndex = 0;
+	}
+	void FMPCDIRegion::FExternalFileReloader::Release()
+	{
+		bIsEnabled = false;
+		FrameIndex = 0;
+	}
+	bool FMPCDIRegion::FExternalFileReloader::IsChanged(bool bForceTest)
+	{
+		if (bIsEnabled)
+		{
+			if (!bForceTest)
+			{
+				// Detect auto reload by frames
+				const int FramePeriod = CVarMPCDIReloadChangedExtFiles.GetValueOnAnyThread();
+				if (FramePeriod > 0)
+				{
+					if (FrameIndex++ > FramePeriod)
+					{
+						FrameIndex = 0;
+						bForceTest = true;
+					}
+				}
+			}
+			if (bForceTest)
+			{
+				// Check for file data is modified				
+				FDateTime CurrentDateTime = IFileManager::Get().GetAccessTimeStamp(*FileName);
+				if (CurrentDateTime != DateTime)
+				{
+					//! Add log. Reload file *FullPathFileName
+					return true;
+				}
+			}
+		}
+		return false;
+	}
+
+	bool FMPCDIRegion::Load(mpcdi::Region* mpcdiRegion, const IMPCDI::EMPCDIProfileType ProfileType)
 	{
 		ID = FString(mpcdiRegion->GetId().c_str());
 
@@ -46,7 +110,7 @@ namespace MPCDI
 			BetaMap.LoadBlendMap(betaMap,1);
 		}
 
-		if (ProfileType != EMPCDIProfileType::mpcdi_SL)
+		if (ProfileType != IMPCDI::EMPCDIProfileType::mpcdi_SL)
 		{
 			mpcdi::GeometryWarpFile *geometryFile = mpcdiRegion->GetFileSet()->GetGeometryWarpFile();
 			if (geometryFile)
@@ -54,5 +118,141 @@ namespace MPCDI
 				WarpMap.LoadWarpMap(geometryFile, ProfileType);
 			}
 		}
+
+		return true;
+	}
+
+	bool FMPCDIRegion::LoadExtGeometry(const TArray<FVector>& PFMPoints, int DimW, int DimH, const IMPCDI::EMPCDIProfileType ProfileType, const float WorldScale, bool bIsUnrealGameSpace)
+	{
+		return WarpMap.LoadCustom3DWarpMap(PFMPoints, DimW, DimH, ProfileType, WorldScale, bIsUnrealGameSpace);
+	}
+
+	bool FMPCDIRegion::LoadExtPFMFile(const FString& LocalPFMFileName, const IMPCDI::EMPCDIProfileType ProfileType, const float PFMScale, bool bIsUnrealGameSpace)
+	{
+		FString PFMFileName = DisplayClusterHelpers::config::GetFullPath(LocalPFMFileName);
+		if (!FPaths::FileExists(PFMFileName))
+		{
+			//! Handle error: pfm not found
+			return false;
+		}
+
+		// Initialize runtime reloader data:
+		ExtPFMFile.File.Initialize(PFMFileName);
+		ExtPFMFile.ProfileType = ProfileType;
+		ExtPFMFile.PFMScale = PFMScale;
+		ExtPFMFile.bIsUnrealGameSpace = bIsUnrealGameSpace;
+
+		return LoadExtPFMFile(ExtPFMFile);
+	}
+
+	bool FMPCDIRegion::LoadExtPFMFile(FExtPFMFile& PFMFile)
+	{
+		bool bResult = false;
+		std::string FileName = TCHAR_TO_ANSI(*PFMFile.File());
+		mpcdi::PFM* PFMData;
+		mpcdi::MPCDI_Error res = mpcdi::PfmIO::Read(FileName, PFMData);
+		if (mpcdi::MPCDI_SUCCESS == res && PFMData)
+		{
+			if (WarpMap.LoadPFMFile(*PFMData, PFMFile.ProfileType, PFMFile.PFMScale, PFMFile.bIsUnrealGameSpace))
+			{				
+				PFMFile.File.SyncDateTime();
+				bResult = true;
+			}
+			delete PFMData;
+		}
+		else
+		{
+			//@todo: Handle error
+		}
+		return bResult;
+	}
+
+	bool FMPCDIRegion::LoadDataMap(const FString& LocalPNGFileName, float GammaValue, EDataMapType DataType)
+	{
+		FString PNGFileName = DisplayClusterHelpers::config::GetFullPath(LocalPNGFileName);
+		if (!FPaths::FileExists(PNGFileName))
+		{
+			//! Handle error: blend map file not found
+			return false;
+		}
+
+		FExtBlendMapFile* ExtBlendMapFile = nullptr;
+		switch (DataType)
+		{
+		case EDataMapType::Alpha:
+			ExtBlendMapFile = &ExtAlphaMap;
+			break;
+		case EDataMapType::Beta:
+			ExtBlendMapFile = &ExtBetaMap;
+			GammaValue = 1;
+			break;
+		}
+
+		if (ExtBlendMapFile != nullptr)
+		{
+			ExtBlendMapFile->File.Initialize(PNGFileName);
+			ExtBlendMapFile->GammaValue = GammaValue;
+			ExtBlendMapFile->DataType = DataType;
+			return LoadDataMap(*ExtBlendMapFile);
+		}
+
+		//! handle error
+		return false;
+	}
+
+	bool FMPCDIRegion::LoadDataMap(FExtBlendMapFile& BlendMapFile)
+	{
+		bool bResult = false;
+		std::string FileName = TCHAR_TO_ANSI(*BlendMapFile.File());
+
+		mpcdi::DataMap* PngData;
+		mpcdi::MPCDI_Error res = mpcdi::PNGReadWrite::Read(FileName, PngData);
+		if (mpcdi::MPCDI_SUCCESS == res && PngData)
+		{
+			switch (BlendMapFile.DataType)
+			{
+			case EDataMapType::Alpha:
+				AlphaMap.LoadBlendMap(PngData, BlendMapFile.GammaValue);
+				bResult = true;
+				break;
+			case EDataMapType::Beta:
+				BetaMap.LoadBlendMap(PngData, BlendMapFile.GammaValue);
+				bResult = true;
+				break;
+			}
+			delete PngData;
+		}
+		else
+		{
+			//@todo: Handle error
+		}
+
+		if (bResult)
+		{
+			BlendMapFile.File.SyncDateTime();
+		}
+
+		return bResult;
+	}
+
+	void FMPCDIRegion::ReloadExternalFiles_RenderThread(bool bForceReload)
+	{
+		check(IsInRenderingThread());
+
+		// Refresh all modified external files
+		//Warp mesh texture
+		if (ExtPFMFile.File.IsChanged(bForceReload))
+		{
+			LoadExtPFMFile(ExtPFMFile);
+		}
+		// Blend maps
+		if (ExtAlphaMap.File.IsChanged(bForceReload))
+		{
+			LoadDataMap(ExtAlphaMap);
+		}
+		if (ExtBetaMap.File.IsChanged(bForceReload))
+		{
+			LoadDataMap(ExtBetaMap);
+		}
 	}
 }
diff --git a/Engine/Plugins/Runtime/nDisplay/Source/MPCDI/Private/MPCDIRegion.h b/Engine/Plugins/Runtime/nDisplay/Source/MPCDI/Private/MPCDIRegion.h
--- a/Engine/Plugins/Runtime/nDisplay/Source/MPCDI/Private/MPCDIRegion.h
+++ b/Engine/Plugins/Runtime/nDisplay/Source/MPCDI/Private/MPCDIRegion.h
@@ -2,7 +2,6 @@
 
 #pragma once
 
-#include "MPCDITypes.h"
 #include "MPCDIBlendTexture.h"
 #include "MPCDIWarpTexture.h"
 
@@ -16,6 +15,26 @@ namespace mpcdi
 
 namespace MPCDI
 {
+	struct FFrustumData // 3D Simulation and Shader Lamps Only
+	{
+		double Yaw;       // Viewer Rotates head upwards. First Rotation.
+		double Pitch;     // Viewer Rotates head clockwise. Second Rotation.
+		double Roll;      // Viewer Rotates head rightwards. Third Rotation.
+
+		double LeftAngle; // Field of View. Typically Negative (Degrees) 
+		double RightAngle;// Field of View. Typically Positive (Degrees)
+		double UpAngle;   // Field of View. Typically Positive (Degrees)
+		double DownAngle; // Field of View. Typically Negative (Degrees)
+	};
+	struct FCoordinateFrameData // Shader Lamps Only
+	{
+		double Posx,   Posy,   Posz;
+		double Yawx,   Yawy,   Yawz;
+		double Pitchx, Pitchy, Pitchz;
+		double Rollx,  Rolly,  Rollz;
+	};
+
+
 	class FDebugMeshExporter;
 
 	struct FMPCDIRegion
@@ -30,6 +49,11 @@ namespace MPCDI
 		uint32_t ResX;
 		uint32_t ResY;
 
+		// SL data support
+		//@ todo: add SL data support
+		FFrustumData         Frustum;
+		FCoordinateFrameData CoordinateFrame;
+
 		// Position/Scale of the viewport in the output framebuffer - This is temp until we implement the divorce
 		// Between input and output framebuffers
 		//float OutX, OutY;
@@ -39,7 +63,7 @@ namespace MPCDI
 		FMPCDIBlendTexture BetaMap;
 		FMPCDIWarpTexture  WarpMap;
 
-		bool isRuntimeData : 1;
+		bool isRuntimeData;
 
 		FMPCDIRegion()
 			: isRuntimeData(false)
@@ -56,7 +80,20 @@ namespace MPCDI
 			, isRuntimeData(true)
 		{ }
 
-		void Load(mpcdi::Region* InMPCIDRegionData, const EMPCDIProfileType ProfileType);
+		bool Load(mpcdi::Region* InMPCIDRegionData, const IMPCDI::EMPCDIProfileType ProfileType);
+
+		void ReloadExternalFiles_RenderThread(bool bForceReload);
+
+		//Support ext loads:
+		bool LoadExtGeometry(const TArray<FVector>& PFMPoints, int DimW, int DimH, const IMPCDI::EMPCDIProfileType ProfileType, const float WorldScale, bool bIsUnrealGameSpace);
+
+		bool LoadExtPFMFile(const FString& LocalPFMFileName, const IMPCDI::EMPCDIProfileType ProfileType, const float PFMScale, bool bIsUnrealGameSpace);
+
+		bool LoadExtAlphaMap(const FString& LocalPNGFileName, float GammaValue)
+			{ return LoadDataMap(LocalPNGFileName, GammaValue, EDataMapType::Alpha);  }
+		bool LoadExtBetaMap(const FString& LocalPNGFileName)
+			{ return LoadDataMap(LocalPNGFileName, 1, EDataMapType::Beta); }
+
 
 		inline FMatrix GetRegionMatrix() const
 		{
@@ -67,5 +104,61 @@ namespace MPCDI
 			RegionMatrix.M[3][1] = Y;
 			return RegionMatrix;
 		}
+
+	private:
+		enum class EDataMapType : uint8
+		{
+			Alpha,
+			Beta,
+		};
+		bool LoadDataMap(const FString& LocalPNGFileName, float GammaValue, EDataMapType DataType);
+
+		// Support runtime reload
+	private:
+		struct FExternalFileReloader
+		{
+			void      Initialize(const FString& FullPathFileName);
+			void      SyncDateTime(); // Get current file modification date time
+			void      Release();
+			bool      IsChanged(bool bForceReload);
+
+			const FString& operator()() const
+			{
+				return FileName;
+			}
+
+			FExternalFileReloader()
+				: bIsEnabled(false)				
+				, FrameIndex(0)
+			{}
+		private:
+			FString   FileName;
+			FDateTime DateTime;
+			bool      bIsEnabled; // Enable reload			
+			int       FrameIndex;
+		};
+
+		struct FExtPFMFile {
+			FExternalFileReloader     File;
+			IMPCDI::EMPCDIProfileType ProfileType;
+			float                     PFMScale;
+			bool                      bIsUnrealGameSpace;
+		};
+		struct FExtBlendMapFile {
+			FExternalFileReloader     File;
+			float                     GammaValue;
+			EDataMapType              DataType;
+		};
+
+	private:
+		bool LoadExtPFMFile(FExtPFMFile& PFMFile);
+		bool LoadDataMap(FExtBlendMapFile& BlendMapFile);
+
+	private:
+		
+		FExtPFMFile      ExtPFMFile;
+		FExtBlendMapFile ExtAlphaMap;
+		FExtBlendMapFile ExtBetaMap;
+
 	};
 }
diff --git a/Engine/Plugins/Runtime/nDisplay/Source/MPCDI/Private/MPCDIShader.cpp b/Engine/Plugins/Runtime/nDisplay/Source/MPCDI/Private/MPCDIShader.cpp
--- a/Engine/Plugins/Runtime/nDisplay/Source/MPCDI/Private/MPCDIShader.cpp
+++ b/Engine/Plugins/Runtime/nDisplay/Source/MPCDI/Private/MPCDIShader.cpp
@@ -18,7 +18,8 @@
 enum class EVarMPCDIShaderType : uint8
 {
 	Default,
-	Passthrought,
+	DisableBlend,
+	Passthrough,
 	ShowWarpTexture,
 	Disable,
 };
@@ -26,7 +27,12 @@ enum class EVarMPCDIShaderType : uint8
 static TAutoConsoleVariable<int32> CVarMPCDIShaderType(
 	TEXT("nDisplay.render.mpcdi.shader"),
 	(int)EVarMPCDIShaderType::Default,
-	TEXT("Select shader for mpcdi:\n0 = default warp shader\n1 = PassThrought shader\n2 = ShowWarpTexture shader\n3 = disable warp shader"),
+	TEXT("Select shader for mpcdi:\n")
+	TEXT("0 : default warp shader\n")
+	TEXT("1 : Default, disabled blend\n")
+	TEXT("2 : Passthrough shader\n")
+	TEXT("3 : ShowWarpTexture shader\n")
+	TEXT("4 : disable warp shader"),
 	ECVF_RenderThreadSafe
 );
 
@@ -36,10 +42,13 @@ static TAutoConsoleVariable<int32> CVarMPCDIShaderType(
 
 // Implement shaders inside UE4
 IMPLEMENT_SHADER_TYPE(, FMpcdiWarpVS, MPCDIShaderFileName, TEXT("MainVS"), SF_Vertex);
+IMPLEMENT_SHADER_TYPE(, FMPCDIDirectProjectionVS, MPCDIShaderFileName, TEXT("DirectProjectionVS"), SF_Vertex);
 
-IMPLEMENT_SHADER_TYPE(template<>, FMPCDIPassThroughtPS, MPCDIShaderFileName, TEXT("PassThrought_PS"), SF_Pixel);
+IMPLEMENT_SHADER_TYPE(template<>, FMPCDIPassthroughPS,   MPCDIShaderFileName, TEXT("Passthrough_PS"), SF_Pixel);
 IMPLEMENT_SHADER_TYPE(template<>, FMPCDIShowWarpTexture, MPCDIShaderFileName, TEXT("ShowWarpTexture_PS"), SF_Pixel);
+
 IMPLEMENT_SHADER_TYPE(template<>, FMPCDIWarpAndBlendPS, MPCDIShaderFileName, TEXT("WarpAndBlend_PS"), SF_Pixel);
+IMPLEMENT_SHADER_TYPE(template<>, FMPCDIWarpPS,         MPCDIShaderFileName, TEXT("WarpAndBlend_PS"), SF_Pixel);
 
 #if 0
 //Cubemap support
@@ -100,25 +109,29 @@ static EMPCDIShader GetPixelShaderType(FMPCDIData *MPCDIData)
 	{
 		switch (MPCDIData->GetProfileType())
 		{
-		case EMPCDIProfileType::mpcdi_A3D:
-		case EMPCDIProfileType::mpcdi_SL:
+		case IMPCDI::EMPCDIProfileType::mpcdi_SL:
+			return EMPCDIShader::Invalid;
+
+		case IMPCDI::EMPCDIProfileType::mpcdi_A3D:
+#if 0
+			//@todo Unsupported now
 			if (MPCDIData->IsCubeMapEnabled())
 			{
-#if 0
 				// Cubemap support
 				return EMPCDIShader::WarpAndBlendCubeMap;
+		}
 #endif
-				return EMPCDIShader::Invalid;
-			}
 
-		case EMPCDIProfileType::mpcdi_2D:
-		case EMPCDIProfileType::mpcdi_3D:
+		case IMPCDI::EMPCDIProfileType::mpcdi_2D:
+		case IMPCDI::EMPCDIProfileType::mpcdi_3D:
 			return EMPCDIShader::WarpAndBlend;
 
-		case EMPCDIProfileType::Invalid:
-			return EMPCDIShader::PassThrought;
-		};
-	}
+		case IMPCDI::EMPCDIProfileType::Invalid:
+			return EMPCDIShader::Passthrough;
+
+		default: EMPCDIShader::Invalid;
+	};
+}
 
 	return EMPCDIShader::Invalid;
 }
@@ -132,7 +145,7 @@ enum class EColorOPMode :uint8
 	AddInvAlpha,
 };
 
-template<uint32 PixelShaderType>
+template<EMPCDIShader PixelShaderType>
 static bool CompleteWarpTempl(FRHICommandListImmediate& RHICmdList, IMPCDI::FTextureWarpData& TextureWarpData, const IMPCDI::FShaderInputData &ShaderInputData, MPCDI::FMPCDIRegion& MPCDIRegion, EColorOPMode ColorOp)
 {
 	FMpcdiWarpDrawRectangleParameters VSData;
@@ -178,10 +191,7 @@ static bool CompleteWarpTempl(FRHICommandListImmediate& RHICmdList, IMPCDI::FTex
 	VertexShader->SetParameters(RHICmdList, VertexShader->GetVertexShader(), VSData);
 	PixelShader->SetParameters(RHICmdList, PixelShader->GetPixelShader(), TextureWarpData.SrcTexture, ShaderInputData, MPCDIRegion);
 
-	{
-		// Render quad
-		FPixelShaderUtils::DrawFullscreenQuad(RHICmdList, 1);
-	}
+	FPixelShaderUtils::DrawFullscreenQuad(RHICmdList, 1); // Render quad
 
 	return true;
 }
@@ -190,6 +200,18 @@ bool FMPCDIShader::ApplyWarpBlend(FRHICommandListImmediate& RHICmdList, IMPCDI::
 {
 	check(IsInRenderingThread());
 
+	MPCDI::FMPCDIRegion* MPCDIRegion = nullptr;
+	if (MPCDIData && MPCDIData->IsValid())
+	{
+		MPCDIRegion = MPCDIData->GetRegion(ShaderInputData.RegionLocator);
+	}
+	if (MPCDIRegion == nullptr)
+	{
+		//Handle error
+		return false;
+	}
+
+
 	// Map ProfileType to shader type to use
 	EMPCDIShader PixelShaderType = GetPixelShaderType(MPCDIData);
 	if (EMPCDIShader::Invalid == PixelShaderType)
@@ -198,27 +220,21 @@ bool FMPCDIShader::ApplyWarpBlend(FRHICommandListImmediate& RHICmdList, IMPCDI::
 		return false;
 	}
 
+	bool bIsBlendDisabled = !MPCDIRegion->AlphaMap.IsInitialized();
+
 	FMatrix StereoMatrix;
 	GetStereoMatrix(TextureWarpData, StereoMatrix);
 
-	MPCDI::FMPCDIRegion* MPCDIRegion = nullptr;
-	if (MPCDIData && MPCDIData->IsValid())
-	{
-		MPCDIRegion = MPCDIData->GetRegion(ShaderInputData.RegionLocator);
-	}
-
-	if (MPCDIRegion == nullptr)
-	{
-		//Handle error
-		return false;
-	}
 
 
 	const EVarMPCDIShaderType ShaderType = (EVarMPCDIShaderType)CVarMPCDIShaderType.GetValueOnAnyThread();
 	switch (ShaderType)
 	{
-	case EVarMPCDIShaderType::Passthrought:
-		PixelShaderType = EMPCDIShader::PassThrought;
+	case EVarMPCDIShaderType::DisableBlend:
+		bIsBlendDisabled = true; // Disable blend color ops
+		break;
+	case EVarMPCDIShaderType::Passthrough:
+		PixelShaderType = EMPCDIShader::Passthrough;
 		break;
 
 	case EVarMPCDIShaderType::ShowWarpTexture:
@@ -227,9 +243,8 @@ bool FMPCDIShader::ApplyWarpBlend(FRHICommandListImmediate& RHICmdList, IMPCDI::
 
 	case EVarMPCDIShaderType::Disable:
 		return false;
-
 	default:
-		return false;
+		break;
 	};
 	
 	{
@@ -240,21 +255,38 @@ bool FMPCDIShader::ApplyWarpBlend(FRHICommandListImmediate& RHICmdList, IMPCDI::
 		{
 			switch (PixelShaderType)
 			{
-			case EMPCDIShader::PassThrought:
+			case EMPCDIShader::Passthrough:
 				ShaderInputData.UVMatrix = StereoMatrix;
-				bIsRenderSuccess = CompleteWarpTempl<(int)EMPCDIShader::PassThrought>(RHICmdList, TextureWarpData, ShaderInputData, *MPCDIRegion, EColorOPMode::Default);
+				bIsRenderSuccess = CompleteWarpTempl<EMPCDIShader::Passthrough>(RHICmdList, TextureWarpData, ShaderInputData, *MPCDIRegion, EColorOPMode::Default);
 				break;
 			case EMPCDIShader::ShowWarpTexture:
 				ShaderInputData.UVMatrix = ShaderInputData.Frustum.UVMatrix*MPCDIData->GetTextureMatrix()*MPCDIRegion->GetRegionMatrix()*StereoMatrix;
-				bIsRenderSuccess = CompleteWarpTempl<(int)EMPCDIShader::WarpAndBlend>(RHICmdList, TextureWarpData, ShaderInputData, *MPCDIRegion, EColorOPMode::Default);
-				bIsRenderSuccess = CompleteWarpTempl<(int)EMPCDIShader::ShowWarpTexture>(RHICmdList, TextureWarpData, ShaderInputData, *MPCDIRegion, EColorOPMode::AddAlpha);
+				{// WarpAndBlend
+					if (bIsBlendDisabled)
+					{
+						bIsRenderSuccess = CompleteWarpTempl<EMPCDIShader::Warp>(RHICmdList, TextureWarpData, ShaderInputData, *MPCDIRegion, EColorOPMode::Default);
+					}
+					else
+					{
+						bIsRenderSuccess = CompleteWarpTempl<EMPCDIShader::WarpAndBlend>(RHICmdList, TextureWarpData, ShaderInputData, *MPCDIRegion, EColorOPMode::Default);
+					}
+				}
+				// Overlay
+				bIsRenderSuccess &= CompleteWarpTempl<EMPCDIShader::ShowWarpTexture>(RHICmdList, TextureWarpData, ShaderInputData, *MPCDIRegion, EColorOPMode::AddAlpha);
 				break;
 
 			case EMPCDIShader::WarpAndBlend:
 				ShaderInputData.UVMatrix = ShaderInputData.Frustum.UVMatrix*MPCDIData->GetTextureMatrix()*MPCDIRegion->GetRegionMatrix()*StereoMatrix;
-				bIsRenderSuccess = CompleteWarpTempl<(int)EMPCDIShader::WarpAndBlend>(RHICmdList, TextureWarpData, ShaderInputData, *MPCDIRegion, EColorOPMode::Default);
+				if (bIsBlendDisabled)
+				{
+					bIsRenderSuccess = CompleteWarpTempl<EMPCDIShader::Warp>(RHICmdList, TextureWarpData, ShaderInputData, *MPCDIRegion, EColorOPMode::Default);
+				}
+				else
+				{
+					bIsRenderSuccess = CompleteWarpTempl<EMPCDIShader::WarpAndBlend>(RHICmdList, TextureWarpData, ShaderInputData, *MPCDIRegion, EColorOPMode::Default);
+				}
+				
 				break;
-
 #if 0
 				//Cubemap support
 			case EMPCDIShader::WarpAndBlendCubeMap:
diff --git a/Engine/Plugins/Runtime/nDisplay/Source/MPCDI/Private/MPCDIShader.h b/Engine/Plugins/Runtime/nDisplay/Source/MPCDI/Private/MPCDIShader.h
--- a/Engine/Plugins/Runtime/nDisplay/Source/MPCDI/Private/MPCDIShader.h
+++ b/Engine/Plugins/Runtime/nDisplay/Source/MPCDI/Private/MPCDIShader.h
@@ -17,6 +17,15 @@
 #include "MPCDIRegion.h"
 
 
+enum class EMPCDIShader : uint8
+{
+	Passthrough,
+	ShowWarpTexture,
+	WarpAndBlend,
+	Warp,
+	Invalid,
+};
+
 class FMPCDIData;
 
 
@@ -34,13 +43,38 @@ struct FMpcdiWarpDrawRectangleParameters
 	FVector4 DrawRectangleUVScaleBias;
 };
 
+class FMPCDIDirectProjectionVS : public FGlobalShader
+{
+	DECLARE_SHADER_TYPE(FMPCDIDirectProjectionVS, Global);
+
+	static bool ShouldCompilePermutation(const FGlobalShaderPermutationParameters& Parameters)
+	{
+		return true;
+	}
+
+	/** Default constructor. */
+	FMPCDIDirectProjectionVS() 
+	{
+	}
+
+
+public:
+	/** Initialization constructor. */
+	FMPCDIDirectProjectionVS(const ShaderMetaType::CompiledShaderInitializerType& Initializer)
+		: FGlobalShader(Initializer)
+	{
+	}
+};
+
+
 class FMpcdiWarpVS : public FGlobalShader
 {
 	DECLARE_SHADER_TYPE(FMpcdiWarpVS, Global);
 
 public:
 	FMpcdiWarpVS()
-	{ }
+	{ 
+	}
 
 	FMpcdiWarpVS(const ShaderMetaType::CompiledShaderInitializerType& Initializer)
 		: FGlobalShader(Initializer)
@@ -86,38 +120,24 @@ private:
 	FShaderParameter DrawRectangleInvTargetSizeAndTextureSizeParameter;
 };
 
-enum class EMPCDIShader : uint8
-{
-	PassThrought,
-	ShowWarpTexture,
-	WarpAndBlend,
-#if 0
-	// Cubemap warpblend support
-	WarpAndBlendCubeMap,
-#endif
-#if 0
-	// GPU perfect frustum
-	BuildProjectedAABB,
-#endif
-	Invalid,
-};
 
 
-template<uint32 ShaderType>
+template<EMPCDIShader ShaderType>
 class TMpcdiWarpBasePS : public FGlobalShader
 {
 	DECLARE_SHADER_TYPE(TMpcdiWarpBasePS, Global);
 
 public:
 	TMpcdiWarpBasePS()
-	{ }
+	{ 
+	}
 
 	TMpcdiWarpBasePS(const ShaderMetaType::CompiledShaderInitializerType& Initializer)
 		: FGlobalShader(Initializer)
 	{
 #if 0
 		//Cubemap supoport
-		if (ShaderType != (int)EMPCDIShader::WarpAndBlendCubeMap)
+		if (ShaderType != EMPCDIShader::WarpAndBlendCubeMap)
 #endif
 		{
 			PostprocessInputParameter0.Bind(Initializer.ParameterMap, TEXT("PostprocessInput0"));
@@ -126,7 +146,7 @@ public:
 
 		TextureProjectionMatrixParameter.Bind(Initializer.ParameterMap, TEXT("TextureProjectionMatrix"));
 
-		if (ShaderType != (int)EMPCDIShader::PassThrought)
+		if (ShaderType != EMPCDIShader::Passthrough)
 		{
 			AlphaEmbeddedGammaParameter.Bind(Initializer.ParameterMap, TEXT("AlphaEmbeddedGamma"));
 
@@ -140,7 +160,7 @@ public:
 
 
 #if 0
-			if (ShaderType == (int)EMPCDIShader::WarpAndBlendCubeMap)
+			if (ShaderType == EMPCDIShader::WarpAndBlendCubeMap)
 			{
 				SceneCubeMapParameter.Bind(Initializer.ParameterMap, TEXT("SceneCubemap"));
 				SceneCubeMapParameterSampler.Bind(Initializer.ParameterMap, TEXT("SceneCubemapSampler"));
@@ -162,6 +182,15 @@ public:
 	static void ModifyCompilationEnvironment(const FGlobalShaderPermutationParameters& Parameters, FShaderCompilerEnvironment& OutEnvironment)
 	{
 		FGlobalShader::ModifyCompilationEnvironment(Parameters, OutEnvironment);
+
+		switch (ShaderType)
+		{
+		case EMPCDIShader::WarpAndBlend:
+			OutEnvironment.SetDefine(TEXT("RENDER_CASE_APPLYBLENDING"), 1);
+			break;
+		default:
+			break;
+		}
 	}
 
 public:
@@ -177,7 +206,7 @@ public:
 		SetShaderValue(RHICmdList, ShaderRHI, TextureProjectionMatrixParameter, ShaderInputData.UVMatrix);
 
 #if 0
-		if (ShaderType == (int)EMPCDIShader::WarpAndBlendCubeMap)
+		if (ShaderType == EMPCDIShader::WarpAndBlendCubeMap)
 		{
 			RHICmdList.SetShaderSampler(ShaderRHI, SceneCubeMapParameterSampler.GetBaseIndex(), TStaticSamplerState<SF_Bilinear>::GetRHI());
 			SetTextureParameter(RHICmdList, ShaderRHI, SceneCubeMapParameter, ShaderInputData.SceneCubeMap);
@@ -194,19 +223,18 @@ public:
 	{
 #if 0
 		//Cubemap support
-		if (ShaderType != (int)EMPCDIShader::WarpAndBlendCubeMap)
+		if (ShaderType != EMPCDIShader::WarpAndBlendCubeMap)
 #endif
 		{
 			SetTextureParameter(RHICmdList, ShaderRHI, PostprocessInputParameter0, SourceTexture);
-			RHICmdList.SetShaderSampler(ShaderRHI, PostprocessInputParameterSampler0.GetBaseIndex(), TStaticSamplerState<SF_Bilinear, AM_Clamp, AM_Clamp, AM_Clamp>::GetRHI());
-			//SetTextureParameter(RHICmdList, ShaderRHI, PostprocessInputParameter0, PostprocessInputParameterSampler0, TStaticSamplerState<SF_Bilinear, AM_Clamp, AM_Clamp, AM_Clamp>::GetRHI(), SourceTexture);
+			RHICmdList.SetShaderSampler(ShaderRHI, PostprocessInputParameterSampler0.GetBaseIndex(), TStaticSamplerState<SF_Trilinear, AM_Clamp, AM_Clamp, AM_Clamp>::GetRHI());
 		}
 	}
 
 	template<typename TShaderRHIParamRef>
 	void SetParameters(FRHICommandListImmediate& RHICmdList, const TShaderRHIParamRef ShaderRHI, MPCDI::FMPCDIRegion& MPCDIRegionData)
 	{
-		if (ShaderType != (int)EMPCDIShader::PassThrought)
+		if (ShaderType != EMPCDIShader::Passthrough)
 		{
 			SetShaderValue(RHICmdList, ShaderRHI, AlphaEmbeddedGammaParameter, MPCDIRegionData.AlphaMap.GetEmbeddedGamma());
 
@@ -277,8 +305,7 @@ private:
 	FShaderParameter         TextureProjectionMatrixParameter;
 	FShaderParameter         AlphaEmbeddedGammaParameter;
 #if 0
-	// Cubemap support
-
+	//@todo  Cubemap support not implemented
 	FShaderResourceParameter SceneCubeMapParameter;
 	FShaderResourceParameter SceneCubeMapParameterSampler;
 #endif
@@ -288,15 +315,16 @@ private:
 	FShaderParameter         VignetteEVParameter;
 };
 
-typedef TMpcdiWarpBasePS<(int)EMPCDIShader::PassThrought>         FMPCDIPassThroughtPS;
-typedef TMpcdiWarpBasePS<(int)EMPCDIShader::ShowWarpTexture>      FMPCDIShowWarpTexture;
-typedef TMpcdiWarpBasePS<(int)EMPCDIShader::WarpAndBlend>         FMPCDIWarpAndBlendPS;
+typedef TMpcdiWarpBasePS<EMPCDIShader::Passthrough>          FMPCDIPassthroughPS;
+typedef TMpcdiWarpBasePS<EMPCDIShader::ShowWarpTexture>      FMPCDIShowWarpTexture;
+typedef TMpcdiWarpBasePS<EMPCDIShader::WarpAndBlend>         FMPCDIWarpAndBlendPS;
+typedef TMpcdiWarpBasePS<EMPCDIShader::Warp>                 FMPCDIWarpPS;
 
 #if 0
 // Cubemap support
-typedef TMpcdiWarpBasePS<(int)EMPCDIShader::WarpAndBlendCubeMap>  FMPCDIWarpAndBlendCubemapPS;
+typedef TMpcdiWarpBasePS<EMPCDIShader::WarpAndBlendCubeMap>  FMPCDIWarpAndBlendCubemapPS;
 #endif
 #if 0
 // Gpu perfect frustum
-typedef TMpcdiWarpBasePS<(int)EMPCDIShader::BuildProjectedAABB>   FMPCDICalcBoundBoxPS;
+typedef TMpcdiWarpBasePS<EMPCDIShader::BuildProjectedAABB>   FMPCDICalcBoundBoxPS;
 #endif
diff --git a/Engine/Plugins/Runtime/nDisplay/Source/MPCDI/Private/MPCDIStrings.h b/Engine/Plugins/Runtime/nDisplay/Source/MPCDI/Private/MPCDIStrings.h
new file mode 100644
--- /dev/null
+++ b/Engine/Plugins/Runtime/nDisplay/Source/MPCDI/Private/MPCDIStrings.h
@@ -0,0 +1,36 @@
+// Copyright 1998-2019 Epic Games, Inc. All Rights Reserved.
+
+#pragma once
+
+#include "CoreMinimal.h"
+
+namespace DisplayClusterStrings
+{
+	namespace cfg
+	{
+		namespace data
+		{
+			namespace mpcdi
+			{
+				static constexpr auto File = TEXT("file");
+				static constexpr auto Buffer = TEXT("buffer");
+				static constexpr auto Region = TEXT("region");
+				static constexpr auto Origin = TEXT("origin");
+
+				//Support external files
+				static constexpr auto MPCDIType = TEXT("mpcdi");
+
+				static constexpr auto FilePFM = TEXT("pfm");
+				static constexpr auto WorldScale = TEXT("scale");
+				static constexpr auto UseUnrealAxis = TEXT("ue4space");
+
+				static constexpr auto FileAlpha = TEXT("alpha");
+				static constexpr auto AlphaGamma = TEXT("alpha_gamma");
+
+				static constexpr auto FileBeta = TEXT("beta");
+
+				static constexpr auto PFMFileDefaultID = TEXT("@PFM");
+			}
+		}
+	}
+};
diff --git a/Engine/Plugins/Runtime/nDisplay/Source/MPCDI/Private/MPCDITexture.cpp b/Engine/Plugins/Runtime/nDisplay/Source/MPCDI/Private/MPCDITexture.cpp
--- a/Engine/Plugins/Runtime/nDisplay/Source/MPCDI/Private/MPCDITexture.cpp
+++ b/Engine/Plugins/Runtime/nDisplay/Source/MPCDI/Private/MPCDITexture.cpp
@@ -18,17 +18,23 @@ namespace
 		TextureData(const void* InData, uint32_t InDataSize)
 			: Data(InData)
 			, DataSize(InDataSize)
-		{ }
+		{ 
+		}
 
 	public:
 		virtual const void* GetResourceBulkData() const
-		{ return Data; }
+		{ 
+			return Data; 
+		}
 
 		virtual uint32 GetResourceBulkDataSize() const
-		{ return DataSize; }
+		{ 
+			return DataSize; 
+		}
 
 		virtual void Discard()
-		{ }
+		{ 
+		}
 
 	private:
 		const void* Data;
@@ -47,14 +53,22 @@ namespace
 //---------------------------------------------
 // FMPCDITexture
 //---------------------------------------------
+void MPCDI::FMPCDITexture::ReleaseTextureData()
+{
+	if (Data != nullptr)
+	{
+		FMemory::Free(Data);
+		Data = nullptr;
+	}
+}
+
 void MPCDI::FMPCDITexture::InitRHI()
 {
 	FTexture2DRHIRef Texture2D = CreateTexture2D(Data, Width, Height, PixelFormat);
 	TextureRHI = Texture2D;
-
+	
 	if (bReleaseData && Data)
 	{
-		FMemory::Free(Data);
-		Data = nullptr;
+		ReleaseTextureData();
 	}
 }
diff --git a/Engine/Plugins/Runtime/nDisplay/Source/MPCDI/Private/MPCDITexture.h b/Engine/Plugins/Runtime/nDisplay/Source/MPCDI/Private/MPCDITexture.h
--- a/Engine/Plugins/Runtime/nDisplay/Source/MPCDI/Private/MPCDITexture.h
+++ b/Engine/Plugins/Runtime/nDisplay/Source/MPCDI/Private/MPCDITexture.h
@@ -3,7 +3,6 @@
 #pragma once
 
 #include "RenderResource.h"
-#include "MPCDITypes.h"
 
 
 namespace MPCDI
@@ -17,7 +16,11 @@ namespace MPCDI
 			, Height(0)
 			, PixelFormat(PF_Unknown)
 			, bReleaseData(false)
-		{ }
+		{ 
+		}
+		virtual ~FMPCDITexture()
+		{
+		}
 
 	public:
 		virtual void InitRHI() override;
@@ -32,16 +35,26 @@ namespace MPCDI
 		}
 
 		inline void* GetData() const
-		{ return Data; }
+		{ 
+			return Data; 
+		}
 
 		inline uint32_t GetWidth() const
-		{ return Width; }
+		{ 
+			return Width; 
+		}
 
 		inline uint32_t GetHeight() const
-		{ return Height; }
+		{ 
+			return Height; 
+		}
 
 		inline EPixelFormat GetPixelFormat() const
-		{ return PixelFormat; }
+		{ 
+			return PixelFormat; 
+		}
+
+		void ReleaseTextureData();
 
 	private:
 		void *Data;
diff --git a/Engine/Plugins/Runtime/nDisplay/Source/MPCDI/Private/MPCDIWarpTexture.cpp b/Engine/Plugins/Runtime/nDisplay/Source/MPCDI/Private/MPCDIWarpTexture.cpp
--- a/Engine/Plugins/Runtime/nDisplay/Source/MPCDI/Private/MPCDIWarpTexture.cpp
+++ b/Engine/Plugins/Runtime/nDisplay/Source/MPCDI/Private/MPCDIWarpTexture.cpp
@@ -14,25 +14,27 @@ THIRD_PARTY_INCLUDES_START
 #include "mpcdiBetaMap.h"
 #include "mpcdiDistortionMap.h"
 #include "mpcdiGeometryWarpFile.h"
+#include "mpcdiPFM.h"
 
 THIRD_PARTY_INCLUDES_END
 
-
 // Select mpcdi frustum calc method
 enum EVarMPCDIFrustumMethod
 {
 	AABB = 0,
 	PerfectCPU,
+	TextureBOX,
 #if 0
 	PerfectGPU, // optimization purpose, project warp texture to one-pixel rendertarget, in min\max colorop pass
 #endif
 };
-
-
 static TAutoConsoleVariable<int32> CVarMPCDIFrustumMethod(
 	TEXT("nDisplay.render.mpcdi.Frustum"),
 	(int)EVarMPCDIFrustumMethod::PerfectCPU,
-	TEXT("Frustum computation method:\n0 = mesh AABB based, lower quality but fast\n1 = mesh vertices based, best quality but slow\n"),
+	TEXT("Frustum computation method:\n")
+	TEXT(" 0: mesh AABB based, lower quality but fast\n")
+	TEXT(" 1: mesh vertices based, best quality but slow\n")
+	TEXT(" 2: texture box, get A*B distributed points from texture, fast, good quality for flat panels\n"),
 	ECVF_RenderThreadSafe
 );
 
@@ -42,29 +44,47 @@ enum EVarMPCDIStereoMode
 	AsymmetricAABB = 0,
 	SymmetricAABB,
 };
-
 static TAutoConsoleVariable<int32> CVarMPCDIStereoMode(
 	TEXT("nDisplay.render.mpcdi.StereoMode"),
 	(int)EVarMPCDIStereoMode::AsymmetricAABB,
-	TEXT("Stereo mode:\n (0 = Asymmetric to AABB center\n(1 = Symmetric to AABB center)\n"),
+	TEXT("Stereo mode:\n")
+	TEXT(" 0: Asymmetric to AABB center\n")
+	TEXT(" 1: Symmetric to AABB center\n"),
 	ECVF_RenderThreadSafe
 );
 
 // Select mpcdi projection mode
 enum EVarMPCDIProjectionMode
 {
-	Dynamic = 0,
-	StaticAxisAligned,
+	StaticSurfaceNormal = 0,
+	DynamicAABBCenter,
+	DynamicAxisAligned,
 };
-
 static TAutoConsoleVariable<int32> CVarMPCDIProjectionMode(
 	TEXT("nDisplay.render.mpcdi.Projection"),
-	(int)EVarMPCDIProjectionMode::Dynamic,
-	TEXT("Projection method:\n0 = Dynamic, to view target\n1 = Static, aligned to mpcdi origin space axis\n"),
+	(int)EVarMPCDIProjectionMode::StaticSurfaceNormal,
+	TEXT("Projection method:\n")
+	TEXT(" 0: Static, aligned to average region surface normal\n")
+	TEXT(" 1: Dynamic, to view target center\n")
+	TEXT(" 2: Dynamic, aligned to cave axis, eye is zero point\n"),
 	ECVF_RenderThreadSafe
 );
 
 
+// Setup frustum projection cache
+static TAutoConsoleVariable<int32> CVarMPCDIFrustumCacheDepth(
+	TEXT("nDisplay.render.mpcdi.cache_depth"),
+	0, // By default cache is disabled (to better performance for EVarMPCDIFrustumMethod::PerfectCPU set value to 8)
+	TEXT("Frustum calculated values cache depth\n"),
+	ECVF_RenderThreadSafe
+);
+static TAutoConsoleVariable<float> CVarMPCDIFrustumCachePrecision(
+	TEXT("nDisplay.render.mpcdi.cache_precision"),
+	0.1f, // 1mm
+	TEXT("Frustum cache values comparison precision (float, unit is sm)\n"),
+	ECVF_RenderThreadSafe
+);
+
 namespace // Helpers
 {
 	FMatrix GetProjectionMatrixAssymetric(float l, float r, float t, float b, float n, float f)
@@ -112,105 +132,276 @@ namespace // Helpers
 
 namespace MPCDI
 {
-	void FMPCDIWarpTexture::CalcFrustum_fullCPU(const IMPCDI::FFrustum& OutFrustum, const FMatrix& world2local, float& top, float& bottom, float& left, float& right) const
+
+	inline void UpdateFrustumFromVertex(const FVector4& PFMVertice, const FMatrix& World2Local, float& Top, float& Bottom, float& Left, float& Right)
 	{
-		int Count = GetWidth()*GetHeight();
-		const FVector4* v = (FVector4*)GetData();
+		if (PFMVertice.W > 0) {
+			FVector4 PrjectedVertice = World2Local.TransformFVector4(PFMVertice);
 
-		// Create a camera space frustum
-		for (int i = 0; i < Count; ++i)
-		{
-			if (v[i].W > 0) {
-				FVector4 prjV = world2local.TransformFVector4(v[i]);
-				float x = 1.0f / prjV.X;
-				prjV.Y *= x;
-				prjV.Z *= x;
-
-				if (prjV.Z > top) { top = prjV.Z; }
-				if (prjV.Z < bottom) { bottom = prjV.Z; }
-				if (prjV.Y > right) { right = prjV.Y; }
-				if (prjV.Y < left) { left = prjV.Y; }
+			float Scale = 1.0f / PrjectedVertice.X;
+			PrjectedVertice.Y *= Scale;
+			PrjectedVertice.Z *= Scale;
+
+			if (PrjectedVertice.Z > Top)
+			{
+				Top = PrjectedVertice.Z;
+			}
+			if (PrjectedVertice.Z < Bottom)
+			{
+				Bottom = PrjectedVertice.Z;
+			}
+			if (PrjectedVertice.Y > Right)
+			{
+				Right = PrjectedVertice.Y;
+			}
+			if (PrjectedVertice.Y < Left)
+			{
+				Left = PrjectedVertice.Y;
 			}
 		}
 	}
-
-	void FMPCDIWarpTexture::CalcFrustum_simpleAABB(const IMPCDI::FFrustum& OutFrustum, const FMatrix& world2local, float& top, float& bottom, float& left, float& right) const
+	
+	struct FValidPFMPoint
 	{
-		FVector4 maxExtent = GetAABB().Max * OutFrustum.WorldScale;
-		FVector4 minExtent = GetAABB().Min * OutFrustum.WorldScale;
+		int X;
+		int Y;
 
+		FValidPFMPoint(const FVector4* InData, int InW, int InH)
+			: Data(InData)
+			, W(InW)
+			, H(InH)
+		{}
 
-		FVector4 v[8];
-		v[0] = FVector4(maxExtent.X, maxExtent.Y, maxExtent.Z);
-		v[1] = FVector4(maxExtent.X, maxExtent.Y, minExtent.Z);
-		v[2] = FVector4(minExtent.X, maxExtent.Y, minExtent.Z);
-		v[3] = FVector4(minExtent.X, maxExtent.Y, maxExtent.Z);
-		v[4] = FVector4(maxExtent.X, minExtent.Y, maxExtent.Z);
-		v[5] = FVector4(maxExtent.X, minExtent.Y, minExtent.Z);
-		v[6] = FVector4(minExtent.X, minExtent.Y, minExtent.Z);
-		v[7] = FVector4(minExtent.X, minExtent.Y, maxExtent.Z);
+		inline int GetSavedPointIndex()
+		{
+			return GetPointIndex(X, Y);
+		}
 
+		inline int GetPointIndex(int InX, int InY)
+		{
+			return InX + InY * W;
+		}
 
-		// Create a camera space frustum
-		FVector4 local[8];
-		for (int i = 0; i < 8; ++i)
+		inline bool IsValidPoint(int InX, int InY)
 		{
-			local[i] = world2local.TransformFVector4(v[i]);
-			float x = local[i].X;
-			local[i].X /= x;
-			local[i].Y /= x;
-			local[i].Z /= x;
-			//! OutFrustum.Bounds[i] = v[i];
+			return Data[GetPointIndex(InX, InY)].W > 0;
 		}
 
-		for (int i = 0; i < 8; ++i)
+		inline bool FindValidPoint(int InX, int InY)
 		{
-			if (local[i].Z > top) { top = local[i].Z; }
-			if (local[i].Z < bottom) { bottom = local[i].Z; }
-			if (local[i].Y > right) { right = local[i].Y; }
-			if (local[i].Y < left) { left = local[i].Y; }
+			X0 = InX;
+			Y0 = InY;
+
+			for (int Range = 1; Range < W; Range++)
+			{
+				if (FindValidPointInRange(Range))
+				{
+					return true;
+				}
+			}
+			return false;
 		}
 
+	private:
+		const FVector4* Data;
+		int X0;
+		int Y0;
+		int W;
+		int H;
 
+	private:
+		inline bool FindValidPointInRange(int Range)
+		{
+			for (int i = -Range; i < Range; i++)
+			{
+				// Top or bottom rows
+				if (IsValid(X0 + i, Y0 - Range) || IsValid(X0 + i, Y0 + Range))
+				{
+					return true;
+				}
+				// Left or Right columns
+				if (IsValid(X0 - Range, Y0 + i) || IsValid(X0 + Range, Y0 + i))
+				{
+					return true;
+				}
+			}
+			return false;
+		}
+
+		inline bool IsValid(int newX, int newY)
+		{
+			if (newX < 0 || newY < 0 || newX >= W || newY >= H)
+			{
+				return false; // Out of texture
+			}
+
+			if (Data[GetPointIndex(newX, newY)].W > 0)
+			{
+				// Store valid result
+				X = newX;
+				Y = newY;
+				return true;
+			}
+			return false;
+		}
+
+	};
+	
+	void FMPCDIWarpTexture::CalcFrustum_TextureBOX(int DivX, int DivY, const IMPCDI::FFrustum& OutFrustum, const FMatrix& World2Local, float& Top, float& Bottom, float& Left, float& Right) const
+	{
+		const FVector4* v = (FVector4*)GetData();
+
+		if (TextureBoxCache.Num() < 1)
+		{
+			TextureBoxCache.Reserve(DivX*DivY);
+
+			FValidPFMPoint PFMPoints(v, GetWidth(), GetHeight());
+			//Generate valid points for texturebox method:
+			for (int low_y = 0; low_y < DivY; low_y++)
+			{
+				int y = (GetHeight() - 1)*(float(low_y) / (DivY - 1));
+
+				for (int low_x = 0; low_x < DivX; low_x++)
+				{
+					int x = (GetWidth() - 1)*(float(low_x) / (DivX - 1));
+
+					if (PFMPoints.IsValidPoint(x, y))
+					{
+						//Just use direct point
+						TextureBoxCache.Add(PFMPoints.GetPointIndex(x, y));
+					}
+					else
+					{
+						//Search for nearset valid point
+						if (PFMPoints.FindValidPoint(x, y))
+						{
+							TextureBoxCache.Add(PFMPoints.GetSavedPointIndex());
+						}
+					}
+				}
+			}
+		}
+
+		// Search a camera space frustum
+		for(const auto It: TextureBoxCache)
+		{
+			UpdateFrustumFromVertex(v[It], World2Local, Top, Bottom, Left, Right);
+		}
+	}
+
+	void FMPCDIWarpTexture::CalcFrustum_fullCPU(const IMPCDI::FFrustum& OutFrustum, const FMatrix& World2Local, float& Top, float& Bottom, float& Left, float& Right) const
+	{
+		int Count = GetWidth()*GetHeight();
+		const FVector4* v = (FVector4*)GetData();
+
+		// Search a camera space frustum
+		for (int i = 0; i < Count; ++i)
+		{
+			UpdateFrustumFromVertex(v[i], World2Local, Top, Bottom, Left, Right);
+		}
 	}
 
-	void FMPCDIWarpTexture::CalcViewProjection(const IMPCDI::FFrustum& Frustum, const FVector& ViewDirection, const FVector& ViewOrigin, const FVector& EyeOrigin, FMatrix& OutViewMatrix) const
+	void FMPCDIWarpTexture::CalcFrustum_simpleAABB(const FVector* AABBoxPts, const IMPCDI::FFrustum& OutFrustum, const FMatrix& World2Local, float& Top, float& Bottom, float& Left, float& Right) const
+	{
+		// Search a camera space frustum
+		for (int i = 0; i < 8; ++i)
+		{
+			UpdateFrustumFromVertex(AABBoxPts[i], World2Local, Top, Bottom, Left, Right);
+		}
+	}
+
+	void FMPCDIWarpTexture::CalcViewProjection(const FVector* AABBoxPts, const IMPCDI::FFrustum& Frustum, const FVector& ViewDirection, const FVector& ViewOrigin, const FVector& EyeOrigin, FMatrix& OutViewMatrix) const
 	{
 		const EVarMPCDIProjectionMode ProjMode = (EVarMPCDIProjectionMode)CVarMPCDIProjectionMode.GetValueOnAnyThread();
 		switch (ProjMode)
 		{
-			case EVarMPCDIProjectionMode::Dynamic:
+			case EVarMPCDIProjectionMode::StaticSurfaceNormal:
+			{
+				// Use fixed surface view normal:
+				OutViewMatrix = FRotationMatrix::MakeFromXZ(SurfaceViewNormal, FVector(0.f, 0.f, 1.f));
+				OutViewMatrix.SetOrigin(EyeOrigin); // Finally set view origin to eye location
+
+				break;
+			}
+			case EVarMPCDIProjectionMode::DynamicAABBCenter:
 			{
 				OutViewMatrix = FRotationMatrix::MakeFromXZ(ViewDirection, FVector(0.f, 0.f, 1.f));
 				OutViewMatrix.SetOrigin(EyeOrigin); // Finally set view origin to eye location
 				break;
 			}
 
-			case EVarMPCDIProjectionMode::StaticAxisAligned:
+			case EVarMPCDIProjectionMode::DynamicAxisAligned:
 			{
-				float X = fabs(ViewDirection.X);
-				float Y = fabs(ViewDirection.Y);
-				float Z = fabs(ViewDirection.Z);
+				// Search best axis direction
+				int PtsCount[6] = {0,0, 0,0, 0,0};
+				for (int i = 0; i < 8; ++i)
+				{
+					FVector LookVector = AABBoxPts[i] - ViewOrigin;
+					if (LookVector.X > 0)
+						{ PtsCount[0]++; }
+					if (LookVector.X < 0)
+						{ PtsCount[1]++; }
+					if (LookVector.Y > 0)
+						{ PtsCount[2]++; }
+					if (LookVector.Y < 0)
+						{ PtsCount[3]++; }
+					if (LookVector.Z > 0)
+						{ PtsCount[4]++; }
+					if (LookVector.Z < 0)
+						{ PtsCount[5]++; }
+				}
+
+				int X = abs(PtsCount[0] - PtsCount[1]);
+				int Y = abs(PtsCount[2] - PtsCount[3]);
+				int Z = abs(PtsCount[4] - PtsCount[5]);
 
-				FVector Direction = ViewDirection;
+				FVector Direction(1, 1, 1);
+
+				if (PtsCount[0] < PtsCount[1])
+				  { Direction.X = -1; }
+				if (PtsCount[2] < PtsCount[3])
+				  { Direction.Y = -1; }
+				if (PtsCount[4] < PtsCount[5])
+				  { Direction.Z = -1; }
 
 				if (X > Y && X > Z)
 				{
 					Direction.Y = Direction.Z = 0;
 				}
+				else
 				if (Y > X && Y > Z)
 				{
 					Direction.X = Direction.Z = 0;
 				}
+				else
 				if (Z > X && Z > Y)
 				{
 					Direction.X = Direction.Y = 0;
 				}
+				else
+				{
+					// Eye is outside of region, optimize surface normal
+					Direction = SurfaceViewNormal;
+
+					if (X == Y && X > Z)
+					{
+						Direction.Z = 0;
+					}
+					else
+					if (X == Z && X > Y)
+					{
+						Direction.Y = 0;
+					}
+					else
+					if (Y == Z && Y > X)
+					{
+						Direction.X = 0;
+					}
+				}
 
 				Direction.GetSafeNormal();
 				OutViewMatrix = FRotationMatrix::MakeFromXZ(Direction, FVector(0.f, 0.f, 1.f));
 				OutViewMatrix.SetOrigin(EyeOrigin); // Finally set view origin to eye location
-
 				break;
 			}
 		}		
@@ -222,10 +413,53 @@ namespace MPCDI
 //---------------------------------------------
 	bool FMPCDIWarpTexture::GetFrustum_A3D(IMPCDI::FFrustum &OutFrustum, float WorldScale, float ZNear, float ZFar) const
 	{
+		const int FrustumCacheDepth = (int)CVarMPCDIFrustumCacheDepth.GetValueOnAnyThread();
+		{// Try to use old frustum values from cache (reduce CPU cost)
+			if (FrustumCacheDepth>0 && FrustumCache.Num() > 0)
+			{
+				FScopeLock lock(&DataGuard);
+
+				// If frustum cache used ,try to search valid value
+				const float FrustumCachePrecision = (float)CVarMPCDIFrustumCachePrecision.GetValueOnAnyThread();
+				for (int i=0;i<FrustumCache.Num(); i++)
+				{
+					if (FrustumCache[i].IsEyeLocationEqual(OutFrustum, FrustumCachePrecision))
+					{
+						// Use cached value
+						OutFrustum = FrustumCache[i];
+						int OnTopIndex = FrustumCache.Num() - 1;
+						if (OnTopIndex > i)
+						{
+							FrustumCache.Add(OutFrustum); // Add to on top of cache
+							FrustumCache.RemoveAt(i, 1); // Remove from prev cache position
+						}
+						return true;
+					}
+				}
+			}
+			else
+			{
+				FrustumCache.Empty();
+			}
+		}
+
+		// Calculate new Frustum value
 		OutFrustum.WorldScale = WorldScale;
 
-		FVector AABBMaxExtent = GetAABB().Max * WorldScale;
-		FVector AABBMinExtent = GetAABB().Min * WorldScale;
+		// Build AABB points
+		FVector AABBoxPts[8];
+
+		FVector4 AABBMaxExtent = GetAABB().Max * WorldScale;
+		FVector4 AABBMinExtent = GetAABB().Min * WorldScale;
+
+		AABBoxPts[0] = FVector(AABBMaxExtent.X, AABBMaxExtent.Y, AABBMaxExtent.Z);
+		AABBoxPts[1] = FVector(AABBMaxExtent.X, AABBMaxExtent.Y, AABBMinExtent.Z);
+		AABBoxPts[2] = FVector(AABBMinExtent.X, AABBMaxExtent.Y, AABBMinExtent.Z);
+		AABBoxPts[3] = FVector(AABBMinExtent.X, AABBMaxExtent.Y, AABBMaxExtent.Z);
+		AABBoxPts[4] = FVector(AABBMaxExtent.X, AABBMinExtent.Y, AABBMaxExtent.Z);
+		AABBoxPts[5] = FVector(AABBMaxExtent.X, AABBMinExtent.Y, AABBMinExtent.Z);
+		AABBoxPts[6] = FVector(AABBMinExtent.X, AABBMinExtent.Y, AABBMinExtent.Z);
+		AABBoxPts[7] = FVector(AABBMinExtent.X, AABBMinExtent.Y, AABBMaxExtent.Z);
 
 		FMatrix Local2world = FMatrix::Identity;
 
@@ -241,9 +475,9 @@ namespace MPCDI
 
 				// Create view transform matrix from look direction vector:
 				FVector LookVector = AABBCenter - LookAt;
-				FVector LookDirection= LookVector.GetSafeNormal();
-				CalcViewProjection(OutFrustum, LookDirection, LookAt, LookAt, Local2world);
-			
+				FVector LookDirection = LookVector.GetSafeNormal();
+				CalcViewProjection(&AABBoxPts[0], OutFrustum, LookDirection, LookAt, LookAt, Local2world);
+
 				break;
 			}
 
@@ -257,7 +491,7 @@ namespace MPCDI
 				// Create view transform matrix from look direction vector:
 				FVector LookVector = AABBCenter - LookAt;
 				FVector LookDirection = LookVector.GetSafeNormal();
-				CalcViewProjection(OutFrustum, LookDirection, LookAt, FVector(LookAt + OutFrustum.OriginEyeOffset), Local2world);
+				CalcViewProjection(&AABBoxPts[0], OutFrustum, LookDirection, LookAt, FVector(LookAt + OutFrustum.OriginEyeOffset), Local2world);
 
 				break;
 			}
@@ -272,19 +506,26 @@ namespace MPCDI
 		float left = FLT_MAX;
 		float right = -FLT_MAX;
 
+		//Texture box frustum method size
+		static FIntPoint TextureBoxSize(8,8);
+
+		//Compute rendering frustum with current method
 		const EVarMPCDIFrustumMethod FrustumComputeMethod = (EVarMPCDIFrustumMethod)CVarMPCDIFrustumMethod.GetValueOnAnyThread();
 		switch (FrustumComputeMethod)
 		{
 		case EVarMPCDIFrustumMethod::AABB:
-			CalcFrustum_simpleAABB(OutFrustum, World2local, top, bottom, left, right);
+			CalcFrustum_simpleAABB(&AABBoxPts[0], OutFrustum, World2local, top, bottom, left, right);
 			break;
 
 		case EVarMPCDIFrustumMethod::PerfectCPU:
 			CalcFrustum_fullCPU(OutFrustum, World2local, top, bottom, left, right);
 			break;
+		case EVarMPCDIFrustumMethod::TextureBOX:			
+			CalcFrustum_TextureBOX(TextureBoxSize.X, TextureBoxSize.Y, OutFrustum, World2local, top, bottom, left, right);
+			break;
 		}
 
-		OutFrustum.ProjectionAngles = IMPCDI::FFrustum::Angles(top, bottom, left, right);
+		OutFrustum.ProjectionAngles = IMPCDI::FFrustum::FAngles(top, bottom, left, right);
 
 		// These matrices were copied from LocalPlayer.cpp.
 		// They change the coordinate system from the Unreal "Game" coordinate system to the Unreal "Render" coordinate system
@@ -293,31 +534,72 @@ namespace MPCDI
 			FPlane(1, 0, 0, 0),
 			FPlane(0, 1, 0, 0),
 			FPlane(0, 0, 0, 1));
-		
+
 		static const FMatrix Render2Game = Game2Render.Inverse();
 
-		//!
 		ZNear = 1.0f;
 
-		
-		// Compute view matrix
-		//!OutFrustum.ViewMatrix = Game2Render;
-
-		// Compute projection matrix
-		FMatrix OriginViewMatrix = World2local * Game2Render;
-
+		// Compute warp projection matrix
 		OutFrustum.Local2WorldMatrix = Local2world;
-		OutFrustum.ProjectionMatrix  = GetProjectionMatrixAssymetric(left, right, top, bottom, ZNear, ZFar);
-		//!OutFrustum.OriginViewMatrix  = OriginViewMatrix;
-		
-		//OutFrustum.UVMatrix = OutFrustum.ViewMatrix * (OutFrustum.OriginViewMatrix * OutFrustum.ProjectionMatrix);
-		OutFrustum.UVMatrix = OriginViewMatrix * OutFrustum.ProjectionMatrix;
+		OutFrustum.ProjectionMatrix = GetProjectionMatrixAssymetric(left, right, top, bottom, ZNear, ZFar);
+		OutFrustum.UVMatrix = World2local * Game2Render * OutFrustum.ProjectionMatrix;
+	
+		{//Store current used frustum value to cache:
+			FScopeLock lock(&DataGuard);
+			// Save current frustum value to cache
+			if (FrustumCacheDepth > 0)
+			{
+				FrustumCache.Add(OutFrustum);
+			}
+
+			// Remove too old cached values
+			int TotalTooOldValuesCount = FrustumCache.Num() - FrustumCacheDepth;
+			if(TotalTooOldValuesCount > 0)
+			{
+				FrustumCache.RemoveAt(0, TotalTooOldValuesCount);
+			}
+		}
 
 		return true;
 	}
 
 	void FMPCDIWarpTexture::BuildAABBox()
 	{
+		{// Build SurfaceViewNormal
+			int Ncount = 0;
+			double Nxyz[3] = { 0,0,0 };
+			for (uint32 Y = 0; Y < GetHeight()-1; ++Y)
+			{
+				for (uint32 X = 0; X < GetWidth()-1; ++X)
+				{
+					const FVector4& Pts0 = ((FVector4*)GetData())[(X + Y * GetWidth())];
+					const FVector4& Pts1 = ((FVector4*)GetData())[(X+1 + Y * GetWidth())];
+					const FVector4& Pts2 = ((FVector4*)GetData())[(X + (Y+1) * GetWidth())];
+					if (Pts0.W > 0 && Pts1.W > 0 && Pts2.W > 0)
+					{
+						const FVector N1 = Pts1 - Pts0;
+						const FVector N2 = Pts2 - Pts0;
+						const FVector N = FVector::CrossProduct(N2, N1).GetSafeNormal();
+						for (int i = 0; i < 3; i++)
+						{
+							Nxyz[i] += N[i];
+						}
+						Ncount++;
+					}
+				}
+			}
+
+			double Scale = double(1) / Ncount;
+			for (int i = 0; i < 3; i++)
+			{
+				SurfaceViewNormal[i] = Nxyz[i]* Scale;
+			}
+			SurfaceViewNormal = SurfaceViewNormal.GetSafeNormal();
+		}
+
+		TextureBoxCache.Empty();
+		FrustumCache.Empty();
+
 		//Build bbox from valid points:
 		for (uint32 Y = 0; Y < GetHeight(); ++Y)
 		{
@@ -338,21 +620,33 @@ namespace MPCDI
 		}
 	}
 
-	void FMPCDIWarpTexture::LoadCustom3DWarpMap(const TArray<FVector>& InPoints, int WarpX, int WarpY)
+	bool FMPCDIWarpTexture::LoadCustom3DWarpMap(const TArray<FVector>& InPoints, int WarpX, int WarpY, IMPCDI::EMPCDIProfileType ProfileType, float WorldScale, bool bIsUnrealGameSpace)
 	{
-		FVector& OutMinBBox = AABBox.Min;
-		FVector& OutMaxBBox = AABBox.Max;
-
-		// Convert from MPCDI convention to Unreal convention
-		// MPCDI is Right Handed (Y is up, X is left, Z is in the screen)
-		// Unreal is Left Handed (Z is up, X in the screen, Y is right)
-		float scale = 100.f;
-		FMatrix m(
-			FPlane(0.f, scale, 0.f, 0.f),
-			FPlane(0.f, 0.f, scale, 0.f),
-			FPlane(-scale, 0.f, 0.f, 0.f),
-			FPlane(0.f, 0.f, 0.f, 1.f));
-		
+		FMatrix m = FMatrix::Identity;
+
+		if (ProfileType == IMPCDI::EMPCDIProfileType::mpcdi_A3D)
+		{
+			if (bIsUnrealGameSpace)
+			{
+				m = FMatrix(
+					FPlane(WorldScale, 0.f, 0.f, 0.f),
+					FPlane(0.f, WorldScale, 0.f, 0.f),
+					FPlane(0.f, 0.f, WorldScale, 0.f),
+					FPlane(0.f, 0.f, 0.f, 1.f));
+			}
+			else
+			{
+				m = // Convert from MPCDI convention to Unreal convention
+					// MPCDI is Right Handed (Y is up, X is left, Z is in the screen)
+					// Unreal is Left Handed (Z is up, X in the screen, Y is right)
+					FMatrix(
+						FPlane(0.f, WorldScale, 0.f, 0.f),
+						FPlane(0.f, 0.f, WorldScale, 0.f),
+						FPlane(-WorldScale, 0.f, 0.f, 0.f),
+						FPlane(0.f, 0.f, 0.f, 1.f));
+			}
+		}
+
 		FVector4* data = new FVector4[WarpX*WarpY];
 
 		static const float kEpsilon = 0.00001f;
@@ -372,23 +666,50 @@ namespace MPCDI
 				Pts = FVector4(0.f, 0.f, 0.f, -1.f);
 			}
 		}
+		
+		ReleaseTextureData();
 
 		// Create texture
 		const EPixelFormat pixelFormat = PF_A32B32G32R32F;
 		SetTextureData(reinterpret_cast<void*>(data), WarpX, WarpY, pixelFormat, false);
 
 		BuildAABBox();
+
+		if (IsInitialized())
+		{
+			BeginUpdateResourceRHI(this);
+		}
 		BeginInitResource(this);
+		return true;
 	}
 
-	void FMPCDIWarpTexture::LoadWarpMap(mpcdi::GeometryWarpFile *SourceWarpMap, EMPCDIProfileType ProfileType)
+	bool FMPCDIWarpTexture::LoadPFMFile(mpcdi::PFM& SourcePFM, IMPCDI::EMPCDIProfileType ProfileType, float PFMScale, bool bIsUnrealGameSpace)
 	{
-		FVector& OutMinBBox = AABBox.Min;
-		FVector& OutMaxBBox = AABBox.Max;
+		int PFMWidth = SourcePFM.GetSizeX();
+		int PFMHeight = SourcePFM.GetSizeY();
+
+		TArray<FVector> WarpMeshPoints;
+		WarpMeshPoints.Reserve(PFMWidth*PFMHeight);
+
+		for (int y = 0; y < PFMHeight; ++y)
+		{
+			for (int x = 0; x < PFMWidth; ++x)
+			{
+				mpcdi::NODE node = SourcePFM(x, y);
 
+				FVector pts = (FVector&)(node);
+				WarpMeshPoints.Add( pts );
+			}
+		}
+
+		return LoadCustom3DWarpMap(WarpMeshPoints, PFMWidth, PFMHeight, ProfileType, PFMScale, bIsUnrealGameSpace);
+	}
+
+	bool FMPCDIWarpTexture::LoadWarpMap(mpcdi::GeometryWarpFile *SourceWarpMap, IMPCDI::EMPCDIProfileType ProfileType)
+	{
 		FMatrix m = FMatrix::Identity;
 		bool is2DData = true;
-		if (ProfileType == EMPCDIProfileType::mpcdi_A3D)
+		if (ProfileType == IMPCDI::EMPCDIProfileType::mpcdi_A3D)
 		{
 			// Unreal is in cm, so we need to convert to cm.
 			float scale = 1.f;
@@ -445,20 +766,26 @@ namespace MPCDI
 			}
 		}
 
-		if (ProfileType == EMPCDIProfileType::mpcdi_A3D)
+		if (ProfileType == IMPCDI::EMPCDIProfileType::mpcdi_A3D)
 		{
 			// Remove noise from warp mesh (small areas less than 3*3 quads)
 			ClearNoise(FIntPoint(3, 3), FIntPoint(2, 3));
 		}
 
+		ReleaseTextureData();
+
 		// Create texture
 		const EPixelFormat pixelFormat = PF_A32B32G32R32F;
 		SetTextureData(reinterpret_cast<void*>(data), WarpX, WarpY, pixelFormat, false);
 
 		BuildAABBox();
-		//BuildWarpMapData(ProfileType);
 
+		if (IsInitialized())
+		{
+			BeginUpdateResourceRHI(this);
+		}
 		BeginInitResource(this);
+		return true;
 	}
 
 	bool FMPCDIWarpTexture::Is3DPointValid(int X, int Y) const
@@ -625,6 +952,10 @@ namespace MPCDI
 
 		}
 	}
+	void FMPCDIWarpTexture::ImportMeshData(const FMPCDIGeometryImportData& Src)
+	{		
+		LoadCustom3DWarpMap(Src.Vertices, Src.Width, Src.Height, IMPCDI::EMPCDIProfileType::mpcdi_A3D, 1, true);
+	}
 }
 
 
diff --git a/Engine/Plugins/Runtime/nDisplay/Source/MPCDI/Private/MPCDIWarpTexture.h b/Engine/Plugins/Runtime/nDisplay/Source/MPCDI/Private/MPCDIWarpTexture.h
--- a/Engine/Plugins/Runtime/nDisplay/Source/MPCDI/Private/MPCDIWarpTexture.h
+++ b/Engine/Plugins/Runtime/nDisplay/Source/MPCDI/Private/MPCDIWarpTexture.h
@@ -3,7 +3,6 @@
 #pragma once
 
 #include "MPCDITexture.h"
-#include "MPCDITypes.h"
 #include "IMPCDI.h"
 
 #include "Blueprints/MPCDIContainers.h"
@@ -12,6 +11,7 @@
 namespace mpcdi
 {
 	struct GeometryWarpFile;
+	struct PFM;
 }
 
 namespace MPCDI
@@ -23,40 +23,53 @@ namespace MPCDI
 			: FMPCDITexture()
 			, AABBox(FVector(FLT_MAX, FLT_MAX, FLT_MAX)
 			, FVector(-FLT_MAX, -FLT_MAX, -FLT_MAX))
-		{ }
+		{ 
+		}
+		virtual ~FMPCDIWarpTexture()
+		{
+		}
 
 	public:
 		bool GetFrustum_A3D(IMPCDI::FFrustum &OutFrustum, float WorldScale, float ZNear, float ZFar) const;
 
-		void LoadCustom3DWarpMap(const TArray<FVector>& Points, int DimW, int DimH);
-		void LoadWarpMap(mpcdi::GeometryWarpFile *SourceWarpMap, EMPCDIProfileType ProfileType);
+		bool LoadCustom3DWarpMap(const TArray<FVector>& Points, int DimW, int DimH, IMPCDI::EMPCDIProfileType ProfileType, float WorldScale, bool bIsUnrealGameSpace);
+		bool LoadPFMFile(mpcdi::PFM& SourcePFM, IMPCDI::EMPCDIProfileType ProfileType, float PFMScale, bool bIsUnrealGameSpace);
+
+		bool LoadWarpMap(mpcdi::GeometryWarpFile *SourceWarpMap, IMPCDI::EMPCDIProfileType ProfileType);
 
 		inline const FBox& GetAABB() const
-		{ return AABBox; }
+		{ 
+			return AABBox; 
+		}
 
-		inline void AppendAABB(FBox& TargetAABBox)
-		{ TargetAABBox += AABBox; }
+		inline void AppendAABB(FBox& TargetAABBox) const
+		{ 
+			TargetAABBox += AABBox; 
+		}
 
 		void ExportMeshData(FMPCDIGeometryExportData& Dst);
+		void ImportMeshData(const FMPCDIGeometryImportData& Src);
 
 	private:
 		bool Is3DPointValid(int X, int Y) const;
-		void BuildWarpMapData(EMPCDIProfileType ProfileType);
+		void BuildWarpMapData(IMPCDI::EMPCDIProfileType ProfileType);
 		void BuildAABBox();
 		void ClearNoise(const FIntPoint& SearchXYDepth, const FIntPoint& AllowedXYDepthRules);
 		int RemoveDetachedPoints(const FIntPoint& SearchXYDepth, const FIntPoint& AllowedXYDepthRules);
 
 	private:
-
-		void CalcFrustum_simpleAABB(const IMPCDI::FFrustum& Frustum, const FMatrix& world2local, float& top, float& bottom, float& left, float& right) const;
-		void CalcFrustum_fullCPU(const IMPCDI::FFrustum& Frustum, const FMatrix& world2local, float& top, float& bottom, float& left, float& right) const;
-
-		void CalcViewProjection(const IMPCDI::FFrustum& Frustum, const FVector& ViewDirection, const FVector& ViewOrigin, const FVector& EyeOrigin, FMatrix& OutViewMatrix) const;
+		void CalcFrustum_simpleAABB(const FVector* AABBoxPts, const IMPCDI::FFrustum& Frustum, const FMatrix& World2Local, float& Top, float& Bottom, float& Left, float& Right) const;
+		void CalcFrustum_fullCPU(const IMPCDI::FFrustum& Frustum, const FMatrix& World2Local, float& Top, float& Bottom, float& Left, float& Right) const;
+		void CalcFrustum_TextureBOX(int DivX, int DivY, const IMPCDI::FFrustum& Frustum, const FMatrix& World2Local, float& Top, float& Bottom, float& Left, float& Right) const;
+		
+		void CalcViewProjection(const FVector* AABBoxPts, const IMPCDI::FFrustum& Frustum, const FVector& ViewDirection, const FVector& ViewOrigin, const FVector& EyeOrigin, FMatrix& OutViewMatrix) const;
 
 	private:
-		FBox AABBox;
+		FBox    AABBox;
+		FVector SurfaceViewNormal; // Static surface average normal for this region
 
-		// Orthogonal to a projection plane
-		FRotator ViewRotation = FRotator::ZeroRotator;
+		mutable FCriticalSection DataGuard;
+		mutable TArray<IMPCDI::FFrustum> FrustumCache;
+		mutable TArray<int> TextureBoxCache;
 	};
 }
