nDisplay features

diff --git a/Engine/Plugins/Runtime/nDisplay/Source/MPCDI/Public/Blueprints/IMPCDIBlueprintAPI.h b/Engine/Plugins/Runtime/nDisplay/Source/MPCDI/Public/Blueprints/IMPCDIBlueprintAPI.h
--- a/Engine/Plugins/Runtime/nDisplay/Source/MPCDI/Public/Blueprints/IMPCDIBlueprintAPI.h
+++ b/Engine/Plugins/Runtime/nDisplay/Source/MPCDI/Public/Blueprints/IMPCDIBlueprintAPI.h
@@ -21,5 +21,14 @@ class IMPCDIBlueprintAPI
 
 public:
 	UFUNCTION(BlueprintCallable, meta = (DisplayName = "Get MPCDI Mesh Data"), Category = "MPCDI")
-	virtual void GetMPCDIMeshData(const FString& MPCDIFile, const FString& BufferName, const FString& RegionName, FMPCDIGeometryExportData& MeshData) = 0;
+	virtual bool GetMPCDIMeshData(const FString& MPCDIFile, const FString& BufferName, const FString& RegionName, FMPCDIGeometryExportData& MeshData) = 0;
+
+	UFUNCTION(BlueprintCallable, meta = (DisplayName = "Get PFM Mesh Data"), Category = "MPCDI")
+	virtual bool GetPFMMeshData(const FString& PFMFile, FMPCDIGeometryExportData& MeshData, float PFMScale=1, bool bIsMPCDIAxis= true) = 0;
+
+	UFUNCTION(BlueprintCallable, meta = (DisplayName = "Set MPCDI Mesh Data"), Category = "MPCDI")
+	virtual void SetMPCDIMeshData(const FString& MPCDIFile, const FString& BufferName, const FString& RegionName, const FMPCDIGeometryImportData& MeshData) = 0;
+
+	UFUNCTION(BlueprintCallable, meta = (DisplayName = "Reload Changed External Files"), Category = "MPCDI")
+	virtual void ReloadChangedExternalFiles() = 0;
 };
diff --git a/Engine/Plugins/Runtime/nDisplay/Source/MPCDI/Public/Blueprints/MPCDIContainers.h b/Engine/Plugins/Runtime/nDisplay/Source/MPCDI/Public/Blueprints/MPCDIContainers.h
--- a/Engine/Plugins/Runtime/nDisplay/Source/MPCDI/Public/Blueprints/MPCDIContainers.h
+++ b/Engine/Plugins/Runtime/nDisplay/Source/MPCDI/Public/Blueprints/MPCDIContainers.h
@@ -8,6 +8,23 @@
 
 #include "MPCDIContainers.generated.h"
 
+USTRUCT(BlueprintType, Category = "MPCDI")
+struct FMPCDIGeometryImportData
+{
+	GENERATED_BODY()
+
+public:
+	UPROPERTY(EditAnywhere, BlueprintReadWrite, Category = "MPCDI")
+	int Width;
+
+	UPROPERTY(EditAnywhere, BlueprintReadWrite, Category = "MPCDI")
+	int Height;
+
+	UPROPERTY(EditAnywhere, BlueprintReadWrite, Category = "MPCDI")
+	TArray<FVector>   Vertices;
+};
+
+
 USTRUCT(BlueprintType, Category = "MPCDI")
 struct FMPCDIGeometryExportData
 {
diff --git a/Engine/Plugins/Runtime/nDisplay/Source/MPCDI/Public/IMPCDI.h b/Engine/Plugins/Runtime/nDisplay/Source/MPCDI/Public/IMPCDI.h
--- a/Engine/Plugins/Runtime/nDisplay/Source/MPCDI/Public/IMPCDI.h
+++ b/Engine/Plugins/Runtime/nDisplay/Source/MPCDI/Public/IMPCDI.h
@@ -20,6 +20,44 @@ public:
 	static constexpr auto ModuleName = TEXT("MPCDI");
 
 public:
+
+	enum EMPCDIProfileType: uint8
+	{
+		mpcdi_2D = 0, // 2D mode
+		mpcdi_3D,     // 3D mode
+		mpcdi_A3D,    // Advanced 3D mode
+		mpcdi_SL,     // Shader lamps
+		Invalid,
+	};
+
+	struct ConfigParser
+	{
+		FString  ConfigLineStr;// Saved viewport config line string
+		FString  MPCDIFileName; // Single mpcdi file name
+
+		FString  BufferId;
+		FString  RegionId;
+
+		FString  OriginType;
+
+		// Support external pfm (warp)  and png(blend) files
+		EMPCDIProfileType  MPCDIType;
+
+		FString  PFMFile;
+		float    PFMFileScale;
+		bool     bIsUnrealGameSpace;
+
+		FString  AlphaFile;
+		float AlphaGamma;
+
+		FString  BetaFile;
+
+		inline bool IsExtConfig() const
+		{ 
+			return !PFMFile.IsEmpty(); 
+		}
+	};
+
 	struct FRegionLocator
 	{
 		// MPCDI data indices
@@ -28,23 +66,32 @@ public:
 		int RegionIndex = -1;
 
 		inline bool isValid() const
-		{ return FileIndex >= 0 && BufferIndex >= 0 && RegionIndex >= 0; }
+		{ 
+			return FileIndex >= 0 && BufferIndex >= 0 && RegionIndex >= 0; 
+		}
 	};
 
 	struct FFrustum
 	{
+		// Frustum projection angles
+		struct FAngles
+		{
+			float Left, Right, Top, Bottom;
+			FAngles()
+			{
+			};
+
+			FAngles(float& top, float& bottom, float& left, float& right) : Left(left), Right(right), Top(top), Bottom(bottom)
+			{
+			};
+		};
+
 		// Frustum origins Input data:
 		FVector OriginLocation;  // Current camera Origin location
 		FVector OriginEyeOffset; // Offset from OriginLocation to Eye view location
 
-		// Output runtime calc data:
-
-		// Frustum projection angles
-		struct Angles {
-			float Left, Right, Top, Bottom;
-			Angles() {};
-			Angles(float& top, float& bottom, float& left, float& right) : Left(left), Right(right), Top(top), Bottom(bottom) {};
-		} ProjectionAngles;
+		// Output runtime calc data: 
+		FAngles ProjectionAngles;
 
 		// Frustum projection matrix
 		FMatrix  ProjectionMatrix;
@@ -59,11 +106,22 @@ public:
 		FMatrix  UVMatrix;
 
 		float    WorldScale;
-		bool     bIsValid : 1;
+		bool     bIsValid;
+
+		inline FVector GetEyeLocation() const
+		{
+			return OriginLocation + OriginEyeOffset;
+		}
+
+		inline bool IsEyeLocationEqual(const FFrustum& InFrustum, float Precision) const
+		{
+			return (GetEyeLocation() - InFrustum.GetEyeLocation()).Size() < Precision;
+		}
 
 		FFrustum()
 			: FFrustum(FVector(0.f, 0.f, 0.f), FVector(0.f, 0.f, 0.f))
-		{ }
+		{ 
+		}
 
 		FFrustum(const FVector &InOriginLocation, const FVector& InOriginEyeOffset)
 			: OriginLocation(InOriginLocation)
@@ -72,7 +130,8 @@ public:
 			, UVMatrix(FMatrix::Identity)
 			, WorldScale(1.f)
 			, bIsValid(false)
-		{ }
+		{ 
+		}
 	};
 
 	struct FTextureWarpData
@@ -188,4 +247,102 @@ public:
 	* @return - ptr if success
 	*/
 	virtual TSharedPtr<FMPCDIData> GetMPCDIData(IMPCDI::FShaderInputData& ShaderInputData) = 0;
+
+	//! Support ext warp&blend calibration files
+	/**
+	* Create new mpcdi file + buffer + region (region is empty, for custom setup)
+	*
+	* @param MPCDIFile  - Fullpath unique mpcdi file name
+	* @param BufferName - Buffer unique name
+	* @param RegionName - region unique name
+	* @param OutRegionLocator - return unique region identifier
+	*
+	* @return - true if success
+	*           false, if region already exist
+	*/
+	virtual bool CreateCustomRegion(const FString& MPCDIFile, const FString& BufferName, const FString& RegionName, IMPCDI::FRegionLocator& OutRegionLocator) = 0;
+
+	/**
+	* Change mpcdi profile type for custom warpblend (for owned mpcdi file)
+	*
+	* @param InRegionLocator - region locator
+	*
+	* @return - true if success
+	*/
+	virtual bool SetMPCDIProfileType(const IMPCDI::FRegionLocator& InRegionLocator, const EMPCDIProfileType ProfileType) = 0;
+
+	/**
+	* Load warp map data from external PFM file, used custom scale and axis orientation (by default asis in mpcdi orientation)
+	*
+	* @param InRegionLocator - region locator
+	* @param PFMFile - full path to PFM file
+	* @param WorldScale - world scale multiplier to UE4 game (value=100 from meters->sm)
+	* @param bIsUnrealGameSpace - let=true, if you dont want to apply axis conversion from mpcdi space
+	*
+	* @return - true if success
+	*/
+	virtual bool LoadPFM(const IMPCDI::FRegionLocator& InRegionLocator, const FString& PFMFile, const float WorldScale, bool bIsUnrealGameSpace=false) = 0;
+	
+	/**
+	* Load warp map data from array, used custom scale and axis orientation (by default asis in mpcdi orientation)
+	*
+	* @param InRegionLocator - region locator
+	* @param PFMPoints - PFM points, row-by-row
+	* @param DimW - Row width
+	* @param DimH - Rows count
+	* @param WorldScale - world scale multiplier to UE4 game (value=100 from meters->sm)
+	* @param bIsUnrealGameSpace - let=true, if you dont want to apply axis conversion from mpcdi space
+	*
+	* @return - true if success
+	*/
+	virtual bool LoadPFMGeometry(const IMPCDI::FRegionLocator& InRegionLocator, const TArray<FVector>& PFMPoints, int DimW, int DimH, const float WorldScale, bool bIsUnrealGameSpace = false) = 0;
+
+	/**
+	* Load alpha map from external PNG file
+	*
+	* @param InRegionLocator - region locator
+	* @param PFMFile - full path to PNG file
+	* @param GammaValue - alpha map gamme value
+	*
+	* @return - true if success
+	*/
+	virtual bool LoadAlphaMap(const IMPCDI::FRegionLocator& InRegionLocator, const FString& PNGFile, float GammaValue) = 0;
+
+	/**
+	* Load beta map from external PNG file
+	*
+	* @param InRegionLocator - region locator
+	* @param PFMFile - full path to PNG file	
+	*
+	* @return - true if success
+	*/
+	virtual bool LoadBetaMap(const IMPCDI::FRegionLocator& InRegionLocator, const FString& PNGFile) = 0;
+
+
+	/**
+	* Helper. Load config data from string
+	*
+	* @param InConfigLineStr - config string
+	* @param OutCfgData - result condig data 
+	*
+	* @return - true if success
+	*/
+	virtual bool LoadConfig(const FString& InConfigLineStr, ConfigParser& OutCfgData) = 0;
+
+	/**
+	* Helper. Load or create mpcdi data
+	*
+	* @param CfgData - Initialized config dta for viewport
+	* @param OutRegionLocator - region locator
+	*
+	* @return - true if success
+	*/
+	virtual bool Load(const ConfigParser& CfgData, IMPCDI::FRegionLocator& OutRegionLocator) = 0;
+
+	/**
+	* Helper. Reload all data from changed external files
+	*
+	*/
+	virtual void ReloadAll() = 0;
+	virtual void ReloadAll_RenderThread() = 0;
 };
diff --git a/Engine/Plugins/Runtime/nDisplay/Source/OutputRemap/OutputRemap.Build.cs b/Engine/Plugins/Runtime/nDisplay/Source/OutputRemap/OutputRemap.Build.cs
new file mode 100644
--- /dev/null
+++ b/Engine/Plugins/Runtime/nDisplay/Source/OutputRemap/OutputRemap.Build.cs
@@ -0,0 +1,36 @@
+// Copyright 1998-2019 Epic Games, Inc. All Rights Reserved.
+
+using UnrealBuildTool;
+using System.IO;
+
+public class OutputRemap : ModuleRules
+{
+	public OutputRemap(ReadOnlyTargetRules ROTargetRules) : base(ROTargetRules)
+	{
+		PublicDefinitions.Add("OutputRemap_STATIC");
+
+		PublicDependencyModuleNames.AddRange(
+			new string[]
+			{
+				"Core",
+				"CoreUObject",
+				"Engine",
+				"Projects",
+            });
+
+		PrivateDependencyModuleNames.AddRange(
+			new string[]
+			{
+                "DisplayCluster",
+				"RenderCore",
+				"RHI",
+				"UtilityShaders"
+			}
+		);
+
+		if (Target.bBuildEditor == true)
+		{
+			PrivateDependencyModuleNames.Add("UnrealEd");
+		}
+	}
+}
diff --git a/Engine/Plugins/Runtime/nDisplay/Source/OutputRemap/Private/Blueprints/OutputRemapBlueprintAPIImpl.cpp b/Engine/Plugins/Runtime/nDisplay/Source/OutputRemap/Private/Blueprints/OutputRemapBlueprintAPIImpl.cpp
new file mode 100644
--- /dev/null
+++ b/Engine/Plugins/Runtime/nDisplay/Source/OutputRemap/Private/Blueprints/OutputRemapBlueprintAPIImpl.cpp
@@ -0,0 +1,14 @@
+// Copyright 1998-2019 Epic Games, Inc. All Rights Reserved.
+
+#include "Blueprints/OutputRemapBlueprintAPIImpl.h"
+
+#include "UObject/Package.h"
+#include "IOutputRemap.h"
+
+
+void UOutputRemapAPIImpl::ReloadChangedExternalFiles()
+{
+	IOutputRemap& OutputRemapModule = IOutputRemap::Get();
+	OutputRemapModule.ReloadAll();
+}
+
diff --git a/Engine/Plugins/Runtime/nDisplay/Source/OutputRemap/Private/Blueprints/OutputRemapBlueprintAPIImpl.h b/Engine/Plugins/Runtime/nDisplay/Source/OutputRemap/Private/Blueprints/OutputRemapBlueprintAPIImpl.h
new file mode 100644
--- /dev/null
+++ b/Engine/Plugins/Runtime/nDisplay/Source/OutputRemap/Private/Blueprints/OutputRemapBlueprintAPIImpl.h
@@ -0,0 +1,26 @@
+// Copyright 1998-2019 Epic Games, Inc. All Rights Reserved.
+
+#pragma once
+
+#include "CoreMinimal.h"
+#include "InputCoreTypes.h"
+
+#include "Blueprints/IOutputRemapBlueprintAPI.h"
+#include "OutputRemapBlueprintAPIImpl.generated.h"
+
+
+/**
+ * Blueprint API interface implementation
+ */
+UCLASS()
+class UOutputRemapAPIImpl
+	: public UObject
+	, public IOutputRemapBlueprintAPI
+{
+	GENERATED_BODY()
+
+public:
+
+	UFUNCTION(BlueprintCallable, meta = (DisplayName = "Reload Changed External Files"), Category = "OutputRemap")
+	virtual void ReloadChangedExternalFiles() override;
+};
\ No newline at end of file
diff --git a/Engine/Plugins/Runtime/nDisplay/Source/OutputRemap/Private/Blueprints/OutputRemapBlueprintLib.cpp b/Engine/Plugins/Runtime/nDisplay/Source/OutputRemap/Private/Blueprints/OutputRemapBlueprintLib.cpp
new file mode 100644
--- /dev/null
+++ b/Engine/Plugins/Runtime/nDisplay/Source/OutputRemap/Private/Blueprints/OutputRemapBlueprintLib.cpp
@@ -0,0 +1,17 @@
+// Copyright 1998-2019 Epic Games, Inc. All Rights Reserved.
+
+#include "Blueprints/OutputRemapBlueprintLib.h"
+#include "Blueprints/OutputRemapBlueprintAPIImpl.h"
+#include "UObject/Package.h"
+
+
+UOutputRemapIBlueprintLib::UOutputRemapIBlueprintLib(class FObjectInitializer const & ObjectInitializer)
+	: Super(ObjectInitializer)
+{
+}
+
+void UOutputRemapIBlueprintLib::GetAPI(TScriptInterface<IOutputRemapBlueprintAPI>& OutAPI)
+{
+	static UOutputRemapAPIImpl* Obj = NewObject<UOutputRemapAPIImpl>(GetTransientPackage(), NAME_None, RF_MarkAsRootSet);
+	OutAPI = Obj;
+}
diff --git a/Engine/Plugins/Runtime/nDisplay/Source/OutputRemap/Private/OutputRemapLog.cpp b/Engine/Plugins/Runtime/nDisplay/Source/OutputRemap/Private/OutputRemapLog.cpp
new file mode 100644
--- /dev/null
+++ b/Engine/Plugins/Runtime/nDisplay/Source/OutputRemap/Private/OutputRemapLog.cpp
@@ -0,0 +1,5 @@
+// Copyright 1998-2019 Epic Games, Inc. All Rights Reserved.
+
+#include "OutputRemapLog.h"
+
+DEFINE_LOG_CATEGORY(LogOutputRemap);
diff --git a/Engine/Plugins/Runtime/nDisplay/Source/OutputRemap/Private/OutputRemapLog.h b/Engine/Plugins/Runtime/nDisplay/Source/OutputRemap/Private/OutputRemapLog.h
new file mode 100644
--- /dev/null
+++ b/Engine/Plugins/Runtime/nDisplay/Source/OutputRemap/Private/OutputRemapLog.h
@@ -0,0 +1,12 @@
+// Copyright 1998-2019 Epic Games, Inc. All Rights Reserved.
+
+#pragma once
+
+#include "CoreMinimal.h"
+
+// Plugin-wide log categories
+#if UE_BUILD_SHIPPING
+DECLARE_LOG_CATEGORY_EXTERN(LogOutputRemap, Warning, Warning);
+#else
+DECLARE_LOG_CATEGORY_EXTERN(LogOutputRemap, Log, All);
+#endif
diff --git a/Engine/Plugins/Runtime/nDisplay/Source/OutputRemap/Private/OutputRemapMesh.cpp b/Engine/Plugins/Runtime/nDisplay/Source/OutputRemap/Private/OutputRemapMesh.cpp
new file mode 100644
--- /dev/null
+++ b/Engine/Plugins/Runtime/nDisplay/Source/OutputRemap/Private/OutputRemapMesh.cpp
@@ -0,0 +1,451 @@
+// Copyright 1998-2019 Epic Games, Inc. All Rights Reserved.
+
+#include "OutputRemapMesh.h"
+#include "OutputRemapLog.h"
+#include "OutputRemapHelpers.h"
+
+#include "Stats/Stats.h"
+#include "Engine/Engine.h"
+#include "Misc/ConfigCacheIni.h"
+#include "Misc/Paths.h"
+#include "Misc/FileHelper.h"
+
+#include "RendererInterface.h"
+#include "RenderResource.h"
+#include "SceneTypes.h"
+
+#include "CommonRenderResources.h"
+
+
+static int VarOutputRemapReloadChangedExtFilesFrame = 0;
+static TAutoConsoleVariable<int32> CVarOutputRemapReloadChangedExtFiles(
+	TEXT("nDisplay.reload.output_remap"),
+	(int) 0, // Disabled by default
+	TEXT("Changed external files reload period for output remap:\n")
+	TEXT("0 : disable\n"),
+	ECVF_RenderThreadSafe
+);
+
+
+
+void FOutputRemapMesh::DrawMesh(FRHICommandList& RHICmdList)
+{
+	check(IsInRenderingThread());
+
+	if (bIsValid)
+	{
+		RHICmdList.SetStreamSource(0, VertexBufferRHI, 0);
+		RHICmdList.DrawIndexedPrimitive(IndexBufferRHI, 0, 0, NumVertices, 0, NumTriangles, 1);
+
+		//!support mesh reload
+		ReloadMesh();
+	}
+	else
+	{
+		//! Add error logs
+	}
+}
+
+namespace ObjMeshStrings
+{
+	static constexpr auto Vertex = TEXT("v ");
+	static constexpr auto UV = TEXT("vt ");
+	static constexpr auto Face = TEXT("f ");
+
+	namespace delims
+	{
+		static constexpr auto Values = TEXT(" ");
+		static constexpr auto Face = TEXT("/");
+	}
+};
+
+bool FOutputRemapMesh::FMeshData::CreateFromFile(const FString& FullPathFileName)
+{	
+	//@load from obj file at runtime:
+	if (FPaths::FileExists(FullPathFileName))
+	{
+		TArray<FString> data;
+		if (FFileHelper::LoadANSITextFileToStrings(*FullPathFileName, nullptr, data) == true)
+		{
+			bool bResult = true;
+
+			// Parse each line from config
+			for (auto line : data)
+			{
+				line.TrimStartAndEndInline();
+				if (!ParseLine(line))
+				{
+					// Handle error
+					bResult = false;
+				}
+			}
+			return bResult;
+		}
+	}
+	return false;
+}
+
+bool FOutputRemapMesh::FMeshData::ParseLine(const FString& Line)
+{
+	if (Line.StartsWith(ObjMeshStrings::Vertex))
+	{
+		return ExtractVertex(Line);
+	}
+
+	if (Line.StartsWith(ObjMeshStrings::UV))
+	{
+		return ExtractUV(Line);
+	}
+
+	if (Line.StartsWith(ObjMeshStrings::Face))
+	{
+		return ExtractFace(Line);
+	}
+
+	//skip
+	return true;
+}
+
+bool FOutputRemapMesh::FMeshData::ExtractVertex(const FString& Line)
+{
+	TArray<FString> Data;
+	if (Line.ParseIntoArray(Data, ObjMeshStrings::delims::Values)==4)
+	{
+		float X = FCString::Atof(*Data[1]);
+		float Y = FCString::Atof(*Data[2]);
+		float Z = FCString::Atof(*Data[3]);
+		InVertex.Add(FVector(X, Y, Z));
+		return true;
+	}
+	//! handle error
+	return false;
+}
+bool FOutputRemapMesh::FMeshData::ExtractUV(const FString& Line)
+{
+	TArray<FString> Data;
+	if (Line.ParseIntoArray(Data, ObjMeshStrings::delims::Values) == 4)
+	{
+		float U = FCString::Atof(*Data[1]);
+		float V = FCString::Atof(*Data[2]);
+		float W = FCString::Atof(*Data[3]);
+		InUV.Add(FVector(U, V, W));
+		return true;
+	}
+	//! handle error
+	return false;
+}
+bool FOutputRemapMesh::FMeshData::ExtractFace(const FString& Line)
+{
+	TArray<FString> Data;
+	Line.ParseIntoArray(Data, ObjMeshStrings::delims::Values);
+
+	if (Data.Num() > 3)
+	{
+		int IndexBase = Vertex.Num();
+		if (ExtractFaceVertex(Data[1])
+			&& ExtractFaceVertex(Data[2])
+			&& ExtractFaceVertex(Data[3])
+			)
+		{
+			// Triangle
+			Index.Add(IndexBase + 0);
+			Index.Add(IndexBase + 1);
+			Index.Add(IndexBase + 2);
+
+			if (Data.Num() > 4)
+			{
+				// Quad face
+				if (!ExtractFaceVertex(Data[4]))
+				{
+					//! handle error
+					return false;
+				}
+
+				Index.Add(IndexBase + 3);
+				Index.Add(IndexBase + 0);
+				Index.Add(IndexBase + 2);
+			}
+			return true;
+		}
+	}
+	//! handle error
+	return false;
+}
+bool FOutputRemapMesh::FMeshData::ExtractFaceVertex(const FString& Line)
+{
+	TArray<FString> Data;
+	if (Line.ParseIntoArray(Data, ObjMeshStrings::delims::Face) > 1)
+	{
+		int InVertexIndex = FCString::Atoi(*Data[0]) - 1;
+		int InUVIndex = FCString::Atoi(*Data[1]) - 1;
+
+		Vertex.Add(InVertex[InVertexIndex]);
+		UV.Add(InUV[InUVIndex]);
+		return true;
+	}
+
+	//! handle error
+	return false;
+}
+
+bool FOutputRemapMesh::FMeshData::CreatePassthrough()
+{
+	//@create dummy Passthrough mesh (test purpose)
+	
+	Vertex.Add(FVector(2.6f, 0, 1));
+	Vertex.Add(FVector(2.6f, 0, 2));
+	Vertex.Add(FVector(5.5f, 0, 2));
+	Vertex.Add(FVector(5.5f, 0, 1));
+
+	UV.Add(FVector(0, 1, 0));
+	UV.Add(FVector(0, 0, 0));
+	UV.Add(FVector(1, 0, 0));
+	UV.Add(FVector(1, 1, 0));
+
+	Index.Add(0);
+	Index.Add(1);
+	Index.Add(2);
+
+	Index.Add(3);
+	Index.Add(0);
+	Index.Add(2);
+
+	return true;
+}
+void FOutputRemapMesh::FMeshData::Normalize()
+{
+	FBox AABBox(FVector(FLT_MAX, FLT_MAX, FLT_MAX), FVector(-FLT_MAX, -FLT_MAX, -FLT_MAX));
+	
+	for (const auto& Pts: Vertex)
+	{
+		AABBox.Min.X = FMath::Min(AABBox.Min.X, Pts.X);
+		AABBox.Min.Y = FMath::Min(AABBox.Min.Y, Pts.Y);
+		AABBox.Min.Z = FMath::Min(AABBox.Min.Z, Pts.Z);
+
+		AABBox.Max.X = FMath::Max(AABBox.Max.X, Pts.X);
+		AABBox.Max.Y = FMath::Max(AABBox.Max.Y, Pts.Y);
+		AABBox.Max.Z = FMath::Max(AABBox.Max.Z, Pts.Z);
+	}
+
+	//Normalize
+	FVector Size(
+		(AABBox.Max.X - AABBox.Min.X),
+		(AABBox.Max.Y - AABBox.Min.Y),
+		(AABBox.Max.Z - AABBox.Min.Z)
+	);
+
+	if (Size.X == 0) { Size.X = 1; }
+	if (Size.Y == 0) { Size.Y = 1; }
+	if (Size.Z == 0) { Size.Z = 1; }
+
+	FVector Scale(1/Size.X, 1 / Size.Y, 1 / Size.Z);
+	
+	for (auto& Pts : Vertex)
+	{
+		float X = (Pts.X - AABBox.Min.X) * Scale.X;
+		float Y = (Pts.Y - AABBox.Min.Y) * Scale.Y;
+		float Z = (Pts.Z - AABBox.Min.Z) * Scale.Z;
+
+		Pts = FVector(X,Y,Z);
+	}
+}
+void FOutputRemapMesh::FMeshData::RemoveInvisibleFaces()
+{
+	int FacesNum = Index.Num() / 3;
+	for (int Face = 0; Face < FacesNum; Face++)
+	{
+		if (!IsVisibleFace(Face))
+		{
+			Index.RemoveAt(Face * 3, 3, true);
+			Face--;
+		}
+	}
+}
+bool FOutputRemapMesh::FMeshData::IsVisibleFace(int Face)
+{
+	int idx0 = Index[Face * 3 + 0];
+	int idx1 = Index[Face * 3 + 1];
+	int idx2 = Index[Face * 3 + 2];
+
+	return IsVisibleUV(idx0) && IsVisibleUV(idx1) && IsVisibleUV(idx2);
+}
+
+bool FOutputRemapMesh::FMeshData::IsVisibleUV(int UVIndex)
+{
+	if (UV[UVIndex].X < 0 || UV[UVIndex].X > 1)
+		return false;
+
+	if (UV[UVIndex].Y < 0 || UV[UVIndex].Y > 1)
+		return false;
+
+	return true;
+}
+
+void FOutputRemapMesh::ReleaseMesh()
+{
+	VertexBufferRHI.SafeRelease();
+	IndexBufferRHI.SafeRelease();
+	NumVertices = 0;
+	NumIndices = 0;
+	NumTriangles = 0;
+	bIsValid = false;
+	bIsExtFilesReloadEnabled = false;
+	bForceExtFilesReload = false;
+}
+
+void FOutputRemapMesh::CreateMesh()
+{
+
+	FMeshData* Mesh = new FMeshData();
+
+	if (FileName == "Passthrough")
+	{
+		//autogenerate test mesh:
+		bIsValid = Mesh->CreatePassthrough();
+	}
+	else
+	{
+		FString FullPathFileName = DisplayClusterHelpers::config::GetFullPath(FileName);
+		bIsValid = Mesh->CreateFromFile(FullPathFileName);
+
+		FileLastAccessDateTime = IFileManager::Get().GetAccessTimeStamp(*FullPathFileName);
+		bIsExtFilesReloadEnabled = true;
+	}
+
+	if (bIsValid)
+	{
+		// Normalize mesh for output screen 0..1
+		Mesh->Normalize();
+		// Remove faces with UV out of 0..1 range
+		Mesh->RemoveInvisibleFaces();
+		// Build runtime mesh in RHI
+		BuildMesh(Mesh); // Initialize mesh RHI, structure released latter on renderthread
+	}
+	else
+	{
+		delete Mesh;
+	}
+}
+
+void FOutputRemapMesh::ReloadMesh()
+{
+	if (bIsExtFilesReloadEnabled)
+	{
+		if (!bForceExtFilesReload)
+		{
+			// Detect auto reload by frames
+			const int FramePeriod = CVarOutputRemapReloadChangedExtFiles.GetValueOnAnyThread();
+			if (FramePeriod > 0)
+			{
+				if (VarOutputRemapReloadChangedExtFilesFrame++ > FramePeriod)
+				{
+					VarOutputRemapReloadChangedExtFilesFrame = 0;
+					bForceExtFilesReload = true; 
+				}
+			}
+		}
+
+		if(bForceExtFilesReload)
+		{				
+			// Check for file data is modified
+			FString FullPathFileName = DisplayClusterHelpers::config::GetFullPath(FileName);
+			FDateTime CurrentDateTime = IFileManager::Get().GetAccessTimeStamp(*FullPathFileName);
+			if (CurrentDateTime != FileLastAccessDateTime)
+			{
+				//! Add log. Reload file *FullPathFileName
+				ReleaseMesh();
+				CreateMesh();
+			}
+		}
+	}
+	bForceExtFilesReload = false;
+}
+
+
+FOutputRemapMesh::FOutputRemapMesh(const FString& InExtFileName)	
+	: bIsValid(false)
+	, FileName(InExtFileName)
+	, bIsExtFilesReloadEnabled(false)
+	, bForceExtFilesReload(false)
+	, NumVertices(0)
+	, NumIndices(0)
+	, NumTriangles(0)
+{
+	CreateMesh();
+}
+
+FOutputRemapMesh::~FOutputRemapMesh()
+{
+	ReleaseMesh();
+}
+
+static void BuildMesh_RenderThread(FOutputRemapMesh& Dst, const FOutputRemapMesh::FMeshData* Mesh)
+{
+	Dst.NumVertices = Mesh->Vertex.Num();
+	Dst.NumTriangles = Mesh->Index.Num() / 3;
+	Dst.NumIndices = Mesh->Index.Num();
+
+	check(Dst.NumVertices > 2 && Dst.NumTriangles > 0);
+
+
+	FRHIResourceCreateInfo CreateInfo;
+	Dst.VertexBufferRHI = RHICreateVertexBuffer(sizeof(FFilterVertex) * Dst.NumVertices, BUF_Static, CreateInfo);
+	void* VoidPtr = RHILockVertexBuffer(Dst.VertexBufferRHI, 0, sizeof(FFilterVertex) * Dst.NumVertices, RLM_WriteOnly);
+	FFilterVertex* pVertices = reinterpret_cast<FFilterVertex*>(VoidPtr);
+
+	Dst.IndexBufferRHI = RHICreateIndexBuffer(sizeof(uint16), sizeof(uint16) * Dst.NumIndices, BUF_Static, CreateInfo);
+	void* VoidPtr2 = RHILockIndexBuffer(Dst.IndexBufferRHI, 0, sizeof(uint16) * Dst.NumIndices, RLM_WriteOnly);
+	uint16* pIndices = reinterpret_cast<uint16*>(VoidPtr2);
+
+	//Copy Vertex
+	int DataIndex = 0;
+	for (auto& It : Mesh->Vertex)
+	{
+		FFilterVertex& Vertex = pVertices[DataIndex++];
+		// Remap from to NDC space [0 1] -> [-1 1]
+		Vertex.Position.X = It.X;
+		Vertex.Position.Y = 1-It.Z;
+		Vertex.Position.Z = 1.0f;
+		Vertex.Position.W = 1.0f;
+	}
+
+	//Copy UV
+	DataIndex = 0;
+	for (auto& It : Mesh->UV)
+	{
+		FFilterVertex& Vertex = pVertices[DataIndex++];		
+		Vertex.UV.X = It.X;
+		Vertex.UV.Y = It.Y;
+	}
+
+	//Copy faces Index
+	DataIndex = 0;
+	for (auto& It : Mesh->Index)
+	{
+		pIndices[DataIndex++] = (uint16)It;
+	}
+	
+	RHIUnlockVertexBuffer(Dst.VertexBufferRHI);
+	RHIUnlockIndexBuffer(Dst.IndexBufferRHI);
+
+	delete Mesh; // Now free tempropary mesh data
+}
+
+void FOutputRemapMesh::BuildMesh(const FMeshData* Mesh)
+{
+	if (Mesh)
+	{
+		if (IsInRenderingThread())
+		{
+			BuildMesh_RenderThread(*this, Mesh);
+		}
+		else
+		{
+			FOutputRemapMesh* Self = this;
+			ENQUEUE_RENDER_COMMAND(SetupOutputRemapMeshCmd)([Self, Mesh](FRHICommandListImmediate& RHICmdList)
+			{
+				BuildMesh_RenderThread(*Self, Mesh);
+			});
+		}
+	}
+}
diff --git a/Engine/Plugins/Runtime/nDisplay/Source/OutputRemap/Private/OutputRemapMesh.h b/Engine/Plugins/Runtime/nDisplay/Source/OutputRemap/Private/OutputRemapMesh.h
new file mode 100644
--- /dev/null
+++ b/Engine/Plugins/Runtime/nDisplay/Source/OutputRemap/Private/OutputRemapMesh.h
@@ -0,0 +1,82 @@
+// Copyright 1998-2019 Epic Games, Inc. All Rights Reserved.
+
+#pragma once
+
+#include "CoreMinimal.h"
+#include "RHI.h"
+#include "RHICommandList.h"
+
+class FOutputRemapMesh
+{
+public:
+	struct FMeshData
+	{
+		TArray<FVector> Vertex; 
+		TArray<FVector> UV;
+		TArray<int>     Index;
+
+		bool CreateFromFile(const FString& FullPathFileName);
+		bool CreatePassthrough();
+
+		void Normalize();
+		void RemoveInvisibleFaces();
+	private:
+		// Mesh post-op's
+		bool IsVisibleFace(int Face);
+		bool IsVisibleUV(int UVIndex);
+
+		// Obj file parser:
+		bool ParseLine(const FString& Line);
+		bool ExtractVertex(const FString& Line);
+		bool ExtractUV(const FString& Line);
+		bool ExtractFace(const FString& Line);
+		bool ExtractFaceVertex(const FString& Line);
+		
+		TArray<FVector> InVertex;
+		TArray<FVector> InUV;
+	};
+
+public:
+	FOutputRemapMesh(const FString& ExtFileName);
+	~FOutputRemapMesh();
+	
+	void DrawMesh(FRHICommandList& RHICmdList);
+
+	const FString& GetFileName() const 
+	{ 
+		return FileName; 
+	}
+	bool IsValid() const 
+	{ 
+		return bIsValid; 
+	}
+
+	void ReloadChangedExtFiles()
+	{
+		bForceExtFilesReload = true;
+	}
+
+private:
+	void BuildMesh(const FMeshData* Mesh);
+
+	void CreateMesh();
+	void ReleaseMesh();
+	void ReloadMesh();
+private:
+	bool       bIsValid;
+
+	FString    FileName;
+	FDateTime  FileLastAccessDateTime;
+	bool       bIsExtFilesReloadEnabled;
+	bool       bForceExtFilesReload;
+
+public:
+	FVertexBufferRHIRef VertexBufferRHI;
+	FIndexBufferRHIRef  IndexBufferRHI;
+
+	uint32  NumVertices;
+	uint32  NumIndices;
+	uint32  NumTriangles;
+};
+
+
diff --git a/Engine/Plugins/Runtime/nDisplay/Source/OutputRemap/Private/OutputRemapModule.cpp b/Engine/Plugins/Runtime/nDisplay/Source/OutputRemap/Private/OutputRemapModule.cpp
new file mode 100644
--- /dev/null
+++ b/Engine/Plugins/Runtime/nDisplay/Source/OutputRemap/Private/OutputRemapModule.cpp
@@ -0,0 +1,113 @@
+// Copyright 1998-2019 Epic Games, Inc. All Rights Reserved.
+
+#include "OutputRemapModule.h"
+#include "OutputRemapShader.h"
+
+#include "Interfaces/IPluginManager.h"
+#include "Misc/Paths.h"
+#include "ShaderCore.h"
+
+#include "RHICommandList.h"
+
+#include "OutputRemapHelpers.h"
+
+#include "OutputRemapLog.h"
+#include "OutputRemapMesh.h"
+
+
+
+//////////////////////////////////////////////////////////////////////////////////////////////
+// IModuleInterface
+//////////////////////////////////////////////////////////////////////////////////////////////
+#define NDISPLAY_SHADERS_MAP TEXT("/Plugin/nDisplay")
+
+void FOutputRemapModule::StartupModule()
+{
+	if (!AllShaderSourceDirectoryMappings().Contains(NDISPLAY_SHADERS_MAP))
+	{
+		FString PluginShaderDir = FPaths::Combine(IPluginManager::Get().FindPlugin(TEXT("nDisplay"))->GetBaseDir(), TEXT("Shaders"));
+		AddShaderSourceDirectoryMapping(NDISPLAY_SHADERS_MAP, PluginShaderDir);
+	}
+}
+
+void FOutputRemapModule::ShutdownModule()
+{
+	ReleaseData();
+}
+
+
+//////////////////////////////////////////////////////////////////////////////////////////////
+// IOutputRemap
+//////////////////////////////////////////////////////////////////////////////////////////////
+void FOutputRemapModule::ReloadAll()
+{
+	for (auto& It : MeshAssets)
+	{
+		It->ReloadChangedExtFiles();
+	}
+}
+
+bool FOutputRemapModule::Load(const FString& LocalRemapMeshFile, uint32& OutMeshRef)
+{
+	FScopeLock lock(&DataGuard);
+
+	if (IsLoaded(LocalRemapMeshFile, OutMeshRef))
+	{
+		return true;
+	}
+
+	TSharedPtr<FOutputRemapMesh> MeshItem(new FOutputRemapMesh(LocalRemapMeshFile));
+	if (!MeshItem->IsValid()){
+		return false;
+	}
+
+	// Load and cache the file
+	MeshAssets.Add(MeshItem);
+	OutMeshRef = MeshAssets.Num() - 1;
+	return true;
+}
+
+bool FOutputRemapModule::IsLoaded(const FString& LocalRemapMeshFile, uint32& OutMeshRef)
+{
+	FScopeLock lock(&DataGuard);
+
+	for (int32 Index = 0; Index != MeshAssets.Num(); ++Index)
+	{
+		if (!MeshAssets[Index]->GetFileName().Compare(LocalRemapMeshFile, ESearchCase::IgnoreCase))
+		{
+			OutMeshRef = Index;
+			return true;
+		}
+	}
+
+	return false;
+}
+
+bool FOutputRemapModule::ApplyOutputRemap_RenderThread(FRHICommandListImmediate& RHICmdList, FRHITexture2D* ShaderResourceTexture, FRHITexture2D* TargetableTexture, uint32 MeshRef)
+{
+	FScopeLock lock(&DataGuard);
+
+	if (MeshAssets.Num() > (int)MeshRef)
+	{
+		return FOutputRemapShader::ApplyOutputRemap_RenderThread(RHICmdList, ShaderResourceTexture, TargetableTexture, MeshAssets[MeshRef].Get());
+	}
+	else
+	{
+		//@todo: handle error
+	}
+
+	return false;
+}
+
+void FOutputRemapModule::ReleaseData()
+{
+	FScopeLock lock(&DataGuard);
+	for (auto& It : MeshAssets)
+	{
+		It.Reset();
+	}
+	MeshAssets.Empty();
+}
+
+IMPLEMENT_MODULE(FOutputRemapModule, OutputRemap);
+
diff --git a/Engine/Plugins/Runtime/nDisplay/Source/OutputRemap/Private/OutputRemapModule.h b/Engine/Plugins/Runtime/nDisplay/Source/OutputRemap/Private/OutputRemapModule.h
new file mode 100644
--- /dev/null
+++ b/Engine/Plugins/Runtime/nDisplay/Source/OutputRemap/Private/OutputRemapModule.h
@@ -0,0 +1,33 @@
+// Copyright 1998-2019 Epic Games, Inc. All Rights Reserved.
+
+#pragma once
+
+#include "IOutputRemap.h"
+
+class FOutputRemapMesh;
+
+class FOutputRemapModule : public IOutputRemap
+{
+public:
+	//////////////////////////////////////////////////////////////////////////////////////////////
+	// IModuleInterface
+	//////////////////////////////////////////////////////////////////////////////////////////////
+	virtual void StartupModule() override;
+	virtual void ShutdownModule() override;
+
+public:
+	//////////////////////////////////////////////////////////////////////////////////////////////
+	// OutputRemap
+	//////////////////////////////////////////////////////////////////////////////////////////////
+	virtual bool Load(const FString& RemapMeshFile, uint32& OutMeshRef) override;
+	virtual void ReloadAll() override;
+	virtual bool ApplyOutputRemap_RenderThread(FRHICommandListImmediate& RHICmdList, FRHITexture2D* ShaderResourceTexture, FRHITexture2D* TargetableTexture, uint32 MeshRef) override;
+
+private:
+	bool IsLoaded(const FString& RemapMeshFile, uint32& OutMeshRef);
+	void ReleaseData();	
+
+private:
+	FCriticalSection DataGuard;
+	TArray< TSharedPtr<FOutputRemapMesh>> MeshAssets;
+};
diff --git a/Engine/Plugins/Runtime/nDisplay/Source/OutputRemap/Private/OutputRemapShader.cpp b/Engine/Plugins/Runtime/nDisplay/Source/OutputRemap/Private/OutputRemapShader.cpp
new file mode 100644
--- /dev/null
+++ b/Engine/Plugins/Runtime/nDisplay/Source/OutputRemap/Private/OutputRemapShader.cpp
@@ -0,0 +1,187 @@
+// Copyright 1998-2019 Epic Games, Inc. All Rights Reserved.
+
+#include "OutputRemapShader.h"
+
+#include "CoreMinimal.h"
+#include "RHI.h"
+#include "RHIResources.h"
+
+#include "CommonRenderResources.h"
+#include "PixelShaderUtils.h"
+
+#include "ShaderParameterUtils.h"
+
+#include "HAL/IConsoleManager.h"
+
+#include "OutputRemapMesh.h"
+
+#define OutputRemapShaderFileName TEXT("/Plugin/nDisplay/Private/OutputRemapShaders.usf")
+
+
+// Select output remap shader
+enum class EVarOutputRemapShaderType : uint8
+{
+	Default,
+	Passthrough,
+	Disable,
+};
+
+static TAutoConsoleVariable<int32> CVarOutputRemapShaderType(
+	TEXT("nDisplay.render.output_remap.shader"),
+	(int)EVarOutputRemapShaderType::Default,
+	TEXT("Select shader for output remap:\n")	
+	TEXT(" 0: default remap shader\n")
+	TEXT(" 1: pass throught shader, test rect mesh\n")
+	TEXT(" 2: Disable remap shaders\n")
+	,ECVF_RenderThreadSafe
+);
+
+
+class FOutputRemapVS : public FGlobalShader
+{
+	DECLARE_SHADER_TYPE(FOutputRemapVS, Global);
+
+	static bool ShouldCompilePermutation(const FGlobalShaderPermutationParameters& Parameters)
+	{
+		return true;
+	}
+
+	/** Default constructor. */
+	FOutputRemapVS() 
+	{
+	}
+public:
+	/** Initialization constructor. */
+	FOutputRemapVS(const ShaderMetaType::CompiledShaderInitializerType& Initializer)
+		: FGlobalShader(Initializer)
+	{
+	}
+};
+
+class FOutputRemapPS : public FGlobalShader
+{
+	DECLARE_SHADER_TYPE(FOutputRemapPS, Global);
+
+public:
+	FOutputRemapPS()
+	{ 
+	}
+
+	FOutputRemapPS(const ShaderMetaType::CompiledShaderInitializerType& Initializer)
+		: FGlobalShader(Initializer)
+	{
+		PostprocessInputParameter0.Bind(Initializer.ParameterMap, TEXT("PostprocessInput0"));
+		PostprocessInputParameterSampler0.Bind(Initializer.ParameterMap, TEXT("PostprocessInput0Sampler"));
+	}
+
+public:
+	static bool ShouldCompilePermutation(const FGlobalShaderPermutationParameters& Parameters)
+	{
+		return IsFeatureLevelSupported(Parameters.Platform, ERHIFeatureLevel::SM4);
+	}
+
+	static void ModifyCompilationEnvironment(const FGlobalShaderPermutationParameters& Parameters, FShaderCompilerEnvironment& OutEnvironment)
+	{
+		FGlobalShader::ModifyCompilationEnvironment(Parameters, OutEnvironment);
+	}
+
+public:
+	template<typename TShaderRHIParamRef>
+	void SetParameters(FRHICommandListImmediate& RHICmdList, const TShaderRHIParamRef ShaderRHI, FRHITexture2D* SourceTexture)
+	{
+		SetTextureParameter(RHICmdList, ShaderRHI, PostprocessInputParameter0, SourceTexture);
+		RHICmdList.SetShaderSampler(ShaderRHI, PostprocessInputParameterSampler0.GetBaseIndex(), TStaticSamplerState<SF_Trilinear, AM_Clamp, AM_Clamp, AM_Clamp>::GetRHI());
+	}
+
+	virtual bool Serialize(FArchive& Ar) override
+	{
+		bool bShaderHasOutdatedParameters = FGlobalShader::Serialize(Ar);
+		Ar
+			<< PostprocessInputParameter0
+			<< PostprocessInputParameterSampler0
+			;
+
+		return bShaderHasOutdatedParameters;
+	}
+
+private:
+	FShaderResourceParameter PostprocessInputParameter0;
+	FShaderResourceParameter PostprocessInputParameterSampler0;
+};
+
+// Implement shaders inside UE4
+IMPLEMENT_SHADER_TYPE(, FOutputRemapVS, OutputRemapShaderFileName, TEXT("OutputRemap_VS"), SF_Vertex);
+IMPLEMENT_SHADER_TYPE(, FOutputRemapPS, OutputRemapShaderFileName, TEXT("OutputRemap_PS"), SF_Pixel);
+
+bool FOutputRemapShader::ApplyOutputRemap_RenderThread(FRHICommandListImmediate& RHICmdList, FRHITexture2D* ShaderResourceTexture, FRHITexture2D* TargetableTexture, FOutputRemapMesh* MeshData)
+{
+	check(IsInRenderingThread());
+
+	if (MeshData == nullptr)
+	{
+		//Handle error
+		return false;
+	}
+
+
+	const EVarOutputRemapShaderType ShaderType = (EVarOutputRemapShaderType)CVarOutputRemapShaderType.GetValueOnAnyThread();
+	switch (ShaderType)
+	{
+	case EVarOutputRemapShaderType::Passthrough:
+	{
+		// Use simple 1:1 test mesh for shader forwarding
+		static FOutputRemapMesh TestMesh("Passthrough");
+		MeshData = &TestMesh;
+	}
+	case EVarOutputRemapShaderType::Default:
+		break;
+	case EVarOutputRemapShaderType::Disable:
+		return false;
+
+	default:
+		return false;
+	};
+
+
+	{// Do remap single render pass		
+		FRHIRenderPassInfo RPInfo(TargetableTexture, ERenderTargetActions::Load_Store);
+		RHICmdList.BeginRenderPass(RPInfo, TEXT("DisplayClusterOutputRemapShader"));
+		{
+			{
+				FIntRect DstRect(FIntPoint(0,0), TargetableTexture->GetSizeXY());
+				//Clear viewport before render
+				RHICmdList.SetViewport(DstRect.Min.X, DstRect.Min.Y, 0.0f, DstRect.Max.X, DstRect.Max.Y, 1.0f);
+				//DrawClearQuad(RHICmdList, FLinearColor::Black);
+			}
+
+			// Set the graphic pipeline state.
+			FGraphicsPipelineStateInitializer GraphicsPSOInit;
+			RHICmdList.ApplyCachedRenderTargets(GraphicsPSOInit);
+
+			GraphicsPSOInit.DepthStencilState = TStaticDepthStencilState<false, CF_Never>::GetRHI();
+			GraphicsPSOInit.BlendState = TStaticBlendState <>::GetRHI();
+			GraphicsPSOInit.RasterizerState = TStaticRasterizerState<>::GetRHI();
+
+			TShaderMap<FGlobalShaderType>* ShaderMap = GetGlobalShaderMap(GMaxRHIFeatureLevel);
+			TShaderMapRef<FOutputRemapVS> VertexShader(ShaderMap);
+			TShaderMapRef<FOutputRemapPS> PixelShader(ShaderMap);
+
+			GraphicsPSOInit.BoundShaderState.VertexDeclarationRHI = GFilterVertexDeclaration.VertexDeclarationRHI;// GetVertexDeclarationFVector4();
+			GraphicsPSOInit.BoundShaderState.VertexShaderRHI = GETSAFERHISHADER_VERTEX(*VertexShader);
+			GraphicsPSOInit.BoundShaderState.PixelShaderRHI = GETSAFERHISHADER_PIXEL(*PixelShader);
+			GraphicsPSOInit.PrimitiveType = PT_TriangleList;
+
+
+			SetGraphicsPipelineState(RHICmdList, GraphicsPSOInit);
+
+			PixelShader->SetParameters(RHICmdList, PixelShader->GetPixelShader(), ShaderResourceTexture);
+			MeshData->DrawMesh(RHICmdList);
+		}
+		RHICmdList.EndRenderPass();
+	}
+
+	// Render pass failed, handle error
+	return true;
+}
+
+
diff --git a/Engine/Plugins/Runtime/nDisplay/Source/OutputRemap/Private/OutputRemapShader.h b/Engine/Plugins/Runtime/nDisplay/Source/OutputRemap/Private/OutputRemapShader.h
new file mode 100644
--- /dev/null
+++ b/Engine/Plugins/Runtime/nDisplay/Source/OutputRemap/Private/OutputRemapShader.h
@@ -0,0 +1,21 @@
+// Copyright 1998-2019 Epic Games, Inc. All Rights Reserved.
+
+#pragma once
+
+#include "CoreMinimal.h"
+#include "RHI.h"
+#include "RHICommandList.h"
+#include "RenderResource.h"
+#include "ShaderParameters.h"
+#include "Shader.h"
+#include "GlobalShader.h"
+#include "ShaderParameterUtils.h"
+
+class FOutputRemapMesh;
+
+
+class FOutputRemapShader
+{
+public:
+	static bool ApplyOutputRemap_RenderThread(FRHICommandListImmediate& RHICmdList, FRHITexture2D* ShaderResourceTexture, FRHITexture2D* TargetableTexture, FOutputRemapMesh* MeshData);
+};
diff --git a/Engine/Plugins/Runtime/nDisplay/Source/OutputRemap/Public/Blueprints/IOutputRemapBlueprintAPI.h b/Engine/Plugins/Runtime/nDisplay/Source/OutputRemap/Public/Blueprints/IOutputRemapBlueprintAPI.h
new file mode 100644
--- /dev/null
+++ b/Engine/Plugins/Runtime/nDisplay/Source/OutputRemap/Public/Blueprints/IOutputRemapBlueprintAPI.h
@@ -0,0 +1,26 @@
+// Copyright 1998-2019 Epic Games, Inc. All Rights Reserved.
+
+#pragma once
+
+#include "CoreMinimal.h"
+#include "UObject/Interface.h"
+#include "IOutputRemapBlueprintAPI.generated.h"
+
+UINTERFACE(meta = (CannotImplementInterfaceInBlueprint))
+class UOutputRemapBlueprintAPI : public UInterface
+{
+	GENERATED_BODY()
+};
+
+class IOutputRemapBlueprintAPI
+{
+	GENERATED_BODY()
+
+public:
+	/**
+	* Reload all changed external remap files, and apply changes runtime
+	*
+	*/
+	UFUNCTION(BlueprintCallable, meta = (DisplayName = "Reload Changed External Files"), Category = "OutputRemap")
+	virtual void ReloadChangedExternalFiles() = 0;
+};
diff --git a/Engine/Plugins/Runtime/nDisplay/Source/OutputRemap/Public/Blueprints/OutputRemapBlueprintLib.h b/Engine/Plugins/Runtime/nDisplay/Source/OutputRemap/Public/Blueprints/OutputRemapBlueprintLib.h
new file mode 100644
--- /dev/null
+++ b/Engine/Plugins/Runtime/nDisplay/Source/OutputRemap/Public/Blueprints/OutputRemapBlueprintLib.h
@@ -0,0 +1,25 @@
+// Copyright 1998-2019 Epic Games, Inc. All Rights Reserved.
+
+#pragma once
+
+#include "CoreMinimal.h"
+#include "Kismet/BlueprintFunctionLibrary.h"
+#include "Blueprints/IOutputRemapBlueprintAPI.h"
+
+#include "OutputRemapBlueprintLib.generated.h"
+
+
+/**
+ * Blueprint API function library
+ */
+UCLASS()
+class UOutputRemapIBlueprintLib
+	: public UBlueprintFunctionLibrary
+{
+	GENERATED_UCLASS_BODY()
+
+public:
+	/** Return Display Cluster API interface. */
+	UFUNCTION(BlueprintPure, meta = (DisplayName = "OutputRemap Module API"), Category = "nDisplay")
+	static void GetAPI(TScriptInterface<IOutputRemapBlueprintAPI>& OutAPI);
+};
diff --git a/Engine/Plugins/Runtime/nDisplay/Source/OutputRemap/Public/IOutputRemap.h b/Engine/Plugins/Runtime/nDisplay/Source/OutputRemap/Public/IOutputRemap.h
new file mode 100644
--- /dev/null
+++ b/Engine/Plugins/Runtime/nDisplay/Source/OutputRemap/Public/IOutputRemap.h
@@ -0,0 +1,71 @@
+// Copyright 1998-2019 Epic Games, Inc. All Rights Reserved.
+
+#pragma once
+
+#include "CoreMinimal.h"
+
+#include "Modules/ModuleInterface.h"
+#include "Modules/ModuleManager.h"
+
+#include "RHI.h"
+#include "RHICommandList.h"
+#include "RHIResources.h"
+#include "RHIUtilities.h"
+
+class FMPCDIData;
+
+class IOutputRemap : public IModuleInterface
+{
+public:
+	static constexpr auto ModuleName = TEXT("OutputRemap");
+public:
+	/**
+	 * Singleton-like access to this module's interface.  This is just for convenience!
+	 * Beware of calling this during the shutdown phase, though.  Your module might have been unloaded already.
+	 *
+	 * @return Returns singleton instance, loading the module on demand if needed
+	 */
+
+	static inline IOutputRemap& Get()
+	{
+		return FModuleManager::LoadModuleChecked<IOutputRemap>(IOutputRemap::ModuleName);
+	}
+
+	/**
+	 * Checks to see if this module is loaded and ready.  It is only valid to call Get() if IsAvailable() returns true.
+	 *
+	 * @return True if the module is loaded and ready to use
+	 */
+	static inline bool IsAvailable()
+	{
+		return FModuleManager::Get().IsModuleLoaded(IOutputRemap::ModuleName);
+	}
+
+public:
+	/**
+	* Loads an external remap mesh file
+	*
+	* @param RemapMeshFile - file path (.obj fmt)
+	* @param OutMeshHandler - handler to loaded data
+	*
+	* @return - true if success
+	*/
+	virtual bool Load(const FString& RemapMeshFile, uint32& OutMeshRef) = 0;
+
+	/**
+	* Reload all changed external mesh files
+	*/
+	virtual void ReloadAll() = 0;
+
+
+	/**
+	* Applies Output Remap
+	*
+	* @param RHICmdList      - RHI command list
+	* @param TextureWarpData - texture warp data
+	* @param ShaderInputData - shader data
+	*
+	* @return - true if success
+	*/
+	virtual bool ApplyOutputRemap_RenderThread(FRHICommandListImmediate& RHICmdList, FRHITexture2D* ShaderResourceTexture, FRHITexture2D* TargetableTexture, uint32 MeshRef) = 0;
+};
diff --git a/Engine/Plugins/Runtime/nDisplay/Source/PicpMPCDI/PicpMPCDI.Build.cs b/Engine/Plugins/Runtime/nDisplay/Source/PicpMPCDI/PicpMPCDI.Build.cs
--- a/Engine/Plugins/Runtime/nDisplay/Source/PicpMPCDI/PicpMPCDI.Build.cs
+++ b/Engine/Plugins/Runtime/nDisplay/Source/PicpMPCDI/PicpMPCDI.Build.cs
@@ -10,11 +10,7 @@ public class PicpMPCDI : ModuleRules
 			new string[] {
 				"PicpMPCDI/Private",
 				"PicpProjection/Private",
-				"MPCDI/Private",
-				"../../../../../Engine/Source/Runtime/Windows/D3D11RHI/Private",
-				"../../../../../Engine/Source/Runtime/Windows/D3D11RHI/Private/Windows",
-				"../../../../../Engine/Source/Runtime/D3D12RHI/Private",
-				"../../../../../Engine/Source/Runtime/D3D12RHI/Private/Windows"
+				"MPCDI/Private"				
 			});
 
 		PublicDependencyModuleNames.AddRange(
@@ -23,21 +19,16 @@ public class PicpMPCDI : ModuleRules
 				"Core",
 				"CoreUObject",
 				"Engine",
-				"Projects",
-				"DisplayCluster",
-				"PicpProjection",
-				"MPCDI",
 			});
 
 		PrivateDependencyModuleNames.AddRange(
 			new string[]
 			{
-				"Core",
-				"CoreUObject",
-				"Engine",
+				"MPCDI",
+				"PicpProjection",
+				"Projects",				
 				"RenderCore",
-				"DisplayCluster",
-				"RHI",
+                "RHI",
 				"UtilityShaders"
 			}
 		);
@@ -46,10 +37,5 @@ public class PicpMPCDI : ModuleRules
 		{
 			PrivateDependencyModuleNames.Add("UnrealEd");
 		}
-
-		PublicAdditionalLibraries.Add("opengl32.lib");
-		AddEngineThirdPartyPrivateStaticDependencies(Target, "OpenGL");
-		AddEngineThirdPartyPrivateStaticDependencies(Target, "DX11");
-		AddEngineThirdPartyPrivateStaticDependencies(Target, "DX12");
 	}
 }
diff --git a/Engine/Plugins/Runtime/nDisplay/Source/PicpMPCDI/Private/PicpBlurPostProcess.cpp b/Engine/Plugins/Runtime/nDisplay/Source/PicpMPCDI/Private/PicpBlurPostProcess.cpp
--- a/Engine/Plugins/Runtime/nDisplay/Source/PicpMPCDI/Private/PicpBlurPostProcess.cpp
+++ b/Engine/Plugins/Runtime/nDisplay/Source/PicpMPCDI/Private/PicpBlurPostProcess.cpp
@@ -17,7 +17,9 @@
 static TAutoConsoleVariable<int32> CVarPicpPostProcess(
 	TEXT("nDisplay.render.picp.EnablePostProcess"),
 	1,
-	TEXT("Enable postprocess shaders for picp (0 = disabled)"),
+	TEXT("Enable postprocess shaders for picp\n")
+	TEXT(" 0: disabled\n")
+	TEXT(" 1: enabled\n"),
 	ECVF_RenderThreadSafe
 );
 
@@ -31,6 +33,11 @@ IMPLEMENT_SHADER_TYPE(template<>, FPicpBlurPostProcessDilatePS,  PostProcessShad
 
 IMPLEMENT_SHADER_TYPE(, FDirectComposePS, PostProcessShaderFileName, TEXT("DirectComposePS"), SF_Pixel);
 
+
+FTextureResource* FPicpBlurPostProcess::SrcTexture = nullptr;
+FTextureRenderTargetResource* FPicpBlurPostProcess::DstTexture = nullptr;
+FTextureRenderTargetResource* FPicpBlurPostProcess::ResultTexture = nullptr;
+
 void FPicpBlurPostProcess::ApplyCompose_RenderThread(
 	FRHICommandListImmediate& RHICmdList, 
 	FRHITexture* OverlayTexture,
@@ -179,24 +186,33 @@ void FPicpBlurPostProcess::ApplyCompose(UTexture* InputTexture, UTextureRenderTa
 		return;
 	}
 
-	FTextureResource* OverlayTextureResource = InputTexture ? InputTexture->Resource : nullptr;
-	FTextureRenderTargetResource* TextureRenderTarget = OutputRenderTarget->GameThread_GetRenderTargetResource();
-	FTextureRenderTargetResource* TextureRenderResult = Result->GameThread_GetRenderTargetResource();
+	SrcTexture = InputTexture ? InputTexture->Resource : nullptr;
+	DstTexture = OutputRenderTarget ? OutputRenderTarget->GameThread_GetRenderTargetResource() : nullptr;
+	ResultTexture = Result ? Result->GameThread_GetRenderTargetResource() : nullptr;
+}
 
-	ENQUEUE_RENDER_COMMAND(CaptureCommand)([OverlayTextureResource, TextureRenderTarget, TextureRenderResult](FRHICommandListImmediate& RHICmdList)
+
+void FPicpBlurPostProcess::ExecuteCompose()
+{
+	if (FPicpBlurPostProcess::DstTexture == nullptr || FPicpBlurPostProcess::ResultTexture == nullptr || FPicpBlurPostProcess::SrcTexture == nullptr)
 	{
-		FRHITexture2D* RTTexture = TextureRenderTarget ? TextureRenderTarget->GetRenderTargetTexture()->GetTexture2D(): nullptr;
-		FRHITexture2D* ResultTexture = TextureRenderResult ? TextureRenderResult->GetRenderTargetTexture()->GetTexture2D() : nullptr;
-		FRHITexture*   OverlayTexture = OverlayTextureResource ? OverlayTextureResource->TextureRHI : nullptr;;
-		
-		if (nullptr == RTTexture)
+		return;
+	}
+
+	FRHITexture*   Src = FPicpBlurPostProcess::SrcTexture ? FPicpBlurPostProcess::SrcTexture->TextureRHI : nullptr;
+	FRHITexture2D* Dst = FPicpBlurPostProcess::DstTexture ? FPicpBlurPostProcess::DstTexture->GetRenderTargetTexture()->GetTexture2D() : nullptr;
+	FRHITexture2D* Res = FPicpBlurPostProcess::ResultTexture ? FPicpBlurPostProcess::ResultTexture->GetRenderTargetTexture()->GetTexture2D() : nullptr;
+	
+
+	ENQUEUE_RENDER_COMMAND(CaptureCommand)([Src, Dst, Res](FRHICommandListImmediate& RHICmdList)
+	{		
+		if (nullptr == Dst)
 		{
 			//@todo handle error
 			return;
 		}
 
-		ApplyCompose_RenderThread(RHICmdList, OverlayTexture, RTTexture, ResultTexture);
+		ApplyCompose_RenderThread(RHICmdList, Src, Dst, Res);
 	});
 }
 
-
diff --git a/Engine/Plugins/Runtime/nDisplay/Source/PicpMPCDI/Private/PicpBlurPostProcess.h b/Engine/Plugins/Runtime/nDisplay/Source/PicpMPCDI/Private/PicpBlurPostProcess.h
--- a/Engine/Plugins/Runtime/nDisplay/Source/PicpMPCDI/Private/PicpBlurPostProcess.h
+++ b/Engine/Plugins/Runtime/nDisplay/Source/PicpMPCDI/Private/PicpBlurPostProcess.h
@@ -22,10 +22,16 @@ public:
 	//on Gamethread
 	static void ApplyBlur(UTextureRenderTarget2D* InOutRenderTarget, UTextureRenderTarget2D* TemporaryRenderTarget, int KernelRadius, float KernelScale, EPicpBlurPostProcessShaderType BlurType);
 	static void ApplyCompose(UTexture* InputTexture, UTextureRenderTarget2D* OutputRenderTarget, UTextureRenderTarget2D* Result);
+	static void ExecuteCompose();
 
 	//onRenderThread
 	static void ApplyBlur_RenderThread(FRHICommandListImmediate& RHICmdList, FRHITexture2D* InOutRT, FRHITexture2D* TempRT, int KernelRadius, float KernelScale, EPicpBlurPostProcessShaderType BlurType);
 	static void ApplyCompose_RenderThread(FRHICommandListImmediate& RHICmdList, FRHITexture* OverlayTexture, FRHITexture2D* DstRenderTarget, FRHITexture2D* DstTexture);
+
+	// blend texture references
+	static FTextureResource* SrcTexture;
+	static FTextureRenderTargetResource* DstTexture;
+	static FTextureRenderTargetResource* ResultTexture;
 };
 
 class FDirectProjectionVS : public FGlobalShader
@@ -38,7 +44,9 @@ class FDirectProjectionVS : public FGlobalShader
 	}
 
 	/** Default constructor. */
-	FDirectProjectionVS() {}
+	FDirectProjectionVS() 
+	{
+	}
 
 
 public:
@@ -62,7 +70,9 @@ class TPicpBlurPostProcessPS : public FGlobalShader
 	}
 
 	/** Default constructor. */
-	TPicpBlurPostProcessPS() {}
+	TPicpBlurPostProcessPS() 
+	{
+	}
 
 	FShaderResourceParameter SrcTextureParameter;
 	FShaderResourceParameter BilinearClampTextureSamplerParameter;
@@ -139,7 +149,9 @@ class FDirectComposePS : public FGlobalShader
 	}
 
 	/** Default constructor. */
-	FDirectComposePS() {}
+	FDirectComposePS() 
+	{
+	}
 
 	FShaderResourceParameter SrcTextureParameter;
 	FShaderResourceParameter BilinearClampTextureSamplerParameter;
diff --git a/Engine/Plugins/Runtime/nDisplay/Source/PicpMPCDI/Private/PicpMPCDIModule.cpp b/Engine/Plugins/Runtime/nDisplay/Source/PicpMPCDI/Private/PicpMPCDIModule.cpp
--- a/Engine/Plugins/Runtime/nDisplay/Source/PicpMPCDI/Private/PicpMPCDIModule.cpp
+++ b/Engine/Plugins/Runtime/nDisplay/Source/PicpMPCDI/Private/PicpMPCDIModule.cpp
@@ -12,10 +12,16 @@
 //////////////////////////////////////////////////////////////////////////////////////////////
 // IModuleInterface
 //////////////////////////////////////////////////////////////////////////////////////////////
+
+#define NDISPLAY_SHADERS_MAP TEXT("/Plugin/nDisplay")
+
 void FPicpMPCDIModule::StartupModule()
 {
-	//FString PluginShaderDir = FPaths::Combine(IPluginManager::Get().FindPlugin(TEXT("nDisplay"))->GetBaseDir(), TEXT("Shaders"));
-	//AddShaderSourceDirectoryMapping(TEXT("/Plugin/nDisplay"), PluginShaderDir);
+	if (!AllShaderSourceDirectoryMappings().Contains(NDISPLAY_SHADERS_MAP))
+	{
+		FString PluginShaderDir = FPaths::Combine(IPluginManager::Get().FindPlugin(TEXT("nDisplay"))->GetBaseDir(), TEXT("Shaders"));
+		AddShaderSourceDirectoryMapping(NDISPLAY_SHADERS_MAP, PluginShaderDir);
+	}
 }
 
 void FPicpMPCDIModule::ShutdownModule()
@@ -43,5 +49,10 @@ void FPicpMPCDIModule::ApplyCompose(UTexture* InputTexture, UTextureRenderTarget
 	FPicpBlurPostProcess::ApplyCompose(InputTexture, OutputRenderTarget, Result);
 }
 
+void FPicpMPCDIModule::ExecuteCompose()
+{
+	FPicpBlurPostProcess::ExecuteCompose();
+}
+
 IMPLEMENT_MODULE(FPicpMPCDIModule, PicpMPCDI);
 
diff --git a/Engine/Plugins/Runtime/nDisplay/Source/PicpMPCDI/Private/PicpMPCDIModule.h b/Engine/Plugins/Runtime/nDisplay/Source/PicpMPCDI/Private/PicpMPCDIModule.h
--- a/Engine/Plugins/Runtime/nDisplay/Source/PicpMPCDI/Private/PicpMPCDIModule.h
+++ b/Engine/Plugins/Runtime/nDisplay/Source/PicpMPCDI/Private/PicpMPCDIModule.h
@@ -24,6 +24,7 @@ public:
 	
 	virtual void ApplyBlur(UTextureRenderTarget2D* InOutRenderTarget, UTextureRenderTarget2D* TemporaryRenderTarget, int KernelRadius, float KernelScale, EPicpBlurPostProcessShaderType BlurType) override;
 	virtual void ApplyCompose(UTexture* InputTexture, UTextureRenderTarget2D* OutputRenderTarget, UTextureRenderTarget2D* Result) override;
+	virtual void ExecuteCompose() override;
 
 private:
 	FCriticalSection DataGuard;
diff --git a/Engine/Plugins/Runtime/nDisplay/Source/PicpMPCDI/Private/PicpMPCDIShader.cpp b/Engine/Plugins/Runtime/nDisplay/Source/PicpMPCDI/Private/PicpMPCDIShader.cpp
--- a/Engine/Plugins/Runtime/nDisplay/Source/PicpMPCDI/Private/PicpMPCDIShader.cpp
+++ b/Engine/Plugins/Runtime/nDisplay/Source/PicpMPCDI/Private/PicpMPCDIShader.cpp
@@ -14,12 +14,24 @@
 #include "MPCDIData.h"
 #include "MPCDIShader.h"
 
-
-
-static TAutoConsoleVariable<int32> CVarPicpEnableSinglePassRender(
-	TEXT("nDisplay.render.picp.EnableSinglePassRender"),
-	1,
-	TEXT("Single render pass for PICP (0 = disabled)"),
+// Select picp shader
+enum class EVarPicpMPCDIShaderType : uint8
+{
+	Default=0,
+	DefaultNoBlend,
+	Passthrough,
+	ForceMultiPassRender,
+	Disable,
+};
+static TAutoConsoleVariable<int32> CVarPicpMPCDIShaderType(
+	TEXT("nDisplay.render.picp.shader"),
+	(int)EVarPicpMPCDIShaderType::Default,
+	TEXT("Select shader for picp:\n")
+	TEXT(" 0: Warp shader (used by default)\n")
+	TEXT(" 1: Warp shader with disabled blend maps\n")
+	TEXT(" 2: Passthrough shader\n")
+	TEXT(" 3: Force Multi-Pass render shaders\n")
+	TEXT(" 4: Disable all shaders\n"),
 	ECVF_RenderThreadSafe
 );
 
@@ -29,11 +41,9 @@ static TAutoConsoleVariable<int32> CVarPicpEnableSinglePassRender(
 // Implement shaders inside UE4
 IMPLEMENT_SHADER_TYPE(, FPicpWarpVS, MPCDIShaderFileName, TEXT("MainVS"), SF_Vertex);
 
-IMPLEMENT_SHADER_TYPE(template<>, FPicpMPCDIPassThroughtPS, MPCDIShaderFileName, TEXT("PassThrought_PS"), SF_Pixel);
+IMPLEMENT_SHADER_TYPE(template<>, FPicpMPCDIPassthroughPS,  MPCDIShaderFileName, TEXT("Passthrough_PS"),  SF_Pixel);
 IMPLEMENT_SHADER_TYPE(template<>, FPicpMPCDIWarpAndBlendPS, MPCDIShaderFileName, TEXT("WarpAndBlend_PS"), SF_Pixel);
 
-//Control code by MACRO:
-//!fixme shader code customization other methods?
 IMPLEMENT_SHADER_TYPE(template<>, FPicpMPCDIWarpAndBlendBgPS,            MPCDIOverlayShaderFileName, TEXT("PicpMPCDIWarpAndBlendOverlay_PS"), SF_Pixel);
 IMPLEMENT_SHADER_TYPE(template<>, FPicpMPCDIWarpAndBlendBgOverlayPS,     MPCDIOverlayShaderFileName, TEXT("PicpMPCDIWarpAndBlendOverlay_PS"), SF_Pixel);
 
@@ -47,6 +57,21 @@ IMPLEMENT_SHADER_TYPE(template<>, FPicpWarpAndBlendOneCameraBackPS,      MPCDIOv
 IMPLEMENT_SHADER_TYPE(template<>, FPicpWarpAndBlendOneCameraFrontPS,     MPCDIOverlayShaderFileName, TEXT("PicpMPCDIWarpAndBlendOneCameraFull_PS"), SF_Pixel);
 IMPLEMENT_SHADER_TYPE(template<>, FPicpWarpAndBlendOneCameraNonePS,      MPCDIOverlayShaderFileName, TEXT("PicpMPCDIWarpAndBlendOneCameraFull_PS"), SF_Pixel);
 
+// No Blend case:
+IMPLEMENT_SHADER_TYPE(template<>, FPicpMPCDIWarpPS,          MPCDIShaderFileName, TEXT("WarpAndBlend_PS"), SF_Pixel);
+
+IMPLEMENT_SHADER_TYPE(template<>, FPicpMPCDIWarpBgPS,        MPCDIOverlayShaderFileName, TEXT("PicpMPCDIWarpAndBlendOverlay_PS"), SF_Pixel);
+IMPLEMENT_SHADER_TYPE(template<>, FPicpMPCDIWarpBgOverlayPS, MPCDIOverlayShaderFileName, TEXT("PicpMPCDIWarpAndBlendOverlay_PS"), SF_Pixel);
+
+IMPLEMENT_SHADER_TYPE(template<>, FPicpMPCDIWarpCameraPS,    MPCDIOverlayShaderFileName, TEXT("PicpMPCDIWarpAndBlendOverlay_PS"), SF_Pixel);
+IMPLEMENT_SHADER_TYPE(template<>, FPicpMPCDIWarpCameraOverlayPS, MPCDIOverlayShaderFileName, TEXT("PicpMPCDIWarpAndBlendOverlay_PS"), SF_Pixel);
+
+IMPLEMENT_SHADER_TYPE(template<>, FPicpMPCDIWarpOverlayPS,   MPCDIOverlayShaderFileName, TEXT("PicpMPCDIWarpAndBlendViewportOverlay_PS"), SF_Pixel);
+
+IMPLEMENT_SHADER_TYPE(template<>, FPicpWarpOneCameraFullPS,  MPCDIOverlayShaderFileName, TEXT("PicpMPCDIWarpAndBlendOneCameraFull_PS"), SF_Pixel);
+IMPLEMENT_SHADER_TYPE(template<>, FPicpWarpOneCameraBackPS,  MPCDIOverlayShaderFileName, TEXT("PicpMPCDIWarpAndBlendOneCameraFull_PS"), SF_Pixel);
+IMPLEMENT_SHADER_TYPE(template<>, FPicpWarpOneCameraFrontPS, MPCDIOverlayShaderFileName, TEXT("PicpMPCDIWarpAndBlendOneCameraFull_PS"), SF_Pixel);
+IMPLEMENT_SHADER_TYPE(template<>, FPicpWarpOneCameraNonePS,  MPCDIOverlayShaderFileName, TEXT("PicpMPCDIWarpAndBlendOneCameraFull_PS"), SF_Pixel);
 
 
 
@@ -96,14 +121,14 @@ static EPicpShaderType GetPixelShaderType(FMPCDIData *MPCDIData)
 	{
 		switch (MPCDIData->GetProfileType())
 		{
-		case EMPCDIProfileType::mpcdi_A3D:
+		case IMPCDI::EMPCDIProfileType::mpcdi_A3D:
 			return EPicpShaderType::WarpAndBlend;
 
-		case EMPCDIProfileType::Invalid:
-		case EMPCDIProfileType::mpcdi_SL:
-		case EMPCDIProfileType::mpcdi_2D:
-		case EMPCDIProfileType::mpcdi_3D:
-			return EPicpShaderType::PassThrought;
+		case IMPCDI::EMPCDIProfileType::Invalid:
+		case IMPCDI::EMPCDIProfileType::mpcdi_SL:
+		case IMPCDI::EMPCDIProfileType::mpcdi_2D:
+		case IMPCDI::EMPCDIProfileType::mpcdi_3D:
+			return EPicpShaderType::Passthrough;
 		};
 	}
 
@@ -139,7 +164,6 @@ static bool CompleteWarpTempl(
 		GetViewportRect(TextureWarpData, MPCDIRegion, MPCDIViewportRect);
 
 		RHICmdList.SetViewport(MPCDIViewportRect.Min.X, MPCDIViewportRect.Min.Y, 0.0f, MPCDIViewportRect.Max.X, MPCDIViewportRect.Max.Y, 1.0f);
-		//DrawClearQuad(RHICmdList, FLinearColor::Black);
 	}
 
 
@@ -170,7 +194,7 @@ static bool CompleteWarpTempl(
 
 	if (PicpProjectionOverlayViewportData != nullptr && CameraIndex >= 0 && PicpProjectionOverlayViewportData->Cameras.Num() > CameraIndex)
 	{
-		FPicpProjectionOverlayCamera* Cam = PicpProjectionOverlayViewportData->Cameras[CameraIndex];
+		FPicpProjectionOverlayCamera Camera = PicpProjectionOverlayViewportData->Cameras[CameraIndex];
 
 		static const FMatrix Game2Render(
 			FPlane(0, 0, 1, 0),
@@ -180,42 +204,49 @@ static bool CompleteWarpTempl(
 		static const FMatrix Render2Game =
 			Game2Render.Inverse();
 
-		FRotationTranslationMatrix CameraTranslationMatrix(Cam->ViewRot, Cam->ViewLoc);
+		FRotationTranslationMatrix CameraTranslationMatrix(Camera.ViewRot, Camera.ViewLoc);
 
 		FMatrix WorldToCamera = CameraTranslationMatrix.Inverse();
 
-		FMatrix UVMatrix = WorldToCamera * Game2Render * Cam->Prj;
+		FMatrix UVMatrix = WorldToCamera * Game2Render * Camera.Prj;
 
-		Cam->RuntimeCameraProjection = UVMatrix * MPCDIData->GetTextureMatrix() *MPCDIRegion.GetRegionMatrix();
+		Camera.RuntimeCameraProjection = UVMatrix * MPCDIData->GetTextureMatrix() *MPCDIRegion.GetRegionMatrix();
 
-		PixelShader->SetPicpParameters(RHICmdList, PixelShader->GetPixelShader(), Cam);
+		PixelShader->SetPicpParameters(RHICmdList, PixelShader->GetPixelShader(), &Camera);
 	}
 
 	switch (PixelShaderType)
 	{
 	case EPicpShaderType::WarpAndBlendAddOverlay:
+	case EPicpShaderType::WarpAddOverlay:
 		if (PicpProjectionOverlayViewportData != nullptr)
 		{
-			PixelShader->SetPicpParameters(RHICmdList, PixelShader->GetPixelShader(), PicpProjectionOverlayViewportData->ViewportOver, nullptr);
+			PixelShader->SetPicpParameters(RHICmdList, PixelShader->GetPixelShader(), &PicpProjectionOverlayViewportData->ViewportOver, nullptr);
 		}
 		break;
 	case EPicpShaderType::WarpAndBlendAddCameraOverlay:
+	case EPicpShaderType::WarpAddCameraOverlay:
 		if(PicpProjectionOverlayViewportData!=nullptr)
 		{
-			PixelShader->SetPicpParameters(RHICmdList, PixelShader->GetPixelShader(), PicpProjectionOverlayViewportData->ViewportOver, nullptr);
+			PixelShader->SetPicpParameters(RHICmdList, PixelShader->GetPixelShader(), &PicpProjectionOverlayViewportData->ViewportOver, nullptr);
 		}
 	case EPicpShaderType::WarpAndBlendBgOverlay:
+	case EPicpShaderType::WarpBgOverlay:
 		if (PicpProjectionOverlayViewportData != nullptr)
 		{
-			PixelShader->SetPicpParameters(RHICmdList, PixelShader->GetPixelShader(), PicpProjectionOverlayViewportData->ViewportUnder, nullptr);
+			PixelShader->SetPicpParameters(RHICmdList, PixelShader->GetPixelShader(), &PicpProjectionOverlayViewportData->ViewportUnder, nullptr);
 		}
 		break;
 	case EPicpShaderType::WarpAndBlendOneCameraFull:
 	case EPicpShaderType::WarpAndBlendOneCameraBack:
 	case EPicpShaderType::WarpAndBlendOneCameraFront:
+
+	case EPicpShaderType::WarpOneCameraFull:
+	case EPicpShaderType::WarpOneCameraBack:
+	case EPicpShaderType::WarpOneCameraFront:
 		if (PicpProjectionOverlayViewportData != nullptr)
 		{
-			PixelShader->SetPicpParameters(RHICmdList, PixelShader->GetPixelShader(), PicpProjectionOverlayViewportData->ViewportOver, PicpProjectionOverlayViewportData->ViewportUnder);
+			PixelShader->SetPicpParameters(RHICmdList, PixelShader->GetPixelShader(), &PicpProjectionOverlayViewportData->ViewportOver, &PicpProjectionOverlayViewportData->ViewportUnder);
 		}
 		break;
 	default:
@@ -227,8 +258,6 @@ static bool CompleteWarpTempl(
 
 	PixelShader->SetParameters(RHICmdList, PixelShader->GetPixelShader(), TextureWarpData.SrcTexture);
 	PixelShader->SetParameters(RHICmdList, PixelShader->GetPixelShader(), MPCDIRegion);
-	//!PixelShader->SetParameters(RHICmdList, PixelShader->GetPixelShader(), ShaderInputData);
-		
 	PixelShader->SetParameters(RHICmdList, PixelShader->GetPixelShader(), ShaderInputData.UVMatrix, StereoMatrix);
 
 	{
@@ -244,15 +273,6 @@ bool FPicpMPCDIShader::ApplyWarpBlend(FRHICommandListImmediate& RHICmdList, IMPC
 {
 	check(IsInRenderingThread());
 
-	// Map ProfileType to shader type to use
-	EPicpShaderType PixelShaderType = GetPixelShaderType(MPCDIData);
-	if (EPicpShaderType::Invalid == PixelShaderType)
-	{
-		//@todo: handle error
-		return false;
-	}
-
-
 	MPCDI::FMPCDIRegion* MPCDIRegion = nullptr;
 	if (MPCDIData && MPCDIData->IsValid())
 	{
@@ -265,6 +285,34 @@ bool FPicpMPCDIShader::ApplyWarpBlend(FRHICommandListImmediate& RHICmdList, IMPC
 		return false;
 	}
 
+	// Map ProfileType to shader type to use
+	EPicpShaderType PixelShaderType = GetPixelShaderType(MPCDIData);
+	if (EPicpShaderType::Invalid == PixelShaderType)
+	{
+		//@todo: handle error
+		return false;
+	}
+
+	bool bIsBlendDisabled = !MPCDIRegion->AlphaMap.IsInitialized();
+	bool bIsRenderSinglePass = ViewportOverlayData->Cameras.Num()==1;
+
+	const EVarPicpMPCDIShaderType ShaderType = (EVarPicpMPCDIShaderType)CVarPicpMPCDIShaderType.GetValueOnAnyThread();
+
+	switch (ShaderType) {
+	case EVarPicpMPCDIShaderType::DefaultNoBlend:
+		bIsBlendDisabled = true;
+		break;
+	case EVarPicpMPCDIShaderType::Passthrough:
+		PixelShaderType = EPicpShaderType::Passthrough;
+		break;
+	case EVarPicpMPCDIShaderType::ForceMultiPassRender:
+		bIsRenderSinglePass = false;
+		break;
+	case EVarPicpMPCDIShaderType::Disable:
+		return false;
+		break;
+	};
+
 	{
 		// Do single warp render pass
 		bool bIsRenderSuccess = false;
@@ -273,9 +321,9 @@ bool FPicpMPCDIShader::ApplyWarpBlend(FRHICommandListImmediate& RHICmdList, IMPC
 		{
 			switch (PixelShaderType)
 			{
-			case EPicpShaderType::PassThrought:
-				ShaderInputData.UVMatrix = FMatrix::Identity;
-				bIsRenderSuccess = CompleteWarpTempl<EPicpShaderType::PassThrought>(RHICmdList, TextureWarpData, ShaderInputData, *MPCDIRegion, MPCDIData, nullptr, 0, EColorOPMode::Default);
+			case EPicpShaderType::Passthrough:
+				ShaderInputData.UVMatrix = FMatrix::Identity;				
+				bIsRenderSuccess = CompleteWarpTempl<EPicpShaderType::Passthrough>(RHICmdList, TextureWarpData, ShaderInputData, *MPCDIRegion, MPCDIData, nullptr, 0, EColorOPMode::Default);
 				break;
 
 			case EPicpShaderType::WarpAndBlend:
@@ -284,8 +332,6 @@ bool FPicpMPCDIShader::ApplyWarpBlend(FRHICommandListImmediate& RHICmdList, IMPC
 				if (ViewportOverlayData)
 				{// Use overlay render tricks:					
 					// First pass render bg
-					bool bIsRenderSinglePass = ((CVarPicpEnableSinglePassRender.GetValueOnAnyThread() != 0) && ViewportOverlayData->Cameras.Num() == 1);
-
 					if (bIsRenderSinglePass)
 					{
 						//One pass solution:
@@ -293,36 +339,77 @@ bool FPicpMPCDIShader::ApplyWarpBlend(FRHICommandListImmediate& RHICmdList, IMPC
 						{
 							if (ViewportOverlayData->iViewportOverUsed())
 							{
-								bIsRenderSuccess = CompleteWarpTempl<EPicpShaderType::WarpAndBlendOneCameraFull>(RHICmdList, TextureWarpData, ShaderInputData, *MPCDIRegion, MPCDIData, ViewportOverlayData, 0, EColorOPMode::Default);
+								if (bIsBlendDisabled)
+								{
+									bIsRenderSuccess = CompleteWarpTempl<EPicpShaderType::WarpOneCameraFull>(RHICmdList, TextureWarpData, ShaderInputData, *MPCDIRegion, MPCDIData, ViewportOverlayData, 0, EColorOPMode::Default);
+								}
+								else
+								{
+									bIsRenderSuccess = CompleteWarpTempl<EPicpShaderType::WarpAndBlendOneCameraFull>(RHICmdList, TextureWarpData, ShaderInputData, *MPCDIRegion, MPCDIData, ViewportOverlayData, 0, EColorOPMode::Default);
+								}
 							}
 							else
 							{
-								bIsRenderSuccess = CompleteWarpTempl<EPicpShaderType::WarpAndBlendOneCameraBack>(RHICmdList, TextureWarpData, ShaderInputData, *MPCDIRegion, MPCDIData, ViewportOverlayData, 0, EColorOPMode::Default);
+								if (bIsBlendDisabled)
+								{
+									bIsRenderSuccess = CompleteWarpTempl<EPicpShaderType::WarpOneCameraBack>(RHICmdList, TextureWarpData, ShaderInputData, *MPCDIRegion, MPCDIData, ViewportOverlayData, 0, EColorOPMode::Default);
+								}
+								else
+								{
+									bIsRenderSuccess = CompleteWarpTempl<EPicpShaderType::WarpAndBlendOneCameraBack>(RHICmdList, TextureWarpData, ShaderInputData, *MPCDIRegion, MPCDIData, ViewportOverlayData, 0, EColorOPMode::Default);
+								}
 							}
 						}
 						else
 						{
 							if (ViewportOverlayData->iViewportOverUsed())
 							{
-								bIsRenderSuccess = CompleteWarpTempl<EPicpShaderType::WarpAndBlendOneCameraFront>(RHICmdList, TextureWarpData, ShaderInputData, *MPCDIRegion, MPCDIData, ViewportOverlayData, 0, EColorOPMode::Default);
+								if (bIsBlendDisabled)
+								{
+									bIsRenderSuccess = CompleteWarpTempl<EPicpShaderType::WarpOneCameraFront>(RHICmdList, TextureWarpData, ShaderInputData, *MPCDIRegion, MPCDIData, ViewportOverlayData, 0, EColorOPMode::Default);
+								}
+								else
+								{
+									bIsRenderSuccess = CompleteWarpTempl<EPicpShaderType::WarpAndBlendOneCameraFront>(RHICmdList, TextureWarpData, ShaderInputData, *MPCDIRegion, MPCDIData, ViewportOverlayData, 0, EColorOPMode::Default);
+								}
 							}
 							else
 							{
-								bIsRenderSuccess = CompleteWarpTempl<EPicpShaderType::WarpAndBlendOneCameraNone>(RHICmdList, TextureWarpData, ShaderInputData, *MPCDIRegion, MPCDIData, ViewportOverlayData, 0, EColorOPMode::Default);
+								if (bIsBlendDisabled)
+								{
+									bIsRenderSuccess = CompleteWarpTempl<EPicpShaderType::WarpOneCameraNone>(RHICmdList, TextureWarpData, ShaderInputData, *MPCDIRegion, MPCDIData, ViewportOverlayData, 0, EColorOPMode::Default);
+								}
+								else
+								{
+									bIsRenderSuccess = CompleteWarpTempl<EPicpShaderType::WarpAndBlendOneCameraNone>(RHICmdList, TextureWarpData, ShaderInputData, *MPCDIRegion, MPCDIData, ViewportOverlayData, 0, EColorOPMode::Default);
+								}
 							}
 						}						
 					}
 					else
 					{
-						//Complex multipass solution
-											//@todo: add optimization in single shader
+						//Complex multipass solution						
 						if (ViewportOverlayData->iViewportUnderUsed())
 						{
-							bIsRenderSuccess = CompleteWarpTempl <EPicpShaderType::WarpAndBlendBgOverlay>(RHICmdList, TextureWarpData, ShaderInputData, *MPCDIRegion, MPCDIData, ViewportOverlayData, -1, EColorOPMode::Default);
+							if (bIsBlendDisabled)
+							{
+								bIsRenderSuccess = CompleteWarpTempl <EPicpShaderType::WarpBgOverlay>(RHICmdList, TextureWarpData, ShaderInputData, *MPCDIRegion, MPCDIData, ViewportOverlayData, -1, EColorOPMode::Default);
+							}
+							else
+							{
+								bIsRenderSuccess = CompleteWarpTempl <EPicpShaderType::WarpAndBlendBgOverlay>(RHICmdList, TextureWarpData, ShaderInputData, *MPCDIRegion, MPCDIData, ViewportOverlayData, -1, EColorOPMode::Default);
+							}
 						}
 						else
 						{
-							bIsRenderSuccess = CompleteWarpTempl<EPicpShaderType::WarpAndBlendBg>(RHICmdList, TextureWarpData, ShaderInputData, *MPCDIRegion, MPCDIData, ViewportOverlayData, -1, EColorOPMode::Default);
+							if (bIsBlendDisabled)
+							{
+								bIsRenderSuccess = CompleteWarpTempl<EPicpShaderType::WarpBg>(RHICmdList, TextureWarpData, ShaderInputData, *MPCDIRegion, MPCDIData, ViewportOverlayData, -1, EColorOPMode::Default);
+							}
+							else
+							{
+								bIsRenderSuccess = CompleteWarpTempl<EPicpShaderType::WarpAndBlendBg>(RHICmdList, TextureWarpData, ShaderInputData, *MPCDIRegion, MPCDIData, ViewportOverlayData, -1, EColorOPMode::Default);
+							}
 						}
 
 						//Second pass (now all cams on second pass, @todo add first cam on single pass with bg
@@ -333,7 +420,14 @@ bool FPicpMPCDIShader::ApplyWarpBlend(FRHICommandListImmediate& RHICmdList, IMPC
 							int RenderPass = 0;
 							for (auto& It : ViewportOverlayData->Cameras)
 							{
-								bIsRenderSuccess = CompleteWarpTempl<EPicpShaderType::WarpAndBlendAddCamera>(RHICmdList, TextureWarpData, ShaderInputData, *MPCDIRegion, MPCDIData, ViewportOverlayData, CameraIndex, EColorOPMode::AddAlpha);
+								if (bIsBlendDisabled)
+								{
+									bIsRenderSuccess = CompleteWarpTempl<EPicpShaderType::WarpAddCamera>(RHICmdList, TextureWarpData, ShaderInputData, *MPCDIRegion, MPCDIData, ViewportOverlayData, CameraIndex, EColorOPMode::AddAlpha);
+								}
+								else
+								{
+									bIsRenderSuccess = CompleteWarpTempl<EPicpShaderType::WarpAndBlendAddCamera>(RHICmdList, TextureWarpData, ShaderInputData, *MPCDIRegion, MPCDIData, ViewportOverlayData, CameraIndex, EColorOPMode::AddAlpha);
+								}
 
 								RenderPass++;
 								CameraIndex++; //! fixme
@@ -342,13 +436,27 @@ bool FPicpMPCDIShader::ApplyWarpBlend(FRHICommandListImmediate& RHICmdList, IMPC
 
 						if (ViewportOverlayData->iViewportOverUsed())
 						{
-							bIsRenderSuccess = CompleteWarpTempl<EPicpShaderType::WarpAndBlendAddOverlay>(RHICmdList, TextureWarpData, ShaderInputData, *MPCDIRegion, MPCDIData, ViewportOverlayData, -1, EColorOPMode::AddInvAlpha);
+							if (bIsBlendDisabled)
+							{
+								bIsRenderSuccess = CompleteWarpTempl<EPicpShaderType::WarpAddOverlay>(RHICmdList, TextureWarpData, ShaderInputData, *MPCDIRegion, MPCDIData, ViewportOverlayData, -1, EColorOPMode::AddInvAlpha);
+							}
+							else
+							{
+								bIsRenderSuccess = CompleteWarpTempl<EPicpShaderType::WarpAndBlendAddOverlay>(RHICmdList, TextureWarpData, ShaderInputData, *MPCDIRegion, MPCDIData, ViewportOverlayData, -1, EColorOPMode::AddInvAlpha);
+							}
 						}
 					}
 				}
 				else
 				{// Use default render 
-					bIsRenderSuccess = CompleteWarpTempl<EPicpShaderType::WarpAndBlend>(RHICmdList, TextureWarpData, ShaderInputData, *MPCDIRegion, MPCDIData, nullptr, 0, EColorOPMode::Default);
+					if (bIsBlendDisabled)
+					{
+						bIsRenderSuccess = CompleteWarpTempl<EPicpShaderType::Warp>(RHICmdList, TextureWarpData, ShaderInputData, *MPCDIRegion, MPCDIData, nullptr, 0, EColorOPMode::Default);
+					}
+					else
+					{
+						bIsRenderSuccess = CompleteWarpTempl<EPicpShaderType::WarpAndBlend>(RHICmdList, TextureWarpData, ShaderInputData, *MPCDIRegion, MPCDIData, nullptr, 0, EColorOPMode::Default);
+					}
 				}
 			};
 		}
diff --git a/Engine/Plugins/Runtime/nDisplay/Source/PicpMPCDI/Private/PicpMPCDIShader.h b/Engine/Plugins/Runtime/nDisplay/Source/PicpMPCDI/Private/PicpMPCDIShader.h
--- a/Engine/Plugins/Runtime/nDisplay/Source/PicpMPCDI/Private/PicpMPCDIShader.h
+++ b/Engine/Plugins/Runtime/nDisplay/Source/PicpMPCDI/Private/PicpMPCDIShader.h
@@ -17,6 +17,41 @@
 #include "Overlay/PicpProjectionOverlayRender.h"
 
 
+enum class EPicpShaderType : uint8
+{
+	Passthrough, // viewport frustum (no warpblend, only frustum math)
+	WarpAndBlend, // Pure mpcdi warpblend for viewport
+	Warp, //No Blend
+
+	//One-pass solution: (+overlay_over+camera+overlay_over +lut)
+	WarpAndBlendOneCameraFront,
+	WarpAndBlendOneCameraFull,
+	WarpAndBlendOneCameraNone,
+	WarpAndBlendOneCameraBack,
+	//No Blend case
+	WarpOneCameraBack,
+	WarpOneCameraFront,
+	WarpOneCameraFull,
+	WarpOneCameraNone,
+
+	// Multi-pass render:
+	WarpAndBlendBg,                 // bg +lut
+	WarpAndBlendBgOverlay,          // bg+overlay_under +lut
+	WarpAndBlendAddCamera,          // +camera +lut   (Additive blend colorop)
+	WarpAndBlendAddCameraOverlay,   // +camera+overlay_over +lut (Additive blend colorop)
+	WarpAndBlendAddOverlay,         // +overlay_over +lut (Additive blend colorop)
+	//No Blend case:
+	WarpBg,                 // bg +lut
+	WarpBgOverlay,          // bg+overlay_under +lut
+	WarpAddCamera,          // +camera +lut   (Additive blend colorop)
+	WarpAddCameraOverlay,   // +camera+overlay_over +lut (Additive blend colorop)
+	WarpAddOverlay,         // +overlay_over +lut (Additive blend colorop)
+
+	Invalid,
+};
+
+
+
 class FPicpMPCDIShader
 {
 public:
@@ -38,7 +73,8 @@ class FPicpWarpVS : public FGlobalShader
 
 public:
 	FPicpWarpVS()
-	{ }
+	{ 
+	}
 
 	FPicpWarpVS(const ShaderMetaType::CompiledShaderInitializerType& Initializer)
 		: FGlobalShader(Initializer)
@@ -85,29 +121,6 @@ private:
 };
 
 
-enum class EPicpShaderType : uint8
-{
-	PassThrought, // viewport frustum (no warpblend, only frustum math)
-	WarpAndBlend, // Pure mpcdi warpblend for viewport
-
-	//One-pass solution: (+overlay_over+camera+overlay_over +lut)
-	WarpAndBlendOneCameraBack,
-	WarpAndBlendOneCameraFront,
-	WarpAndBlendOneCameraFull,
-	WarpAndBlendOneCameraNone,
-
-		// Multi-pass render:
-	WarpAndBlendBg,                 // bg +lut
-	WarpAndBlendBgOverlay,          // bg+overlay_under +lut
-	WarpAndBlendAddCamera,          // +camera +lut   (Additive blend colorop)
-	WarpAndBlendAddCameraOverlay,   // +camera+overlay_over +lut (Additive blend colorop)
-	WarpAndBlendAddOverlay,         // +overlay_over +lut (Additive blend colorop)
-	
-	Invalid,
-};
-
-
-
 template<EPicpShaderType ShaderType>
 class TPicpWarpBasePS : public FGlobalShader
 {
@@ -115,7 +128,8 @@ class TPicpWarpBasePS : public FGlobalShader
 
 public:
 	TPicpWarpBasePS()
-	{ }
+	{ 
+	}
 
 	TPicpWarpBasePS(const ShaderMetaType::CompiledShaderInitializerType& Initializer)
 		: FGlobalShader(Initializer)
@@ -128,7 +142,7 @@ public:
 		TextureProjectionMatrixParameter.Bind(Initializer.ParameterMap, TEXT("TextureProjectionMatrix"));
 		MultiViewportMatrixParameter.Bind(Initializer.ParameterMap, TEXT("MultiViewportMatrix"));
 
-		if (ShaderType != EPicpShaderType::PassThrought)
+		if (ShaderType != EPicpShaderType::Passthrough)
 		{
 			AlphaEmbeddedGammaParameter.Bind(Initializer.ParameterMap, TEXT("AlphaEmbeddedGamma"));
 
@@ -166,46 +180,74 @@ public:
 	{
 		FGlobalShader::ModifyCompilationEnvironment(Parameters, OutEnvironment);
 
+		//Enable blend:
 		switch (ShaderType)
 		{
+		case EPicpShaderType::WarpAndBlend:
+
+		case EPicpShaderType::WarpAndBlendOneCameraFront:
+		case EPicpShaderType::WarpAndBlendOneCameraFull:
 		case EPicpShaderType::WarpAndBlendOneCameraNone:
+		case EPicpShaderType::WarpAndBlendOneCameraBack:
+
+		case EPicpShaderType::WarpAndBlendBg:
+		case EPicpShaderType::WarpAndBlendBgOverlay:
+		case EPicpShaderType::WarpAndBlendAddCamera:
+		case EPicpShaderType::WarpAndBlendAddCameraOverlay:
+		case EPicpShaderType::WarpAndBlendAddOverlay:
+
+			OutEnvironment.SetDefine(TEXT("RENDER_CASE_APPLYBLENDING"), 1);
 			break;
+		default:
+			// Do nothing
+			break;
+		}
 
+		switch (ShaderType)
+		{
 		case EPicpShaderType::WarpAndBlendOneCameraBack:
+		case EPicpShaderType::WarpOneCameraBack:
 			OutEnvironment.SetDefine(TEXT("USE_OVERLAY_BACK"), 1);
 			break;
 
 		case EPicpShaderType::WarpAndBlendOneCameraFront:
+		case EPicpShaderType::WarpOneCameraFront:			
 			OutEnvironment.SetDefine(TEXT("USE_OVERLAY_FRONT"), 1);
 			break;
 
 		case EPicpShaderType::WarpAndBlendOneCameraFull:
+		case EPicpShaderType::WarpOneCameraFull:			
 			OutEnvironment.SetDefine(TEXT("USE_OVERLAY_BACK"), 1);
 			OutEnvironment.SetDefine(TEXT("USE_OVERLAY_FRONT"), 1);
 			break;
 
 		case EPicpShaderType::WarpAndBlendAddCamera:
+		case EPicpShaderType::WarpAddCamera:			
 			OutEnvironment.SetDefine(TEXT("RENDER_CASE_CAMERA"), 1);
 			break;
 
 		case EPicpShaderType::WarpAndBlendAddCameraOverlay:
+		case EPicpShaderType::WarpAddCameraOverlay:			
 			OutEnvironment.SetDefine(TEXT("RENDER_CASE_CAMERA"), 1);
 			OutEnvironment.SetDefine(TEXT("RENDER_CASE_VIEWPORT"), 1);
 			OutEnvironment.SetDefine(TEXT("RENDER_CASE_VIEWPORT_OVERLAY"), 1);
 			break;
 
 		case EPicpShaderType::WarpAndBlendAddOverlay:
+		case EPicpShaderType::WarpAddOverlay:			
 			OutEnvironment.SetDefine(TEXT("RENDER_CASE_VIEWPORT"), 1);
 			OutEnvironment.SetDefine(TEXT("RENDER_CASE_VIEWPORT_OVERLAY"), 1);
 			break;
 
 		case EPicpShaderType::WarpAndBlendBgOverlay:
+		case EPicpShaderType::WarpBgOverlay:			
 			OutEnvironment.SetDefine(TEXT("RENDER_CASE_VIEWPORT"), 1);
 			OutEnvironment.SetDefine(TEXT("RENDER_CASE_VIEWPORT_OVERLAY"), 1);
 			OutEnvironment.SetDefine(TEXT("RENDER_CASE_VIEWPORT_BASE"), 1);
 			break;
 
 		case EPicpShaderType::WarpAndBlendBg:
+		case EPicpShaderType::WarpBg:			
 			OutEnvironment.SetDefine(TEXT("RENDER_CASE_VIEWPORT"), 1);
 			OutEnvironment.SetDefine(TEXT("RENDER_CASE_VIEWPORT_BASE"), 1);
 			break;
@@ -229,7 +271,7 @@ public:
 	{
 		{
 			SetTextureParameter(RHICmdList, ShaderRHI, PostprocessInputParameter0, SourceTexture);
-			RHICmdList.SetShaderSampler(ShaderRHI, PostprocessInputParameterSampler0.GetBaseIndex(), TStaticSamplerState<SF_Bilinear, AM_Clamp, AM_Clamp, AM_Clamp>::GetRHI());
+			RHICmdList.SetShaderSampler(ShaderRHI, PostprocessInputParameterSampler0.GetBaseIndex(), TStaticSamplerState<SF_Trilinear, AM_Clamp, AM_Clamp, AM_Clamp>::GetRHI());
 			//SetTextureParameter(RHICmdList, ShaderRHI, PostprocessInputParameter0, PostprocessInputParameterSampler0, TStaticSamplerState<SF_Bilinear, AM_Clamp, AM_Clamp, AM_Clamp>::GetRHI(), SourceTexture);
 		}
 	}
@@ -237,12 +279,12 @@ public:
 	template<typename TShaderRHIParamRef>
 	void SetParameters(FRHICommandListImmediate& RHICmdList, const TShaderRHIParamRef ShaderRHI, MPCDI::FMPCDIRegion& MPCDIRegionData)
 	{
-		if (ShaderType != EPicpShaderType::PassThrought)
+		if (ShaderType != EPicpShaderType::Passthrough)
 		{
 			SetShaderValue(RHICmdList, ShaderRHI, AlphaEmbeddedGammaParameter, MPCDIRegionData.AlphaMap.GetEmbeddedGamma());
 
 			SetTextureParameter(RHICmdList, ShaderRHI, WarpMapParameter, MPCDIRegionData.WarpMap.TextureRHI);
-			RHICmdList.SetShaderSampler(ShaderRHI, WarpMapParameterSampler.GetBaseIndex(), TStaticSamplerState<SF_Bilinear, AM_Clamp, AM_Clamp, AM_Clamp>::GetRHI());
+			RHICmdList.SetShaderSampler(ShaderRHI, WarpMapParameterSampler.GetBaseIndex(), TStaticSamplerState<SF_Trilinear, AM_Clamp, AM_Clamp, AM_Clamp>::GetRHI());
 
 			SetTextureParameter(RHICmdList, ShaderRHI, AlphaMapParameter, MPCDIRegionData.AlphaMap.TextureRHI);
 			RHICmdList.SetShaderSampler(ShaderRHI, AlphaMapParameterSampler.GetBaseIndex(), TStaticSamplerState<SF_Bilinear, AM_Clamp, AM_Clamp, AM_Clamp>::GetRHI());
@@ -261,7 +303,7 @@ public:
 			SetShaderValue(RHICmdList, ShaderRHI, CameraSoftEdgeParameter, OverlayCamera->SoftEdge);
 
 			SetTextureParameter(RHICmdList, ShaderRHI, CameraMapParameter, OverlayCamera->CameraTexture);
-			RHICmdList.SetShaderSampler(ShaderRHI, CameraMapParameterSampler.GetBaseIndex(), TStaticSamplerState<SF_Bilinear, AM_Clamp, AM_Clamp, AM_Clamp>::GetRHI());
+			RHICmdList.SetShaderSampler(ShaderRHI, CameraMapParameterSampler.GetBaseIndex(), TStaticSamplerState<SF_AnisotropicLinear, AM_Clamp, AM_Clamp, AM_Clamp>::GetRHI());
 		}
 	}
 
@@ -271,13 +313,13 @@ public:
 		if (OverlayOverViewport && ShaderType > EPicpShaderType::WarpAndBlend) // Only Picp techniques
 		{
 			SetTextureParameter(RHICmdList, ShaderRHI, ViewportOverlayMapParameter, OverlayOverViewport->ViewportTexture);
-			RHICmdList.SetShaderSampler(ShaderRHI, ViewportOverlayMapParameterSampler.GetBaseIndex(), TStaticSamplerState<SF_Bilinear, AM_Clamp, AM_Clamp, AM_Clamp>::GetRHI());
+			RHICmdList.SetShaderSampler(ShaderRHI, ViewportOverlayMapParameterSampler.GetBaseIndex(), TStaticSamplerState<SF_AnisotropicLinear, AM_Clamp, AM_Clamp, AM_Clamp>::GetRHI());
 		}
 
 		if (OverlayBackViewport && ShaderType > EPicpShaderType::WarpAndBlend) // Only Picp techniques
 		{
 			SetTextureParameter(RHICmdList, ShaderRHI, ViewportOverlayBackMapParameter, OverlayBackViewport->ViewportTexture);
-			RHICmdList.SetShaderSampler(ShaderRHI, ViewportOverlayBackMapParameterSampler.GetBaseIndex(), TStaticSamplerState<SF_Bilinear, AM_Clamp, AM_Clamp, AM_Clamp>::GetRHI());
+			RHICmdList.SetShaderSampler(ShaderRHI, ViewportOverlayBackMapParameterSampler.GetBaseIndex(), TStaticSamplerState<SF_AnisotropicLinear, AM_Clamp, AM_Clamp, AM_Clamp>::GetRHI());
 		}
 	}
 
@@ -359,7 +401,7 @@ private:
 	FShaderParameter         VignetteEVParameter;
 };
 
-typedef TPicpWarpBasePS<EPicpShaderType::PassThrought>         FPicpMPCDIPassThroughtPS;
+typedef TPicpWarpBasePS<EPicpShaderType::Passthrough>          FPicpMPCDIPassthroughPS;
 typedef TPicpWarpBasePS<EPicpShaderType::WarpAndBlend>         FPicpMPCDIWarpAndBlendPS;
 
 typedef TPicpWarpBasePS<EPicpShaderType::WarpAndBlendBg>                 FPicpMPCDIWarpAndBlendBgPS;
@@ -374,3 +416,17 @@ typedef TPicpWarpBasePS<EPicpShaderType::WarpAndBlendOneCameraBack>      FPicpWa
 typedef TPicpWarpBasePS<EPicpShaderType::WarpAndBlendOneCameraFront>     FPicpWarpAndBlendOneCameraFrontPS;
 typedef TPicpWarpBasePS<EPicpShaderType::WarpAndBlendOneCameraNone>      FPicpWarpAndBlendOneCameraNonePS;
 
+// No Blend case
+typedef TPicpWarpBasePS<EPicpShaderType::Warp>                 FPicpMPCDIWarpPS;
+
+typedef TPicpWarpBasePS<EPicpShaderType::WarpBg>                 FPicpMPCDIWarpBgPS;
+typedef TPicpWarpBasePS<EPicpShaderType::WarpBgOverlay>          FPicpMPCDIWarpBgOverlayPS;
+
+typedef TPicpWarpBasePS<EPicpShaderType::WarpAddCamera>          FPicpMPCDIWarpCameraPS;
+typedef TPicpWarpBasePS<EPicpShaderType::WarpAddCameraOverlay>   FPicpMPCDIWarpCameraOverlayPS;
+typedef TPicpWarpBasePS<EPicpShaderType::WarpAddOverlay>         FPicpMPCDIWarpOverlayPS;
+
+typedef TPicpWarpBasePS<EPicpShaderType::WarpOneCameraFull>      FPicpWarpOneCameraFullPS;
+typedef TPicpWarpBasePS<EPicpShaderType::WarpOneCameraBack>      FPicpWarpOneCameraBackPS;
+typedef TPicpWarpBasePS<EPicpShaderType::WarpOneCameraFront>     FPicpWarpOneCameraFrontPS;
+typedef TPicpWarpBasePS<EPicpShaderType::WarpOneCameraNone>      FPicpWarpOneCameraNonePS;
diff --git a/Engine/Plugins/Runtime/nDisplay/Source/PicpMPCDI/Public/IPicpMPCDI.h b/Engine/Plugins/Runtime/nDisplay/Source/PicpMPCDI/Public/IPicpMPCDI.h
--- a/Engine/Plugins/Runtime/nDisplay/Source/PicpMPCDI/Public/IPicpMPCDI.h
+++ b/Engine/Plugins/Runtime/nDisplay/Source/PicpMPCDI/Public/IPicpMPCDI.h
@@ -35,7 +35,7 @@ public:
 	 */
 	static inline IPicpMPCDI& Get()
 	{
-		return FModuleManager::LoadModuleChecked<IPicpMPCDI>(IPicpMPCDI::ModuleName);
+		return FModuleManager::GetModuleChecked<IPicpMPCDI>(IPicpMPCDI::ModuleName);
 	}
 
 	/**
@@ -53,4 +53,5 @@ public:
 
 	virtual void ApplyBlur(UTextureRenderTarget2D* InOutRenderTarget, UTextureRenderTarget2D* TemporaryRenderTarget, int KernelRadius, float KernelScale, EPicpBlurPostProcessShaderType BlurType) = 0;
 	virtual void ApplyCompose(UTexture* InputTexture, UTextureRenderTarget2D* OutputRenderTarget, UTextureRenderTarget2D* Result) = 0;
+	virtual void ExecuteCompose() = 0;
 };
diff --git a/Engine/Plugins/Runtime/nDisplay/Source/PicpProjection/PicpProjection.Build.cs b/Engine/Plugins/Runtime/nDisplay/Source/PicpProjection/PicpProjection.Build.cs
--- a/Engine/Plugins/Runtime/nDisplay/Source/PicpProjection/PicpProjection.Build.cs
+++ b/Engine/Plugins/Runtime/nDisplay/Source/PicpProjection/PicpProjection.Build.cs
@@ -12,10 +12,6 @@ public class PicpProjection : ModuleRules
 			{
 				"PicpProjection/Private",
 				"DisplayClusterProjection/Private",
-				"../../../../Source/Runtime/Renderer/Private",
-				"../../../../Source/Runtime/Engine/Classes/Components",
-				"../../../../../Engine/Source/Runtime/Windows/D3D11RHI/Private",
-				"../../../../../Engine/Source/Runtime/Windows/D3D11RHI/Private/Windows",
 			}
 		);
 
@@ -24,34 +20,26 @@ public class PicpProjection : ModuleRules
 			{
 				"Core",
 				"CoreUObject",
-				"D3D11RHI",
-				"DisplayCluster",
-				"PicpMPCDI",
 				"Engine",
-				"HeadMountedDisplay",
-				"MPCDI",
-				"RenderCore",
-				"RHI",
-				"Composure"
-			});
+            });
 
 		PrivateDependencyModuleNames.AddRange(
 			new string[]
 			{
 				"ApplicationCore",
-				"Core",
-				"CoreUObject",
-				"Engine"
-			}
-		);
+                "Composure",
+				"CinematicCamera",
+                "DisplayCluster",
+                "MPCDI",
+                "PicpMPCDI",
+                "RenderCore",
+				"RHI",				
+            }
+        );
 
 		if (Target.bBuildEditor == true)
 		{
 			PrivateDependencyModuleNames.Add("UnrealEd");
 		}
-
-		AddEngineThirdPartyPrivateStaticDependencies(Target, "DX11");
-		AddEngineThirdPartyPrivateStaticDependencies(Target, "NVAftermath");
-		AddEngineThirdPartyPrivateStaticDependencies(Target, "IntelMetricsDiscovery");
 	}
 }
diff --git a/Engine/Plugins/Runtime/nDisplay/Source/PicpProjection/Private/Blueprints/PicpProjectionBlueprintAPIImpl.cpp b/Engine/Plugins/Runtime/nDisplay/Source/PicpProjection/Private/Blueprints/PicpProjectionBlueprintAPIImpl.cpp
--- a/Engine/Plugins/Runtime/nDisplay/Source/PicpProjection/Private/Blueprints/PicpProjectionBlueprintAPIImpl.cpp
+++ b/Engine/Plugins/Runtime/nDisplay/Source/PicpProjection/Private/Blueprints/PicpProjectionBlueprintAPIImpl.cpp
@@ -43,17 +43,14 @@ void UPicpProjectionAPIImpl::SetupOverlayCaptures(const TArray<struct FPicpOverl
 	FPicpProjectionOverlayFrameData overlayFrameData;
 
 	for (auto& capture : captures)
-	{
-		FRotator CameraRotation = capture.CameraOverlayFrameCapture->K2_GetComponentRotation();
-		FVector  CameraLocation = capture.CameraOverlayFrameCapture->K2_GetComponentLocation();
+	{		
+		FRotator CameraRotation = capture.CineCamera->K2_GetComponentRotation();
+		FVector  CameraLocation = capture.CineCamera->K2_GetComponentLocation();
 
-		float fov = capture.CameraOverlayFrameCapture->FOVAngle;
+		float fov = capture.CineCamera->FieldOfView;
 
 		FComposurePostMoveSettings settings;
-		float width = capture.CameraOverlayFrameCapture->TextureTarget->GetSurfaceWidth();
-		float height = capture.CameraOverlayFrameCapture->TextureTarget->GetSurfaceHeight();
-
-		float aspectRatio = width / height;
+		float aspectRatio = capture.CineCamera->AspectRatio;
 
 		FMatrix Prj = settings.GetProjectionMatrix(fov, aspectRatio);
 
@@ -61,11 +58,12 @@ void UPicpProjectionAPIImpl::SetupOverlayCaptures(const TArray<struct FPicpOverl
 		FTextureRenderTarget2DResource* overlayRTTRes2D = (FTextureRenderTarget2DResource*)overlayRTTRes;
 
 		FRHITexture2D* CameraTextureRef = overlayRTTRes2D->GetTextureRHI();
-		FPicpProjectionOverlayCamera* overlayCamera = new FPicpProjectionOverlayCamera(CameraRotation, CameraLocation, Prj, CameraTextureRef);
 
 		// add inner camera
-		overlayFrameData.Cameras.Add(overlayCamera);
-		overlayCamera->SoftEdge = capture.SoftEdge;
+		FPicpProjectionOverlayCamera NewCamera(CameraRotation, CameraLocation, Prj, CameraTextureRef, capture.RTTViewportId);
+		NewCamera.SoftEdge = capture.SoftEdge;
+		overlayFrameData.Cameras.Add(NewCamera);
+		
 
 		// add blending layers (lights?)
 		for (auto& layer : capture.OverlayBlendFrames)
@@ -74,15 +72,14 @@ void UPicpProjectionAPIImpl::SetupOverlayCaptures(const TArray<struct FPicpOverl
 
 			FTextureRenderTarget2DResource* rtt2d = (FTextureRenderTarget2DResource*)res;
 			FRHITexture2D* paramRef = rtt2d->GetTextureRHI();
-			FPicpProjectionOverlayViewport* overlayViewport = new FPicpProjectionOverlayViewport(paramRef);
 
 			if (layer.OverlayBlendMode == ECameraOverlayRenderMode::Over)
 			{
-				overlayFrameData.ViewportsOver.Add(layer.Id, overlayViewport);
+				overlayFrameData.ViewportsOver.Add(layer.Id, FPicpProjectionOverlayViewport(paramRef));
 			}
 			else if (layer.OverlayBlendMode == ECameraOverlayRenderMode::Under)
 			{
-				overlayFrameData.ViewportsUnder.Add(layer.Id, overlayViewport);
+				overlayFrameData.ViewportsUnder.Add(layer.Id, FPicpProjectionOverlayViewport(paramRef));
 			}
 		}
 	}
@@ -95,4 +92,10 @@ void UPicpProjectionAPIImpl::SetupOverlayCaptures(const TArray<struct FPicpOverl
 void UPicpProjectionAPIImpl::SetWarpTextureCaptureState(UTextureRenderTarget2D* dstTexture, const FString& ViewportId, const int ViewIdx, bool bCaptureNow)
 {
 	IPicpProjection::Get().CaptureWarpTexture(dstTexture, ViewportId, ViewIdx, bCaptureNow);
+}
+void UPicpProjectionAPIImpl::EnableTextureRenderTargetMips(UTextureRenderTarget2D* Texture)
+{
+	Texture->bAutoGenerateMips = true;
+	Texture->UpdateResource();
+	Texture->UpdateResourceImmediate(true);
 }
\ No newline at end of file
diff --git a/Engine/Plugins/Runtime/nDisplay/Source/PicpProjection/Private/Blueprints/PicpProjectionBlueprintAPIImpl.h b/Engine/Plugins/Runtime/nDisplay/Source/PicpProjection/Private/Blueprints/PicpProjectionBlueprintAPIImpl.h
--- a/Engine/Plugins/Runtime/nDisplay/Source/PicpProjection/Private/Blueprints/PicpProjectionBlueprintAPIImpl.h
+++ b/Engine/Plugins/Runtime/nDisplay/Source/PicpProjection/Private/Blueprints/PicpProjectionBlueprintAPIImpl.h
@@ -42,4 +42,7 @@ public:
 
 	UFUNCTION(BlueprintCallable, meta = (DisplayName = "Capture Final Warped Frame"), Category = "PICP")
 	virtual void SetWarpTextureCaptureState(UTextureRenderTarget2D* dstTexture, const FString& ViewportId, const int ViewIdx, bool bCaptureNow) override;
+
+	UFUNCTION(BlueprintCallable, meta = (DisplayName = "Enable mips for Texture Render Target"), Category = "PICP")
+	virtual void EnableTextureRenderTargetMips(UTextureRenderTarget2D* Texture) override;
 };
\ No newline at end of file
diff --git a/Engine/Plugins/Runtime/nDisplay/Source/PicpProjection/Private/PicpProjectionModule.cpp b/Engine/Plugins/Runtime/nDisplay/Source/PicpProjection/Private/PicpProjectionModule.cpp
--- a/Engine/Plugins/Runtime/nDisplay/Source/PicpProjection/Private/PicpProjectionModule.cpp
+++ b/Engine/Plugins/Runtime/nDisplay/Source/PicpProjection/Private/PicpProjectionModule.cpp
@@ -101,18 +101,18 @@ void FPicpProjectionModule::SetOverlayFrameData(const FString& PolicyType, FPicp
 	TSharedPtr<IDisplayClusterProjectionPolicyFactory> Factory = GetProjectionFactory(PolicyType);
 	if (Factory.IsValid())
 	{
-		FPicpProjectionMPCDIPolicyFactory* MPCDIFactory = static_cast<FPicpProjectionMPCDIPolicyFactory*>(Factory.Get());
-		if (MPCDIFactory)
+		FPicpProjectionMPCDIPolicyFactory* PicpMPCDIFactory = static_cast<FPicpProjectionMPCDIPolicyFactory*>(Factory.Get());
+		if (PicpMPCDIFactory)
 		{
-			TArray<TSharedPtr<IDisplayClusterProjectionPolicy>> UsedPolicy = MPCDIFactory->GetMPCDIPolicy();
+			TArray<TSharedPtr<IDisplayClusterProjectionPolicy>> UsedPolicy = PicpMPCDIFactory->GetMPCDIPolicy();
 			for (auto It : UsedPolicy)
 			{
 				if (It.IsValid())
 				{
-					FPicpProjectionMPCDIPolicy* MPCDIPolicy = static_cast<FPicpProjectionMPCDIPolicy*>(It.Get());
-					if (MPCDIPolicy != nullptr)
+					FPicpProjectionMPCDIPolicy* PicpMPCDIPolicy = static_cast<FPicpProjectionMPCDIPolicy*>(It.Get());
+					if (PicpMPCDIPolicy != nullptr)
 					{
-						MPCDIPolicy->UpdateOverlayViewportData(OverlayFrameData);
+						PicpMPCDIPolicy->UpdateOverlayViewportData(OverlayFrameData);
 					}
 				}
 			}
