nDisplay features

diff --git a/Engine/Plugins/Runtime/nDisplay/Source/PicpProjection/Private/Policy/MPCDI/PicpProjectionMPCDIPolicy.cpp b/Engine/Plugins/Runtime/nDisplay/Source/PicpProjection/Private/Policy/MPCDI/PicpProjectionMPCDIPolicy.cpp
--- a/Engine/Plugins/Runtime/nDisplay/Source/PicpProjection/Private/Policy/MPCDI/PicpProjectionMPCDIPolicy.cpp
+++ b/Engine/Plugins/Runtime/nDisplay/Source/PicpProjection/Private/Policy/MPCDI/PicpProjectionMPCDIPolicy.cpp
@@ -13,8 +13,7 @@
 
 #include "Config/DisplayClusterConfigTypes.h"
 #include "Game/IDisplayClusterGameManager.h"
-
-#include "XRThreadUtils.h"
+#include "Render/IDisplayClusterRenderManager.h"
 
 #include "Engine/RendererSettings.h"
 #include "Misc/Paths.h"
@@ -23,7 +22,9 @@
 static TAutoConsoleVariable<int32> CVarPicpShowFakeCamera(
 	TEXT("nDisplay.render.picp.ShowFakeCamera"),
 	0,
-	TEXT("Show fake debug camera (0 = disabled)"),
+	TEXT("Show fake debug camera\n")
+	TEXT(" 0: disabled\n")
+	TEXT(" 1: enabled\n"),
 	ECVF_RenderThreadSafe
 );
 
@@ -66,39 +67,32 @@ bool FPicpProjectionMPCDIPolicy::HandleAddViewport(const FIntPoint& InViewportSi
 	check(IsInGameThread());
 	check(InViewsAmount > 0);
 
-	FString File;
-	FString Buffer;
-	FString Region;
 
-	// Read MPCDI config data from nDisplay config file
-	if (!ReadConfigData(GetViewportId(), File, Buffer, Region, OriginCompId))
-	{
-		UE_LOG(LogPicpProjectionMPCDI, Error, TEXT("Couldn't read MPCDI configuration from the config file"));
-		return false;
-	}
+	IMPCDI::ConfigParser CfgData;
+	{//Load config:
+		FDisplayClusterConfigProjection CfgProjection;
+		if (!DisplayClusterHelpers::config::GetViewportProjection(GetViewportId(), CfgProjection))
+		{
+			UE_LOG(LogPicpProjectionMPCDI, Error, TEXT("Couldn't obtain projection info for '%s' viewport"), *GetViewportId());
+			return false;
+		}
 
-	// Check if MPCDI file exists
-	if (!FPaths::FileExists(File))
-	{
-		UE_LOG(LogPicpProjectionMPCDI, Warning, TEXT("File not found: %s"), *File);
-		return false;
+		if (!MPCDIAPI.LoadConfig(CfgProjection.Params, CfgData))
+		{
+			UE_LOG(LogPicpProjectionMPCDI, Error, TEXT("Couldn't read MPCDI configuration from the config file"));
+			return false;
+		}
 	}
 
-	// Load MPCDI file
-	if (!MPCDIAPI.Load(File))
-	{
-		UE_LOG(LogPicpProjectionMPCDI, Warning, TEXT("Couldn't load MPCDI file: %s"), *File);
-		return false;
-	}
 
-	// Store MPCDI region locator for this viewport
-	if (!MPCDIAPI.GetRegionLocator(File, Buffer, Region, WarpRef))
+	// Load MPCDI config
+	if (!MPCDIAPI.Load(CfgData, WarpRef))
 	{
-		UE_LOG(LogPicpProjectionMPCDI, Warning, TEXT("Couldn't get region locator for <buf %s, reg %s> in file: %s # %s "), *Buffer, *Region, *File);
+		UE_LOG(LogPicpProjectionMPCDI, Warning, TEXT("Couldn't load PICP MPCDI config: %s"), *CfgData.ConfigLineStr);
 		return false;
 	}
-
-	UE_LOG(LogPicpProjectionMPCDI, Log, TEXT("MPCDI policy has been initialized [%s:%s in %s]"), *Buffer, *Region, *File);
+	
+	UE_LOG(LogPicpProjectionMPCDI, Log, TEXT("PICP MPCDI policy has been initialized [%s]"), *CfgData.ConfigLineStr);
 
 	// Finally, initialize internal views data container
 	Views.AddDefaulted(InViewsAmount);
@@ -188,8 +182,12 @@ void FPicpProjectionMPCDIPolicy::ApplyWarpBlend_RenderThread(const uint32 ViewId
 	{
 		UE_LOG(LogPicpProjectionMPCDI, Warning, TEXT("Couldn't initialize rendering resources"));
 		return;
+	
 	}
 
+	//Support runtime PFM reload
+	MPCDIAPI.ReloadAll_RenderThread();
+	
 	FPicpProjectionOverlayViewportData ViewportOverlayData;
 	ViewportOverlayData.Initialize(OverlayViewportData);
 
@@ -211,8 +209,44 @@ void FPicpProjectionMPCDIPolicy::ApplyWarpBlend_RenderThread(const uint32 ViewId
 	{
 		UE_LOG(LogPicpProjectionMPCDI, Warning, TEXT("Couldn't apply warp&blend"));
 		return;
+	}	
+	
+	// Copy rtt viewport into camera texture RTT
+	IDisplayClusterRenderManager* const Manager = IDisplayCluster::Get().GetRenderMgr();
+	if (Manager)
+	{
+		FIntRect RTTViewportRect;
+		if (ViewportOverlayData.Cameras.Num() > 0)
+		{
+			bool RTTViewportFound = Manager->GetViewportRect(ViewportOverlayData.Cameras[0].RTTViewportId, RTTViewportRect);
+			if (RTTViewportFound)
+			{
+				FResolveParams RTTViewportCopyParams;
+				RTTViewportCopyParams.DestArrayIndex = 0;
+				RTTViewportCopyParams.SourceArrayIndex = 0;
+
+				RTTViewportCopyParams.Rect.X1 = RTTViewportRect.Min.X;
+				RTTViewportCopyParams.Rect.X2 = RTTViewportRect.Max.X;
+				RTTViewportCopyParams.Rect.Y1 = RTTViewportRect.Min.Y;
+				RTTViewportCopyParams.Rect.Y2 = RTTViewportRect.Max.Y;
+
+				RTTViewportCopyParams.DestRect.X1 = 0;
+				RTTViewportCopyParams.DestRect.Y1 = 0;
+				RTTViewportCopyParams.DestRect.X2 = ViewportOverlayData.Cameras[0].CameraTexture->GetSizeXYZ().X;
+				RTTViewportCopyParams.DestRect.Y2 = ViewportOverlayData.Cameras[0].CameraTexture->GetSizeXYZ().Y;
+
+				RHICmdList.CopyToResolveTarget(SrcTexture, ViewportOverlayData.Cameras[0].CameraTexture, RTTViewportCopyParams);
+			}
+			else
+			{
+				UE_LOG(LogPicpProjectionMPCDI, Error, TEXT("PICP Policy: RTT Viewport is not set. Assign in PICPOverlayFrameBlendingParameters."));
+			}
+		}
 	}
 
+	//Execute composure
+	IPicpMPCDI::Get().ExecuteCompose();
+
 	// Perform warp&blend shader call (Picp module with data from base mpdi module)
 	if (!PicpMPCDIAPI.ApplyWarpBlend(RHICmdList, TextureWarpData, ShaderInputData, MPCDIData.Get(), &ViewportOverlayData))
 	{
@@ -241,6 +275,7 @@ void FPicpProjectionMPCDIPolicy::ApplyWarpBlend_RenderThread(const uint32 ViewId
 	{
 		RHICmdList.CopyToResolveTarget(Views[ViewIdx].RTTexture, Views[ViewIdx].ExtWarpTexture, copyParams);
 		Views[ViewIdx].ExtWarpFrustum = Views[ViewIdx].Frustum; // Make snapshot of Frustum data for captured warp texture
+		Views[ViewIdx].ExtWarpTexture = nullptr; // Stop capture right now
 	}
 }
 
@@ -248,50 +283,6 @@ void FPicpProjectionMPCDIPolicy::ApplyWarpBlend_RenderThread(const uint32 ViewId
 //////////////////////////////////////////////////////////////////////////////////////////////
 // FDisplayClusterProjectionMPCDIPolicy
 //////////////////////////////////////////////////////////////////////////////////////////////
-bool FPicpProjectionMPCDIPolicy::ReadConfigData(const FString& ViewportId, FString& OutFile, FString& OutBuffer, FString& OutRegion, FString& OutOrigin)
-{
-	// Get projection settings of the specified viewport
-	FDisplayClusterConfigProjection CfgProjection;
-	if (!DisplayClusterHelpers::config::GetViewportProjection(ViewportId, CfgProjection))
-	{
-		UE_LOG(LogPicpProjectionMPCDI, Error, TEXT("Couldn't obtain projection info for '%s' viewport"), *ViewportId);
-		return false;
-	}
-
-	// MPCDI file
-	if (!DisplayClusterHelpers::str::ExtractValue(CfgProjection.Params, PicpProjectionStrings::cfg::data::projection::mpcdi::File, OutFile))
-	{
-		UE_LOG(LogPicpProjectionMPCDI, Error, TEXT("Argument '%s' not found in the config file"), PicpProjectionStrings::cfg::data::projection::mpcdi::File);
-		return false;
-	}
-
-	// Buffer
-	if (!DisplayClusterHelpers::str::ExtractValue(CfgProjection.Params, PicpProjectionStrings::cfg::data::projection::mpcdi::Buffer, OutBuffer))
-	{
-		UE_LOG(LogPicpProjectionMPCDI, Error, TEXT("Argument '%s' not found in the config file"), PicpProjectionStrings::cfg::data::projection::mpcdi::Buffer);
-		return false;
-	}
-
-	// Region
-	if (!DisplayClusterHelpers::str::ExtractValue(CfgProjection.Params, PicpProjectionStrings::cfg::data::projection::mpcdi::Region, OutRegion))
-	{
-		UE_LOG(LogPicpProjectionMPCDI, Error, TEXT("Argument '%s' not found in the config file"), PicpProjectionStrings::cfg::data::projection::mpcdi::Region);
-		return false;
-	}
-
-	// Origin node (optional)
-	if (DisplayClusterHelpers::str::ExtractValue(CfgProjection.Params, PicpProjectionStrings::cfg::data::projection::mpcdi::Origin, OutOrigin))
-	{
-		UE_LOG(LogPicpProjectionMPCDI, Log, TEXT("Found origin node for %s:%s - %s"), *OutBuffer, *OutRegion, *OutOrigin);
-	}
-	else
-	{
-		UE_LOG(LogPicpProjectionMPCDI, Log, TEXT("No origin node found for %s:%s. VR root will be used as default."), *OutBuffer, *OutRegion);
-	}
-
-	return true;
-}
-
 bool FPicpProjectionMPCDIPolicy::InitializeResources_RenderThread()
 {
 	check(IsInRenderingThread());
@@ -311,10 +302,12 @@ bool FPicpProjectionMPCDIPolicy::InitializeResources_RenderThread()
 				FTexture2DRHIRef DummyTexRef;
 				RHICreateTargetableShaderResource2D(ViewportSize.X, ViewportSize.Y, SceneTargetFormat, 1, TexCreate_None, TexCreate_RenderTargetable, false, CreateInfo, It.RTTexture, DummyTexRef);
 			}
+
+			bIsRenderResourcesInitialized = true;
 		}
 	}
 
-	return true;
+	return bIsRenderResourcesInitialized;
 }
 
 
@@ -374,24 +367,13 @@ void FPicpProjectionOverlayViewportData::Initialize(const FPicpProjectionOverlay
 {
 	Empty(); //Clear prev data immediatelly
 
-	if (Source.LUTCorrection)
-	{
-		LUTCorrection = new FPicpProjectionOverlayLUT(*Source.LUTCorrection);
-	}
-
-	if (Source.ViewportOver)
-	{
-		ViewportOver = new FPicpProjectionOverlayViewport(*Source.ViewportOver);
-	}
+	LUTCorrection  = Source.LUTCorrection;
+	 ViewportOver  = Source.ViewportOver;
+	 ViewportUnder = Source.ViewportUnder;
 
-	if (Source.ViewportUnder)
+	for (const auto& It : Source.Cameras)
 	{
-		ViewportUnder = new FPicpProjectionOverlayViewport(*Source.ViewportUnder);
-	}
-
-	for (auto It : Source.Cameras)
-	{
-		Cameras.Add(new FPicpProjectionOverlayCamera(*It));
+		Cameras.Add(It);
 	}
 }
 
@@ -400,29 +382,12 @@ void FPicpProjectionOverlayViewportData::Empty()
 {
 	check(IsInRenderingThread());
 
-	if (LUTCorrection)
-	{
-		delete LUTCorrection;
-		LUTCorrection = nullptr;
-	}
-
-	if (ViewportOver)
-	{
-		delete ViewportOver;
-		ViewportOver = nullptr;
-	}
-
-	if (ViewportUnder)
-	{
-		delete ViewportUnder;
-		ViewportUnder = nullptr;
-	}
+	LUTCorrection.Empty();
+	ViewportOver.Empty();
+	ViewportUnder.Empty();
 
 	for (auto& It : Cameras)
-	{
-		delete It;
-	}
-
+		{ It.Empty(); }
 	Cameras.Empty();
 }
 
@@ -488,27 +453,6 @@ static void SetOverlayData_RenderThread(FRHICommandListImmediate& RHICmdList, FP
 
 void FPicpProjectionOverlayFrameData::Empty()
 {
-	if (LUTCorrection)
-	{
-		delete LUTCorrection;
-		LUTCorrection = nullptr;
-	}
-
-	for (auto& It : ViewportsOver)
-	{
-		delete It.Value;
-	}
-
-	for (auto& It : ViewportsUnder)
-	{
-		delete It.Value;
-	}
-
-	for (auto& It : Cameras)
-	{
-		delete It;
-	}
-
 	ViewportsOver.Empty();
 	ViewportsUnder.Empty();
 	Cameras.Empty();
@@ -520,31 +464,28 @@ void FPicpProjectionOverlayFrameData::GetViewportData(const FString& ViewportId,
 
 	FPicpProjectionOverlayViewportData* GameThreadData = new FPicpProjectionOverlayViewportData();
 
-	// Check TSharedPtr counters
-	if (LUTCorrection)
-	{
-		GameThreadData->LUTCorrection = new FPicpProjectionOverlayLUT(*LUTCorrection);
-	}
+	GameThreadData->LUTCorrection = LUTCorrection;
 
-	for (auto It : Cameras) {
+	// Copy cameras data
+	for (const auto& It : Cameras) {
 
-		FPicpProjectionOverlayCamera* Cam = new FPicpProjectionOverlayCamera(*It);
+		FPicpProjectionOverlayCamera Camera(It);
 
 		// Update view matrix to cave space
-		Cam->ViewRot = Origin2WorldTransform.InverseTransformRotation(It->ViewRot.Quaternion()).Rotator();
-		Cam->ViewLoc = Origin2WorldTransform.InverseTransformPosition(It->ViewLoc);
+		Camera.ViewRot = Origin2WorldTransform.InverseTransformRotation(Camera.ViewRot.Quaternion()).Rotator();
+		Camera.ViewLoc = Origin2WorldTransform.InverseTransformPosition(Camera.ViewLoc);
 
-		GameThreadData->Cameras.Add(Cam);
+		GameThreadData->Cameras.Add(Camera);
 	}
 
 	if (ViewportsOver.Contains(ViewportId))
 	{
-		GameThreadData->ViewportOver = new FPicpProjectionOverlayViewport(*ViewportsOver[ViewportId]);
+		GameThreadData->ViewportOver = ViewportsOver[ViewportId];
 	}
 
 	if (ViewportsUnder.Contains(ViewportId))
 	{
-		GameThreadData->ViewportUnder = new FPicpProjectionOverlayViewport(*ViewportsUnder[ViewportId]);
+		GameThreadData->ViewportUnder = ViewportsUnder[ViewportId];
 	}
 
 	//Send data to render thread:
@@ -561,9 +502,9 @@ void FPicpProjectionOverlayFrameData::GenerateDebugContent(const FString& Viewpo
 	float Yaw = 45; // Floor rotation
 	float Fov = 10;
 
-	FPicpProjectionOverlayCamera* DebugCamera = new FPicpProjectionOverlayCamera(
+	FPicpProjectionOverlayCamera DebugCamera(
 		FRotator(0, Yaw, 0),FVector(0,0,0),
-		GetProjectionMatrix(Fov, 1, 20000), GWhiteTexture->TextureRHI);
+		GetProjectionMatrix(Fov, 1, 20000), GWhiteTexture->TextureRHI, TEXT("DebugInnerCameraViewport"));
 
 	FPicpProjectionOverlayViewportData* GameThreadData = new FPicpProjectionOverlayViewportData();
 	GameThreadData->Cameras.Add(DebugCamera);
diff --git a/Engine/Plugins/Runtime/nDisplay/Source/PicpProjection/Private/Policy/MPCDI/PicpProjectionMPCDIPolicy.h b/Engine/Plugins/Runtime/nDisplay/Source/PicpProjection/Private/Policy/MPCDI/PicpProjectionMPCDIPolicy.h
--- a/Engine/Plugins/Runtime/nDisplay/Source/PicpProjection/Private/Policy/MPCDI/PicpProjectionMPCDIPolicy.h
+++ b/Engine/Plugins/Runtime/nDisplay/Source/PicpProjection/Private/Policy/MPCDI/PicpProjectionMPCDIPolicy.h
@@ -30,8 +30,8 @@ public:
 	virtual ~FPicpProjectionMPCDIPolicy();
 
 	//////////////////////////////////////////////////////////////////////////////////////////////
-		// IDisplayClusterProjectionPolicy
-		//////////////////////////////////////////////////////////////////////////////////////////////
+	// IDisplayClusterProjectionPolicy
+	//////////////////////////////////////////////////////////////////////////////////////////////
 	virtual void StartScene(UWorld* World) override;
 	virtual void EndScene() override;
 	virtual bool HandleAddViewport(const FIntPoint& ViewportSize, const uint32 ViewsAmount) override;
@@ -49,9 +49,8 @@ public:
 	void SetWarpTextureCapture(const uint32 ViewIdx, FRHITexture2D* target);
 	IMPCDI::FFrustum GetWarpFrustum(const uint32 ViewIdx, bool bIsCaptureWarpTextureFrustum);
 
-protected:
-	bool ReadConfigData(const FString& ViewportId, FString& OutFile, FString& OutBuffer, FString& OutRegion, FString& OutOrigin);
-	bool InitializeResources_RenderThread();
+protected:	
+	bool InitializeResources_RenderThread();	
 
 private:
 	FString OriginCompId;
@@ -77,6 +76,6 @@ private:
 
 	TArray<FViewData> Views;
 
-	bool bIsRenderResourcesInitialized = false;
-	FCriticalSection RenderingResourcesInitializationCS;
+	mutable bool bIsRenderResourcesInitialized = false;
+	mutable FCriticalSection RenderingResourcesInitializationCS;
 };
diff --git a/Engine/Plugins/Runtime/nDisplay/Source/PicpProjection/Public/Blueprints/IPicpProjectionBlueprintAPI.h b/Engine/Plugins/Runtime/nDisplay/Source/PicpProjection/Public/Blueprints/IPicpProjectionBlueprintAPI.h
--- a/Engine/Plugins/Runtime/nDisplay/Source/PicpProjection/Public/Blueprints/IPicpProjectionBlueprintAPI.h
+++ b/Engine/Plugins/Runtime/nDisplay/Source/PicpProjection/Public/Blueprints/IPicpProjectionBlueprintAPI.h
@@ -48,4 +48,10 @@ public:
 
 	UFUNCTION(BlueprintCallable, meta = (DisplayName = "Capture Final Warped Frame"), Category = "PICP")
 	virtual void SetWarpTextureCaptureState(UTextureRenderTarget2D* dstTexture, const FString& ViewportId, const int ViewIdx, bool bCaptureNow) = 0;
+
+	/**
+	 * Creates a new render target and initializes it to the specified dimensions
+	 */
+	UFUNCTION(BlueprintCallable, meta = (DisplayName = "Enable mips for Texture Render Target"), Category = "PICP")
+	virtual void EnableTextureRenderTargetMips(UTextureRenderTarget2D* Texture) = 0;
 };
diff --git a/Engine/Plugins/Runtime/nDisplay/Source/PicpProjection/Public/Blueprints/ReprojectionBlendingData.h b/Engine/Plugins/Runtime/nDisplay/Source/PicpProjection/Public/Blueprints/ReprojectionBlendingData.h
--- a/Engine/Plugins/Runtime/nDisplay/Source/PicpProjection/Public/Blueprints/ReprojectionBlendingData.h
+++ b/Engine/Plugins/Runtime/nDisplay/Source/PicpProjection/Public/Blueprints/ReprojectionBlendingData.h
@@ -7,6 +7,7 @@
 #include "OverlayRenderingParameters.h"
 
 #include "Components/SceneCaptureComponent2D.h"
+#include "CineCameraComponent.h"
 #include "TextureResource.h"
 
 #include "ReprojectionBlendingData.generated.h"
@@ -58,7 +59,10 @@ public:
 	FVector  SoftEdge;    // Basic soft edges values
 
 	UPROPERTY(EditAnywhere, BlueprintReadWrite, Category = "PICP")
-	USceneCaptureComponent2D* CameraOverlayFrameCapture;
+	UCineCameraComponent* CineCamera;
+
+	UPROPERTY(EditAnywhere, BlueprintReadWrite, Category = "PICP")
+	FString RTTViewportId;
 
 	UPROPERTY(EditAnywhere, BlueprintReadWrite, Category = "PICP")
 	TArray<struct FPicpOverlayFrameBlendingPair> OverlayBlendFrames;
diff --git a/Engine/Plugins/Runtime/nDisplay/Source/PicpProjection/Public/Overlay/PicpProjectionOverlayCamera.h b/Engine/Plugins/Runtime/nDisplay/Source/PicpProjection/Public/Overlay/PicpProjectionOverlayCamera.h
--- a/Engine/Plugins/Runtime/nDisplay/Source/PicpProjection/Public/Overlay/PicpProjectionOverlayCamera.h
+++ b/Engine/Plugins/Runtime/nDisplay/Source/PicpProjection/Public/Overlay/PicpProjectionOverlayCamera.h
@@ -2,7 +2,6 @@
 
 #pragma once
 
-#include "PicpProjectionOverlayBase.h"
 #include "CoreMinimal.h"
 #include "RHI.h"
 
@@ -10,13 +9,13 @@ class FMPCDIData;
 
 
 class FPicpProjectionOverlayCamera
-	: public FPicpProjectionOverlayBase
 {
 public:
 	FVector  SoftEdge;    // Basic soft edges values
 	//@ Add more render options here
 	
 	//Camera setup
+	FString RTTViewportId;
 	FRHITexture* CameraTexture; // Texture to render
 	FMatrix Prj; // Projection matrix
 
@@ -24,23 +23,22 @@ public:
 	FVector  ViewLoc;
 
 public:
-	FPicpProjectionOverlayCamera(const FRotator& CameraRotation, const FVector& CameraLocation, const FMatrix& CameraPrj, FRHITexture* CameraTextureRef)
-		: FPicpProjectionOverlayBase()
-		, SoftEdge(0.1f, 0.1f, 0.1f)
+	FPicpProjectionOverlayCamera(const FRotator& CameraRotation, const FVector& CameraLocation, const FMatrix& CameraPrj, FRHITexture* CameraTextureRef, const FString& ViewportId)
+		: SoftEdge(0.1f, 0.1f, 0.1f)
+		, RTTViewportId(ViewportId)
 		, CameraTexture(CameraTextureRef)
 		, Prj(CameraPrj)
 		, ViewRot(CameraRotation)
 		, ViewLoc(CameraLocation)
-	{ 
-		SetEnable(true);
+	{
 	}
 
-	virtual ~FPicpProjectionOverlayCamera()
-	{ }
-
 	const FMatrix GetRuntimeCameraProjection() const
 	{ return RuntimeCameraProjection; }
 
+	void Empty()
+	{}
+
 public:
 	FMatrix RuntimeCameraProjection;
 };
diff --git a/Engine/Plugins/Runtime/nDisplay/Source/PicpProjection/Public/Overlay/PicpProjectionOverlayLUT.h b/Engine/Plugins/Runtime/nDisplay/Source/PicpProjection/Public/Overlay/PicpProjectionOverlayLUT.h
--- a/Engine/Plugins/Runtime/nDisplay/Source/PicpProjection/Public/Overlay/PicpProjectionOverlayLUT.h
+++ b/Engine/Plugins/Runtime/nDisplay/Source/PicpProjection/Public/Overlay/PicpProjectionOverlayLUT.h
@@ -2,12 +2,8 @@
 
 #pragma once
 
-#include "PicpProjectionOverlayBase.h"
-
-
 // Root data container. Render all over basic
 class FPicpProjectionOverlayLUT
-	: public FPicpProjectionOverlayBase
 {
 public:
 	FRHITexture2D* LUTTexture;           // if texture defined, applyed brightness correction from view angle
@@ -20,16 +16,30 @@ public:
 		: LUTTexture(nullptr)
 		, ViewCorrectionCoeff(0,0,0) 
 		, EyeOrigin(0,0,0)
-	{ }
+		, bIsEnabled(false)
+	{ 
+	}
 
 	void Initialize(FRHITexture2D*& TextureRef, const FVector& ViewCorrection, const FVector& EyeLocation)
 	{
 		LUTTexture = TextureRef;
 		ViewCorrectionCoeff = ViewCorrection;
 		EyeOrigin = EyeLocation;
-		SetEnable(true);
+		bIsEnabled = true;
+	}
+
+	bool IsEnabled() const 
+	{ 
+		return bIsEnabled; 
+	}
+
+	void Empty()
+	{
+		LUTTexture = nullptr;
+		bIsEnabled = false;
 	}
 
 private:
+	bool    bIsEnabled;
 	FVector EyeOriginLocal; // eye in local mpcdi space
 };
diff --git a/Engine/Plugins/Runtime/nDisplay/Source/PicpProjection/Public/Overlay/PicpProjectionOverlayRender.h b/Engine/Plugins/Runtime/nDisplay/Source/PicpProjection/Public/Overlay/PicpProjectionOverlayRender.h
--- a/Engine/Plugins/Runtime/nDisplay/Source/PicpProjection/Public/Overlay/PicpProjectionOverlayRender.h
+++ b/Engine/Plugins/Runtime/nDisplay/Source/PicpProjection/Public/Overlay/PicpProjectionOverlayRender.h
@@ -12,10 +12,8 @@ class FPicpProjectionOverlayViewportData
 {
 public:
 	FPicpProjectionOverlayViewportData()
-		: LUTCorrection(nullptr)
-		, ViewportOver(nullptr)
-		, ViewportUnder(nullptr)
-	{ }
+	{ 
+	}
 
 	~FPicpProjectionOverlayViewportData()
 	{
@@ -26,30 +24,42 @@ public:
 	void Empty();
 
 	bool isAnyViewportUsed() const 
-	{ return iViewportUnderUsed() || iViewportOverUsed(); }
+	{ 
+		return iViewportUnderUsed() || iViewportOverUsed(); 
+	}
 
 	bool iViewportOverUsed() const
-	{ return ViewportOver!=nullptr && ViewportOver->IsEnabled(); }
+	{ 
+		return ViewportOver.IsEnabled(); 
+	}
 
 	bool iViewportUnderUsed() const
-	{ return ViewportUnder!=nullptr && ViewportUnder->IsEnabled(); }
+	{ 
+		return ViewportUnder.IsEnabled(); 
+	}
 
 	bool isAnyCameraUsed() const
-	{ return (Cameras.Num() > 0); }
+	{ 
+		return (Cameras.Num() > 0); 
+	}
 
 	bool isLUTused() const
-	{ return LUTCorrection!=nullptr && LUTCorrection->IsEnabled(); }
+	{ 
+		return LUTCorrection.IsEnabled(); 
+	}
 
 	bool isValid()
-	{ return isAnyViewportUsed() || isAnyCameraUsed(); }
+	{ 
+		return isAnyViewportUsed() || isAnyCameraUsed(); 
+	}
 
 	void Initialize(const FPicpProjectionOverlayViewportData& Source);
 
 public:
-	FPicpProjectionOverlayLUT*               LUTCorrection;
-	FPicpProjectionOverlayViewport*          ViewportOver; // viewport overlay texture
-	FPicpProjectionOverlayViewport*          ViewportUnder; // viewport overlay texture
-	TArray<FPicpProjectionOverlayCamera*>    Cameras; // Multi cams, in render order
+	FPicpProjectionOverlayLUT                LUTCorrection;
+	FPicpProjectionOverlayViewport           ViewportOver; // viewport overlay texture
+	FPicpProjectionOverlayViewport           ViewportUnder; // viewport overlay texture
+	TArray<FPicpProjectionOverlayCamera>     Cameras; // Multi cams, in render order
 };
 
 
@@ -61,8 +71,8 @@ class FPicpProjectionOverlayFrameData
 {
 public:
 	FPicpProjectionOverlayFrameData()
-		: LUTCorrection(nullptr)
-	{ }
+	{ 
+	}
 
 	~FPicpProjectionOverlayFrameData()
 	{
@@ -75,8 +85,8 @@ public:
 	void Empty();
 
 public:
-	FPicpProjectionOverlayLUT*                      LUTCorrection;  // Global LUT correction
-	TMap<FString, FPicpProjectionOverlayViewport*>  ViewportsOver;  // Overlay under inner frame for all viewports by name
-	TMap<FString, FPicpProjectionOverlayViewport*>  ViewportsUnder; // Overlay on top of inner frame for all viewports by name
-	TArray<FPicpProjectionOverlayCamera*>           Cameras;        // Camera's overlay, in render order
+	FPicpProjectionOverlayLUT                      LUTCorrection;  // Global LUT correction
+	TMap<FString, FPicpProjectionOverlayViewport>  ViewportsOver;  // Overlay under inner frame for all viewports by name
+	TMap<FString, FPicpProjectionOverlayViewport>  ViewportsUnder; // Overlay on top of inner frame for all viewports by name
+	TArray<FPicpProjectionOverlayCamera>           Cameras;        // Camera's overlay, in render order
 };
diff --git a/Engine/Plugins/Runtime/nDisplay/Source/PicpProjection/Public/Overlay/PicpProjectionOverlayViewport.h b/Engine/Plugins/Runtime/nDisplay/Source/PicpProjection/Public/Overlay/PicpProjectionOverlayViewport.h
--- a/Engine/Plugins/Runtime/nDisplay/Source/PicpProjection/Public/Overlay/PicpProjectionOverlayViewport.h
+++ b/Engine/Plugins/Runtime/nDisplay/Source/PicpProjection/Public/Overlay/PicpProjectionOverlayViewport.h
@@ -2,26 +2,34 @@
 
 #pragma once
 
-#include "PicpProjectionOverlayBase.h"
 #include "CoreMinimal.h"
 #include "RHI.h"
 
 
 class FPicpProjectionOverlayViewport
-	: public FPicpProjectionOverlayBase
 {
 public:
 	FRHITexture2D* ViewportTexture; // Texture to render
 	//@ Add more render options here
 
 public:
+	FPicpProjectionOverlayViewport()
+		: ViewportTexture(nullptr)
+	{ 
+	}
+
 	FPicpProjectionOverlayViewport(FRHITexture2D* TextureRef)
-		: FPicpProjectionOverlayBase()
-		, ViewportTexture(TextureRef)
+		: ViewportTexture(TextureRef)
 	{ 
-		SetEnable(TextureRef->IsValid());
+	}
+	
+	bool IsEnabled() const
+	{ 
+		return ViewportTexture!=nullptr && ViewportTexture->IsValid(); 
 	}
 
-	virtual ~FPicpProjectionOverlayViewport()
-	{ }
+	void Empty()
+	{ 
+		ViewportTexture = nullptr;  
+	}
 };
diff --git a/Engine/Plugins/Runtime/nDisplay/nDisplay.uplugin b/Engine/Plugins/Runtime/nDisplay/nDisplay.uplugin
--- a/Engine/Plugins/Runtime/nDisplay/nDisplay.uplugin
+++ b/Engine/Plugins/Runtime/nDisplay/nDisplay.uplugin
@@ -47,6 +47,14 @@
 				"Win64"
 			]
 		},
+		{
+			"Name": "OutputRemap",
+			"Type": "Runtime",
+			"LoadingPhase": "PostConfigInit",
+			"WhitelistPlatforms": [
+				"Win64"
+			]
+		},
 		{
 			"Name": "DisplayClusterProjection",
 			"Type": "Runtime",
@@ -55,6 +63,14 @@
 				"Win64"
 			]
 		},
+		{
+			"Name": "DisplayClusterPostprocess",
+			"Type": "Runtime",
+			"LoadingPhase": "Default",
+			"WhitelistPlatforms": [
+				"Win64"
+			]
+		},
 		{
 			"Name": "PicpMPCDI",
 			"Type": "Runtime",
diff --git a/Engine/Source/Runtime/Engine/Private/SceneView.cpp b/Engine/Source/Runtime/Engine/Private/SceneView.cpp
--- a/Engine/Source/Runtime/Engine/Private/SceneView.cpp
+++ b/Engine/Source/Runtime/Engine/Private/SceneView.cpp
@@ -1900,11 +1900,9 @@ void FSceneView::EndFinalPostprocessSettings(const FSceneViewInitOptions& ViewIn
 #endif
 
 	{
-		const bool bStereoEnabled = StereoPass != eSSP_FULL;
-		const bool bScaledToRenderTarget = GEngine->XRSystem.IsValid() && bStereoEnabled && GEngine->XRSystem->GetHMDDevice();
-		if (bScaledToRenderTarget)
+		if (GEngine->StereoRenderingDevice.IsValid())
 		{
-			GEngine->XRSystem->GetHMDDevice()->UpdatePostProcessSettings(&FinalPostProcessSettings);
+			GEngine->StereoRenderingDevice->UpdatePostProcessSettings(&FinalPostProcessSettings, StereoPass);
 		}
 	}
 
diff --git a/Engine/Source/Runtime/Engine/Public/StereoRendering.h b/Engine/Source/Runtime/Engine/Public/StereoRendering.h
--- a/Engine/Source/Runtime/Engine/Public/StereoRendering.h
+++ b/Engine/Source/Runtime/Engine/Public/StereoRendering.h
@@ -161,4 +161,9 @@ public:
 	 */
 	virtual IStereoLayers* GetStereoLayers () { return nullptr; }
 
+	/**
+	 * Allows to override the PostProcessSettings in the last moment e.g. allows up sampled 3D rendering
+	 */
+	virtual void UpdatePostProcessSettings(struct FPostProcessSettings* FinalPostProcessingSettings, const enum EStereoscopicPass StereoPassType) {}
+
 };
diff --git a/Engine/Source/Runtime/HeadMountedDisplay/Public/IHeadMountedDisplay.h b/Engine/Source/Runtime/HeadMountedDisplay/Public/IHeadMountedDisplay.h
--- a/Engine/Source/Runtime/HeadMountedDisplay/Public/IHeadMountedDisplay.h
+++ b/Engine/Source/Runtime/HeadMountedDisplay/Public/IHeadMountedDisplay.h
@@ -203,12 +203,7 @@ public:
 	 * It is called at the beginning of UGameViewportClient::Draw() method.
 	 * We might remove this one as UpdatePostProcessSettings should be able to capture all needed cases
 	 */
-	virtual void UpdateScreenSettings(const class FViewport* InViewport) {}
-
-	/**
-	 * Allows to override the PostProcessSettings in the last moment e.g. allows up sampled 3D rendering
-	 */
-	virtual void UpdatePostProcessSettings(FPostProcessSettings*) {}
+	virtual void UpdateScreenSettings(const class FViewport* InViewport) {}	
 
 	/** 
 	 * Additional optional distortion rendering parameters
