Fix CIS, missing include

diff --git a/Engine/Plugins/Runtime/nDisplay/Source/DisplayCluster/Public/Blueprints/IDisplayClusterBlueprintAPI.h b/Engine/Plugins/Runtime/nDisplay/Source/DisplayCluster/Public/Blueprints/IDisplayClusterBlueprintAPI.h
--- a/Engine/Plugins/Runtime/nDisplay/Source/DisplayCluster/Public/Blueprints/IDisplayClusterBlueprintAPI.h
+++ b/Engine/Plugins/Runtime/nDisplay/Source/DisplayCluster/Public/Blueprints/IDisplayClusterBlueprintAPI.h
@@ -10,10 +10,10 @@
 #include "Cluster/IDisplayClusterClusterEventListener.h"
 
 #include "Config/DisplayClusterConfigTypes.h"
+#include "Engine/Scene.h"
 
 #include "IDisplayClusterBlueprintAPI.generated.h"
 
-struct FPostProcessSettings;
 
 
 
