Fix CIS NonUnity build

diff --git a/Engine/Plugins/Runtime/nDisplay/Source/DisplayCluster/Private/Render/Device/DisplayClusterDeviceBase.h b/Engine/Plugins/Runtime/nDisplay/Source/DisplayCluster/Private/Render/Device/DisplayClusterDeviceBase.h
--- a/Engine/Plugins/Runtime/nDisplay/Source/DisplayCluster/Private/Render/Device/DisplayClusterDeviceBase.h
+++ b/Engine/Plugins/Runtime/nDisplay/Source/DisplayCluster/Private/Render/Device/DisplayClusterDeviceBase.h
@@ -6,6 +6,7 @@
 
 #include "RHI.h"
 #include "RHIResources.h"
+#include "Engine/Scene.h"
 #include "StereoRenderTargetManager.h"
 
 #include "Render/Device/DisplayClusterDeviceBase_PostProcess.h"
diff --git a/Engine/Plugins/Runtime/nDisplay/Source/DisplayCluster/Private/Render/Presentation/DisplayClusterDevicePresentationDX11.cpp b/Engine/Plugins/Runtime/nDisplay/Source/DisplayCluster/Private/Render/Presentation/DisplayClusterDevicePresentationDX11.cpp
--- a/Engine/Plugins/Runtime/nDisplay/Source/DisplayCluster/Private/Render/Presentation/DisplayClusterDevicePresentationDX11.cpp
+++ b/Engine/Plugins/Runtime/nDisplay/Source/DisplayCluster/Private/Render/Presentation/DisplayClusterDevicePresentationDX11.cpp
@@ -4,6 +4,7 @@
 
 #include "DisplayClusterLog.h"
 
+#include "RHI.h"
 #include "RHIResources.h"
 
 #include "Windows/AllowWindowsPlatformTypes.h"
diff --git a/Engine/Plugins/Runtime/nDisplay/Source/DisplayCluster/Private/Render/Presentation/DisplayClusterDevicePresentationDX12.cpp b/Engine/Plugins/Runtime/nDisplay/Source/DisplayCluster/Private/Render/Presentation/DisplayClusterDevicePresentationDX12.cpp
--- a/Engine/Plugins/Runtime/nDisplay/Source/DisplayCluster/Private/Render/Presentation/DisplayClusterDevicePresentationDX12.cpp
+++ b/Engine/Plugins/Runtime/nDisplay/Source/DisplayCluster/Private/Render/Presentation/DisplayClusterDevicePresentationDX12.cpp
@@ -4,6 +4,7 @@
 
 #include "DisplayClusterLog.h"
 
+#include "RHI.h"
 #include "RHIResources.h"
 
 #include "Windows/AllowWindowsPlatformTypes.h"
diff --git a/Engine/Plugins/Runtime/nDisplay/Source/DisplayCluster/Private/Render/Synchronization/DisplayClusterRenderSyncPolicySoftwareDX11.cpp b/Engine/Plugins/Runtime/nDisplay/Source/DisplayCluster/Private/Render/Synchronization/DisplayClusterRenderSyncPolicySoftwareDX11.cpp
--- a/Engine/Plugins/Runtime/nDisplay/Source/DisplayCluster/Private/Render/Synchronization/DisplayClusterRenderSyncPolicySoftwareDX11.cpp
+++ b/Engine/Plugins/Runtime/nDisplay/Source/DisplayCluster/Private/Render/Synchronization/DisplayClusterRenderSyncPolicySoftwareDX11.cpp
@@ -11,6 +11,7 @@
 #include "DirectX/Include/DXGI.h"
 #include "dwmapi.h"
 #include "Windows/HideWindowsPlatformTypes.h"
+#include "RenderCore/Public/ShaderCore.h"
 
 #include "D3D11RHI/Private/Windows/D3D11RHIBasePrivate.h"
 
@@ -20,7 +21,6 @@
 
 #include "NVIDIA/nvapi/nvapi.h"
 
-
 static TAutoConsoleVariable<int32> CVarAdvancedSyncEnabled(
 	TEXT("nDisplay.render.softsync.AdvancedSyncEnabled"),
 	0,
diff --git a/Engine/Plugins/Runtime/nDisplay/Source/DisplayClusterProjection/Private/Policy/EasyBlend/DX11/DisplayClusterProjectionEasyBlendLibraryDX11.cpp b/Engine/Plugins/Runtime/nDisplay/Source/DisplayClusterProjection/Private/Policy/EasyBlend/DX11/DisplayClusterProjectionEasyBlendLibraryDX11.cpp
--- a/Engine/Plugins/Runtime/nDisplay/Source/DisplayClusterProjection/Private/Policy/EasyBlend/DX11/DisplayClusterProjectionEasyBlendLibraryDX11.cpp
+++ b/Engine/Plugins/Runtime/nDisplay/Source/DisplayClusterProjection/Private/Policy/EasyBlend/DX11/DisplayClusterProjectionEasyBlendLibraryDX11.cpp
@@ -7,6 +7,7 @@
 
 #include "HAL/PlatformProcess.h"
 #include "Interfaces/IPluginManager.h"
+#include "Misc/Paths.h"
 
 #if PLATFORM_WINDOWS
 #include "EasyBlendSDKDXApi.h"
diff --git a/Engine/Plugins/Runtime/nDisplay/Source/DisplayClusterProjection/Private/Policy/EasyBlend/DX11/DisplayClusterProjectionEasyBlendLibraryDX11.h b/Engine/Plugins/Runtime/nDisplay/Source/DisplayClusterProjection/Private/Policy/EasyBlend/DX11/DisplayClusterProjectionEasyBlendLibraryDX11.h
--- a/Engine/Plugins/Runtime/nDisplay/Source/DisplayClusterProjection/Private/Policy/EasyBlend/DX11/DisplayClusterProjectionEasyBlendLibraryDX11.h
+++ b/Engine/Plugins/Runtime/nDisplay/Source/DisplayClusterProjection/Private/Policy/EasyBlend/DX11/DisplayClusterProjectionEasyBlendLibraryDX11.h
@@ -4,6 +4,7 @@
 
 #include "CoreMinimal.h"
 
+#include "RHI.h"
 #include "RHIResources.h"
 
 #include "Windows/AllowWindowsPlatformTypes.h"
diff --git a/Engine/Plugins/Runtime/nDisplay/Source/PicpMPCDI/Private/PicpBlurPostProcess.h b/Engine/Plugins/Runtime/nDisplay/Source/PicpMPCDI/Private/PicpBlurPostProcess.h
--- a/Engine/Plugins/Runtime/nDisplay/Source/PicpMPCDI/Private/PicpBlurPostProcess.h
+++ b/Engine/Plugins/Runtime/nDisplay/Source/PicpMPCDI/Private/PicpBlurPostProcess.h
@@ -12,6 +12,7 @@
 #include "GlobalShader.h"
 #include "ShaderParameterUtils.h"
 #include "PicpPostProcessing.h"
+#include "TextureResource.h"
 
 class UTexture;
 class UTextureRenderTarget2D;
