nDisplay fixes
- simplify LED wall setup
- change some pfm loading warnings to errors

diff --git a/Engine/Plugins/Runtime/nDisplay/Source/MPCDI/Private/Blueprints/MPCDIBlueprintAPIImpl.cpp b/Engine/Plugins/Runtime/nDisplay/Source/MPCDI/Private/Blueprints/MPCDIBlueprintAPIImpl.cpp
--- a/Engine/Plugins/Runtime/nDisplay/Source/MPCDI/Private/Blueprints/MPCDIBlueprintAPIImpl.cpp
+++ b/Engine/Plugins/Runtime/nDisplay/Source/MPCDI/Private/Blueprints/MPCDIBlueprintAPIImpl.cpp
@@ -56,7 +56,7 @@ bool UMPCDIAPIImpl::GetPFMMeshData(const FString& LocalPFMFile, FMPCDIGeometryEx
 	//! todo: Force reload
 	if (!MpcdiModule.LoadPFM(RegionLocator, LocalPFMFile, PFMScale, !bIsMPCDIAxis))
 	{
-		UE_LOG(LogMPCDI, Error, TEXT("Couldn't get load ext pfm from file: %s"), *LocalPFMFile);
+		UE_LOG(LogMPCDI, Error, TEXT("Failed to load PFM from file: %s"), *LocalPFMFile);
 		return false;
 	}
 
diff --git a/Engine/Plugins/Runtime/nDisplay/Source/MPCDI/Private/MPCDIModule.cpp b/Engine/Plugins/Runtime/nDisplay/Source/MPCDI/Private/MPCDIModule.cpp
--- a/Engine/Plugins/Runtime/nDisplay/Source/MPCDI/Private/MPCDIModule.cpp
+++ b/Engine/Plugins/Runtime/nDisplay/Source/MPCDI/Private/MPCDIModule.cpp
@@ -323,7 +323,7 @@ bool FMPCDIModule::Load(const ConfigParser& CfgData, IMPCDI::FRegionLocator& Out
 		{
 			if (!LoadPFM(OutRegionLocator, CfgData.PFMFile, CfgData.PFMFileScale, CfgData.bIsUnrealGameSpace))
 			{
-				UE_LOG(LogMPCDI, Warning, TEXT("Couldn't get load ext pfm for <buf %s, reg %s> from file: %s"), *CfgData.BufferId, *CfgData.RegionId, *CfgData.PFMFile);
+				UE_LOG(LogMPCDI, Error, TEXT("Failed to load PFM <buf %s, reg %s> from file: %s"), *CfgData.BufferId, *CfgData.RegionId, *CfgData.PFMFile);
 			}
 		}
 		// AlphaMAP
