Merging fix from dev-enterprise, CL# 7647632, to avoid crashing when launching nDisplay with dx12

diff --git a/Engine/Plugins/Runtime/nDisplay/Source/DisplayCluster/Private/Render/Synchronization/DisplayClusterRenderSyncPolicyBase.cpp b/Engine/Plugins/Runtime/nDisplay/Source/DisplayCluster/Private/Render/Synchronization/DisplayClusterRenderSyncPolicyBase.cpp
--- a/Engine/Plugins/Runtime/nDisplay/Source/DisplayCluster/Private/Render/Synchronization/DisplayClusterRenderSyncPolicyBase.cpp
+++ b/Engine/Plugins/Runtime/nDisplay/Source/DisplayCluster/Private/Render/Synchronization/DisplayClusterRenderSyncPolicyBase.cpp
@@ -20,8 +20,6 @@ FDisplayClusterRenderSyncPolicyBase::~FDisplayClusterRenderSyncPolicyBase()
 
 void FDisplayClusterRenderSyncPolicyBase::SyncBarrierRenderThread()
 {
-	check(IsInRenderingThread());
-
 	DISPLAY_CLUSTER_FUNC_TRACE(LogDisplayClusterRender);
 
 	if (GDisplayCluster->GetOperationMode() == EDisplayClusterOperationMode::Disabled)
diff --git a/Engine/Plugins/Runtime/nDisplay/Source/DisplayCluster/Private/Render/Synchronization/DisplayClusterRenderSyncPolicyNone.cpp b/Engine/Plugins/Runtime/nDisplay/Source/DisplayCluster/Private/Render/Synchronization/DisplayClusterRenderSyncPolicyNone.cpp
--- a/Engine/Plugins/Runtime/nDisplay/Source/DisplayCluster/Private/Render/Synchronization/DisplayClusterRenderSyncPolicyNone.cpp
+++ b/Engine/Plugins/Runtime/nDisplay/Source/DisplayCluster/Private/Render/Synchronization/DisplayClusterRenderSyncPolicyNone.cpp
@@ -14,8 +14,6 @@ FDisplayClusterRenderSyncPolicyNone::~FDisplayClusterRenderSyncPolicyNone()
 
 bool FDisplayClusterRenderSyncPolicyNone::SynchronizeClusterRendering(int32& InOutSyncInterval)
 {
-	check(IsInRenderingThread());
-
 	// Override sync interval with 0 to show a frame ASAP. We don't care about tearing in this policy.
 	InOutSyncInterval = 0;
 	// Tell a caller that he still needs to present a frame
diff --git a/Engine/Plugins/Runtime/nDisplay/Source/DisplayCluster/Private/Render/Synchronization/DisplayClusterRenderSyncPolicyNvidiaDX11.cpp b/Engine/Plugins/Runtime/nDisplay/Source/DisplayCluster/Private/Render/Synchronization/DisplayClusterRenderSyncPolicyNvidiaDX11.cpp
--- a/Engine/Plugins/Runtime/nDisplay/Source/DisplayCluster/Private/Render/Synchronization/DisplayClusterRenderSyncPolicyNvidiaDX11.cpp
+++ b/Engine/Plugins/Runtime/nDisplay/Source/DisplayCluster/Private/Render/Synchronization/DisplayClusterRenderSyncPolicyNvidiaDX11.cpp
@@ -14,8 +14,6 @@ FDisplayClusterRenderSyncPolicyNvidiaDX11::~FDisplayClusterRenderSyncPolicyNvidi
 
 bool FDisplayClusterRenderSyncPolicyNvidiaDX11::SynchronizeClusterRendering(int32& InOutSyncInterval)
 {
-	check(IsInRenderingThread());
-
 	// Unfortunately NVIDIA doesn't export required functions in the public SDK. This make impossible
 	// to implement the hardware base synchronization approach. As a fallback we use our own soft sync.
 	// As soon as the required functions become available, the proper synchonization approach will be implemented.
diff --git a/Engine/Plugins/Runtime/nDisplay/Source/DisplayCluster/Private/Render/Synchronization/DisplayClusterRenderSyncPolicyNvidiaDX12.cpp b/Engine/Plugins/Runtime/nDisplay/Source/DisplayCluster/Private/Render/Synchronization/DisplayClusterRenderSyncPolicyNvidiaDX12.cpp
--- a/Engine/Plugins/Runtime/nDisplay/Source/DisplayCluster/Private/Render/Synchronization/DisplayClusterRenderSyncPolicyNvidiaDX12.cpp
+++ b/Engine/Plugins/Runtime/nDisplay/Source/DisplayCluster/Private/Render/Synchronization/DisplayClusterRenderSyncPolicyNvidiaDX12.cpp
@@ -14,8 +14,6 @@ FDisplayClusterRenderSyncPolicyNvidiaDX12::~FDisplayClusterRenderSyncPolicyNvidi
 
 bool FDisplayClusterRenderSyncPolicyNvidiaDX12::SynchronizeClusterRendering(int32& InOutSyncInterval)
 {
-	check(IsInRenderingThread());
-
 	// Unfortunately NVIDIA doesn't export required functions in the public SDK. This make impossible
 	// to implement the hardware base synchronization approach. As a fallback we use our own soft sync.
 	// As soon as the required functions become available, the proper synchonization approach will be implemented.
diff --git a/Engine/Plugins/Runtime/nDisplay/Source/DisplayCluster/Private/Render/Synchronization/DisplayClusterRenderSyncPolicySoftwareDX11.cpp b/Engine/Plugins/Runtime/nDisplay/Source/DisplayCluster/Private/Render/Synchronization/DisplayClusterRenderSyncPolicySoftwareDX11.cpp
--- a/Engine/Plugins/Runtime/nDisplay/Source/DisplayCluster/Private/Render/Synchronization/DisplayClusterRenderSyncPolicySoftwareDX11.cpp
+++ b/Engine/Plugins/Runtime/nDisplay/Source/DisplayCluster/Private/Render/Synchronization/DisplayClusterRenderSyncPolicySoftwareDX11.cpp
@@ -135,8 +135,6 @@ FDisplayClusterRenderSyncPolicySoftwareDX11::~FDisplayClusterRenderSyncPolicySof
 
 bool FDisplayClusterRenderSyncPolicySoftwareDX11::SynchronizeClusterRendering(int32& InOutSyncInterval)
 {
-	check(IsInRenderingThread());
-
 	if(!(GEngine && GEngine->GameViewport && GEngine->GameViewport->Viewport))
 	{
 		return false;
diff --git a/Engine/Plugins/Runtime/nDisplay/Source/DisplayCluster/Private/Render/Synchronization/DisplayClusterRenderSyncPolicySoftwareDX12.cpp b/Engine/Plugins/Runtime/nDisplay/Source/DisplayCluster/Private/Render/Synchronization/DisplayClusterRenderSyncPolicySoftwareDX12.cpp
--- a/Engine/Plugins/Runtime/nDisplay/Source/DisplayCluster/Private/Render/Synchronization/DisplayClusterRenderSyncPolicySoftwareDX12.cpp
+++ b/Engine/Plugins/Runtime/nDisplay/Source/DisplayCluster/Private/Render/Synchronization/DisplayClusterRenderSyncPolicySoftwareDX12.cpp
@@ -24,8 +24,6 @@ FDisplayClusterRenderSyncPolicySoftwareDX12::~FDisplayClusterRenderSyncPolicySof
 
 bool FDisplayClusterRenderSyncPolicySoftwareDX12::SynchronizeClusterRendering(int32& InOutSyncInterval)
 {
-	check(IsInRenderingThread());
-
 	if(!(GEngine && GEngine->GameViewport && GEngine->GameViewport->Viewport))
 	{
 		return false;
diff --git a/Engine/Plugins/Runtime/nDisplay/Source/DisplayCluster/Private/Render/Synchronization/DisplayClusterRenderSyncPolicySoftwareGeneric.cpp b/Engine/Plugins/Runtime/nDisplay/Source/DisplayCluster/Private/Render/Synchronization/DisplayClusterRenderSyncPolicySoftwareGeneric.cpp
--- a/Engine/Plugins/Runtime/nDisplay/Source/DisplayCluster/Private/Render/Synchronization/DisplayClusterRenderSyncPolicySoftwareGeneric.cpp
+++ b/Engine/Plugins/Runtime/nDisplay/Source/DisplayCluster/Private/Render/Synchronization/DisplayClusterRenderSyncPolicySoftwareGeneric.cpp
@@ -14,8 +14,6 @@ FDisplayClusterRenderSyncPolicySoftwareGeneric::~FDisplayClusterRenderSyncPolicy
 
 bool FDisplayClusterRenderSyncPolicySoftwareGeneric::SynchronizeClusterRendering(int32& InOutSyncInterval)
 {
-	check(IsInRenderingThread());
-
 	// Synchronize by a barrier only
 	SyncBarrierRenderThread();
 	// Tell a caller that he is still need to present a frame
diff --git a/Engine/Plugins/Runtime/nDisplay/Source/DisplayCluster/Private/Render/Synchronization/DisplayClusterRenderSyncPolicySoftwareOpenGL.cpp b/Engine/Plugins/Runtime/nDisplay/Source/DisplayCluster/Private/Render/Synchronization/DisplayClusterRenderSyncPolicySoftwareOpenGL.cpp
--- a/Engine/Plugins/Runtime/nDisplay/Source/DisplayCluster/Private/Render/Synchronization/DisplayClusterRenderSyncPolicySoftwareOpenGL.cpp
+++ b/Engine/Plugins/Runtime/nDisplay/Source/DisplayCluster/Private/Render/Synchronization/DisplayClusterRenderSyncPolicySoftwareOpenGL.cpp
@@ -14,8 +14,6 @@ FDisplayClusterRenderSyncPolicySoftwareOpenGL::~FDisplayClusterRenderSyncPolicyS
 
 bool FDisplayClusterRenderSyncPolicySoftwareOpenGL::SynchronizeClusterRendering(int32& InOutSyncInterval)
 {
-	check(IsInRenderingThread());
-
 	// Synchronize by a barrier only
 	SyncBarrierRenderThread();
 	// Tell a caller that he is still need to present a frame
