Added a nDisplay module to setup message bus interception to sync across the display clusters
Merged from CL 7437705, CL 7438413

diff --git a/Engine/Plugins/Runtime/nDisplay/Source/DisplayClusterMessageInterceptor/DisplayClusterMessageInterception.Build.cs b/Engine/Plugins/Runtime/nDisplay/Source/DisplayClusterMessageInterceptor/DisplayClusterMessageInterception.Build.cs
new file mode 100644
--- /dev/null
+++ b/Engine/Plugins/Runtime/nDisplay/Source/DisplayClusterMessageInterceptor/DisplayClusterMessageInterception.Build.cs
@@ -0,0 +1,20 @@
+// Copyright 1998-2019 Epic Games, Inc. All Rights Reserved.
+
+using UnrealBuildTool;
+using System.IO;
+
+public class DisplayClusterMessageInterception : ModuleRules
+{
+	public DisplayClusterMessageInterception(ReadOnlyTargetRules Target) : base(Target)
+	{
+		PrivateDependencyModuleNames.AddRange(
+			new string[]
+			{
+				"Core",
+				"CoreUObject",
+				"DisplayCluster",
+				"Engine",
+				"Messaging",
+			});
+	}
+}
diff --git a/Engine/Plugins/Runtime/nDisplay/Source/DisplayClusterMessageInterceptor/Private/DisplayClusterMessageInterceptionModule.cpp b/Engine/Plugins/Runtime/nDisplay/Source/DisplayClusterMessageInterceptor/Private/DisplayClusterMessageInterceptionModule.cpp
new file mode 100644
--- /dev/null
+++ b/Engine/Plugins/Runtime/nDisplay/Source/DisplayClusterMessageInterceptor/Private/DisplayClusterMessageInterceptionModule.cpp
@@ -0,0 +1,148 @@
+// Copyright 1998-2019 Epic Games, Inc. All Rights Reserved.
+
+#include "Modules/ModuleInterface.h"
+#include "Modules/ModuleManager.h"
+
+#include "HAL/IConsoleManager.h"
+#include "IMessagingModule.h"
+#include "IMessageBus.h"
+#include "IDisplayCluster.h"
+#include "Cluster/IDisplayClusterClusterManager.h"
+#include "DisplayClusterMessageInterceptor.h"
+#include "DisplayClusterMessageInterceptionSettings.h"
+
+#if WITH_EDITOR
+	#include "ISettingsModule.h"
+	#include "ISettingsSection.h"
+#endif 
+
+#define LOCTEXT_NAMESPACE "DisplayClusterInterception"
+
+
+/**
+ * Display Cluster Message Interceptor module
+ * Intercept a specified set of message bus messages that are received across all the display nodes
+ * to process them in sync across the cluster.
+ */
+class FDisplayClusterMessageInterceptionModule : public IModuleInterface
+{
+public:
+	virtual void StartupModule() override
+	{
+		// Create the message interceptor
+		Interceptor = MakeShared<FDisplayClusterMessageInterceptor, ESPMode::ThreadSafe>();
+
+		// Register cluster event listeners
+		IDisplayClusterClusterManager* ClusterManager = IDisplayCluster::Get().GetClusterMgr();
+		if (ClusterManager && !ListenerDelegate.IsBound())
+		{
+			ListenerDelegate = FOnClusterEventListener::CreateRaw(this, &FDisplayClusterMessageInterceptionModule::HandleClusterEvent);
+			ClusterManager->AddClusterEventListener(ListenerDelegate);
+		}
+
+		// Register cluster session events
+		IDisplayCluster::Get().OnDisplayClusterStartSession().AddRaw(this, &FDisplayClusterMessageInterceptionModule::StartInterception);
+		IDisplayCluster::Get().OnDisplayClusterEndSession().AddRaw(this, &FDisplayClusterMessageInterceptionModule::StopInterception);
+		IDisplayCluster::Get().OnDisplayClusterPreTick().AddRaw(this, &FDisplayClusterMessageInterceptionModule::HandleClusterPreTick);
+
+		// Setup console command to start/stop interception
+		StartMessageSyncCommand = MakeUnique<FAutoConsoleCommand>(
+			TEXT("nDisplay.MessageBusSync.Start"),
+			TEXT("Start MessageBus syncing"),
+			FConsoleCommandDelegate::CreateRaw(this, &FDisplayClusterMessageInterceptionModule::StartInterception)
+			);
+		StopMessageSyncCommand = MakeUnique<FAutoConsoleCommand>(
+			TEXT("nDisplay.MessageBusSync.Stop"),
+			TEXT("Stop MessageBus syncing"),
+			FConsoleCommandDelegate::CreateRaw(this, &FDisplayClusterMessageInterceptionModule::StopInterception)
+			);
+
+		if (ISettingsModule* SettingsModule = FModuleManager::GetModulePtr<ISettingsModule>("Settings"))
+		{
+			SettingsModule->RegisterSettings(
+				"Project", "Plugins", "nDisplay Message Interception",
+				LOCTEXT("InterceptionSettingsName", "nDisplay Message Interception"),
+				LOCTEXT("InterceptionSettingsDescription", "Configure nDisplay Message Interception."),
+				GetMutableDefault<UDisplayClusterMessageInterceptionSettings>()
+			);
+		}
+	}
+
+	virtual void ShutdownModule() override
+	{
+		if (ISettingsModule* SettingsModule = FModuleManager::GetModulePtr<ISettingsModule>("Settings"))
+		{
+			SettingsModule->UnregisterSettings("Project", "Plugins", "nDisplay Message Interception");
+		}
+
+		if (IDisplayCluster::IsAvailable())
+		{
+			// Unregister cluster event listening
+			IDisplayClusterClusterManager* ClusterManager = IDisplayCluster::Get().GetClusterMgr();
+			if (ClusterManager && ListenerDelegate.IsBound())
+			{
+				ClusterManager->RemoveClusterEventListener(ListenerDelegate);
+				ListenerDelegate.Unbind();
+			}
+
+			// Unregister cluster session events
+			IDisplayCluster::Get().OnDisplayClusterStartSession().RemoveAll(this);
+			IDisplayCluster::Get().OnDisplayClusterEndSession().RemoveAll(this);
+			IDisplayCluster::Get().OnDisplayClusterPreTick().RemoveAll(this);
+		}
+		Interceptor.Reset();
+		StartMessageSyncCommand.Reset();
+		StopMessageSyncCommand.Reset();
+	}
+
+private:
+	void HandleClusterEvent(const FDisplayClusterClusterEvent& InEvent)
+	{
+		if (Interceptor)
+		{
+			Interceptor->HandleClusterEvent(InEvent);
+		}
+	}
+
+	void HandleClusterPreTick()
+	{
+		if (Interceptor)
+		{
+			Interceptor->SyncMessages();
+		}
+	}
+
+	void StartInterception()
+	{
+		IDisplayClusterClusterManager* ClusterManager = IDisplayCluster::Get().GetClusterMgr();
+		if (Interceptor && ClusterManager)
+		{
+			Interceptor->Setup(ClusterManager, IMessagingModule::Get().GetDefaultBus().ToSharedRef());
+			Interceptor->Start();
+		}		
+	}
+
+	void StopInterception()
+	{
+		if (Interceptor)
+		{
+			Interceptor->Stop();
+		}
+	}
+
+	/** MessageBus interceptor */
+	TSharedPtr<FDisplayClusterMessageInterceptor, ESPMode::ThreadSafe> Interceptor;
+
+	/** Cluster event listener delegate */
+	FOnClusterEventListener ListenerDelegate;
+
+	// Console commands handle.
+	TUniquePtr<FAutoConsoleCommand> StartMessageSyncCommand;
+	TUniquePtr<FAutoConsoleCommand> StopMessageSyncCommand;
+};
+
+#undef LOCTEXT_NAMESPACE
+
+
+IMPLEMENT_MODULE(FDisplayClusterMessageInterceptionModule, DisplayClusterMessageInterception);
+
diff --git a/Engine/Plugins/Runtime/nDisplay/Source/DisplayClusterMessageInterceptor/Private/DisplayClusterMessageInterceptor.cpp b/Engine/Plugins/Runtime/nDisplay/Source/DisplayClusterMessageInterceptor/Private/DisplayClusterMessageInterceptor.cpp
new file mode 100644
--- /dev/null
+++ b/Engine/Plugins/Runtime/nDisplay/Source/DisplayClusterMessageInterceptor/Private/DisplayClusterMessageInterceptor.cpp
@@ -0,0 +1,186 @@
+// Copyright 1998-2019 Epic Games, Inc. All Rights Reserved.
+
+#include "DisplayClusterMessageInterceptor.h"
+
+#include "Cluster/DisplayClusterClusterEvent.h"
+#include "Misc/ScopeLock.h"
+#include "IMessageContext.h"
+
+
+DEFINE_LOG_CATEGORY_STATIC(LogDisplayClusterInterception, Log, All);
+
+FDisplayClusterMessageInterceptor::FDisplayClusterMessageInterceptor()
+	: bIsIntercepting(false)
+	, InterceptorId(FGuid::NewGuid())
+	, Address(FMessageAddress::NewAddress())
+{}
+
+void FDisplayClusterMessageInterceptor::Setup(IDisplayClusterClusterManager* InClusterManager, TSharedRef<IMessageBus, ESPMode::ThreadSafe> InBus)
+{
+	ClusterManager = InClusterManager;
+	InterceptedBus = InBus;
+}
+
+void FDisplayClusterMessageInterceptor::Purge()
+{
+	TArray<TSharedPtr<IMessageContext, ESPMode::ThreadSafe>> ContextToForward;
+	{
+		FScopeLock Lock(&ContextQueueCS);
+		for (const auto& ContextPair : ContextMap)
+		{
+			ContextToForward.Add(ContextPair.Value.ContextPtr);
+		}
+		ContextMap.Empty();
+	}
+	for (const auto& Context : ContextToForward)
+	{
+		InterceptedBus->Forward(Context.ToSharedRef(), Context->GetRecipients(), FTimespan::Zero(), AsShared());
+	}
+}
+
+void FDisplayClusterMessageInterceptor::Start()
+{
+	const UDisplayClusterMessageInterceptionSettings* InterceptionSettings = GetDefault<UDisplayClusterMessageInterceptionSettings>();
+	if (!bIsIntercepting && InterceptionSettings->bIsEnabled)
+	{
+		check(InterceptedBus);
+		for (const FName& MessageType : InterceptionSettings->MessageTypes)
+		{
+			InterceptedBus->Intercept(AsShared(), MessageType);
+		}
+		InterceptedAnnotation = InterceptionSettings->Annotation;
+		bIsIntercepting = true;
+	}
+}
+
+void FDisplayClusterMessageInterceptor::Stop()
+{
+	if (bIsIntercepting)
+	{
+		check(InterceptedBus);
+		InterceptedBus->Unintercept(AsShared(), NAME_All);
+		bIsIntercepting = false;
+		Purge();
+	}
+}
+
+void FDisplayClusterMessageInterceptor::SyncMessages()
+{
+	if (ClusterManager)
+	{
+		TArray<FString> MessageIds;
+		{
+			FScopeLock Lock(&ContextQueueCS);
+			ContextMap.GenerateKeyArray(MessageIds);
+		}
+		FDisplayClusterClusterEvent SyncMessagesEvent;
+		SyncMessagesEvent.Category = TEXT("nDCI");				// message bus sync message
+		SyncMessagesEvent.Name = ClusterManager->GetNodeId();	// which node got the message
+		for (const FString& MessageId : MessageIds)
+		{
+			SyncMessagesEvent.Type = MessageId;					// the actually message id we received
+			ClusterManager->EmitClusterEvent(SyncMessagesEvent, false);
+			UE_LOG(LogDisplayClusterInterception, VeryVerbose, TEXT("Emitting cluster event for message %s on frame %d"), *MessageId, GFrameCounter);
+		}
+	}
+
+	// remove out of date messages that are not marked reliable
+	const FTimespan MessageTimeoutSpan = FTimespan(0, 0, 1);
+	FDateTime UtcNow = FDateTime::UtcNow();
+
+	TArray<TSharedPtr<IMessageContext, ESPMode::ThreadSafe>> ContextToForward;
+	{
+		FScopeLock Lock(&ContextQueueCS);
+		for (auto It = ContextMap.CreateIterator(); It; ++It)
+		{
+			if (It.Value().ContextPtr->GetTimeSent() + MessageTimeoutSpan <= UtcNow)
+			{
+				if (!EnumHasAnyFlags(It.Value().ContextPtr->GetFlags(), EMessageFlags::Reliable))
+				{
+					UE_LOG(LogDisplayClusterInterception, VeryVerbose, TEXT("discarding unreliable message %s left intercepted for more than 1s"), *It.Key());
+					It.RemoveCurrent();
+				}
+				else
+				{
+					UE_LOG(LogDisplayClusterInterception, Warning, TEXT("Forcing dispatching of reliable message %s left intercepted for more than 1s"), *It.Key());
+					// Force treatment if not synced after a second and the message is reliable
+					ContextToForward.Add(It.Value().ContextPtr);
+				}
+			}
+		}
+	}
+	if (InterceptedBus)
+	{
+		for (const auto& Context : ContextToForward)
+		{
+			InterceptedBus->Forward(Context.ToSharedRef(), Context->GetRecipients(), FTimespan::Zero(), AsShared());
+		}
+	}
+}
+
+void FDisplayClusterMessageInterceptor::HandleClusterEvent(const FDisplayClusterClusterEvent& InEvent)
+{
+	TArray<TSharedPtr<IMessageContext, ESPMode::ThreadSafe>> ContextToForward;
+	if (InEvent.Category == TEXT("nDCI"))
+	{
+		FScopeLock Lock(&ContextQueueCS);
+		if (FContextSync* ContextSync = ContextMap.Find(InEvent.Type))
+		{
+			ContextSync->NodesReceived.Add(InEvent.Name);
+			if (ContextSync->NodesReceived.Num() >= (int32)ClusterManager->GetNodesAmount())
+			{
+				UE_LOG(LogDisplayClusterInterception, VeryVerbose, TEXT("Fowarding message for message id %s on frame %d"), *InEvent.Type, GFrameCounter);
+				ContextToForward.Add(ContextSync->ContextPtr);
+				ContextMap.Remove(InEvent.Type);
+			}
+		}
+	}
+	if (InterceptedBus)
+	{
+		for (const auto& Context : ContextToForward)
+		{
+			InterceptedBus->Forward(Context.ToSharedRef(), Context->GetRecipients(), FTimespan::Zero(), AsShared());
+		}
+	}
+}
+
+FName FDisplayClusterMessageInterceptor::GetDebugName() const
+{
+	return FName("DisplayClusterInterceptor");
+}
+
+const FGuid& FDisplayClusterMessageInterceptor::GetInterceptorId() const
+{
+	return InterceptorId;
+}
+
+
+bool FDisplayClusterMessageInterceptor::InterceptMessage(const TSharedRef<IMessageContext, ESPMode::ThreadSafe>& Context)
+{
+	// we do not intercept forwarded message, they will be either coming off of the network or being forwarded by ourselves
+	if (!bIsIntercepting || !Context->IsForwarded())
+	{
+		return false;
+	}
+
+	const FString* MessageId = Context->GetAnnotations().Find(InterceptedAnnotation);
+	if (Context->GetForwarder() != Address && MessageId)
+	{
+		UE_LOG(LogDisplayClusterInterception, VeryVerbose, TEXT("Intercepting message %s"), **MessageId);
+
+		FScopeLock Lock(&ContextQueueCS);
+		ContextMap.Add(*MessageId, FContextSync(Context));
+		return true;
+	}
+	return false;
+}
+
+FMessageAddress FDisplayClusterMessageInterceptor::GetSenderAddress()
+{
+	return Address;
+}
+
+void FDisplayClusterMessageInterceptor::NotifyMessageError(const TSharedRef<IMessageContext, ESPMode::ThreadSafe>& Context, const FString& Error)
+{
+	// deprecated
+}
diff --git a/Engine/Plugins/Runtime/nDisplay/Source/DisplayClusterMessageInterceptor/Private/DisplayClusterMessageInterceptor.h b/Engine/Plugins/Runtime/nDisplay/Source/DisplayClusterMessageInterceptor/Private/DisplayClusterMessageInterceptor.h
new file mode 100644
--- /dev/null
+++ b/Engine/Plugins/Runtime/nDisplay/Source/DisplayClusterMessageInterceptor/Private/DisplayClusterMessageInterceptor.h
@@ -0,0 +1,80 @@
+// Copyright 1998-2019 Epic Games, Inc. All Rights Reserved.
+
+#pragma once
+
+#include "CoreMinimal.h"
+#include "IMessageInterceptor.h"
+#include "IMessageSender.h"
+#include "DisplayClusterMessageInterceptionSettings.h"
+
+class IMessageBus;
+class IDisplayClusterClusterManager;
+struct FDisplayClusterClusterEvent;
+
+/**
+ * Display Cluster Message Interceptor 
+ * Intercept message bus messages to process them in sync across the display cluster
+ */
+class FDisplayClusterMessageInterceptor : public TSharedFromThis<FDisplayClusterMessageInterceptor, ESPMode::ThreadSafe>
+	, public IMessageInterceptor
+	, public IMessageSender
+{
+public:
+	FDisplayClusterMessageInterceptor();
+
+	void Setup(IDisplayClusterClusterManager* InClusterManager, TSharedRef<IMessageBus, ESPMode::ThreadSafe> InBus);
+	
+	void Start();
+	void Stop();
+	
+	void SyncMessages();
+	void HandleClusterEvent(const FDisplayClusterClusterEvent& InEvent);
+
+	// IMessageInterceptor interface
+	virtual FName GetDebugName() const override;
+	virtual const FGuid& GetInterceptorId() const override;
+	virtual bool InterceptMessage(const TSharedRef<IMessageContext, ESPMode::ThreadSafe>& Context) override;
+
+	// IMessageSender interface
+	virtual FMessageAddress GetSenderAddress() override;
+	virtual void NotifyMessageError(const TSharedRef<IMessageContext, ESPMode::ThreadSafe>& Context, const FString& Error) override;
+private:
+	/** Purge the message map and force their forwarding. */
+	void Purge();
+
+	/** Currently intercepting message bus. */
+	bool bIsIntercepting;
+
+	/** InterceptorId needed by message bus. */
+	FGuid InterceptorId;
+
+	/** Sender address needed by message bus. */
+	FMessageAddress Address;
+
+	/** Message with this annotation will be intercepted. */
+	FName InterceptedAnnotation;
+
+	/** The bus this interceptor is intercepting messages on. */
+	TSharedPtr<IMessageBus, ESPMode::ThreadSafe> InterceptedBus;
+
+	/** The cluster manager to sync messages across. */
+	IDisplayClusterClusterManager* ClusterManager;
+
+	/** Critical section for the intercepted messages. */
+	FCriticalSection ContextQueueCS;
+
+	/** Holds intercepted messages and their current sync state across the cluster. */
+	struct FContextSync
+	{
+		FContextSync(TSharedPtr<IMessageContext, ESPMode::ThreadSafe> InContext)
+			: ContextPtr(InContext)
+		{}
+
+		TSet<FString> NodesReceived;
+		TSharedPtr<IMessageContext, ESPMode::ThreadSafe> ContextPtr;
+	};
+
+	/** The intercepted message map. */
+	TMap<FString, FContextSync> ContextMap;
+};
+
diff --git a/Engine/Plugins/Runtime/nDisplay/Source/DisplayClusterMessageInterceptor/Public/DisplayClusterMessageInterceptionSettings.h b/Engine/Plugins/Runtime/nDisplay/Source/DisplayClusterMessageInterceptor/Public/DisplayClusterMessageInterceptionSettings.h
new file mode 100644
--- /dev/null
+++ b/Engine/Plugins/Runtime/nDisplay/Source/DisplayClusterMessageInterceptor/Public/DisplayClusterMessageInterceptionSettings.h
@@ -0,0 +1,25 @@
+// Copyright 1998-2019 Epic Games, Inc. All Rights Reserved.
+
+#pragma once
+#include "UObject/Object.h"
+#include "UObject/ObjectMacros.h"
+#include "DisplayClusterMessageInterceptionSettings.generated.h"
+
+UCLASS(config=Engine)
+class DISPLAYCLUSTERMESSAGEINTERCEPTION_API UDisplayClusterMessageInterceptionSettings : public UObject
+{
+	GENERATED_BODY()
+
+public:
+	/** Indicates if message interception is enabled. */
+	UPROPERTY(config, EditAnywhere, Category = "Interception Settings")
+	bool bIsEnabled = false;
+
+	/** List of message types that should be taken into consideration for interception. */
+	UPROPERTY(config, EditAnywhere, Category = "Interception Settings")
+	TArray<FName> MessageTypes;
+
+	/** The annotations to look for identifying messages across the cluster. */
+	UPROPERTY(config, EditAnywhere, Category = "Interception Settings")
+	FName Annotation;
+};
\ No newline at end of file
diff --git a/Engine/Plugins/Runtime/nDisplay/nDisplay.uplugin b/Engine/Plugins/Runtime/nDisplay/nDisplay.uplugin
--- a/Engine/Plugins/Runtime/nDisplay/nDisplay.uplugin
+++ b/Engine/Plugins/Runtime/nDisplay/nDisplay.uplugin
@@ -86,6 +86,14 @@
 			"WhitelistPlatforms": [
 				"Win64"
 			]
+		},
+		{
+			"Name": "DisplayClusterMessageInterception",
+			"Type": "Runtime",
+			"LoadingPhase": "Default",
+			"WhitelistPlatforms": [
+				"Win64"
+			]
 		}
 	],
 	"Plugins": [
