Removed the sending of json response for cluster events
Fixed cluster events responses being shared across multiple threads in a non thread-safe manner
Merged from CL 7478463, CL 7478448, CL 7660218

diff --git a/Engine/Plugins/Runtime/nDisplay/Source/DisplayCluster/Private/Network/Service/ClusterEvents/DisplayClusterClusterEventsService.cpp b/Engine/Plugins/Runtime/nDisplay/Source/DisplayCluster/Private/Network/Service/ClusterEvents/DisplayClusterClusterEventsService.cpp
--- a/Engine/Plugins/Runtime/nDisplay/Source/DisplayCluster/Private/Network/Service/ClusterEvents/DisplayClusterClusterEventsService.cpp
+++ b/Engine/Plugins/Runtime/nDisplay/Source/DisplayCluster/Private/Network/Service/ClusterEvents/DisplayClusterClusterEventsService.cpp
@@ -17,14 +17,6 @@
 FDisplayClusterClusterEventsService::FDisplayClusterClusterEventsService(const FString& InAddr, const int32 InPort) :
 	FDisplayClusterService(FString("SRV_CE"), InAddr, InPort)
 {
-	ResponseErrorMissedMandatoryFields = MakeShareable(new FJsonObject());
-	ResponseErrorMissedMandatoryFields->SetNumberField(FString(FDisplayClusterClusterEventsMsg::ArgError), (double)EDisplayClusterJsonError::MissedMandatoryFields);
-
-	ResponseErrorUnknown = MakeShareable(new FJsonObject());
-	ResponseErrorUnknown->SetNumberField(FString(FDisplayClusterClusterEventsMsg::ArgError), (double)EDisplayClusterJsonError::UnknownError);
-
-	ResponseOk = MakeShareable(new FJsonObject());
-	ResponseOk->SetNumberField(FString(FDisplayClusterClusterEventsMsg::ArgError), (double)EDisplayClusterJsonError::Ok);
 }
 
 FDisplayClusterClusterEventsService::~FDisplayClusterClusterEventsService()
@@ -88,7 +80,7 @@ TSharedPtr<FJsonObject> FDisplayClusterClusterEventsService::ProcessJson(const T
 	if (MandatoryFieldsExist == false)
 	{
 		UE_LOG(LogDisplayClusterNetwork, Warning, TEXT("Json message doesn't have a mandatory field"), *GetName());
-		return ResponseErrorMissedMandatoryFields;
+		return MakeResponseErrorMissedMandatoryFields();
 	}
 
 	// Convert Json request to internal cluster event structure
@@ -98,7 +90,7 @@ TSharedPtr<FJsonObject> FDisplayClusterClusterEventsService::ProcessJson(const T
 	// Emit the event
 	EmitClusterEvent(ClusterEvent);
 
-	return ResponseOk;
+	return MakeResponseOk();
 }
 
 
@@ -109,3 +101,24 @@ void FDisplayClusterClusterEventsService::EmitClusterEvent(const FDisplayCluster
 {
 	GDisplayCluster->GetPrivateClusterMgr()->EmitClusterEvent(Event, true);
 }
+
+TSharedPtr<FJsonObject> FDisplayClusterClusterEventsService::MakeResponseErrorMissedMandatoryFields()
+{
+	TSharedPtr<FJsonObject> ResponseErrorMissedMandatoryFields = MakeShared<FJsonObject>();
+	ResponseErrorMissedMandatoryFields->SetNumberField(FString(FDisplayClusterClusterEventsMsg::ArgError), (double)EDisplayClusterJsonError::MissedMandatoryFields);
+	return ResponseErrorMissedMandatoryFields;
+}
+
+TSharedPtr<FJsonObject> FDisplayClusterClusterEventsService::MakeResponseErrorUnknown()
+{
+	TSharedPtr<FJsonObject> ResponseErrorUnknown = MakeShared<FJsonObject>();
+	ResponseErrorUnknown->SetNumberField(FString(FDisplayClusterClusterEventsMsg::ArgError), (double)EDisplayClusterJsonError::UnknownError);
+	return ResponseErrorUnknown;
+}
+
+TSharedPtr<FJsonObject> FDisplayClusterClusterEventsService::MakeResponseOk()
+{
+	TSharedPtr<FJsonObject> ResponseOk = MakeShared<FJsonObject>();
+	ResponseOk->SetNumberField(FString(FDisplayClusterClusterEventsMsg::ArgError), (double)EDisplayClusterJsonError::Ok);
+	return ResponseOk;
+}
diff --git a/Engine/Plugins/Runtime/nDisplay/Source/DisplayCluster/Private/Network/Service/ClusterEvents/DisplayClusterClusterEventsService.h b/Engine/Plugins/Runtime/nDisplay/Source/DisplayCluster/Private/Network/Service/ClusterEvents/DisplayClusterClusterEventsService.h
--- a/Engine/Plugins/Runtime/nDisplay/Source/DisplayCluster/Private/Network/Service/ClusterEvents/DisplayClusterClusterEventsService.h
+++ b/Engine/Plugins/Runtime/nDisplay/Source/DisplayCluster/Private/Network/Service/ClusterEvents/DisplayClusterClusterEventsService.h
@@ -6,9 +6,7 @@
 #include "Network/Service/DisplayClusterService.h"
 #include "Network/Protocol/IPDisplayClusterClusterEventsProtocol.h"
 #include "Network/DisplayClusterMessage.h"
-
-class FJsonObject;
-
+#include "Dom/JsonObject.h"
 
 /**
  * Cluster events server
@@ -56,8 +54,7 @@ protected:
 private:
 	FDisplayClusterClusterEvent BuildClusterEventFromJson(const FString& EventName, const FString& EventType, const FString& EventCategory, const TSharedPtr<FJsonObject>& JsonMessage) const;
 
-private:
-	TSharedPtr<FJsonObject> ResponseErrorMissedMandatoryFields;
-	TSharedPtr<FJsonObject> ResponseErrorUnknown;
-	TSharedPtr<FJsonObject> ResponseOk;
+	TSharedPtr<FJsonObject> MakeResponseErrorMissedMandatoryFields();
+	TSharedPtr<FJsonObject> MakeResponseErrorUnknown();
+	TSharedPtr<FJsonObject> MakeResponseOk();
 };
diff --git a/Engine/Plugins/Runtime/nDisplay/Source/DisplayCluster/Private/Network/Session/DisplayClusterSessionExternal.cpp b/Engine/Plugins/Runtime/nDisplay/Source/DisplayCluster/Private/Network/Session/DisplayClusterSessionExternal.cpp
--- a/Engine/Plugins/Runtime/nDisplay/Source/DisplayCluster/Private/Network/Session/DisplayClusterSessionExternal.cpp
+++ b/Engine/Plugins/Runtime/nDisplay/Source/DisplayCluster/Private/Network/Session/DisplayClusterSessionExternal.cpp
@@ -35,7 +35,6 @@ uint32 FDisplayClusterSessionExternal::Run()
 		if (Response.IsValid())
 		{
 			UE_LOG(LogDisplayClusterNetwork, Log, TEXT("Json based message has been processed"), *GetName());
-			SendJson(Response);
 		}
 	}
 
