Fixed missing include in DisplayClusterMessageInterceptor.cpp

diff --git a/Engine/Plugins/Runtime/nDisplay/Source/DisplayClusterMessageInterceptor/Private/DisplayClusterMessageInterceptor.cpp b/Engine/Plugins/Runtime/nDisplay/Source/DisplayClusterMessageInterceptor/Private/DisplayClusterMessageInterceptor.cpp
--- a/Engine/Plugins/Runtime/nDisplay/Source/DisplayClusterMessageInterceptor/Private/DisplayClusterMessageInterceptor.cpp
+++ b/Engine/Plugins/Runtime/nDisplay/Source/DisplayClusterMessageInterceptor/Private/DisplayClusterMessageInterceptor.cpp
@@ -5,6 +5,8 @@
 #include "Cluster/DisplayClusterClusterEvent.h"
 #include "Misc/ScopeLock.h"
 #include "IMessageContext.h"
+#include "IMessageBus.h"
+#include "Cluster/IDisplayClusterClusterManager.h"
 
 
 DEFINE_LOG_CATEGORY_STATIC(LogDisplayClusterInterception, Log, All);
