Fixed CIS failure on DisplayClusterMessageInterceptionModule

diff --git a/Engine/Plugins/Runtime/nDisplay/Source/DisplayClusterMessageInterceptor/Private/DisplayClusterMessageInterceptionModule.cpp b/Engine/Plugins/Runtime/nDisplay/Source/DisplayClusterMessageInterceptor/Private/DisplayClusterMessageInterceptionModule.cpp
--- a/Engine/Plugins/Runtime/nDisplay/Source/DisplayClusterMessageInterceptor/Private/DisplayClusterMessageInterceptionModule.cpp
+++ b/Engine/Plugins/Runtime/nDisplay/Source/DisplayClusterMessageInterceptor/Private/DisplayClusterMessageInterceptionModule.cpp
@@ -57,6 +57,7 @@ public:
 			FConsoleCommandDelegate::CreateRaw(this, &FDisplayClusterMessageInterceptionModule::StopInterception)
 			);
 
+#if WITH_EDITOR
 		if (ISettingsModule* SettingsModule = FModuleManager::GetModulePtr<ISettingsModule>("Settings"))
 		{
 			SettingsModule->RegisterSettings(
@@ -66,14 +67,17 @@ public:
 				GetMutableDefault<UDisplayClusterMessageInterceptionSettings>()
 			);
 		}
+#endif
 	}
 
 	virtual void ShutdownModule() override
 	{
+#if WITH_EDITOR
 		if (ISettingsModule* SettingsModule = FModuleManager::GetModulePtr<ISettingsModule>("Settings"))
 		{
 			SettingsModule->UnregisterSettings("Project", "Plugins", "nDisplay Message Interception");
 		}
+#endif
 
 		if (IDisplayCluster::IsAvailable())
 		{
