nDisplay merge Quokka -> Release23

diff --git a/Engine/Plugins/Runtime/nDisplay/Shaders/Private/MPCDIOverlayShaders.usf b/Engine/Plugins/Runtime/nDisplay/Shaders/Private/MPCDIOverlayShaders.usf
--- a/Engine/Plugins/Runtime/nDisplay/Shaders/Private/MPCDIOverlayShaders.usf
+++ b/Engine/Plugins/Runtime/nDisplay/Shaders/Private/MPCDIOverlayShaders.usf
@@ -49,19 +49,6 @@ struct OutputPS
 	float4 Color : SV_Target0;
 };
 
-float3 ApplyPicpTextureGammaCorrection(float3 Color)
-{
-	//float3 OutColor = ApplyGammaCorrection(Color, 2.2);
-	       //OutColor = LinearToSrgb(OutColor);
-
-	return Color;
-}
-
-float4 ApplyPicpTextureGammaCorrection(float4 Color)
-{
-	return Color;
-	//return float4(ApplyPicpTextureGammaCorrection(Color.xyz), Color.w);
-}
 
 /** Used for calculating vertex positions and UVs when drawing with DrawRectangle */
 void DrawRectangle(in float4 InPosition, in float2 InTexCoord, out float4 OutPosition, out float2 OutTexCoord)
@@ -169,7 +156,7 @@ OutputPS PicpMPCDIWarpAndBlendOverlay_PS(OutputVS IN)
 #if RENDER_CASE_VIEWPORT_OVERLAY
 	float4 ViewportOverlayUVW = mul(WarpedUV, TextureProjectionMatrix);
 	float2 ViewportOverlayUV = ViewportOverlayUVW.xy / ViewportOverlayUVW.w;
-	float4 ViewportOverlayColor = ApplyPicpTextureGammaCorrection(ViewportOverlayMap.Sample(ViewportOverlayMapSampler, ViewportOverlayUV).rgba);
+	float4 ViewportOverlayColor = ViewportOverlayMap.Sample(ViewportOverlayMapSampler, ViewportOverlayUV).rgba;
 #endif
 
 #if RENDER_CASE_CAMERA
@@ -177,7 +164,7 @@ OutputPS PicpMPCDIWarpAndBlendOverlay_PS(OutputVS IN)
 	float4 CameraUVW = mul(WarpedUV, CameraProjectionMatrix);
 	float2 CameraUV = CameraUVW.xy / CameraUVW.w;
 	// Load from Cameramap
-	float3 CameraBaseColor        = ApplyPicpTextureGammaCorrection(CameraMap.Sample(CameraMapSampler, CameraUV).rgb);
+	float3 CameraBaseColor        = CameraMap.Sample(CameraMapSampler, CameraUV).rgb;
 
 	float4 CameraColor = float4(CameraBaseColor, 0);
 
@@ -239,7 +226,7 @@ OutputPS PicpMPCDIWarpAndBlendViewportOverlay_PS(OutputVS IN)
 	float4 OverlayUVW = mul(WarpedUV, TextureProjectionMatrix);
 	float2 OverlayUV = OverlayUVW.xy / OverlayUVW.w;
 	
-	float4 ViewportOverlayColor = ApplyPicpTextureGammaCorrection(ViewportOverlayMap.Sample(ViewportOverlayMapSampler, OverlayUV).rgba);
+	float4 ViewportOverlayColor = ViewportOverlayMap.Sample(ViewportOverlayMapSampler, OverlayUV).rgba;
 
 	//!fixme: add LUT here
 
@@ -305,13 +292,13 @@ OutputPS PicpMPCDIWarpAndBlendOneCameraFull_PS(OutputVS IN)
 	float4 OutColor = ViewportBaseColor;
 	
 #if USE_OVERLAY_BACK	
-	OutColor.xyz = lerp(OutColor.xyz, ApplyPicpTextureGammaCorrection(ViewportOverlayBackColor.xyz), saturate(1.0f - ViewportOverlayBackColor.w));
+	OutColor.xyz = lerp(OutColor.xyz, ViewportOverlayBackColor.xyz, saturate(1.0f - ViewportOverlayBackColor.w));
 #endif
 
-	OutColor.xyz = lerp(OutColor.xyz, ApplyPicpTextureGammaCorrection(CameraColor.xyz),  CameraColor.w);
+	OutColor.xyz = lerp(OutColor.xyz, CameraColor.xyz,  CameraColor.w);
 
 #if USE_OVERLAY_FRONT
-	OutColor.xyz = lerp(OutColor.xyz, ApplyPicpTextureGammaCorrection(ViewportOverlayColor.xyz), saturate(1.0f - ViewportOverlayColor.w));
+	OutColor.xyz = lerp(OutColor.xyz, ViewportOverlayColor.xyz, saturate(1.0f - ViewportOverlayColor.w));
 #endif
 
 	//!fixme: add LUT here
diff --git a/Engine/Plugins/Runtime/nDisplay/Source/DisplayCluster/Private/Blueprints/DisplayClusterBlueprintAPIImpl.cpp b/Engine/Plugins/Runtime/nDisplay/Source/DisplayCluster/Private/Blueprints/DisplayClusterBlueprintAPIImpl.cpp
--- a/Engine/Plugins/Runtime/nDisplay/Source/DisplayCluster/Private/Blueprints/DisplayClusterBlueprintAPIImpl.cpp
+++ b/Engine/Plugins/Runtime/nDisplay/Source/DisplayCluster/Private/Blueprints/DisplayClusterBlueprintAPIImpl.cpp
@@ -15,6 +15,8 @@
 #include "Config/DisplayClusterConfigTypes.h"
 #include "Misc/DisplayClusterHelpers.h"
 
+#include "DisplayClusterGlobals.h"
+
 
 //////////////////////////////////////////////////////////////////////////////////////////////
 // DisplayCluster module API
@@ -138,7 +140,7 @@ void UDisplayClusterBlueprintAPIImpl::EmitClusterEvent(const FDisplayClusterClus
 //////////////////////////////////////////////////////////////////////////////////////////////
 // Config API
 //////////////////////////////////////////////////////////////////////////////////////////////
-void UDisplayClusterBlueprintAPIImpl::GetViewports(bool IsRTT, TArray<FString>& ViewportIDs, TArray<FString>& ViewportTypes, TArray<FIntPoint>& ViewportLocations, TArray<FIntPoint>& ViewportSizes)
+void UDisplayClusterBlueprintAPIImpl::GetLocalViewports(bool IsRTT, TArray<FString>& ViewportIDs, TArray<FString>& ViewportTypes, TArray<FIntPoint>& ViewportLocations, TArray<FIntPoint>& ViewportSizes)
 {
 	DISPLAY_CLUSTER_FUNC_TRACE(LogDisplayClusterBlueprint);
 
@@ -584,14 +586,36 @@ void UDisplayClusterBlueprintAPIImpl::SetViewportCamera(const FString& InCameraI
 	return;
 }
 
-void UDisplayClusterBlueprintAPIImpl::SetCustomPostProcessing(const FString& ViewportID, const FPostProcessSettings& PostProcessingSettings)
+void UDisplayClusterBlueprintAPIImpl::SetStartPostProcessingSettings(const FString& ViewportID, const FPostProcessSettings& StartPostProcessingSettings)
+{
+	DISPLAY_CLUSTER_FUNC_TRACE(LogDisplayClusterBlueprint);
+
+	IDisplayClusterRenderManager* const Manager = IDisplayCluster::Get().GetRenderMgr();
+	if (Manager)
+	{
+		Manager->SetStartPostProcessingSettings(ViewportID, StartPostProcessingSettings);
+	}
+}
+
+void UDisplayClusterBlueprintAPIImpl::SetOverridePostProcessingSettings(const FString& ViewportID, const FPostProcessSettings& OverridePostProcessingSettings, float BlendWeight)
+{
+	DISPLAY_CLUSTER_FUNC_TRACE(LogDisplayClusterBlueprint);
+
+	IDisplayClusterRenderManager* const Manager = IDisplayCluster::Get().GetRenderMgr();
+	if (Manager)
+	{
+		Manager->SetOverridePostProcessingSettings(ViewportID, OverridePostProcessingSettings, BlendWeight);
+	}
+}
+
+void UDisplayClusterBlueprintAPIImpl::SetFinalPostProcessingSettings(const FString& ViewportID, const FPostProcessSettings& FinalPostProcessingSettings)
 {
 	DISPLAY_CLUSTER_FUNC_TRACE(LogDisplayClusterBlueprint);
 
 	IDisplayClusterRenderManager* const Manager = IDisplayCluster::Get().GetRenderMgr();
 	if (Manager)
 	{
-		Manager->SetCustomPostProcessing(ViewportID, PostProcessingSettings);
+		Manager->SetFinalPostProcessingSettings(ViewportID, FinalPostProcessingSettings);
 	}
 }
 
diff --git a/Engine/Plugins/Runtime/nDisplay/Source/DisplayCluster/Private/Blueprints/DisplayClusterBlueprintAPIImpl.h b/Engine/Plugins/Runtime/nDisplay/Source/DisplayCluster/Private/Blueprints/DisplayClusterBlueprintAPIImpl.h
--- a/Engine/Plugins/Runtime/nDisplay/Source/DisplayCluster/Private/Blueprints/DisplayClusterBlueprintAPIImpl.h
+++ b/Engine/Plugins/Runtime/nDisplay/Source/DisplayCluster/Private/Blueprints/DisplayClusterBlueprintAPIImpl.h
@@ -69,8 +69,8 @@ public:
 	//////////////////////////////////////////////////////////////////////////////////////////////
 	// Config API
 	//////////////////////////////////////////////////////////////////////////////////////////////
-	UFUNCTION(BlueprintCallable, meta = (DisplayName = "Get viewports"), Category = "DisplayCluster|Config")
-	virtual void GetViewports(bool IsRTT, TArray<FString>& ViewportIDs, TArray<FString>& ViewportTypes, TArray<FIntPoint>& ViewportLocations, TArray<FIntPoint>& ViewportSizes) override;
+	UFUNCTION(BlueprintCallable, meta = (DisplayName = "Get local viewports"), Category = "DisplayCluster|Config")
+	virtual void GetLocalViewports(bool IsRTT, TArray<FString>& ViewportIDs, TArray<FString>& ViewportTypes, TArray<FIntPoint>& ViewportLocations, TArray<FIntPoint>& ViewportSizes) override;
 
 public:
 	//////////////////////////////////////////////////////////////////////////////////////////////
@@ -192,8 +192,14 @@ public:
 	UFUNCTION(BlueprintCallable, meta = (DisplayName = "Set viewport camera"), Category = "DisplayCluster|Render")
 	virtual void SetViewportCamera(const FString& InCameraId, const FString& InViewportId) override;
 	
-	UFUNCTION(BlueprintCallable, meta = (DisplayName = "Set custom post processing settings for viewport"), Category = "DisplayCluster|Render")
-	virtual void SetCustomPostProcessing(const FString& ViewportID, const FPostProcessSettings& PostProcessingSettings) override;
+	UFUNCTION(BlueprintCallable, meta = (DisplayName = "Set start post processing settings for viewport"), Category = "DisplayCluster|Render")
+	virtual void SetStartPostProcessingSettings(const FString& ViewportID, const FPostProcessSettings& StartPostProcessingSettings) override;
+
+	UFUNCTION(BlueprintCallable, meta = (DisplayName = "Set start post override settings for viewport"), Category = "DisplayCluster|Render")
+	virtual void SetOverridePostProcessingSettings(const FString& ViewportID, const FPostProcessSettings& OverridePostProcessingSettings, float BlendWeight = 1.0f) override;
+
+	UFUNCTION(BlueprintCallable, meta = (DisplayName = "Set start post processing settings for viewport"), Category = "DisplayCluster|Render")
+	virtual void SetFinalPostProcessingSettings(const FString& ViewportID, const FPostProcessSettings& FinalPostProcessingSettings) override;
 
 	UFUNCTION(BlueprintCallable, meta = (DisplayName = "Get viewport rectangle"), Category = "DisplayCluster|Render")
 	virtual bool GetViewportRect(const FString& ViewportID, FIntPoint& ViewportLoc, FIntPoint& ViewportSize) override;
diff --git a/Engine/Plugins/Runtime/nDisplay/Source/DisplayCluster/Private/Cluster/DisplayClusterClusterManager.cpp b/Engine/Plugins/Runtime/nDisplay/Source/DisplayCluster/Private/Cluster/DisplayClusterClusterManager.cpp
--- a/Engine/Plugins/Runtime/nDisplay/Source/DisplayCluster/Private/Cluster/DisplayClusterClusterManager.cpp
+++ b/Engine/Plugins/Runtime/nDisplay/Source/DisplayCluster/Private/Cluster/DisplayClusterClusterManager.cpp
@@ -187,11 +187,12 @@ void FDisplayClusterClusterManager::PreTick(float DeltaSeconds)
 {
 	DISPLAY_CLUSTER_FUNC_TRACE(LogDisplayClusterCluster);
 
-	// Clear cached data from previous game frame
-	{
-		FScopeLock lock(&ObjectsToSyncCritSec);
-		SyncObjectsCache.Empty(SyncObjectsCache.Num() | 0x07);
-	}
+	// Moved clearing of the sync object cache to the end of the frame
+	//// Clear cached data from previous game frame
+	//{
+	//	FScopeLock lock(&ObjectsToSyncCritSec);
+	//	SyncObjectsCache.Empty(SyncObjectsCache.Num() | 0x07);
+	//}
 
 	// Move cluster events from the primary pool to the output pool. These will be synchronized on the current frame.
 	{
@@ -366,9 +367,9 @@ void FDisplayClusterClusterManager::ExportSyncData(FDisplayClusterMessage::DataT
 				}
 			}
 		}
-	}
 
-	data = SyncObjectsCache;
+		data = SyncObjectsCache;
+	}
 }
 
 void FDisplayClusterClusterManager::ImportSyncData(const FDisplayClusterMessage::DataType& data)
@@ -523,6 +524,12 @@ void FDisplayClusterClusterManager::SyncEvents()
 	}
 }
 
+void FDisplayClusterClusterManager::ClearSyncObjects()
+{
+	FScopeLock lock(&ObjectsToSyncCritSec);
+	SyncObjectsCache.Empty(SyncObjectsCache.Num() | 0x07);
+}
+
 //////////////////////////////////////////////////////////////////////////////////////////////
 // FDisplayClusterClusterManager
 //////////////////////////////////////////////////////////////////////////////////////////////
diff --git a/Engine/Plugins/Runtime/nDisplay/Source/DisplayCluster/Private/Cluster/DisplayClusterClusterManager.h b/Engine/Plugins/Runtime/nDisplay/Source/DisplayCluster/Private/Cluster/DisplayClusterClusterManager.h
--- a/Engine/Plugins/Runtime/nDisplay/Source/DisplayCluster/Private/Cluster/DisplayClusterClusterManager.h
+++ b/Engine/Plugins/Runtime/nDisplay/Source/DisplayCluster/Private/Cluster/DisplayClusterClusterManager.h
@@ -88,6 +88,9 @@ public:
 	virtual void SyncInput()   override;
 	virtual void SyncEvents()  override;
 
+	virtual void ClearSyncObjects() override;
+
+
 private:
 	bool GetResolvedNodeId(FString& id) const;
 
diff --git a/Engine/Plugins/Runtime/nDisplay/Source/DisplayCluster/Private/Cluster/IPDisplayClusterClusterManager.h b/Engine/Plugins/Runtime/nDisplay/Source/DisplayCluster/Private/Cluster/IPDisplayClusterClusterManager.h
--- a/Engine/Plugins/Runtime/nDisplay/Source/DisplayCluster/Private/Cluster/IPDisplayClusterClusterManager.h
+++ b/Engine/Plugins/Runtime/nDisplay/Source/DisplayCluster/Private/Cluster/IPDisplayClusterClusterManager.h
@@ -45,4 +45,6 @@ public:
 	virtual void SyncObjects() = 0;
 	virtual void SyncInput()   = 0;
 	virtual void SyncEvents()  = 0;
+
+	virtual void ClearSyncObjects() = 0;
 };
diff --git a/Engine/Plugins/Runtime/nDisplay/Source/DisplayCluster/Private/DisplayClusterLog.cpp b/Engine/Plugins/Runtime/nDisplay/Source/DisplayCluster/Private/DisplayClusterLog.cpp
--- a/Engine/Plugins/Runtime/nDisplay/Source/DisplayCluster/Private/DisplayClusterLog.cpp
+++ b/Engine/Plugins/Runtime/nDisplay/Source/DisplayCluster/Private/DisplayClusterLog.cpp
@@ -14,4 +14,5 @@ DEFINE_LOG_CATEGORY(LogDisplayClusterInputVRPN);
 DEFINE_LOG_CATEGORY(LogDisplayClusterNetwork);
 DEFINE_LOG_CATEGORY(LogDisplayClusterNetworkMsg);
 DEFINE_LOG_CATEGORY(LogDisplayClusterRender);
+DEFINE_LOG_CATEGORY(LogDisplayClusterRenderSync);
 DEFINE_LOG_CATEGORY(LogDisplayClusterBlueprint);
diff --git a/Engine/Plugins/Runtime/nDisplay/Source/DisplayCluster/Private/DisplayClusterLog.h b/Engine/Plugins/Runtime/nDisplay/Source/DisplayCluster/Private/DisplayClusterLog.h
--- a/Engine/Plugins/Runtime/nDisplay/Source/DisplayCluster/Private/DisplayClusterLog.h
+++ b/Engine/Plugins/Runtime/nDisplay/Source/DisplayCluster/Private/DisplayClusterLog.h
@@ -17,6 +17,7 @@ DECLARE_LOG_CATEGORY_EXTERN(LogDisplayClusterInputVRPN,  Warning, Warning);
 DECLARE_LOG_CATEGORY_EXTERN(LogDisplayClusterNetwork,    Warning, Warning);
 DECLARE_LOG_CATEGORY_EXTERN(LogDisplayClusterNetworkMsg, Warning, Warning);
 DECLARE_LOG_CATEGORY_EXTERN(LogDisplayClusterRender,     Warning, Warning);
+DECLARE_LOG_CATEGORY_EXTERN(LogDisplayClusterRenderSync, Warning, Warning);
 DECLARE_LOG_CATEGORY_EXTERN(LogDisplayClusterBlueprint,  Warning, Warning);
 #else
 DECLARE_LOG_CATEGORY_EXTERN(LogDisplayClusterGameMode,   Log, All);
@@ -30,6 +31,7 @@ DECLARE_LOG_CATEGORY_EXTERN(LogDisplayClusterInputVRPN,  Log, All);
 DECLARE_LOG_CATEGORY_EXTERN(LogDisplayClusterNetwork,    Log, All);
 DECLARE_LOG_CATEGORY_EXTERN(LogDisplayClusterNetworkMsg, Log, All);
 DECLARE_LOG_CATEGORY_EXTERN(LogDisplayClusterRender,     Log, All);
+DECLARE_LOG_CATEGORY_EXTERN(LogDisplayClusterRenderSync, Log, All);
 DECLARE_LOG_CATEGORY_EXTERN(LogDisplayClusterBlueprint,  Log, All);
 #endif
 
diff --git a/Engine/Plugins/Runtime/nDisplay/Source/DisplayCluster/Private/Game/Classes/Basics/DisplayClusterGameEngine.cpp b/Engine/Plugins/Runtime/nDisplay/Source/DisplayCluster/Private/Game/Classes/Basics/DisplayClusterGameEngine.cpp
--- a/Engine/Plugins/Runtime/nDisplay/Source/DisplayCluster/Private/Game/Classes/Basics/DisplayClusterGameEngine.cpp
+++ b/Engine/Plugins/Runtime/nDisplay/Source/DisplayCluster/Private/Game/Classes/Basics/DisplayClusterGameEngine.cpp
@@ -200,14 +200,13 @@ void UDisplayClusterGameEngine::Tick(float DeltaSeconds, bool bIdleMode)
 		// Update delta time. Cluster slaves will get this value from the master few steps later
 		ClusterMgr->SetDeltaTime(DeltaSeconds);
 
-		// Sync cluster objects
-		ClusterMgr->SyncObjects();
 
 		//////////////////////////////////////////////////////////////////////////////////////////////
 		// Frame start barrier
 		NodeController->WaitForFrameStart();
 		UE_LOG(LogDisplayClusterEngine, Verbose, TEXT("Sync frame start"));
 
+
 		// Get DisplayCluster time delta
 		NodeController->GetDeltaTime(DeltaSeconds);
 		NodeController->GetTimecode(Timecode, FrameRate);
@@ -225,6 +224,9 @@ void UDisplayClusterGameEngine::Tick(float DeltaSeconds, bool bIdleMode)
 		UE_LOG(LogDisplayClusterEngine, Verbose, TEXT("Perform PreTick()"));
 		GDisplayCluster->PreTick(DeltaSeconds);
 
+		// Sync cluster objects
+		ClusterMgr->SyncObjects();
+
 		// Sync cluster events
 		ClusterMgr->SyncEvents();
 
@@ -252,6 +254,8 @@ void UDisplayClusterGameEngine::Tick(float DeltaSeconds, bool bIdleMode)
 		//////////////////////////////////////////////////////////////////////////////////////////////
 		// Frame end barrier
 		NodeController->WaitForFrameEnd();
+		ClusterMgr->ClearSyncObjects();
+
 		UE_LOG(LogDisplayClusterEngine, Verbose, TEXT("Sync frame end"));
 	}
 	else
diff --git a/Engine/Plugins/Runtime/nDisplay/Source/DisplayCluster/Private/Game/Classes/Basics/DisplayClusterViewportClient.cpp b/Engine/Plugins/Runtime/nDisplay/Source/DisplayCluster/Private/Game/Classes/Basics/DisplayClusterViewportClient.cpp
new file mode 100644
--- /dev/null
+++ b/Engine/Plugins/Runtime/nDisplay/Source/DisplayCluster/Private/Game/Classes/Basics/DisplayClusterViewportClient.cpp
@@ -0,0 +1,673 @@
+// Copyright 1998-2019 Epic Games, Inc. All Rights Reserved.
+
+#include "DisplayClusterViewportClient.h"
+
+#include "SceneView.h"
+#include "Engine/Canvas.h"
+#include "SceneViewExtension.h"
+#include "IHeadMountedDisplay.h"
+#include "IXRTrackingSystem.h"
+#include "BufferVisualizationData.h"
+#include "Engine/Engine.h"
+#include "Engine/Console.h"
+#include "Engine/LocalPlayer.h"
+#include "UnrealEngine.h"
+#include "EngineUtils.h"
+
+#include "AudioDeviceManager.h"
+#include "AudioDevice.h"
+#include "Audio/AudioDebug.h"
+
+#include "GameFramework/PlayerController.h"
+#include "Debug/DebugDrawService.h"
+#include "ContentStreaming.h"
+#include "EngineModule.h"
+#include "FXSystem.h"
+#include "GameFramework/HUD.h"
+#include "SubtitleManager.h"
+#include "Components/LineBatchComponent.h"
+#include "GameFramework/GameUserSettings.h"
+#include "Framework/Application/SlateApplication.h"
+
+#include "LegacyScreenPercentageDriver.h"
+#include "DynamicResolutionState.h"
+#include "EngineStats.h"
+
+
+
+//DECLARE_CYCLE_STAT(TEXT("UI Drawing Time"), STAT_UIDrawingTime, STATGROUP_UI);
+
+UDisplayClusterViewportClient::UDisplayClusterViewportClient(FVTableHelper& Helper) : Super(Helper)
+{
+	
+}
+
+UDisplayClusterViewportClient::~UDisplayClusterViewportClient()
+{
+
+}
+
+/** Util to find named canvas in transient package, and create if not found */
+static UCanvas* GetCanvasByName(FName CanvasName)
+{
+	// Cache to avoid FString/FName conversions/compares
+	static TMap<FName, UCanvas*> CanvasMap;
+	UCanvas** FoundCanvas = CanvasMap.Find(CanvasName);
+	if (!FoundCanvas)
+	{
+		UCanvas* CanvasObject = FindObject<UCanvas>(GetTransientPackage(), *CanvasName.ToString());
+		if (!CanvasObject)
+		{
+			CanvasObject = NewObject<UCanvas>(GetTransientPackage(), CanvasName);
+			CanvasObject->AddToRoot();
+		}
+
+		CanvasMap.Add(CanvasName, CanvasObject);
+		return CanvasObject;
+	}
+
+	return *FoundCanvas;
+}
+
+void UDisplayClusterViewportClient::Init(struct FWorldContext& WorldContext, UGameInstance* OwningGameInstance, bool bCreateNewAudioDevice)
+{
+	IConsoleVariable* const ForceLoadCVar = IConsoleManager::Get().FindConsoleVariable(TEXT("r.CompositionForceRenderTargetLoad"));
+	if (ForceLoadCVar)
+	{
+		ForceLoadCVar->Set(int32(1));
+	}
+
+	Super::Init(WorldContext, OwningGameInstance, bCreateNewAudioDevice);
+}
+
+void UDisplayClusterViewportClient::Draw(FViewport* InViewport, FCanvas* SceneCanvas)
+{
+	//Valid SceneCanvas is required.  Make this explicit.
+	check(SceneCanvas);
+
+	OnBeginDraw().Broadcast();
+	
+	const bool bStereoRendering = GEngine->IsStereoscopic3D(InViewport);
+	FCanvas* DebugCanvas = InViewport->GetDebugCanvas();
+
+	
+	// Create a temporary canvas if there isn't already one.
+	static FName CanvasObjectName(TEXT("CanvasObject"));
+	UCanvas* CanvasObject = GetCanvasByName(CanvasObjectName);
+	CanvasObject->Canvas = SceneCanvas;
+	
+	// Create temp debug canvas object
+	FIntPoint DebugCanvasSize = InViewport->GetSizeXY();
+	if (bStereoRendering && GEngine->XRSystem.IsValid() && GEngine->XRSystem->GetHMDDevice())
+	{
+		DebugCanvasSize = GEngine->XRSystem->GetHMDDevice()->GetIdealDebugCanvasRenderTargetSize();
+	}
+	
+	static FName DebugCanvasObjectName(TEXT("DebugCanvasObject"));
+	UCanvas* DebugCanvasObject = GetCanvasByName(DebugCanvasObjectName);
+	DebugCanvasObject->Init(DebugCanvasSize.X, DebugCanvasSize.Y, NULL, DebugCanvas);
+
+	if (DebugCanvas)
+	{
+		DebugCanvas->SetScaledToRenderTarget(bStereoRendering);
+		DebugCanvas->SetStereoRendering(bStereoRendering);
+	}
+	if (SceneCanvas)
+	{
+		SceneCanvas->SetScaledToRenderTarget(bStereoRendering);
+		SceneCanvas->SetStereoRendering(bStereoRendering);
+	}
+
+	UWorld* MyWorld = GetWorld();
+	
+	// Force path tracing view mode, and extern code set path tracer show flags
+	const bool bForcePathTracing = InViewport->GetClient()->GetEngineShowFlags()->PathTracing;
+	if (bForcePathTracing)
+	{
+		EngineShowFlags.SetPathTracing(true);
+		ViewModeIndex = VMI_PathTracing;
+	}
+
+	const bool bEnableStereo = GEngine->IsStereoscopic3D(InViewport);
+	int32 NumFamilies = 1;
+	int32 NumViews = bStereoRendering ? GEngine->StereoRenderingDevice->GetDesiredNumberOfViews(bStereoRendering) : 1;
+	int32 StartViewIndex = 0;
+	bool nDisplay = false;
+	// check if HMD or nDisplay connected 
+	if (!GEngine->XRSystem.IsValid())
+	{
+		if (GEngine->StereoRenderingDevice.IsValid())
+		{
+			NumFamilies = GEngine->StereoRenderingDevice->GetDesiredNumberOfViews(bStereoRendering);
+			nDisplay = true;
+		}
+	}
+	
+	for (int32 viewFamily = 0; viewFamily < NumFamilies; ++viewFamily)
+	{
+		if (nDisplay)
+		{
+			StartViewIndex = viewFamily;
+			NumViews = StartViewIndex + 1;
+		}
+
+		// create the view family for rendering the world scene to the viewport's render target
+		FSceneViewFamilyContext ViewFamily(FSceneViewFamily::ConstructionValues(
+			InViewport,
+			MyWorld->Scene,
+			EngineShowFlags)
+			.SetRealtimeUpdate(true));
+
+#if WITH_EDITOR
+		if (GIsEditor)
+		{
+			// Force enable view family show flag for HighDPI derived's screen percentage.
+			ViewFamily.EngineShowFlags.ScreenPercentage = true;
+		}
+
+		UpdateDebugViewModeShaders();
+#endif
+
+		ViewFamily.ViewExtensions = GEngine->ViewExtensions->GatherActiveExtensions(InViewport);
+
+		for (auto ViewExt : ViewFamily.ViewExtensions)
+		{
+			ViewExt->SetupViewFamily(ViewFamily);
+		}
+
+		if (bStereoRendering && GEngine->XRSystem.IsValid() && GEngine->XRSystem->GetHMDDevice())
+		{
+			// Allow HMD to modify screen settings
+			GEngine->XRSystem->GetHMDDevice()->UpdateScreenSettings(Viewport);
+		}
+
+		ESplitScreenType::Type SplitScreenConfig = GetCurrentSplitscreenConfiguration();
+		ViewFamily.ViewMode = EViewModeIndex(ViewModeIndex);
+		EngineShowFlagOverride(ESFIM_Game, ViewFamily.ViewMode, ViewFamily.EngineShowFlags, false);
+
+		if (ViewFamily.EngineShowFlags.VisualizeBuffer && AllowDebugViewmodes())
+		{
+			// Process the buffer visualization console command
+			FName NewBufferVisualizationMode = NAME_None;
+			static IConsoleVariable* ICVar = IConsoleManager::Get().FindConsoleVariable(FBufferVisualizationData::GetVisualizationTargetConsoleCommandName());
+			if (ICVar)
+			{
+				static const FName OverviewName = TEXT("Overview");
+				FString ModeNameString = ICVar->GetString();
+				FName ModeName = *ModeNameString;
+				if (ModeNameString.IsEmpty() || ModeName == OverviewName || ModeName == NAME_None)
+				{
+					NewBufferVisualizationMode = NAME_None;
+				}
+				else
+				{
+					if (GetBufferVisualizationData().GetMaterial(ModeName) == NULL)
+					{
+						// Mode is out of range, so display a message to the user, and reset the mode back to the previous valid one
+						UE_LOG(LogConsoleResponse, Warning, TEXT("Buffer visualization mode '%s' does not exist"), *ModeNameString);
+						NewBufferVisualizationMode = GetCurrentBufferVisualizationMode();
+						// todo: cvars are user settings, here the cvar state is used to avoid log spam and to auto correct for the user (likely not what the user wants)
+						ICVar->Set(*NewBufferVisualizationMode.GetPlainNameString(), ECVF_SetByCode);
+					}
+					else
+					{
+						NewBufferVisualizationMode = ModeName;
+					}
+				}
+			}
+
+			if (NewBufferVisualizationMode != GetCurrentBufferVisualizationMode())
+			{
+				SetCurrentBufferVisualizationMode(NewBufferVisualizationMode);
+			}
+		}
+
+		TMap<ULocalPlayer*, FSceneView*> PlayerViewMap;
+
+		FAudioDevice* AudioDevice = MyWorld->GetAudioDevice();
+		TArray<FSceneView*> Views;
+
+		for (FLocalPlayerIterator Iterator(GEngine, MyWorld); Iterator; ++Iterator)
+		{
+			ULocalPlayer* LocalPlayer = *Iterator;
+			if (LocalPlayer)
+			{
+				APlayerController* PlayerController = LocalPlayer->PlayerController;
+
+				for (int32 viewIndex = StartViewIndex; viewIndex < NumViews; ++viewIndex)
+				{
+					// Calculate the player's view information.
+					FVector		ViewLocation;
+					FRotator	ViewRotation;
+
+					EStereoscopicPass PassType = bStereoRendering ? GEngine->StereoRenderingDevice->GetViewPassForIndex(bStereoRendering, viewIndex) : eSSP_FULL;
+
+					FSceneView* View = LocalPlayer->CalcSceneView(&ViewFamily, ViewLocation, ViewRotation, InViewport, nullptr, PassType);
+
+					if (View)
+					{
+						Views.Add(View);
+
+						if (View->Family->EngineShowFlags.Wireframe)
+						{
+							// Wireframe color is emissive-only, and mesh-modifying materials do not use material substitution, hence...
+							View->DiffuseOverrideParameter = FVector4(0.f, 0.f, 0.f, 0.f);
+							View->SpecularOverrideParameter = FVector4(0.f, 0.f, 0.f, 0.f);
+						}
+						else if (View->Family->EngineShowFlags.OverrideDiffuseAndSpecular)
+						{
+							View->DiffuseOverrideParameter = FVector4(GEngine->LightingOnlyBrightness.R, GEngine->LightingOnlyBrightness.G, GEngine->LightingOnlyBrightness.B, 0.0f);
+							View->SpecularOverrideParameter = FVector4(.1f, .1f, .1f, 0.0f);
+						}
+						else if (View->Family->EngineShowFlags.ReflectionOverride)
+						{
+							View->DiffuseOverrideParameter = FVector4(0.f, 0.f, 0.f, 0.f);
+							View->SpecularOverrideParameter = FVector4(1, 1, 1, 0.0f);
+							View->NormalOverrideParameter = FVector4(0, 0, 1, 0.0f);
+							View->RoughnessOverrideParameter = FVector2D(0.0f, 0.0f);
+						}
+
+						if (!View->Family->EngineShowFlags.Diffuse)
+						{
+							View->DiffuseOverrideParameter = FVector4(0.f, 0.f, 0.f, 0.f);
+						}
+
+						if (!View->Family->EngineShowFlags.Specular)
+						{
+							View->SpecularOverrideParameter = FVector4(0.f, 0.f, 0.f, 0.f);
+						}
+
+						View->CurrentBufferVisualizationMode = GetCurrentBufferVisualizationMode();
+
+						View->CameraConstrainedViewRect = View->UnscaledViewRect;
+
+
+
+						// If this is the primary drawing pass, update things that depend on the view location
+						//if (viewIndex == 0)
+						{
+							// Save the location of the view.
+							LocalPlayer->LastViewLocation = ViewLocation;
+
+							PlayerViewMap.Add(LocalPlayer, View);
+
+							// Update the listener.
+							if (AudioDevice != NULL && PlayerController != NULL)
+							{
+								bool bUpdateListenerPosition = true;
+
+								// If the main audio device is used for multiple PIE viewport clients, we only
+								// want to update the main audio device listener position if it is in focus
+								if (GEngine)
+								{
+									FAudioDeviceManager* AudioDeviceManager = GEngine->GetAudioDeviceManager();
+
+									// If there is more than one world referencing the main audio device
+									if (AudioDeviceManager->GetNumMainAudioDeviceWorlds() > 1)
+									{
+										uint32 MainAudioDeviceHandle = GEngine->GetAudioDeviceHandle();
+										if (AudioDevice->DeviceHandle == MainAudioDeviceHandle && !HasAudioFocus())
+										{
+											bUpdateListenerPosition = false;
+										}
+									}
+								}
+
+								if (bUpdateListenerPosition)
+								{
+									FVector Location;
+									FVector ProjFront;
+									FVector ProjRight;
+									PlayerController->GetAudioListenerPosition(Location, ProjFront, ProjRight);
+
+									FTransform ListenerTransform(FRotationMatrix::MakeFromXY(ProjFront, ProjRight));
+
+									// Allow the HMD to adjust based on the head position of the player, as opposed to the view location
+									if (GEngine->XRSystem.IsValid() && GEngine->StereoRenderingDevice.IsValid() && GEngine->StereoRenderingDevice->IsStereoEnabled())
+									{
+										const FVector Offset = GEngine->XRSystem->GetAudioListenerOffset();
+										Location += ListenerTransform.TransformPositionNoScale(Offset);
+									}
+
+									ListenerTransform.SetTranslation(Location);
+									ListenerTransform.NormalizeRotation();
+
+									uint32 ViewportIndex = PlayerViewMap.Num() - 1;
+									AudioDevice->SetListener(MyWorld, ViewportIndex, ListenerTransform, (View->bCameraCut ? 0.f : MyWorld->GetDeltaSeconds()));
+
+									FVector OverrideAttenuation;
+									if (PlayerController->GetAudioListenerAttenuationOverridePosition(OverrideAttenuation))
+									{
+										AudioDevice->SetListenerAttenuationOverride(OverrideAttenuation);
+									}
+									else
+									{
+										AudioDevice->ClearListenerAttenuationOverride();
+									}
+								}
+							}
+
+#if RHI_RAYTRACING
+							View->SetupRayTracedRendering();
+#endif
+						}
+
+						// Add view information for resource streaming. Allow up to 5X boost for small FOV.
+						const float StreamingScale = 1.f / FMath::Clamp<float>(View->LODDistanceFactor, .2f, 1.f);
+						IStreamingManager::Get().AddViewInformation(View->ViewMatrices.GetViewOrigin(), View->UnscaledViewRect.Width(), View->UnscaledViewRect.Width() * View->ViewMatrices.GetProjectionMatrix().M[0][0], StreamingScale);
+						MyWorld->ViewLocationsRenderedLastFrame.Add(View->ViewMatrices.GetViewOrigin());
+					}
+				}
+			}
+		}
+
+		FinalizeViews(&ViewFamily, PlayerViewMap);
+
+		// Update level streaming.
+		MyWorld->UpdateLevelStreaming();
+
+		// Find largest rectangle bounded by all rendered views.
+		uint32 MinX = InViewport->GetSizeXY().X, MinY = InViewport->GetSizeXY().Y, MaxX = 0, MaxY = 0;
+		uint32 TotalArea = 0;
+		{
+			for (int32 ViewIndex = 0; ViewIndex < ViewFamily.Views.Num(); ++ViewIndex)
+			{
+				const FSceneView* View = ViewFamily.Views[ViewIndex];
+
+				FIntRect UpscaledViewRect = View->UnscaledViewRect;
+
+				MinX = FMath::Min<uint32>(UpscaledViewRect.Min.X, MinX);
+				MinY = FMath::Min<uint32>(UpscaledViewRect.Min.Y, MinY);
+				MaxX = FMath::Max<uint32>(UpscaledViewRect.Max.X, MaxX);
+				MaxY = FMath::Max<uint32>(UpscaledViewRect.Max.Y, MaxY);
+				TotalArea += FMath::TruncToInt(UpscaledViewRect.Width()) * FMath::TruncToInt(UpscaledViewRect.Height());
+			}
+
+			// To draw black borders around the rendered image (prevents artifacts from post processing passes that read outside of the image e.g. PostProcessAA)
+			{
+				int32 BlackBordersCVar = IConsoleManager::Get().FindConsoleVariable(TEXT("r.BlackBorders"), false)->GetInt();
+				int32 BlackBorders = FMath::Clamp(BlackBordersCVar, 0, 10);
+
+				if (ViewFamily.Views.Num() == 1 && BlackBorders)
+				{
+					MinX += BlackBorders;
+					MinY += BlackBorders;
+					MaxX -= BlackBorders;
+					MaxY -= BlackBorders;
+					TotalArea = (MaxX - MinX) * (MaxY - MinY);
+				}
+			}
+		}
+
+		// If the views don't cover the entire bounding rectangle, clear the entire buffer.
+		bool bBufferCleared = false;
+		bool bStereoscopicPass = (ViewFamily.Views.Num() != 0 && ViewFamily.Views[0]->StereoPass != eSSP_FULL);
+		if (ViewFamily.Views.Num() == 0 || TotalArea != (MaxX - MinX)*(MaxY - MinY) || bDisableWorldRendering || bStereoscopicPass)
+		{
+			if (bDisableWorldRendering || !bStereoscopicPass) // TotalArea computation does not work correctly for stereoscopic views
+			{
+				//SceneCanvas->Clear(FLinearColor::Transparent);
+			}
+
+			bBufferCleared = true;
+		}
+
+		// Force screen percentage show flag to be turned off if not supported.
+		if (!ViewFamily.SupportsScreenPercentage())
+		{
+			ViewFamily.EngineShowFlags.ScreenPercentage = false;
+		}
+		// Update level streaming.
+		MyWorld->UpdateLevelStreaming();
+
+
+		// Set up secondary resolution fraction for the view family.
+		if (!bStereoRendering && ViewFamily.SupportsScreenPercentage())
+		{
+			float BlackBordersCVar = IConsoleManager::Get().FindConsoleVariable(TEXT("r.SecondaryScreenPercentage.GameViewport"), false)->GetFloat();
+			float CustomSecondaruScreenPercentage = BlackBordersCVar;
+
+			if (CustomSecondaruScreenPercentage > 0.0)
+			{
+				// Override secondary resolution fraction with CVar.
+				ViewFamily.SecondaryViewFraction = FMath::Min(CustomSecondaruScreenPercentage / 100.0f, 1.0f);
+			}
+			else
+			{
+				// Automatically compute secondary resolution fraction from DPI.
+				ViewFamily.SecondaryViewFraction = GetDPIDerivedResolutionFraction();
+			}
+
+			check(ViewFamily.SecondaryViewFraction > 0.0f);
+		}
+
+		checkf(ViewFamily.GetScreenPercentageInterface() == nullptr,
+			TEXT("Some code has tried to set up an alien screen percentage driver, that could be wrong if not supported very well by the RHI."));
+
+		// Setup main view family with screen percentage interface by dynamic resolution if screen percentage is enabled.
+#if WITH_DYNAMIC_RESOLUTION
+		if (ViewFamily.EngineShowFlags.ScreenPercentage)
+		{
+			FDynamicResolutionStateInfos DynamicResolutionStateInfos;
+			GEngine->GetDynamicResolutionCurrentStateInfos( DynamicResolutionStateInfos);
+
+			// Do not allow dynamic resolution to touch the view family if not supported to ensure there is no possibility to ruin
+			// game play experience on platforms that does not support it, but have it enabled by mistake.
+			if (DynamicResolutionStateInfos.Status == EDynamicResolutionStatus::Enabled)
+			{
+				GEngine->EmitDynamicResolutionEvent(EDynamicResolutionStateEvent::BeginDynamicResolutionRendering);
+				GEngine->GetDynamicResolutionState()->SetupMainViewFamily(ViewFamily);
+			}
+			else if (DynamicResolutionStateInfos.Status == EDynamicResolutionStatus::DebugForceEnabled)
+			{
+				GEngine->EmitDynamicResolutionEvent(EDynamicResolutionStateEvent::BeginDynamicResolutionRendering);
+				ViewFamily.SetScreenPercentageInterface(new FLegacyScreenPercentageDriver(
+					ViewFamily,
+					DynamicResolutionStateInfos.ResolutionFractionApproximation,
+					false,
+					DynamicResolutionStateInfos.ResolutionFractionUpperBound));
+			}
+
+#if CSV_PROFILER
+			if (DynamicResolutionStateInfos.ResolutionFractionApproximation >= 0.0f)
+			{
+				CSV_CUSTOM_STAT_GLOBAL(DynamicResolutionPercentage, DynamicResolutionStateInfos.ResolutionFractionApproximation * 100.0f, ECsvCustomStatOp::Set);
+			}
+#endif
+		}
+#endif
+
+		// If a screen percentage interface was not set by dynamic resolution, then create one matching legacy behavior.
+		if (ViewFamily.GetScreenPercentageInterface() == nullptr)
+		{
+			bool AllowPostProcessSettingsScreenPercentage = false;
+			float GlobalResolutionFraction = 1.0f;
+
+			if (ViewFamily.EngineShowFlags.ScreenPercentage)
+			{
+				// Allow FPostProcessSettings::ScreenPercentage.
+				AllowPostProcessSettingsScreenPercentage = true;
+
+				// Get global view fraction set by r.ScreenPercentage.
+				GlobalResolutionFraction = FLegacyScreenPercentageDriver::GetCVarResolutionFraction();
+			}
+
+			ViewFamily.SetScreenPercentageInterface(new FLegacyScreenPercentageDriver(
+				ViewFamily, GlobalResolutionFraction, AllowPostProcessSettingsScreenPercentage));
+		}
+		else if (bStereoRendering)
+		{
+			// Change screen percentage method to raw output when doing dynamic resolution with VR if not using TAA upsample.
+			for (FSceneView* View : Views)
+			{
+				if (View->PrimaryScreenPercentageMethod == EPrimaryScreenPercentageMethod::SpatialUpscale)
+				{
+					View->PrimaryScreenPercentageMethod = EPrimaryScreenPercentageMethod::RawOutput;
+				}
+			}
+		}
+
+		// Draw the player views.
+		if (!bDisableWorldRendering && PlayerViewMap.Num() > 0 && FSlateApplication::Get().GetPlatformApplication()->IsAllowedToRender()) //-V560
+		{
+			GetRendererModule().BeginRenderingViewFamily(SceneCanvas, &ViewFamily);
+		}
+		else
+		{
+			// Make sure RHI resources get flushed if we're not using a renderer
+			ENQUEUE_RENDER_COMMAND(UGameViewportClient_FlushRHIResources)(
+				[](FRHICommandListImmediate& RHICmdList)
+			{
+				RHICmdList.ImmediateFlush(EImmediateFlushType::FlushRHIThreadFlushResources);
+			});
+		}
+
+		//SCOPE_CYCLE_COUNTER(STAT_UIDrawingTime);
+		//CSV_SCOPED_TIMING_STAT_EXCLUSIVE(UI);
+
+		// render HUD
+		bool bDisplayedSubtitles = false;
+		for (FConstPlayerControllerIterator Iterator = MyWorld->GetPlayerControllerIterator(); Iterator; ++Iterator)
+		{
+			APlayerController* PlayerController = Iterator->Get();
+			if (PlayerController)
+			{
+				ULocalPlayer* LocalPlayer = Cast<ULocalPlayer>(PlayerController->Player);
+				if (LocalPlayer)
+				{
+					FSceneView* View = PlayerViewMap.FindRef(LocalPlayer);
+					if (View != NULL)
+					{
+						// rendering to directly to viewport target
+						FVector CanvasOrigin(FMath::TruncToFloat(View->UnscaledViewRect.Min.X), FMath::TruncToInt(View->UnscaledViewRect.Min.Y), 0.f);
+
+						CanvasObject->Init(View->UnscaledViewRect.Width(), View->UnscaledViewRect.Height(), View, SceneCanvas);
+
+						// Set the canvas transform for the player's view rectangle.
+						check(SceneCanvas);
+						SceneCanvas->PushAbsoluteTransform(FTranslationMatrix(CanvasOrigin));
+						CanvasObject->ApplySafeZoneTransform();
+
+						// Render the player's HUD.
+						if (PlayerController->MyHUD)
+						{
+							//SCOPE_CYCLE_COUNTER(STAT_HudTime);
+
+							DebugCanvasObject->SceneView = View;
+							PlayerController->MyHUD->SetCanvas(CanvasObject, DebugCanvasObject);
+
+							PlayerController->MyHUD->PostRender();
+
+							// Put these pointers back as if a blueprint breakpoint hits during HUD PostRender they can
+							// have been changed
+							CanvasObject->Canvas = SceneCanvas;
+							DebugCanvasObject->Canvas = DebugCanvas;
+
+							// A side effect of PostRender is that the playercontroller could be destroyed
+							if (!PlayerController->IsPendingKill())
+							{
+								PlayerController->MyHUD->SetCanvas(NULL, NULL);
+							}
+						}
+
+						if (DebugCanvas != NULL)
+						{
+							DebugCanvas->PushAbsoluteTransform(FTranslationMatrix(CanvasOrigin));
+							UDebugDrawService::Draw(ViewFamily.EngineShowFlags, InViewport, View, DebugCanvas, DebugCanvasObject);
+							DebugCanvas->PopTransform();
+						}
+
+						CanvasObject->PopSafeZoneTransform();
+						SceneCanvas->PopTransform();
+
+						// draw subtitles
+						if (!bDisplayedSubtitles)
+						{
+							FVector2D MinPos(0.f, 0.f);
+							FVector2D MaxPos(1.f, 1.f);
+							GetSubtitleRegion(MinPos, MaxPos);
+
+							const uint32 SizeX = SceneCanvas->GetRenderTarget()->GetSizeXY().X;
+							const uint32 SizeY = SceneCanvas->GetRenderTarget()->GetSizeXY().Y;
+							FIntRect SubtitleRegion(FMath::TruncToInt(SizeX * MinPos.X), FMath::TruncToInt(SizeY * MinPos.Y), FMath::TruncToInt(SizeX * MaxPos.X), FMath::TruncToInt(SizeY * MaxPos.Y));
+							FSubtitleManager::GetSubtitleManager()->DisplaySubtitles(SceneCanvas, SubtitleRegion, MyWorld->GetAudioTimeSeconds());
+							bDisplayedSubtitles = true;
+						}
+					}
+				}
+			}
+		}
+	}
+
+	// Beyond this point, only UI rendering independent from dynamc resolution.
+	GEngine->EmitDynamicResolutionEvent(EDynamicResolutionStateEvent::EndDynamicResolutionRendering);
+
+	// Remove temporary debug lines.
+	if (MyWorld->LineBatcher != nullptr)
+	{
+		MyWorld->LineBatcher->Flush();
+	}
+
+	if (MyWorld->ForegroundLineBatcher != nullptr)
+	{
+		MyWorld->ForegroundLineBatcher->Flush();
+	}
+
+	// Draw FX debug information.
+	if (MyWorld->FXSystem)
+	{
+		MyWorld->FXSystem->DrawDebug(SceneCanvas);
+	}
+
+	{
+		//ensure canvas has been flushed before rendering UI
+		SceneCanvas->Flush_GameThread();
+
+		OnDrawn().Broadcast();
+
+		// Allow the viewport to render additional stuff
+		PostRender(DebugCanvasObject);
+	}
+
+
+	// Grab the player camera location and orientation so we can pass that along to the stats drawing code.
+	FVector PlayerCameraLocation = FVector::ZeroVector;
+	FRotator PlayerCameraRotation = FRotator::ZeroRotator;
+	{
+		for (FConstPlayerControllerIterator Iterator = MyWorld->GetPlayerControllerIterator(); Iterator; ++Iterator)
+		{
+			if (APlayerController* PC = Iterator->Get())
+			{
+				PC->GetPlayerViewPoint(PlayerCameraLocation, PlayerCameraRotation);
+			}
+		}
+	}
+
+	if (DebugCanvas)
+	{
+		// Reset the debug canvas to be full-screen before drawing the console
+		// (the debug draw service above has messed with the viewport size to fit it to a single player's subregion)
+		DebugCanvasObject->Init(DebugCanvasSize.X, DebugCanvasSize.Y, NULL, DebugCanvas);
+
+		DrawStatsHUD(MyWorld, InViewport, DebugCanvas, DebugCanvasObject, DebugProperties, PlayerCameraLocation, PlayerCameraRotation);
+
+		if (GEngine->IsStereoscopic3D(InViewport))
+		{
+#if 0 //!UE_BUILD_SHIPPING
+			// TODO: replace implementation in OculusHMD with a debug renderer
+			if (GEngine->XRSystem.IsValid())
+			{
+				GEngine->XRSystem->DrawDebug(DebugCanvasObject);
+			}
+#endif
+		}
+
+		// Render the console absolutely last because developer input is was matter the most.
+		if (ViewportConsole)
+		{
+			ViewportConsole->PostRender_Console(DebugCanvasObject);
+		}
+	}
+
+	OnEndDraw().Broadcast();
+}
diff --git a/Engine/Plugins/Runtime/nDisplay/Source/DisplayCluster/Private/Render/Device/DisplayClusterDeviceBase.cpp b/Engine/Plugins/Runtime/nDisplay/Source/DisplayCluster/Private/Render/Device/DisplayClusterDeviceBase.cpp
--- a/Engine/Plugins/Runtime/nDisplay/Source/DisplayCluster/Private/Render/Device/DisplayClusterDeviceBase.cpp
+++ b/Engine/Plugins/Runtime/nDisplay/Source/DisplayCluster/Private/Render/Device/DisplayClusterDeviceBase.cpp
@@ -203,24 +203,43 @@ void FDisplayClusterDeviceBase::SetViewportCamera(const FString& InCameraId /* =
 	UE_LOG(LogDisplayClusterRender, Warning, TEXT("Couldn't assign '%s' camera. Viewport '%s' not found"), *InCameraId, *InViewportId);
 }
 
-void FDisplayClusterDeviceBase::SetCustomPostProcessing(const FString& ViewportID, const FPostProcessSettings& PostProcessingSettings)
+void FDisplayClusterDeviceBase::SetStartPostProcessingSettings(const FString& ViewportID, const FPostProcessSettings& StartPostProcessingSettings)
 {	
-	// find viewport index 
-	int ViewportIdx = -1;
-	for(int i = 0; i < RenderViewports.Num(); i++)
+	for(int ViewportIndex = 0; ViewportIndex < RenderViewports.Num(); ViewportIndex++)
 	{
-		if (RenderViewports[i].GetId() == ViewportID)
+		if (RenderViewports[ViewportIndex].GetId() == ViewportID)
 		{
-			ViewportIdx = i;
+			ViewportStartPostProcessingSettings.Emplace(ViewportIndex, StartPostProcessingSettings);
 			break;
 		}
 	}
+}
 
-	// check if post processing settings assigned then override or add
-	if (ViewportIdx != -1)
+void FDisplayClusterDeviceBase::SetOverridePostProcessingSettings(const FString& ViewportID, const FPostProcessSettings& OverridePostProcessingSettings, float BlendWeight)
+{
+	for (int ViewportIndex = 0; ViewportIndex < RenderViewports.Num(); ViewportIndex++)
 	{
-		ViewportFinalPostProcessingSettingsOverride.Remove(ViewportIdx);
-		ViewportFinalPostProcessingSettingsOverride.Emplace(ViewportIdx, PostProcessingSettings);
+		if (RenderViewports[ViewportIndex].GetId() == ViewportID)
+		{
+			FOverridePostProcessingSettings OverrideSettings;
+			OverrideSettings.BlendWeight = BlendWeight;
+			OverrideSettings.PostProcessingSettings = OverridePostProcessingSettings;
+			ViewportOverridePostProcessingSettings.Emplace(ViewportIndex, OverrideSettings);
+			break;
+		}
+	}
+}
+
+
+void FDisplayClusterDeviceBase::SetFinalPostProcessingSettings(const FString& ViewportID, const FPostProcessSettings& FinalPostProcessingSettings)
+	{
+	for (int ViewportIndex = 0; ViewportIndex < RenderViewports.Num(); ViewportIndex++)
+	{
+		if (RenderViewports[ViewportIndex].GetId() == ViewportID)
+		{
+			ViewportFinalPostProcessingSettings.Emplace(ViewportIndex, FinalPostProcessingSettings);
+			break;
+		}
 	}
 }
 
@@ -769,13 +788,39 @@ void FDisplayClusterDeviceBase::CopyTextureToBackBuffer_RenderThread(FRHICommand
 	RHICmdList.CopyToResolveTarget(SrcTexture, BackBuffer, copyParams);
 }
 
-void FDisplayClusterDeviceBase::UpdatePostProcessSettings(struct FPostProcessSettings* FinalPostProcessingSettings, const enum EStereoscopicPass StereoPassType)
+void FDisplayClusterDeviceBase::StartFinalPostprocessSettings(struct FPostProcessSettings* StartPostProcessingSettings, const enum EStereoscopicPass StereoPassType)
+{
+	int ViewportNumber = DecodeViewportIndex(StereoPassType);
+	auto CustomPostProcessAvailable = ViewportStartPostProcessingSettings.Find(ViewportNumber);
+	if (CustomPostProcessAvailable)
+	{
+		*StartPostProcessingSettings = ViewportStartPostProcessingSettings[ViewportNumber];
+		ViewportStartPostProcessingSettings.Remove(ViewportNumber);
+	}
+}
+
+bool FDisplayClusterDeviceBase::OverrideFinalPostprocessSettings(struct FPostProcessSettings* OverridePostProcessingSettings, const enum EStereoscopicPass StereoPassType, float& BlendWeight)
+{
+	int ViewportNumber = DecodeViewportIndex(StereoPassType);
+	auto CustomPostProcessAvailable = ViewportOverridePostProcessingSettings.Find(ViewportNumber);
+	if (CustomPostProcessAvailable)
+	{
+		*OverridePostProcessingSettings = ViewportOverridePostProcessingSettings[ViewportNumber].PostProcessingSettings;
+		BlendWeight = ViewportOverridePostProcessingSettings[ViewportNumber].BlendWeight;
+		ViewportOverridePostProcessingSettings.Remove(ViewportNumber);
+		return true;
+	}
+
+	return false;
+}
+
+void FDisplayClusterDeviceBase::EndFinalPostprocessSettings(struct FPostProcessSettings* FinalPostProcessingSettings, const enum EStereoscopicPass StereoPassType)
 {
 	int ViewportNumber = DecodeViewportIndex(StereoPassType);
-	auto CustomPostProcessAvailable = ViewportFinalPostProcessingSettingsOverride.Find(ViewportNumber);
+	auto CustomPostProcessAvailable = ViewportFinalPostProcessingSettings.Find(ViewportNumber);
 	if (CustomPostProcessAvailable)
 	{
-		*FinalPostProcessingSettings = ViewportFinalPostProcessingSettingsOverride[ViewportNumber];
-		//ViewportFinalPostProcessingSettingsOverride.Remove(ViewportNumber);
+		*FinalPostProcessingSettings = ViewportFinalPostProcessingSettings[ViewportNumber];
+		ViewportFinalPostProcessingSettings.Remove(ViewportNumber);
 	}		
 }
\ No newline at end of file
diff --git a/Engine/Plugins/Runtime/nDisplay/Source/DisplayCluster/Private/Render/Device/DisplayClusterDeviceBase.h b/Engine/Plugins/Runtime/nDisplay/Source/DisplayCluster/Private/Render/Device/DisplayClusterDeviceBase.h
--- a/Engine/Plugins/Runtime/nDisplay/Source/DisplayCluster/Private/Render/Device/DisplayClusterDeviceBase.h
+++ b/Engine/Plugins/Runtime/nDisplay/Source/DisplayCluster/Private/Render/Device/DisplayClusterDeviceBase.h
@@ -39,7 +39,9 @@ public:
 	virtual bool Initialize() override;
 	virtual void InitializeWorldContent(UWorld* InWorld) override;
 	virtual void SetViewportCamera(const FString& InCameraId = FString(), const FString& InViewportId = FString()) override;
-	virtual void SetCustomPostProcessing(const FString& ViewportID, const FPostProcessSettings& PostProcessingSettings) override;
+	virtual void SetStartPostProcessingSettings(const FString& ViewportID, const FPostProcessSettings& StartPostProcessingSettings) override;
+	virtual void SetOverridePostProcessingSettings(const FString& ViewportID, const FPostProcessSettings& OverridePostProcessingSettings, float BlendWeight = 1.0f) override;
+	virtual void SetFinalPostProcessingSettings(const FString& ViewportID, const FPostProcessSettings& FinalPostProcessingSettings) override;
 	virtual bool GetViewportRect(const FString& InViewportID, FIntRect& Rect) override;
 
 public:
@@ -84,6 +86,15 @@ protected:
 	virtual bool AllocateDepthTexture(uint32 Index, uint32 SizeX, uint32 SizeY, uint8 Format, uint32 NumMips, uint32 Flags, uint32 TargetableTextureFlags, FTexture2DRHIRef& OutTargetableTexture, FTexture2DRHIRef& OutShaderResourceTexture, uint32 NumSamples = 1)
 	{ return false; }
 
+	virtual bool DeviceIsAPrimaryView(EStereoscopicPass Pass) override
+	{ return true; }
+	
+	virtual bool DeviceIsASecondaryView(EStereoscopicPass Pass) override
+	{ return false; }
+
+	virtual bool DeviceIsAnAdditionalView(EStereoscopicPass Pass) override
+	{ return false; }
+
 protected:
 	//////////////////////////////////////////////////////////////////////////////////////////////
 	// FRHICustomPresent
@@ -126,10 +137,19 @@ protected:
 	void AddViewport(const FString& InViewportId, const FIntPoint& InViewportLocation, const FIntPoint& InViewportSize, TSharedPtr<IDisplayClusterProjectionPolicy> InProjPolicy, const FString& InCameraId, bool IsRTT = false);
 	// Performs copying of render target data to the back buffer
 	virtual void CopyTextureToBackBuffer_RenderThread(FRHICommandListImmediate& RHICmdList, FRHITexture2D* BackBuffer, FRHITexture2D* SrcTexture, FVector2D WindowSize) const;
+	
 	// Checks if custom post processing settings is assigned for specific viewport and assign them to be used
-	virtual void UpdatePostProcessSettings(struct FPostProcessSettings* FinalPostProcessingSettings, const enum EStereoscopicPass StereoPassType) override;
+	virtual void StartFinalPostprocessSettings(struct FPostProcessSettings* StartPostProcessingSettings, const enum EStereoscopicPass StereoPassType) override;
+	virtual bool OverrideFinalPostprocessSettings(struct FPostProcessSettings* OverridePostProcessingSettings, const enum EStereoscopicPass StereoPassType, float& BlendWeight) override;
+	virtual void EndFinalPostprocessSettings(struct FPostProcessSettings* FinalPostProcessingSettings, const enum EStereoscopicPass StereoPassType) override;
 
 protected:
+	struct FOverridePostProcessingSettings
+	{
+		float BlendWeight = 1.0f;
+		FPostProcessSettings PostProcessingSettings;
+	};
+
 	// Viewports
 	mutable TArray<FDisplayClusterRenderViewport> RenderViewports;
 	// Views per viewport (render passes)
@@ -137,7 +157,9 @@ protected:
 	// UE4 main viewport
 	FViewport* MainViewport = nullptr;
 	// custom post processing settings
-	TMap<int, FPostProcessSettings> ViewportFinalPostProcessingSettingsOverride;	 
+	TMap<int, FPostProcessSettings> ViewportStartPostProcessingSettings;	 
+	TMap<int, FOverridePostProcessingSettings> ViewportOverridePostProcessingSettings;
+	TMap<int, FPostProcessSettings> ViewportFinalPostProcessingSettings;
 
 	// Data access synchronization
 	mutable FCriticalSection InternalsSyncScope;
diff --git a/Engine/Plugins/Runtime/nDisplay/Source/DisplayCluster/Private/Render/DisplayClusterRenderManager.cpp b/Engine/Plugins/Runtime/nDisplay/Source/DisplayCluster/Private/Render/DisplayClusterRenderManager.cpp
--- a/Engine/Plugins/Runtime/nDisplay/Source/DisplayCluster/Private/Render/DisplayClusterRenderManager.cpp
+++ b/Engine/Plugins/Runtime/nDisplay/Source/DisplayCluster/Private/Render/DisplayClusterRenderManager.cpp
@@ -466,14 +466,34 @@ bool FDisplayClusterRenderManager::GetViewportRect(const FString& InViewportID,
 	return RenderDevice->GetViewportRect(InViewportID, Rect);
 }
 
-void FDisplayClusterRenderManager::SetCustomPostProcessing(const FString& ViewportID, const FPostProcessSettings& PostProcessingSettings)
+void FDisplayClusterRenderManager::SetStartPostProcessingSettings(const FString& ViewportID, const FPostProcessSettings& StartPostProcessingSettings)
 {
 	if (!RenderDevice.IsValid())
 	{
 		return;
 	}
 
-	RenderDevice->SetCustomPostProcessing(ViewportID, PostProcessingSettings);
+	RenderDevice->SetStartPostProcessingSettings(ViewportID, StartPostProcessingSettings);
+}
+
+void FDisplayClusterRenderManager::SetOverridePostProcessingSettings(const FString& ViewportID, const FPostProcessSettings& OverridePostProcessingSettings, float BlendWeight)
+{
+	if (!RenderDevice.IsValid())
+	{
+		return;
+	}
+
+	RenderDevice->SetOverridePostProcessingSettings(ViewportID, OverridePostProcessingSettings, BlendWeight);
+}
+
+void FDisplayClusterRenderManager::SetFinalPostProcessingSettings(const FString& ViewportID, const FPostProcessSettings& FinalPostProcessingSettings)
+{
+	if (!RenderDevice.IsValid())
+	{
+		return;
+	}
+
+	RenderDevice->SetFinalPostProcessingSettings(ViewportID, FinalPostProcessingSettings);
 }
 
 void FDisplayClusterRenderManager::SetInterpupillaryDistance(const FString& CameraId, float EyeDistance)
diff --git a/Engine/Plugins/Runtime/nDisplay/Source/DisplayCluster/Private/Render/DisplayClusterRenderManager.h b/Engine/Plugins/Runtime/nDisplay/Source/DisplayCluster/Private/Render/DisplayClusterRenderManager.h
--- a/Engine/Plugins/Runtime/nDisplay/Source/DisplayCluster/Private/Render/DisplayClusterRenderManager.h
+++ b/Engine/Plugins/Runtime/nDisplay/Source/DisplayCluster/Private/Render/DisplayClusterRenderManager.h
@@ -59,7 +59,9 @@ public:
 	virtual bool UnregisterPostprocessOperation(const FString& InName) override;
 	virtual TMap<FString, IPDisplayClusterRenderManager::FDisplayClusterPPInfo> GetRegisteredPostprocessOperations() const override;
 	// Custom Rendering Post-process
-	virtual void SetCustomPostProcessing(const FString& ViewportID, const FPostProcessSettings& PostProcessingSettings) override;
+	virtual void SetStartPostProcessingSettings(const FString& ViewportID, const FPostProcessSettings& StartPostProcessingSettings) override;
+	virtual void SetOverridePostProcessingSettings(const FString& ViewportID, const FPostProcessSettings& OverridePostProcessingSettings, float BlendWeight = 1.0f) override;
+	virtual void SetFinalPostProcessingSettings(const FString& ViewportID, const FPostProcessSettings& FinalPostProcessingSettings) override;
 
 	// Camera
 	virtual void SetViewportCamera(const FString& InCameraId = FString(), const FString& InViewportId = FString()) override;
diff --git a/Engine/Plugins/Runtime/nDisplay/Source/DisplayCluster/Private/Render/Synchronization/DisplayClusterRenderSyncPolicyBase.cpp b/Engine/Plugins/Runtime/nDisplay/Source/DisplayCluster/Private/Render/Synchronization/DisplayClusterRenderSyncPolicyBase.cpp
--- a/Engine/Plugins/Runtime/nDisplay/Source/DisplayCluster/Private/Render/Synchronization/DisplayClusterRenderSyncPolicyBase.cpp
+++ b/Engine/Plugins/Runtime/nDisplay/Source/DisplayCluster/Private/Render/Synchronization/DisplayClusterRenderSyncPolicyBase.cpp
@@ -20,7 +20,7 @@ FDisplayClusterRenderSyncPolicyBase::~FDisplayClusterRenderSyncPolicyBase()
 
 void FDisplayClusterRenderSyncPolicyBase::SyncBarrierRenderThread()
 {
-	DISPLAY_CLUSTER_FUNC_TRACE(LogDisplayClusterRender);
+	DISPLAY_CLUSTER_FUNC_TRACE(LogDisplayClusterRenderSync);
 
 	if (GDisplayCluster->GetOperationMode() == EDisplayClusterOperationMode::Disabled)
 	{
@@ -36,5 +36,5 @@ void FDisplayClusterRenderSyncPolicyBase::SyncBarrierRenderThread()
 		pController->WaitForSwapSync(&tTime, &bTime);
 	}
 
-	UE_LOG(LogDisplayClusterRender, Verbose, TEXT("Render barrier wait: t=%lf b=%lf"), tTime, bTime);
+	UE_LOG(LogDisplayClusterRenderSync, Verbose, TEXT("Render barrier wait: t=%lf b=%lf"), tTime, bTime);
 }
diff --git a/Engine/Plugins/Runtime/nDisplay/Source/DisplayCluster/Private/Render/Synchronization/DisplayClusterRenderSyncPolicyNvidiaOpenGL.cpp b/Engine/Plugins/Runtime/nDisplay/Source/DisplayCluster/Private/Render/Synchronization/DisplayClusterRenderSyncPolicyNvidiaOpenGL.cpp
--- a/Engine/Plugins/Runtime/nDisplay/Source/DisplayCluster/Private/Render/Synchronization/DisplayClusterRenderSyncPolicyNvidiaOpenGL.cpp
+++ b/Engine/Plugins/Runtime/nDisplay/Source/DisplayCluster/Private/Render/Synchronization/DisplayClusterRenderSyncPolicyNvidiaOpenGL.cpp
@@ -104,13 +104,13 @@ bool FDisplayClusterRenderSyncPolicyNvidiaOpenGL::InitializeNvidiaSwapLock()
 
 	if (!DisplayCluster_wglJoinSwapGroupNV_ProcAddress || !DisplayCluster_wglBindSwapBarrierNV_ProcAddress)
 	{
-		UE_LOG(LogDisplayClusterRender, Error, TEXT("Group/Barrier functions not available"));
+		UE_LOG(LogDisplayClusterRenderSync, Error, TEXT("Group/Barrier functions not available"));
 		return false;
 	}
 
 	if (!GEngine || !GEngine->GameViewport || !GEngine->GameViewport->Viewport || !GEngine->GameViewport->Viewport->GetViewportRHI().IsValid())
 	{
-		UE_LOG(LogDisplayClusterRender, Warning, TEXT("Game viewport hasn't been initialized yet"));
+		UE_LOG(LogDisplayClusterRenderSync, Warning, TEXT("Game viewport hasn't been initialized yet"));
 		return false;
 	}
 
@@ -125,36 +125,36 @@ bool FDisplayClusterRenderSyncPolicyNvidiaOpenGL::InitializeNvidiaSwapLock()
 
 	if (!DisplayCluster_wglQueryMaxSwapGroupsNV_ProcAddress(OpenGLContext->DeviceContext, &maxGroups, &maxBarriers))
 	{
-		UE_LOG(LogDisplayClusterRender, Error, TEXT("Couldn't query gr/br limits: %d"), glGetError());
+		UE_LOG(LogDisplayClusterRenderSync, Error, TEXT("Couldn't query gr/br limits: %d"), glGetError());
 		return false;
 	}
 
-	UE_LOG(LogDisplayClusterRender, Log, TEXT("max_groups=%d max_barriers=%d"), (int)maxGroups, (int)maxBarriers);
+	UE_LOG(LogDisplayClusterRenderSync, Log, TEXT("max_groups=%d max_barriers=%d"), (int)maxGroups, (int)maxBarriers);
 
 	if (!(maxGroups > 0 && maxBarriers > 0))
 	{
-		UE_LOG(LogDisplayClusterRender, Error, TEXT("There are no available groups or barriers"));
+		UE_LOG(LogDisplayClusterRenderSync, Error, TEXT("There are no available groups or barriers"));
 		return false;
 	}
 
 	if (!DisplayCluster_wglJoinSwapGroupNV_ProcAddress(OpenGLContext->DeviceContext, 1))
 	{
-		UE_LOG(LogDisplayClusterRender, Error, TEXT("Couldn't join swap group: %d"), glGetError());
+		UE_LOG(LogDisplayClusterRenderSync, Error, TEXT("Couldn't join swap group: %d"), glGetError());
 		return false;
 	}
 	else
 	{
-		UE_LOG(LogDisplayClusterRender, Log, TEXT("Successfully joined the swap group: 1"));
+		UE_LOG(LogDisplayClusterRenderSync, Log, TEXT("Successfully joined the swap group: 1"));
 	}
 
 	if (!DisplayCluster_wglBindSwapBarrierNV_ProcAddress(1, 1))
 	{
-		UE_LOG(LogDisplayClusterRender, Error, TEXT("Couldn't bind to swap barrier: %d"), glGetError());
+		UE_LOG(LogDisplayClusterRenderSync, Error, TEXT("Couldn't bind to swap barrier: %d"), glGetError());
 		return false;
 	}
 	else
 	{
-		UE_LOG(LogDisplayClusterRender, Log, TEXT("Successfully binded to the swap barrier: 1"));
+		UE_LOG(LogDisplayClusterRenderSync, Log, TEXT("Successfully binded to the swap barrier: 1"));
 	}
 
 	return true;
@@ -222,5 +222,5 @@ void FDisplayClusterRenderSyncPolicyNvidiaOpenGL::UpdateSwapInterval(int32 InSyn
 
 	// Perform that each frame
 	if (!DisplayCluster_wglSwapIntervalEXT_ProcAddress(InSyncInterval))
-		UE_LOG(LogDisplayClusterRender, Warning, TEXT("Couldn't set swap interval: %d"), InSyncInterval);
+		UE_LOG(LogDisplayClusterRenderSync, Warning, TEXT("Couldn't set swap interval: %d"), InSyncInterval);
 }
diff --git a/Engine/Plugins/Runtime/nDisplay/Source/DisplayCluster/Private/Render/Synchronization/DisplayClusterRenderSyncPolicySoftwareDX11.cpp b/Engine/Plugins/Runtime/nDisplay/Source/DisplayCluster/Private/Render/Synchronization/DisplayClusterRenderSyncPolicySoftwareDX11.cpp
--- a/Engine/Plugins/Runtime/nDisplay/Source/DisplayCluster/Private/Render/Synchronization/DisplayClusterRenderSyncPolicySoftwareDX11.cpp
+++ b/Engine/Plugins/Runtime/nDisplay/Source/DisplayCluster/Private/Render/Synchronization/DisplayClusterRenderSyncPolicySoftwareDX11.cpp
@@ -13,6 +13,7 @@
 #include "Windows/HideWindowsPlatformTypes.h"
 #include "RenderCore/Public/ShaderCore.h"
 
+#include "ShaderCore.h"
 #include "D3D11RHI/Private/Windows/D3D11RHIBasePrivate.h"
 
 #include "D3D11State.h"
@@ -21,6 +22,7 @@
 
 #include "NVIDIA/nvapi/nvapi.h"
 
+
 static TAutoConsoleVariable<int32> CVarAdvancedSyncEnabled(
 	TEXT("nDisplay.render.softsync.AdvancedSyncEnabled"),
 	0,
@@ -28,6 +30,20 @@ static TAutoConsoleVariable<int32> CVarAdvancedSyncEnabled(
 	ECVF_RenderThreadSafe
 );
 
+static TAutoConsoleVariable<int32> CVarUseCustomRefreshRate(
+	TEXT("nDisplay.render.softsync.UseCustomRefreshRate"),
+	0,
+	TEXT("Force custom refresh rate to be used in synchronization math (0 = disabled)"),
+	ECVF_RenderThreadSafe
+);
+
+static TAutoConsoleVariable<float> CVarCustomRefreshRate(
+	TEXT("nDisplay.render.softsync.CustomRefreshRate"),
+	60.f,
+	TEXT("Custom refresh rate for synchronization math (Hz)"),
+	ECVF_RenderThreadSafe
+);
+
 static TAutoConsoleVariable<float> CVarVBlankThreshold(
 	TEXT("nDisplay.render.softsync.VBlankThreshold"),
 	0.003f,
@@ -89,37 +105,37 @@ FDisplayClusterRenderSyncPolicySoftwareDX11::FDisplayClusterRenderSyncPolicySoft
 		const int Result = NvAPI_Initialize();
 		if (Result != NVAPI_OK)
 		{
-			UE_LOG(LogDisplayClusterRender, Error, TEXT("Couldn't initialize NvAPI. Error code: %d"), Result);
+			UE_LOG(LogDisplayClusterRenderSync, Error, TEXT("Couldn't initialize NvAPI. Error code: %d"), Result);
 			return;
 		}
 
-		UE_LOG(LogDisplayClusterRender, Log, TEXT("NvAPI has been initialized"));
+		UE_LOG(LogDisplayClusterRenderSync, Log, TEXT("NvAPI has been initialized"));
 
 		bNvApiInitialized = (Result == NVAPI_OK);
 		if (bNvApiInitialized)
 		{
 			if (NvAPI_GSync_EnumSyncDevices(GSyncDeviceHandles, &GSyncDeviceAmount) != NVAPI_OK)
 			{
-				UE_LOG(LogDisplayClusterRender, Warning, TEXT("Couldn't enumerate GSync devices or no devices found"));
+				UE_LOG(LogDisplayClusterRenderSync, Warning, TEXT("Couldn't enumerate GSync devices or no devices found"));
 			}
 			else
 			{
-				UE_LOG(LogDisplayClusterRender, Log, TEXT("Found sync devices: %u"), GSyncDeviceAmount);
+				UE_LOG(LogDisplayClusterRenderSync, Log, TEXT("Found sync devices: %u"), GSyncDeviceAmount);
 			}
 
 			if (NvAPI_EnumPhysicalGPUs(PhysGPUHandles, &PhysGPUAmount) != NVAPI_OK)
 			{
-				UE_LOG(LogDisplayClusterRender, Error, TEXT("Couldn't enumerate physical GPU devices"));
+				UE_LOG(LogDisplayClusterRenderSync, Error, TEXT("Couldn't enumerate physical GPU devices"));
 			}
 			else
 			{
-				UE_LOG(LogDisplayClusterRender, Log, TEXT("Found physical GPUs: %u"), PhysGPUAmount);
+				UE_LOG(LogDisplayClusterRenderSync, Log, TEXT("Found physical GPUs: %u"), PhysGPUAmount);
 			}
 
 			while (NvAPI_EnumNvidiaDisplayHandle(DisplayAmount, &DisplayHandles[DisplayAmount]) != NVAPI_END_ENUMERATION)
 			{
 				++DisplayAmount;
-				UE_LOG(LogDisplayClusterRender, Log, TEXT("Found displays: %u"), DisplayAmount);
+				UE_LOG(LogDisplayClusterRenderSync, Log, TEXT("Found display: %u"), DisplayAmount);
 			}
 		}
 	}
@@ -127,7 +143,7 @@ FDisplayClusterRenderSyncPolicySoftwareDX11::FDisplayClusterRenderSyncPolicySoft
 
 FDisplayClusterRenderSyncPolicySoftwareDX11::~FDisplayClusterRenderSyncPolicySoftwareDX11()
 {
-	if (bUseAdvancedSynchronization)
+	if (bNvApiInitialized)
 	{
 		NvAPI_Unload();
 	}
@@ -137,7 +153,9 @@ bool FDisplayClusterRenderSyncPolicySoftwareDX11::SynchronizeClusterRendering(in
 {
 	if(!(GEngine && GEngine->GameViewport && GEngine->GameViewport->Viewport))
 	{
-		return false;
+		UE_LOG(LogDisplayClusterRenderSync, Error, TEXT("Couldn't get a D3D11 viewport."));
+		// Tell a caller to present the frame by himself
+		return true;
 	}
 
 	FD3D11Viewport* const Viewport = static_cast<FD3D11Viewport*>(GEngine->GameViewport->Viewport->GetViewportRHI().GetReference());
@@ -163,11 +181,6 @@ bool FDisplayClusterRenderSyncPolicySoftwareDX11::SynchronizeClusterRendering(in
 	//////////////////////////////////////////////////////
 	SCOPE_CYCLE_COUNTER(STAT_nDisplayPresent_SoftSwapSync);
 
-	if (!GEngine || !GEngine->GameViewport || !GEngine->GameViewport->Viewport)
-	{
-		return false;
-	}
-
 	double B1B = 0.f;  // Barrier 1 before
 	double B1A = 0.f;  // Barrier 1 after
 	double TToB = 0.f; // Time to VBlank
@@ -178,13 +191,12 @@ bool FDisplayClusterRenderSyncPolicySoftwareDX11::SynchronizeClusterRendering(in
 	double PB = 0.f;   // Present before
 	double PA = 0.f;   // Present after
 
-	++FrameCounter;
-
-	IDXGISwapChain* const SwapChain = (IDXGISwapChain*)GEngine->GameViewport->Viewport->GetViewportRHI()->GetNativeSwapChain();
+	IDXGISwapChain* const SwapChain = static_cast<IDXGISwapChain*>(Viewport->GetNativeSwapChain());
 	check(SwapChain);
 	if (!SwapChain)
 	{
-		UE_LOG(LogDisplayClusterRender, Error, TEXT("Couldn't get a swap chain. Native present will be performed."));
+		UE_LOG(LogDisplayClusterRenderSync, Error, TEXT("Couldn't get a swap chain."));
+		// Tell a caller that he still needs to present a frame
 		return true;
 	}
 
@@ -193,13 +205,16 @@ bool FDisplayClusterRenderSyncPolicySoftwareDX11::SynchronizeClusterRendering(in
 	check(DXOutput);
 	if (!DXOutput)
 	{
-		UE_LOG(LogDisplayClusterRender, Error, TEXT("Couldn't get a DX output device. Native present will be performed."));
+		UE_LOG(LogDisplayClusterRenderSync, Error, TEXT("Couldn't get a DX output device."));
+		// Tell a caller that he still needs to present a frame
 		return true;
 	}
 
 	// Get dynamic values from console variables
 	const float VBlankThreshold = CVarVBlankThreshold.GetValueOnRenderThread();
 	const float VBlankThresholdSleepMultiplier = CVarVBlankThresholdSleepMultiplier.GetValueOnRenderThread();
+	const int32 VBlankBasisUpdate = CVarVBlankBasisUpdate.GetValueOnRenderThread();
+	const float VBlankBasisUpdatePeriod = CVarVBlankBasisUpdatePeriod.GetValueOnRenderThread();
 
 	if (!bTimersInitialized)
 	{
@@ -208,18 +223,18 @@ bool FDisplayClusterRenderSyncPolicySoftwareDX11::SynchronizeClusterRendering(in
 		VBlankBasis = GetVBlankTimestamp(DXOutput);
 		RefreshPeriod = GetRefreshPeriod();
 
-		UE_LOG(LogDisplayClusterRender, Log, TEXT("##SYNC_LOG: Refresh period:     %lf"), RefreshPeriod);
-		UE_LOG(LogDisplayClusterRender, Log, TEXT("##SYNC_LOG: VBlank basis:       %lf"), VBlankBasis);
-		UE_LOG(LogDisplayClusterRender, Log, TEXT("##SYNC_LOG: VBlank threshold:   %lf"), VBlankThreshold);
-		UE_LOG(LogDisplayClusterRender, Log, TEXT("##SYNC_LOG: VBlank sleep mult:  %lf"), VBlankThresholdSleepMultiplier);
-		UE_LOG(LogDisplayClusterRender, Log, TEXT("##SYNC_LOG: VBlank sleep:       %lf"), VBlankThreshold * VBlankThresholdSleepMultiplier);
+		UE_LOG(LogDisplayClusterRenderSync, Log, TEXT("##SYNC_LOG: Refresh period:     %lf"), RefreshPeriod);
+		UE_LOG(LogDisplayClusterRenderSync, Log, TEXT("##SYNC_LOG: VBlank basis:       %lf"), VBlankBasis);
+		UE_LOG(LogDisplayClusterRenderSync, Log, TEXT("##SYNC_LOG: VBlank threshold:   %lf"), VBlankThreshold);
+		UE_LOG(LogDisplayClusterRenderSync, Log, TEXT("##SYNC_LOG: VBlank sleep mult:  %lf"), VBlankThresholdSleepMultiplier);
+		UE_LOG(LogDisplayClusterRenderSync, Log, TEXT("##SYNC_LOG: VBlank sleep:       %lf"), VBlankThreshold * VBlankThresholdSleepMultiplier);
 
 #ifdef UE_BUILD_DEBUG
 		const int SamplesCount = 25;
 		for (int i = 0; i < SamplesCount; ++i)
 		{
 			DXOutput->WaitForVBlank();
-			UE_LOG(LogDisplayClusterRender, Log, TEXT("##SYNC_LOG: VBlank Sample #%2d: %lf"), i, FPlatformTime::Seconds());
+			UE_LOG(LogDisplayClusterRenderSync, Log, TEXT("##SYNC_LOG: VBlank Sample #%2d: %lf"), i, FPlatformTime::Seconds());
 		}
 #endif
 
@@ -227,8 +242,9 @@ bool FDisplayClusterRenderSyncPolicySoftwareDX11::SynchronizeClusterRendering(in
 	}
 
 	{
-		// Align render threads with a barrier
 		QUICK_SCOPE_CYCLE_COUNTER(STAT_nDisplayPresent_SoftSync_WaitBarrier);
+
+		// Align render threads with a barrier
 		B1B = FPlatformTime::Seconds();
 		SyncBarrierRenderThread();
 		B1A = FPlatformTime::Seconds();
@@ -240,26 +256,29 @@ bool FDisplayClusterRenderSyncPolicySoftwareDX11::SynchronizeClusterRendering(in
 		// Here we calculate how much time left before the next VBlank
 		const double CurTime = FPlatformTime::Seconds();
 		const double DiffTime = CurTime - VBlankBasis;
-		const double TimeRem = ::fmodl(DiffTime, RefreshPeriod);
-		const double LeftToVBlankTime = RefreshPeriod - TimeRem;
+		const double TimeRemainder    = ::fmodl(DiffTime, RefreshPeriod);
+		const double TimeLeftToVBlank = RefreshPeriod - TimeRemainder;
 
-		TToB = TimeRem;
+		TToB = TimeLeftToVBlank;
 
 		// Skip upcoming VBlank if we're in red zone
 		SB = FPlatformTime::Seconds();
-		if (LeftToVBlankTime < VBlankThreshold)
+		if (TimeLeftToVBlank < VBlankThreshold)
 		{
 			const double SleepTime = VBlankThreshold * VBlankThresholdSleepMultiplier;
-			UE_LOG(LogDisplayClusterRender, Warning, TEXT("FDisplayClusterNativePresentHandler::Present_SoftSwapSync - Skipped VBlank. Sleeping %f"), SleepTime);
-			FPlatformProcess::Sleep(SleepTime);
+			UE_LOG(LogDisplayClusterRenderSync, Log, TEXT("DX11 advanced soft sync - Skipped VBlank. Sleeping %f"), SleepTime);
+			//! FPlatformProcess::Sleep(SleepTime);
+			VBlankBasis = GetVBlankTimestamp(DXOutput);
+			FPlatformProcess::Sleep(0.005);
 		}
 		SA = FPlatformTime::Seconds();
 	}
 
 	{
+		QUICK_SCOPE_CYCLE_COUNTER(STAT_nDisplayPresent_SoftSync_WaitBarrier_Second);
+
 		// Align render threads again. After that, all render threads we'll be either before VBlank or after. The threshold,
 		// sleep, this barrier and MaxFrameLatency==1 (for blocking Present calls) guarantee it.
-		QUICK_SCOPE_CYCLE_COUNTER(STAT_nDisplayPresent_SoftSync_WaitBarrier_Second);
 		B2B = FPlatformTime::Seconds();
 		SyncBarrierRenderThread();
 		B2A = FPlatformTime::Seconds();
@@ -268,10 +287,6 @@ bool FDisplayClusterRenderSyncPolicySoftwareDX11::SynchronizeClusterRendering(in
 	{
 		QUICK_SCOPE_CYCLE_COUNTER(STAT_nDisplayPresent_SoftSync_SwapChainPresent);
 
-		// Get dynamic values from console variables
-		const int32 VBlankBasisUpdate = CVarVBlankBasisUpdate.GetValueOnRenderThread();
-		const float VBlankBasisUpdatePeriod = CVarVBlankBasisUpdatePeriod.GetValueOnRenderThread();
-
 		// Regardless of where we are, it's safe to present a frame now. If we need to update the VBlank basis,
 		// we have to wait for a VBlank and store the time. We don't want to miss a frame presentation so we
 		// present it with swap interval 0 right after VBlank signal.
@@ -287,7 +302,7 @@ bool FDisplayClusterRenderSyncPolicySoftwareDX11::SynchronizeClusterRendering(in
 		PA = FPlatformTime::Seconds();
 	}
 
-	UE_LOG(LogDisplayClusterRender, Verbose, TEXT("##SYNC_LOG - %d:%lf:%lf:%lf:%lf:%lf:%lf:%lf:%lf:%lf"), FrameCounter, B1B, B1A, TToB, SB, SA, B2B, B2A, PB, PA);
+	UE_LOG(LogDisplayClusterRenderSync, Verbose, TEXT("##SYNC_LOG - %d:%lf:%lf:%lf:%lf:%lf:%lf:%lf:%lf:%lf"), FrameCounter, B1B, B1A, TToB, SB, SA, B2B, B2A, PB, PA);
 
 	const bool LogDwmStats = (CVarLogDwmStats.GetValueOnRenderThread() != 0);
 	if (LogDwmStats)
@@ -295,7 +310,9 @@ bool FDisplayClusterRenderSyncPolicySoftwareDX11::SynchronizeClusterRendering(in
 		PrintDwmStats(FrameCounter);
 	}
 
-	// No need to perform native present
+	++FrameCounter;
+
+	// Tell a caller there is no need to present the frame
 	return false;
 }
 
@@ -312,12 +329,24 @@ double FDisplayClusterRenderSyncPolicySoftwareDX11::GetVBlankTimestamp(IDXGIOutp
 
 double FDisplayClusterRenderSyncPolicySoftwareDX11::GetRefreshPeriod() const
 {
-	// Obtain frame interval from the DWM
-	DWM_TIMING_INFO TimingInfo;
-	FMemory::Memzero(TimingInfo);
-	TimingInfo.cbSize = sizeof(DWM_TIMING_INFO);
-	DwmGetCompositionTimingInfo(nullptr, &TimingInfo);
-	return FPlatformTime::ToSeconds(TimingInfo.qpcRefreshPeriod);
+	const bool  bUseCustomRefreshRate = !!CVarUseCustomRefreshRate.GetValueOnRenderThread();
+
+	// Sometimes the DWM returns a refresh rate value that doesn't correspond to the real system.
+	// Use custom refresh rate for the synchronization algorithm below if required.
+	if (bUseCustomRefreshRate)
+	{
+		const float CustomRefreshRate = FMath::Abs(CVarCustomRefreshRate.GetValueOnRenderThread());
+		return 1.L / CustomRefreshRate;
+	}
+	else
+	{
+		// Obtain frame interval from the DWM
+		DWM_TIMING_INFO TimingInfo;
+		FMemory::Memzero(TimingInfo);
+		TimingInfo.cbSize = sizeof(DWM_TIMING_INFO);
+		DwmGetCompositionTimingInfo(nullptr, &TimingInfo);
+		return FPlatformTime::ToSeconds(TimingInfo.qpcRefreshPeriod);
+	}
 }
 
 void FDisplayClusterRenderSyncPolicySoftwareDX11::PrintDwmStats(uint32 FrameNum)
@@ -327,27 +356,27 @@ void FDisplayClusterRenderSyncPolicySoftwareDX11::PrintDwmStats(uint32 FrameNum)
 	TimingInfo.cbSize = sizeof(DWM_TIMING_INFO);
 	DwmGetCompositionTimingInfo(nullptr, &TimingInfo);
 
-	UE_LOG(LogDisplayClusterRender, Log, TEXT(">>>>>>>>>>>>>>>>>>>>>>> DWM"));
-	UE_LOG(LogDisplayClusterRender, Log, TEXT("##DWM_STAT(%d) cRefresh:               %llu"), FrameNum, TimingInfo.cRefresh);
-	UE_LOG(LogDisplayClusterRender, Log, TEXT("##DWM_STAT(%d) cDXRefresh:             %u"),   FrameNum, TimingInfo.cDXRefresh);
-	UE_LOG(LogDisplayClusterRender, Log, TEXT("##DWM_STAT(%d) qpcRefreshPeriod:       %lf"),  FrameNum, FPlatformTime::ToSeconds(TimingInfo.qpcRefreshPeriod));
-	UE_LOG(LogDisplayClusterRender, Log, TEXT("##DWM_STAT(%d) qpcVBlank:              %lf"),  FrameNum, FPlatformTime::ToSeconds(TimingInfo.qpcVBlank));
-	UE_LOG(LogDisplayClusterRender, Log, TEXT("##DWM_STAT(%d) cFrame:                 %llu"), FrameNum, TimingInfo.cFrame);
-	UE_LOG(LogDisplayClusterRender, Log, TEXT("##DWM_STAT(%d) cDXPresent:             %u"),   FrameNum, TimingInfo.cDXPresent);
-	UE_LOG(LogDisplayClusterRender, Log, TEXT("##DWM_STAT(%d) cRefreshFrame:          %llu"), FrameNum, TimingInfo.cRefreshFrame);
-	UE_LOG(LogDisplayClusterRender, Log, TEXT("##DWM_STAT(%d) cDXRefreshConfirmed:    %llu"), FrameNum, TimingInfo.cDXRefreshConfirmed);
-	UE_LOG(LogDisplayClusterRender, Log, TEXT("##DWM_STAT(%d) cFramesLate:            %llu"), FrameNum, TimingInfo.cFramesLate);
-	UE_LOG(LogDisplayClusterRender, Log, TEXT("##DWM_STAT(%d) cFramesOutstanding:     %u"),   FrameNum, TimingInfo.cFramesOutstanding);
-	UE_LOG(LogDisplayClusterRender, Log, TEXT("##DWM_STAT(%d) cFrameDisplayed:        %llu"), FrameNum, TimingInfo.cFrameDisplayed);
-	UE_LOG(LogDisplayClusterRender, Log, TEXT("##DWM_STAT(%d) cRefreshFrameDisplayed: %llu"), FrameNum, TimingInfo.cRefreshFrameDisplayed);
-	UE_LOG(LogDisplayClusterRender, Log, TEXT("##DWM_STAT(%d) cFrameComplete:         %llu"), FrameNum, TimingInfo.cFrameComplete);
-	UE_LOG(LogDisplayClusterRender, Log, TEXT("##DWM_STAT(%d) qpcFrameComplete:       %lf"),  FrameNum, FPlatformTime::ToSeconds(TimingInfo.qpcFrameComplete));
-	UE_LOG(LogDisplayClusterRender, Log, TEXT("##DWM_STAT(%d) cFramePending:          %llu"), FrameNum, TimingInfo.cFramePending);
-	UE_LOG(LogDisplayClusterRender, Log, TEXT("##DWM_STAT(%d) qpcFramePending:        %lf"),  FrameNum, FPlatformTime::ToSeconds(TimingInfo.qpcFramePending));
-	UE_LOG(LogDisplayClusterRender, Log, TEXT("##DWM_STAT(%d) cFramesDisplayed:       %llu"), FrameNum, TimingInfo.cFramesDisplayed);
-	UE_LOG(LogDisplayClusterRender, Log, TEXT("##DWM_STAT(%d) cFramesComplete:        %llu"), FrameNum, TimingInfo.cFramesComplete);
-	UE_LOG(LogDisplayClusterRender, Log, TEXT("##DWM_STAT(%d) cFramesPending:         %llu"), FrameNum, TimingInfo.cFramesPending);
-	UE_LOG(LogDisplayClusterRender, Log, TEXT("##DWM_STAT(%d) cFramesDropped:         %llu"), FrameNum, TimingInfo.cFramesDropped);
-	UE_LOG(LogDisplayClusterRender, Log, TEXT("##DWM_STAT(%d) cFramesMissed:          %llu"), FrameNum, TimingInfo.cFramesMissed);
-	UE_LOG(LogDisplayClusterRender, Log, TEXT(">>>>>>"));
+	UE_LOG(LogDisplayClusterRenderSync, Log, TEXT(">>>>>>>>>>>>>>>>>>>>>>> DWM"));
+	UE_LOG(LogDisplayClusterRenderSync, Log, TEXT("##DWM_STAT(%d) cRefresh:               %llu"), FrameNum, TimingInfo.cRefresh);
+	UE_LOG(LogDisplayClusterRenderSync, Log, TEXT("##DWM_STAT(%d) cDXRefresh:             %u"),   FrameNum, TimingInfo.cDXRefresh);
+	UE_LOG(LogDisplayClusterRenderSync, Log, TEXT("##DWM_STAT(%d) qpcRefreshPeriod:       %lf"),  FrameNum, FPlatformTime::ToSeconds(TimingInfo.qpcRefreshPeriod));
+	UE_LOG(LogDisplayClusterRenderSync, Log, TEXT("##DWM_STAT(%d) qpcVBlank:              %lf"),  FrameNum, FPlatformTime::ToSeconds(TimingInfo.qpcVBlank));
+	UE_LOG(LogDisplayClusterRenderSync, Log, TEXT("##DWM_STAT(%d) cFrame:                 %llu"), FrameNum, TimingInfo.cFrame);
+	UE_LOG(LogDisplayClusterRenderSync, Log, TEXT("##DWM_STAT(%d) cDXPresent:             %u"),   FrameNum, TimingInfo.cDXPresent);
+	UE_LOG(LogDisplayClusterRenderSync, Log, TEXT("##DWM_STAT(%d) cRefreshFrame:          %llu"), FrameNum, TimingInfo.cRefreshFrame);
+	UE_LOG(LogDisplayClusterRenderSync, Log, TEXT("##DWM_STAT(%d) cDXRefreshConfirmed:    %llu"), FrameNum, TimingInfo.cDXRefreshConfirmed);
+	UE_LOG(LogDisplayClusterRenderSync, Log, TEXT("##DWM_STAT(%d) cFramesLate:            %llu"), FrameNum, TimingInfo.cFramesLate);
+	UE_LOG(LogDisplayClusterRenderSync, Log, TEXT("##DWM_STAT(%d) cFramesOutstanding:     %u"),   FrameNum, TimingInfo.cFramesOutstanding);
+	UE_LOG(LogDisplayClusterRenderSync, Log, TEXT("##DWM_STAT(%d) cFrameDisplayed:        %llu"), FrameNum, TimingInfo.cFrameDisplayed);
+	UE_LOG(LogDisplayClusterRenderSync, Log, TEXT("##DWM_STAT(%d) cRefreshFrameDisplayed: %llu"), FrameNum, TimingInfo.cRefreshFrameDisplayed);
+	UE_LOG(LogDisplayClusterRenderSync, Log, TEXT("##DWM_STAT(%d) cFrameComplete:         %llu"), FrameNum, TimingInfo.cFrameComplete);
+	UE_LOG(LogDisplayClusterRenderSync, Log, TEXT("##DWM_STAT(%d) qpcFrameComplete:       %lf"),  FrameNum, FPlatformTime::ToSeconds(TimingInfo.qpcFrameComplete));
+	UE_LOG(LogDisplayClusterRenderSync, Log, TEXT("##DWM_STAT(%d) cFramePending:          %llu"), FrameNum, TimingInfo.cFramePending);
+	UE_LOG(LogDisplayClusterRenderSync, Log, TEXT("##DWM_STAT(%d) qpcFramePending:        %lf"),  FrameNum, FPlatformTime::ToSeconds(TimingInfo.qpcFramePending));
+	UE_LOG(LogDisplayClusterRenderSync, Log, TEXT("##DWM_STAT(%d) cFramesDisplayed:       %llu"), FrameNum, TimingInfo.cFramesDisplayed);
+	UE_LOG(LogDisplayClusterRenderSync, Log, TEXT("##DWM_STAT(%d) cFramesComplete:        %llu"), FrameNum, TimingInfo.cFramesComplete);
+	UE_LOG(LogDisplayClusterRenderSync, Log, TEXT("##DWM_STAT(%d) cFramesPending:         %llu"), FrameNum, TimingInfo.cFramesPending);
+	UE_LOG(LogDisplayClusterRenderSync, Log, TEXT("##DWM_STAT(%d) cFramesDropped:         %llu"), FrameNum, TimingInfo.cFramesDropped);
+	UE_LOG(LogDisplayClusterRenderSync, Log, TEXT("##DWM_STAT(%d) cFramesMissed:          %llu"), FrameNum, TimingInfo.cFramesMissed);
+	UE_LOG(LogDisplayClusterRenderSync, Log, TEXT(">>>>>>"));
 }
diff --git a/Engine/Plugins/Runtime/nDisplay/Source/DisplayCluster/Private/Render/Synchronization/DisplayClusterRenderSyncPolicySoftwareDX11.h b/Engine/Plugins/Runtime/nDisplay/Source/DisplayCluster/Private/Render/Synchronization/DisplayClusterRenderSyncPolicySoftwareDX11.h
--- a/Engine/Plugins/Runtime/nDisplay/Source/DisplayCluster/Private/Render/Synchronization/DisplayClusterRenderSyncPolicySoftwareDX11.h
+++ b/Engine/Plugins/Runtime/nDisplay/Source/DisplayCluster/Private/Render/Synchronization/DisplayClusterRenderSyncPolicySoftwareDX11.h
@@ -31,13 +31,10 @@ private:
 private:
 	bool bTimersInitialized = false;
 	bool bNvApiInitialized = false;
-	bool bNvBarriersInitialized = false;
+	bool bUseAdvancedSynchronization = false;
 
 	double VBlankBasis = 0;
 	double RefreshPeriod = 0;
 
 	uint32 FrameCounter = 0;
-
-private:
-	bool bUseAdvancedSynchronization = false;
 };
diff --git a/Engine/Plugins/Runtime/nDisplay/Source/DisplayCluster/Public/Blueprints/IDisplayClusterBlueprintAPI.h b/Engine/Plugins/Runtime/nDisplay/Source/DisplayCluster/Public/Blueprints/IDisplayClusterBlueprintAPI.h
--- a/Engine/Plugins/Runtime/nDisplay/Source/DisplayCluster/Public/Blueprints/IDisplayClusterBlueprintAPI.h
+++ b/Engine/Plugins/Runtime/nDisplay/Source/DisplayCluster/Public/Blueprints/IDisplayClusterBlueprintAPI.h
@@ -14,6 +14,7 @@
 
 #include "IDisplayClusterBlueprintAPI.generated.h"
 
+struct FPostProcessSettings;
 
 
 
@@ -87,8 +88,8 @@ public:
 	//////////////////////////////////////////////////////////////////////////////////////////////
 	// Config API
 	//////////////////////////////////////////////////////////////////////////////////////////////
-	UFUNCTION(BlueprintCallable, meta = (DisplayName = "Get viewports"), Category = "DisplayCluster|Config")
-	virtual void GetViewports(bool IsRTT, TArray<FString>& ViewportIDs, TArray<FString>& ViewportTypes, TArray<FIntPoint>& ViewportLocations, TArray<FIntPoint>& ViewportSizes) = 0;
+	UFUNCTION(BlueprintCallable, meta = (DisplayName = "Get local viewports"), Category = "DisplayCluster|Config")
+	virtual void GetLocalViewports(bool IsRTT, TArray<FString>& ViewportIDs, TArray<FString>& ViewportTypes, TArray<FIntPoint>& ViewportLocations, TArray<FIntPoint>& ViewportSizes) = 0;
 
 
 public:
@@ -243,8 +244,12 @@ public:
 	UFUNCTION(BlueprintCallable, meta = (DisplayName = "Set viewport camera"), Category = "DisplayCluster|Render")
 	virtual void SetViewportCamera(const FString& InCameraId, const FString& InViewportId) = 0;
 
-	UFUNCTION(BlueprintCallable, meta = (DisplayName = "Set custom post processing settings for viewport"), Category = "DisplayCluster|Render")
-	virtual void SetCustomPostProcessing(const FString& ViewportID, const FPostProcessSettings& PostProcessingSettings) = 0;
+	UFUNCTION(BlueprintCallable, meta = (DisplayName = "Set start post processing settings for viewport"), Category = "DisplayCluster|Render")
+	virtual void SetStartPostProcessingSettings(const FString& ViewportID, const FPostProcessSettings& StartPostProcessingSettings) = 0;
+	UFUNCTION(BlueprintCallable, meta = (DisplayName = "Set override post processing settings for viewport"), Category = "DisplayCluster|Render")
+	virtual void SetOverridePostProcessingSettings(const FString& ViewportID, const FPostProcessSettings& OverridePostProcessingSettings, float BlendWeight = 1.0f) = 0;
+	UFUNCTION(BlueprintCallable, meta = (DisplayName = "Set final post processing settings for viewport"), Category = "DisplayCluster|Render")
+	virtual void SetFinalPostProcessingSettings(const FString& ViewportID, const FPostProcessSettings& FinalPostProcessingSettings) = 0;
 
 	UFUNCTION(BlueprintCallable, meta = (DisplayName = "Get viewport rectangle"), Category = "DisplayCluster|Render")
 	virtual bool GetViewportRect(const FString& ViewportID, FIntPoint& ViewportLoc, FIntPoint& ViewportSize) = 0;
diff --git a/Engine/Plugins/Runtime/nDisplay/Source/DisplayCluster/Public/DisplayClusterViewportClient.h b/Engine/Plugins/Runtime/nDisplay/Source/DisplayCluster/Public/DisplayClusterViewportClient.h
new file mode 100644
--- /dev/null
+++ b/Engine/Plugins/Runtime/nDisplay/Source/DisplayCluster/Public/DisplayClusterViewportClient.h
@@ -0,0 +1,21 @@
+// Copyright 1998-2019 Epic Games, Inc. All Rights Reserved.
+
+#pragma once
+
+#include "Engine/GameViewportClient.h"
+
+#include "DisplayClusterViewportClient.generated.h"
+
+UCLASS()
+class DISPLAYCLUSTER_API UDisplayClusterViewportClient
+	: public UGameViewportClient
+{
+	GENERATED_BODY()
+
+public:
+	UDisplayClusterViewportClient(FVTableHelper& Helper);
+	virtual ~UDisplayClusterViewportClient();
+
+	virtual void Init(struct FWorldContext& WorldContext, UGameInstance* OwningGameInstance, bool bCreateNewAudioDevice = true) override;
+	virtual void Draw(FViewport* Viewport, FCanvas* SceneCanvas) override;
+};
diff --git a/Engine/Plugins/Runtime/nDisplay/Source/DisplayCluster/Public/Render/Device/IDisplayClusterRenderDevice.h b/Engine/Plugins/Runtime/nDisplay/Source/DisplayCluster/Public/Render/Device/IDisplayClusterRenderDevice.h
--- a/Engine/Plugins/Runtime/nDisplay/Source/DisplayCluster/Public/Render/Device/IDisplayClusterRenderDevice.h
+++ b/Engine/Plugins/Runtime/nDisplay/Source/DisplayCluster/Public/Render/Device/IDisplayClusterRenderDevice.h
@@ -44,7 +44,9 @@ public:
 	*/
 	virtual void SetViewportCamera(const FString& InCameraId = FString(), const FString& InViewportId = FString()) = 0;
 
-	virtual void SetCustomPostProcessing(const FString& ViewportID, const FPostProcessSettings& PostProcessingSettings) = 0;
+	virtual void SetStartPostProcessingSettings(const FString& ViewportID, const FPostProcessSettings& StartPostProcessingSettings) = 0;
+	virtual void SetOverridePostProcessingSettings(const FString& ViewportID, const FPostProcessSettings& OverridePostProcessingSettings, float BlendWeight = 1.0f) = 0;
+	virtual void SetFinalPostProcessingSettings(const FString& ViewportID, const FPostProcessSettings& FinalPostProcessingSettings) = 0;
 
 	virtual bool GetViewportRect(const FString& InViewportID, FIntRect& Rect) = 0;
 };
diff --git a/Engine/Plugins/Runtime/nDisplay/Source/DisplayCluster/Public/Render/IDisplayClusterRenderManager.h b/Engine/Plugins/Runtime/nDisplay/Source/DisplayCluster/Public/Render/IDisplayClusterRenderManager.h
--- a/Engine/Plugins/Runtime/nDisplay/Source/DisplayCluster/Public/Render/IDisplayClusterRenderManager.h
+++ b/Engine/Plugins/Runtime/nDisplay/Source/DisplayCluster/Public/Render/IDisplayClusterRenderManager.h
@@ -145,7 +145,9 @@ public:
 	*/
 	virtual TMap<FString, FDisplayClusterPPInfo> GetRegisteredPostprocessOperations() const = 0;
 
-	virtual void SetCustomPostProcessing(const FString& ViewportID, const FPostProcessSettings& PostProcessingSettings) = 0;
+	virtual void SetStartPostProcessingSettings(const FString& ViewportID, const FPostProcessSettings& StartPostProcessingSettings) = 0;
+	virtual void SetOverridePostProcessingSettings(const FString& ViewportID, const FPostProcessSettings& OverridePostProcessingSettings, float BlendWeight = 1.0f) = 0;
+	virtual void SetFinalPostProcessingSettings(const FString& ViewportID, const FPostProcessSettings& FinalPostProcessingSettings) = 0;
 
 	/**
 	* Assigns camera to a specified viewport. If InViewportId is empty, all viewports will be assigned to a new camera. Empty camera ID means default active camera.
diff --git a/Engine/Plugins/Runtime/nDisplay/Source/DisplayClusterProjection/Private/Blueprints/DisplayClusterProjectionBlueprintAPIImpl.cpp b/Engine/Plugins/Runtime/nDisplay/Source/DisplayClusterProjection/Private/Blueprints/DisplayClusterProjectionBlueprintAPIImpl.cpp
--- a/Engine/Plugins/Runtime/nDisplay/Source/DisplayClusterProjection/Private/Blueprints/DisplayClusterProjectionBlueprintAPIImpl.cpp
+++ b/Engine/Plugins/Runtime/nDisplay/Source/DisplayClusterProjection/Private/Blueprints/DisplayClusterProjectionBlueprintAPIImpl.cpp
@@ -16,8 +16,11 @@
 //////////////////////////////////////////////////////////////////////////////////////////////
 // Policy: CAMERA
 //////////////////////////////////////////////////////////////////////////////////////////////
-void UDisplayClusterProjectionBlueprintAPIImpl::CameraPolicySetCamera(const FString& ViewportId, UCameraComponent* NewCamera)
+void UDisplayClusterProjectionBlueprintAPIImpl::CameraPolicySetCamera(const FString& ViewportId, UCameraComponent* NewCamera, float FOVMultiplier)
 {
+	check(NewCamera);
+	check(FOVMultiplier >= 0.1f);
+
 	IDisplayClusterProjection& Module = IDisplayClusterProjection::Get();
 	
 	TSharedPtr<IDisplayClusterProjectionPolicyFactory> Factory = Module.GetProjectionFactory(DisplayClusterStrings::projection::Camera);
@@ -32,7 +35,7 @@ void UDisplayClusterProjectionBlueprintAPIImpl::CameraPolicySetCamera(const FStr
 				TSharedPtr<FDisplayClusterProjectionCameraPolicy> CameraPolicyInstance = StaticCastSharedPtr<FDisplayClusterProjectionCameraPolicy>(PolicyInstance);
 				if (CameraPolicyInstance)
 				{
-					CameraPolicyInstance->SetCamera(NewCamera);
+					CameraPolicyInstance->SetCamera(NewCamera, FOVMultiplier);
 				}
 			}
 		}
diff --git a/Engine/Plugins/Runtime/nDisplay/Source/DisplayClusterProjection/Private/Blueprints/DisplayClusterProjectionBlueprintAPIImpl.h b/Engine/Plugins/Runtime/nDisplay/Source/DisplayClusterProjection/Private/Blueprints/DisplayClusterProjectionBlueprintAPIImpl.h
--- a/Engine/Plugins/Runtime/nDisplay/Source/DisplayClusterProjection/Private/Blueprints/DisplayClusterProjectionBlueprintAPIImpl.h
+++ b/Engine/Plugins/Runtime/nDisplay/Source/DisplayClusterProjection/Private/Blueprints/DisplayClusterProjectionBlueprintAPIImpl.h
@@ -23,5 +23,5 @@ public:
 	//////////////////////////////////////////////////////////////////////////////////////////////
 	/** Sets active camera component for camera policy */
 	UFUNCTION(BlueprintCallable, meta = (DisplayName = "Set camera"), Category = "DisplayClusterProjection|Camera")
-	virtual void CameraPolicySetCamera(const FString& ViewportId, UCameraComponent* NewCamera) override;
+	virtual void CameraPolicySetCamera(const FString& ViewportId, UCameraComponent* NewCamera, float FOVMultiplier = 1.f) override;
 };
diff --git a/Engine/Plugins/Runtime/nDisplay/Source/DisplayClusterProjection/Private/Policy/Camera/DisplayClusterProjectionCameraPolicy.cpp b/Engine/Plugins/Runtime/nDisplay/Source/DisplayClusterProjection/Private/Policy/Camera/DisplayClusterProjectionCameraPolicy.cpp
--- a/Engine/Plugins/Runtime/nDisplay/Source/DisplayClusterProjection/Private/Policy/Camera/DisplayClusterProjectionCameraPolicy.cpp
+++ b/Engine/Plugins/Runtime/nDisplay/Source/DisplayClusterProjection/Private/Policy/Camera/DisplayClusterProjectionCameraPolicy.cpp
@@ -87,7 +87,7 @@ bool FDisplayClusterProjectionCameraPolicy::GetProjectionMatrix(const uint32 Vie
 	}
 
 	FComposurePostMoveSettings ComposureSettings;
-	OutPrjMatrix = ComposureSettings.GetProjectionMatrix(AssignedCamera->FieldOfView, AssignedCamera->AspectRatio);
+	OutPrjMatrix = ComposureSettings.GetProjectionMatrix(AssignedCamera->FieldOfView * CurrentFovMultiplier, AssignedCamera->AspectRatio);
 
 	return true;
 }
@@ -96,8 +96,11 @@ bool FDisplayClusterProjectionCameraPolicy::GetProjectionMatrix(const uint32 Vie
 //////////////////////////////////////////////////////////////////////////////////////////////
 // FDisplayClusterProjectionCameraPolicy
 //////////////////////////////////////////////////////////////////////////////////////////////
-void FDisplayClusterProjectionCameraPolicy::SetCamera(UCameraComponent* NewCamera)
+void FDisplayClusterProjectionCameraPolicy::SetCamera(UCameraComponent* NewCamera, float FOVMultiplier)
 {
+	check(NewCamera);
+	check(FOVMultiplier >= 0.1f);
+
 	if (NewCamera)
 	{
 		UE_LOG(LogDisplayClusterProjectionCamera, Verbose, TEXT("New camera set: %s"), *NewCamera->GetFullName());
@@ -107,4 +110,14 @@ void FDisplayClusterProjectionCameraPolicy::SetCamera(UCameraComponent* NewCamer
 	{
 		UE_LOG(LogDisplayClusterProjectionCamera, Warning, TEXT("Trying to set nullptr camera pointer"));
 	}
+
+	if (FOVMultiplier >= 0.1f)
+	{
+		UE_LOG(LogDisplayClusterProjectionCamera, Verbose, TEXT("New FOV multiplier set: %f"), FOVMultiplier);
+		CurrentFovMultiplier = FOVMultiplier;
+	}
+	else
+	{
+		UE_LOG(LogDisplayClusterProjectionCamera, Warning, TEXT("FOV multiplier is too small"));
+	}
 }
diff --git a/Engine/Plugins/Runtime/nDisplay/Source/DisplayClusterProjection/Private/Policy/Camera/DisplayClusterProjectionCameraPolicy.h b/Engine/Plugins/Runtime/nDisplay/Source/DisplayClusterProjection/Private/Policy/Camera/DisplayClusterProjectionCameraPolicy.h
--- a/Engine/Plugins/Runtime/nDisplay/Source/DisplayClusterProjection/Private/Policy/Camera/DisplayClusterProjectionCameraPolicy.h
+++ b/Engine/Plugins/Runtime/nDisplay/Source/DisplayClusterProjection/Private/Policy/Camera/DisplayClusterProjectionCameraPolicy.h
@@ -35,9 +35,11 @@ public:
 	{ return false; }
 
 public:
-	void SetCamera(UCameraComponent* NewCamera);
+	void SetCamera(UCameraComponent* NewCamera, float FOVMultiplier);
 
 private:
 	// Camera to use for rendering
 	UCameraComponent* AssignedCamera = nullptr;
+	// FOV multiplier
+	float CurrentFovMultiplier = 1.f;
 };
diff --git a/Engine/Plugins/Runtime/nDisplay/Source/DisplayClusterProjection/Public/Blueprints/IDisplayClusterProjectionBlueprintAPI.h b/Engine/Plugins/Runtime/nDisplay/Source/DisplayClusterProjection/Public/Blueprints/IDisplayClusterProjectionBlueprintAPI.h
--- a/Engine/Plugins/Runtime/nDisplay/Source/DisplayClusterProjection/Public/Blueprints/IDisplayClusterProjectionBlueprintAPI.h
+++ b/Engine/Plugins/Runtime/nDisplay/Source/DisplayClusterProjection/Public/Blueprints/IDisplayClusterProjectionBlueprintAPI.h
@@ -30,5 +30,5 @@ public:
 	//////////////////////////////////////////////////////////////////////////////////////////////
 	/** Sets active camera component for camera policy */
 	UFUNCTION(BlueprintCallable, meta = (DisplayName = "Set camera"), Category = "DisplayClusterProjection|Camera")
-	virtual void CameraPolicySetCamera(const FString& ViewportId, UCameraComponent* NewCamera) = 0;
+	virtual void CameraPolicySetCamera(const FString& ViewportId, UCameraComponent* NewCamera, float FOVMultiplier = 1.f) = 0;
 };
diff --git a/Engine/Plugins/Runtime/nDisplay/Source/PicpProjection/Private/Blueprints/PicpProjectionBlueprintAPIImpl.cpp b/Engine/Plugins/Runtime/nDisplay/Source/PicpProjection/Private/Blueprints/PicpProjectionBlueprintAPIImpl.cpp
--- a/Engine/Plugins/Runtime/nDisplay/Source/PicpProjection/Private/Blueprints/PicpProjectionBlueprintAPIImpl.cpp
+++ b/Engine/Plugins/Runtime/nDisplay/Source/PicpProjection/Private/Blueprints/PicpProjectionBlueprintAPIImpl.cpp
@@ -3,13 +3,14 @@
 #include "Blueprints/PicpProjectionBlueprintAPIImpl.h"
 
 #include "Engine/TextureRenderTarget2D.h"
-
+#include "ComposurePostMoves.h"
 #include "UObject/Package.h"
-#include "IPicpProjection.h"
 
-#include "Overlay/PicpProjectionOverlayRender.h"
+#include "IPicpProjection.h"
 #include "IPicpMPCDI.h"
-#include "ComposurePostMoves.h"
+#include "PicpProjectionStrings.h"
+#include "Overlay/PicpProjectionOverlayRender.h"
+
 
 
 int UPicpProjectionAPIImpl::GetViewportCount()
@@ -18,12 +19,23 @@ int UPicpProjectionAPIImpl::GetViewportCount()
 	return PicpModule.GetPolicyCount("picp_mpcdi");
 }
 
-void UPicpProjectionAPIImpl::SetProjectionDataListener(TScriptInterface<IPicpProjectionFrustumDataListener> Listener)
+void UPicpProjectionAPIImpl::AddProjectionDataListener(TScriptInterface<IPicpProjectionFrustumDataListener> Listener)
 {
 	IPicpProjection& PicpModule = IPicpProjection::Get();
 	PicpModule.AddProjectionDataListener(Listener);
 } 
 
+void UPicpProjectionAPIImpl::RemoveProjectionDataListener(TScriptInterface<IPicpProjectionFrustumDataListener> Listener)
+{
+	IPicpProjection& PicpModule = IPicpProjection::Get();
+	PicpModule.RemoveProjectionDataListener(Listener);
+}
+
+void UPicpProjectionAPIImpl::CleanProjectionDataListener(TScriptInterface<IPicpProjectionFrustumDataListener> Listener)
+{
+	IPicpProjection& PicpModule = IPicpProjection::Get();
+	PicpModule.CleanProjectionDataListeners();
+}
 
 void UPicpProjectionAPIImpl::BlendCameraFrameCaptures(UTexture* SrcFrame, UTextureRenderTarget2D* DstFrame, UTextureRenderTarget2D* Result)
 {	
@@ -35,9 +47,6 @@ void UPicpProjectionAPIImpl::ApplyBlurPostProcess(UTextureRenderTarget2D* InOutR
 	IPicpMPCDI::Get().ApplyBlur(InOutRenderTarget, TemporaryRenderTarget, KernelRadius, KernelScale, BlurType);
 }
 
-//!fixme send policy name from script
-#define TEMP_POLICY_NAME TEXT("Picp_mpcdi")
-
 void UPicpProjectionAPIImpl::SetupOverlayCaptures(const TArray<struct FPicpOverlayFrameBlendingParameters>& captures)
 {
 	FPicpProjectionOverlayFrameData overlayFrameData;
@@ -47,7 +56,7 @@ void UPicpProjectionAPIImpl::SetupOverlayCaptures(const TArray<struct FPicpOverl
 		FRotator CameraRotation = capture.CineCamera->K2_GetComponentRotation();
 		FVector  CameraLocation = capture.CineCamera->K2_GetComponentLocation();
 
-		float fov = capture.CineCamera->FieldOfView;
+		float fov = capture.CineCamera->FieldOfView * capture.FieldOfViewMultiplier;
 
 		FComposurePostMoveSettings settings;
 		float aspectRatio = capture.CineCamera->AspectRatio;
@@ -85,14 +94,14 @@ void UPicpProjectionAPIImpl::SetupOverlayCaptures(const TArray<struct FPicpOverl
 	}
 
 	IPicpProjection& PicpModule = IPicpProjection::Get();
-	PicpModule.SetOverlayFrameData(TEMP_POLICY_NAME, overlayFrameData);
+	PicpModule.SetOverlayFrameData(PicpProjectionStrings::projection::PicpMPCDI, overlayFrameData);
 }
 
-
 void UPicpProjectionAPIImpl::SetWarpTextureCaptureState(UTextureRenderTarget2D* dstTexture, const FString& ViewportId, const int ViewIdx, bool bCaptureNow)
 {
 	IPicpProjection::Get().CaptureWarpTexture(dstTexture, ViewportId, ViewIdx, bCaptureNow);
 }
+
 void UPicpProjectionAPIImpl::EnableTextureRenderTargetMips(UTextureRenderTarget2D* Texture)
 {
 	Texture->bAutoGenerateMips = true;
diff --git a/Engine/Plugins/Runtime/nDisplay/Source/PicpProjection/Private/Blueprints/PicpProjectionBlueprintAPIImpl.h b/Engine/Plugins/Runtime/nDisplay/Source/PicpProjection/Private/Blueprints/PicpProjectionBlueprintAPIImpl.h
--- a/Engine/Plugins/Runtime/nDisplay/Source/PicpProjection/Private/Blueprints/PicpProjectionBlueprintAPIImpl.h
+++ b/Engine/Plugins/Runtime/nDisplay/Source/PicpProjection/Private/Blueprints/PicpProjectionBlueprintAPIImpl.h
@@ -28,8 +28,14 @@ public:
 	UFUNCTION(BlueprintCallable, meta = (DisplayName = "Get Viewport Count"), Category = "PICP")
 	virtual int GetViewportCount() override;
 
-	UFUNCTION(BlueprintCallable, meta = (DisplayName = "Set Projection Data Listener"), Category = "PICP")
-	virtual void SetProjectionDataListener(TScriptInterface<IPicpProjectionFrustumDataListener> Listener) override;
+	UFUNCTION(BlueprintCallable, meta = (DisplayName = "Add Projection Data Listener"), Category = "PICP")
+	virtual void AddProjectionDataListener(TScriptInterface<IPicpProjectionFrustumDataListener> Listener) override;
+
+	UFUNCTION(BlueprintCallable, meta = (DisplayName = "Remove Projection Data Listener"), Category = "PICP")
+	virtual void RemoveProjectionDataListener(TScriptInterface<IPicpProjectionFrustumDataListener> Listener) override;
+
+	UFUNCTION(BlueprintCallable, meta = (DisplayName = "Clean Projection Data Listener"), Category = "PICP")
+	virtual void CleanProjectionDataListener(TScriptInterface<IPicpProjectionFrustumDataListener> Listener) override;
 
 	UFUNCTION(BlueprintCallable, meta = (DisplayName = "Blend Frames"), Category = "PICP")
 	virtual void BlendCameraFrameCaptures(UTexture* SrcFrame, UTextureRenderTarget2D* DstFrame, UTextureRenderTarget2D* Result) override;
diff --git a/Engine/Plugins/Runtime/nDisplay/Source/PicpProjection/Private/PicpProjectionModule.cpp b/Engine/Plugins/Runtime/nDisplay/Source/PicpProjection/Private/PicpProjectionModule.cpp
--- a/Engine/Plugins/Runtime/nDisplay/Source/PicpProjection/Private/PicpProjectionModule.cpp
+++ b/Engine/Plugins/Runtime/nDisplay/Source/PicpProjection/Private/PicpProjectionModule.cpp
@@ -219,7 +219,9 @@ bool FPicpProjectionModule::GetWarpFrustum(const FString& ViewportId, const uint
 
 void FPicpProjectionModule::SetViewport(const FString& Id, FRotator& OutViewRotation, FVector& OutViewLocation, FMatrix& OutPrjMatrix)
 {
-	for (auto listener : PicpEventListeners)
+	TArray<TScriptInterface<IPicpProjectionFrustumDataListener>> EmptyListeners;
+
+	for (auto Listener : PicpEventListeners)
 	{
 		FPicpProjectionFrustumData data;
 		data.Id = Id;
@@ -227,7 +229,19 @@ void FPicpProjectionModule::SetViewport(const FString& Id, FRotator& OutViewRota
 		data.ViewLocation = OutViewLocation;
 		data.PrjMatrix = OutPrjMatrix;
 
-		listener->Execute_OnProjectionDataUpdate(listener.GetObject(), data);
+		if (Listener.GetObject() != nullptr && !Listener.GetObject()->IsPendingKill())
+		{
+			Listener->Execute_OnProjectionDataUpdate(Listener.GetObject(), data);
+		}		
+		else
+		{
+			EmptyListeners.Add(Listener);
+		}
+	}
+
+	for (auto EmptyListener : EmptyListeners)
+	{
+		PicpEventListeners.Remove(EmptyListener);
 	}
 }
 
diff --git a/Engine/Plugins/Runtime/nDisplay/Source/PicpProjection/Private/PicpProjectionStrings.h b/Engine/Plugins/Runtime/nDisplay/Source/PicpProjection/Private/PicpProjectionStrings.h
--- a/Engine/Plugins/Runtime/nDisplay/Source/PicpProjection/Private/PicpProjectionStrings.h
+++ b/Engine/Plugins/Runtime/nDisplay/Source/PicpProjection/Private/PicpProjectionStrings.h
@@ -25,6 +25,6 @@ namespace PicpProjectionStrings
 
 	namespace projection
 	{
-		static constexpr auto PicpMPCDI     = TEXT("Picp_mpcdi");		
+		static constexpr auto PicpMPCDI = TEXT("picp_mpcdi");		
 	}
 };
diff --git a/Engine/Plugins/Runtime/nDisplay/Source/PicpProjection/Public/Blueprints/IPicpProjectionBlueprintAPI.h b/Engine/Plugins/Runtime/nDisplay/Source/PicpProjection/Public/Blueprints/IPicpProjectionBlueprintAPI.h
--- a/Engine/Plugins/Runtime/nDisplay/Source/PicpProjection/Public/Blueprints/IPicpProjectionBlueprintAPI.h
+++ b/Engine/Plugins/Runtime/nDisplay/Source/PicpProjection/Public/Blueprints/IPicpProjectionBlueprintAPI.h
@@ -34,8 +34,14 @@ public:
 	UFUNCTION(BlueprintCallable, meta = (DisplayName = "Get Viewport Count"), Category = "PICP")
 	virtual int GetViewportCount() = 0;
 
-	UFUNCTION(BlueprintCallable, meta = (DisplayName = "Set Projection Data Listener"), Category = "PICP")
-	virtual void SetProjectionDataListener(TScriptInterface<IPicpProjectionFrustumDataListener> Listener) = 0;
+	UFUNCTION(BlueprintCallable, meta = (DisplayName = "Add Projection Data Listener"), Category = "PICP")
+	virtual void AddProjectionDataListener(TScriptInterface<IPicpProjectionFrustumDataListener> Listener) = 0;
+
+	UFUNCTION(BlueprintCallable, meta = (DisplayName = "Remove Projection Data Listener"), Category = "PICP")
+	virtual void RemoveProjectionDataListener(TScriptInterface<IPicpProjectionFrustumDataListener> Listener) = 0;
+
+	UFUNCTION(BlueprintCallable, meta = (DisplayName = "Clean Projection Data Listener"), Category = "PICP")
+	virtual void CleanProjectionDataListener(TScriptInterface<IPicpProjectionFrustumDataListener> Listener) = 0;
 
 	UFUNCTION(BlueprintCallable, meta = (DisplayName = "Blend Frames"), Category = "PICP")
 	virtual void BlendCameraFrameCaptures(UTexture* SrcFrame, UTextureRenderTarget2D* DstFrame, UTextureRenderTarget2D* Result) = 0;
diff --git a/Engine/Plugins/Runtime/nDisplay/Source/PicpProjection/Public/Blueprints/ReprojectionBlendingData.h b/Engine/Plugins/Runtime/nDisplay/Source/PicpProjection/Public/Blueprints/ReprojectionBlendingData.h
--- a/Engine/Plugins/Runtime/nDisplay/Source/PicpProjection/Public/Blueprints/ReprojectionBlendingData.h
+++ b/Engine/Plugins/Runtime/nDisplay/Source/PicpProjection/Public/Blueprints/ReprojectionBlendingData.h
@@ -61,6 +61,9 @@ public:
 	UPROPERTY(EditAnywhere, BlueprintReadWrite, Category = "PICP")
 	UCineCameraComponent* CineCamera;
 
+	UPROPERTY(EditAnywhere, BlueprintReadWrite, Category = "PICP")
+	float FieldOfViewMultiplier = 1.0f;
+
 	UPROPERTY(EditAnywhere, BlueprintReadWrite, Category = "PICP")
 	FString RTTViewportId;
 
diff --git a/Engine/Source/Runtime/Engine/Classes/Engine/GameViewportClient.h b/Engine/Source/Runtime/Engine/Classes/Engine/GameViewportClient.h
--- a/Engine/Source/Runtime/Engine/Classes/Engine/GameViewportClient.h
+++ b/Engine/Source/Runtime/Engine/Classes/Engine/GameViewportClient.h
@@ -580,6 +580,9 @@ public:
 protected:
 
 	bool GetUseMouseForTouch() const;
+	void SetCurrentBufferVisualizationMode(FName NewBufferVisualizationMode) { CurrentBufferVisualizationMode = NewBufferVisualizationMode; }
+	FName GetCurrentBufferVisualizationMode() const { return CurrentBufferVisualizationMode; }
+	bool HasAudioFocus() const { return bHasAudioFocus; }
 
 protected:
 	/** FCommonViewportClient interface */
diff --git a/Engine/Source/Runtime/Engine/Private/LocalPlayer.cpp b/Engine/Source/Runtime/Engine/Private/LocalPlayer.cpp
--- a/Engine/Source/Runtime/Engine/Private/LocalPlayer.cpp
+++ b/Engine/Source/Runtime/Engine/Private/LocalPlayer.cpp
@@ -871,6 +871,17 @@ FSceneView* ULocalPlayer::CalcSceneView( class FSceneViewFamily* ViewFamily,
 			PlayerController->PlayerCameraManager->UpdatePhotographyPostProcessing(View->FinalPostProcessSettings);
 		}
 
+		if (GEngine->StereoRenderingDevice.IsValid())
+		{
+			FPostProcessSettings StereoDeviceOverridePostProcessinSettings;
+			float BlendWeight = 1.0f;
+			bool StereoSettingsAvailable = GEngine->StereoRenderingDevice->OverrideFinalPostprocessSettings(&StereoDeviceOverridePostProcessinSettings, StereoPass, BlendWeight);
+			if (StereoSettingsAvailable)
+			{
+				View->OverridePostProcessSettings(StereoDeviceOverridePostProcessinSettings, BlendWeight);
+			}
+		}
+
 #if !(UE_BUILD_SHIPPING || UE_BUILD_TEST)
 		ADebugCameraController* DebugCameraController = Cast<ADebugCameraController>(PlayerController);
 		if (DebugCameraController != nullptr)
diff --git a/Engine/Source/Runtime/Engine/Private/SceneView.cpp b/Engine/Source/Runtime/Engine/Private/SceneView.cpp
--- a/Engine/Source/Runtime/Engine/Private/SceneView.cpp
+++ b/Engine/Source/Runtime/Engine/Private/SceneView.cpp
@@ -1710,6 +1710,13 @@ void FSceneView::StartFinalPostprocessSettings(FVector InViewLocation)
 		}
 	}
 
+	{
+		if (GEngine->StereoRenderingDevice.IsValid())
+		{
+			GEngine->StereoRenderingDevice->StartFinalPostprocessSettings(&FinalPostProcessSettings, StereoPass);
+		}
+	}
+
 	if (State != nullptr)
 	{
 		State->OnStartPostProcessing(*this);
@@ -1902,7 +1909,7 @@ void FSceneView::EndFinalPostprocessSettings(const FSceneViewInitOptions& ViewIn
 	{
 		if (GEngine->StereoRenderingDevice.IsValid())
 		{
-			GEngine->StereoRenderingDevice->UpdatePostProcessSettings(&FinalPostProcessSettings, StereoPass);
+			GEngine->StereoRenderingDevice->EndFinalPostprocessSettings(&FinalPostProcessSettings, StereoPass);
 		}
 	}
 
diff --git a/Engine/Source/Runtime/Engine/Private/UserInterface/Console.cpp b/Engine/Source/Runtime/Engine/Private/UserInterface/Console.cpp
--- a/Engine/Source/Runtime/Engine/Private/UserInterface/Console.cpp
+++ b/Engine/Source/Runtime/Engine/Private/UserInterface/Console.cpp
@@ -38,6 +38,26 @@ static const FName NAME_Open = FName(TEXT("Open"));
 
 UConsole::FRegisterConsoleAutoCompleteEntries UConsole::RegisterConsoleAutoCompleteEntries;
 UConsole::FOnConsoleActivationStateChanged UConsole::OnConsoleActivationStateChanged;
+
+static TAutoConsoleVariable<int32> CVarCustomConsolePosEnabled(
+	TEXT("console.position.enable"),
+	0,
+	TEXT("Enable custom console positioning \n"),
+	ECVF_Default);
+
+static TAutoConsoleVariable<int32> CVarConsoleXPos(
+	TEXT("console.position.x"),
+	0,
+	TEXT("Console X offset from left border \n"),
+	ECVF_Default);
+
+static TAutoConsoleVariable<int32> CVarConsoleYPos(
+	TEXT("console.position.y"),
+	0,
+	TEXT("Console Y offset from bottom border \n"),
+	ECVF_Default);
+
+
 namespace ConsoleDefs
 {
 	/** Colors */
@@ -1097,18 +1117,27 @@ void UConsole::PostRender_Console_Typing(UCanvas* Canvas)
 	float ClipY = Canvas->ClipY;
 	float LeftPos = 0;
 
-	if (GEngine->IsConsoleBuild())
+	if (CVarCustomConsolePosEnabled.GetValueOnAnyThread())
 	{
-		ClipX	-= 64;
-		ClipY	-= 32;
-		LeftPos	 = 32;
+		LeftPos = (float)CVarConsoleXPos.GetValueOnAnyThread();
+		float BottomOffset = (float)CVarConsoleYPos.GetValueOnAnyThread();
+		ClipY = ClipY - BottomOffset;
 	}
-
-	if (GEngine->IsStereoscopic3D())
+	else
 	{
-		LeftPos = ClipX / 3;
-		ClipX -= LeftPos;
-		ClipY = ClipY * 0.60;
+		if (GEngine->IsConsoleBuild())
+		{
+			ClipX	-= 64;
+			ClipY	-= 32;
+			LeftPos	 = 32;
+		}
+
+		if (GEngine->IsStereoscopic3D())
+		{
+			LeftPos = ClipX / 3;
+			ClipX -= LeftPos;
+			ClipY = ClipY * 0.60;
+		}
 	}
 
 	PostRender_InputLine(Canvas, FIntPoint(LeftPos, ClipY));
@@ -1188,17 +1217,27 @@ void UConsole::PostRender_Console_Open(UCanvas* Canvas)
 	float TopPos = 0;
 	float LeftPos = 0;
 
-	if (GEngine->IsConsoleBuild())
+	if (CVarCustomConsolePosEnabled.GetValueOnAnyThread())
 	{
-		ClipX	-= 80;
-		TopPos	 = 30;
-		LeftPos	 = 40;
+		LeftPos = (float)CVarConsoleXPos.GetValueOnAnyThread();
+		float BottomOffset = (float)CVarConsoleYPos.GetValueOnAnyThread();
+		Height = Canvas->ClipY - BottomOffset;
 	}
-	if (GEngine->IsStereoscopic3D())
+	else
 	{
-		LeftPos = ClipX / 3;
-		ClipX -= LeftPos;
-		Height = Canvas->ClipY * 0.60;
+		if (GEngine->IsConsoleBuild())
+		{
+			ClipX	-= 80;
+			TopPos	 = 30;
+			LeftPos	 = 40;
+		}
+
+		if (GEngine->IsStereoscopic3D())
+		{
+			LeftPos = ClipX / 3;
+			ClipX -= LeftPos;
+			Height = Canvas->ClipY * 0.60;
+		}
 	}
 
 	UFont* Font = GEngine->GetSmallFont();
diff --git a/Engine/Source/Runtime/Engine/Public/StereoRendering.h b/Engine/Source/Runtime/Engine/Public/StereoRendering.h
--- a/Engine/Source/Runtime/Engine/Public/StereoRendering.h
+++ b/Engine/Source/Runtime/Engine/Public/StereoRendering.h
@@ -92,27 +92,65 @@ public:
 		return !(Pass == EStereoscopicPass::eSSP_FULL);
 	}
 	
+
+	virtual bool DeviceIsAPrimaryView(EStereoscopicPass Pass)
+	{
+		return Pass == eSSP_FULL || Pass == eSSP_LEFT_EYE;
+	}
+
+	/**
+	* Return true if this pass is for a view for which we share some work done for eSSP_LEFT_EYE (ie borrow some intermediate state from that eye)
+	*/
+	virtual bool DeviceIsASecondaryView(EStereoscopicPass Pass)
+	{
+		return !DeviceIsAPrimaryView(Pass);
+	}
+
+	/**
+	* Return true for additional eyes past the first two (a plugin could implement additional 'eyes').
+	*/
+	virtual bool DeviceIsAnAdditionalView(EStereoscopicPass Pass)
+	{
+		return Pass > eSSP_RIGHT_EYE;
+	}
+
+
 	/**
 	* Return true if this pass is for a view we do all the work for (ie this view can't borrow from another)
 	*/
-	static bool IsAPrimaryView(EStereoscopicPass Pass)
+	static bool IsAPrimaryView(EStereoscopicPass Pass, TSharedPtr< class IStereoRendering, ESPMode::ThreadSafe > StereoRenderingDevice)
 	{
+		if (StereoRenderingDevice.IsValid())
+		{
+			return StereoRenderingDevice->DeviceIsAPrimaryView(Pass);
+		}
+
 		return Pass == eSSP_FULL || Pass == eSSP_LEFT_EYE;
 	}
 	
 	/**
 	* Return true if this pass is for a view for which we share some work done for eSSP_LEFT_EYE (ie borrow some intermediate state from that eye)
 	*/
-	static bool IsASecondaryView(EStereoscopicPass Pass)
+	static bool IsASecondaryView(EStereoscopicPass Pass, TSharedPtr< class IStereoRendering, ESPMode::ThreadSafe > StereoRenderingDevice)
 	{
-		return !IsAPrimaryView(Pass);
+		if (StereoRenderingDevice.IsValid())
+		{
+			return StereoRenderingDevice->DeviceIsASecondaryView(Pass);
+		}
+
+		return !IsAPrimaryView(Pass, StereoRenderingDevice);
 	}
 	
 	/**
 	* Return true for additional eyes past the first two (a plugin could implement additional 'eyes').
 	*/
-	static bool IsAnAdditionalView(EStereoscopicPass Pass)
+	static bool IsAnAdditionalView(EStereoscopicPass Pass, TSharedPtr< class IStereoRendering, ESPMode::ThreadSafe > StereoRenderingDevice)
 	{
+		if (StereoRenderingDevice.IsValid())
+		{
+			return StereoRenderingDevice->DeviceIsAnAdditionalView(Pass);
+		}
+
 		return Pass > eSSP_RIGHT_EYE;
 	}
 
@@ -161,9 +199,9 @@ public:
 	 */
 	virtual IStereoLayers* GetStereoLayers () { return nullptr; }
 
-	/**
-	 * Allows to override the PostProcessSettings in the last moment e.g. allows up sampled 3D rendering
-	 */
-	virtual void UpdatePostProcessSettings(struct FPostProcessSettings* FinalPostProcessingSettings, const enum EStereoscopicPass StereoPassType) {}
+	
+	virtual void StartFinalPostprocessSettings(struct FPostProcessSettings* StartPostProcessingSettings, const enum EStereoscopicPass StereoPassType) {}
+	virtual bool OverrideFinalPostprocessSettings(struct FPostProcessSettings* OverridePostProcessingSettings, const enum EStereoscopicPass StereoPassType, float& BlendWeight) { return false; }
+	virtual void EndFinalPostprocessSettings(struct FPostProcessSettings* FinalPostProcessingSettings, const enum EStereoscopicPass StereoPassType) {}
 
 };
diff --git a/Engine/Source/Runtime/Renderer/Private/LightPropagationVolume.cpp b/Engine/Source/Runtime/Renderer/Private/LightPropagationVolume.cpp
--- a/Engine/Source/Runtime/Renderer/Private/LightPropagationVolume.cpp
+++ b/Engine/Source/Runtime/Renderer/Private/LightPropagationVolume.cpp
@@ -1459,7 +1459,7 @@ void FSceneViewState::SetupLightPropagationVolume(FSceneView& View, FSceneViewFa
 
 	const ERHIFeatureLevel::Type ViewFeatureLevel = View.GetFeatureLevel();
 
-	if (IStereoRendering::IsASecondaryView(View.StereoPass))
+	if (IStereoRendering::IsASecondaryView(View.StereoPass, GEngine->StereoRenderingDevice))
 	{
 		// The right eye will reference the left eye's LPV with the assumption that the left eye uses the primary view (index 0)
 		const FSceneView* PrimaryView = ViewFamily.Views[0];
diff --git a/Engine/Source/Runtime/Renderer/Private/PostProcess/PostProcessing.cpp b/Engine/Source/Runtime/Renderer/Private/PostProcess/PostProcessing.cpp
--- a/Engine/Source/Runtime/Renderer/Private/PostProcess/PostProcessing.cpp
+++ b/Engine/Source/Runtime/Renderer/Private/PostProcess/PostProcessing.cpp
@@ -376,7 +376,7 @@ static FRCPassPostProcessTonemap* AddTonemapper(
 	const bool bIsComputePass = ShouldDoComputePostProcessing(View);
 
 	FRenderingCompositeOutputRef TonemapperCombinedLUTOutputRef;
-	if (StereoPass != eSSP_RIGHT_EYE)
+	if (IStereoRendering::IsAPrimaryView(StereoPass, GEngine->StereoRenderingDevice))
 	{
 		bool bNeedFloatOutput = View.Family->SceneCaptureSource == SCS_FinalColorHDR;
 		bool bAllocateOutput = View.State == nullptr;
@@ -1577,7 +1577,7 @@ void FPostProcessing::Process(FRHICommandListImmediate& RHICmdList, const FViewI
 					bHistogramNeeded = true;
 				}
 
-				if (!GIsHighResScreenshot && bHistogramNeeded && FeatureLevel >= ERHIFeatureLevel::SM5 && IStereoRendering::IsAPrimaryView(StereoPass))
+				if (!GIsHighResScreenshot && bHistogramNeeded && FeatureLevel >= ERHIFeatureLevel::SM5 && IStereoRendering::IsAPrimaryView(StereoPass, GEngine->StereoRenderingDevice))
 				{
 					FRenderingCompositePass* NodeHistogram = Context.Graph.RegisterPass(new(FMemStack::Get()) FRCPassPostProcessHistogram());
 
@@ -1601,7 +1601,7 @@ void FPostProcessing::Process(FRHICommandListImmediate& RHICmdList, const FViewI
 				if (Context.View.FinalPostProcessSettings.BloomThreshold <= -1 && Context.View.Family->Views.Num() == 1)
 				{
 					if (!GIsHighResScreenshot && View.State &&
-						IStereoRendering::IsAPrimaryView(StereoPass) &&
+						IStereoRendering::IsAPrimaryView(StereoPass, GEngine->StereoRenderingDevice) &&
 						AutoExposure.MethodId == EAutoExposureMethod::AEM_Basic)
 					{
 						BloomAndEyeDownSamplesPtr = CreateDownSampleArray(Context, SceneColorHalfRes, true /*bGenerateLog2Alpha*/);
@@ -1610,7 +1610,7 @@ void FPostProcessing::Process(FRHICommandListImmediate& RHICmdList, const FViewI
 			}
 
 			// some views don't have a state (thumbnail rendering)
-			if(!GIsHighResScreenshot && View.State && IStereoRendering::IsAPrimaryView(StereoPass))
+			if(!GIsHighResScreenshot && View.State && IStereoRendering::IsAPrimaryView(StereoPass, GEngine->StereoRenderingDevice))
 			{
 				const bool bUseBasicEyeAdaptation = (AutoExposure.MethodId == EAutoExposureMethod::AEM_Basic);
 
diff --git a/Engine/Source/Runtime/Renderer/Private/PostProcess/RenderingCompositionGraph.cpp b/Engine/Source/Runtime/Renderer/Private/PostProcess/RenderingCompositionGraph.cpp
--- a/Engine/Source/Runtime/Renderer/Private/PostProcess/RenderingCompositionGraph.cpp
+++ b/Engine/Source/Runtime/Renderer/Private/PostProcess/RenderingCompositionGraph.cpp
@@ -36,6 +36,13 @@ static TAutoConsoleVariable<int32> CVarCompositionGraphOrder(
 	TEXT(" 1: RegisterPass() call order, unless the dependencies (input and additional) require a different order (might become new default as it provides more control, executes all registered nodes)"),
 	ECVF_RenderThreadSafe);
 
+static TAutoConsoleVariable<int32> CVarCompositionForceRenderTargetLoad(
+	TEXT("r.CompositionForceRenderTargetLoad"),
+	0,
+	TEXT("0: default engine behaviour\n")
+	TEXT("1: force ERenderTargetLoadAction::ELoad for all render targets"),
+	ECVF_RenderThreadSafe);
+
 #if !UE_BUILD_SHIPPING
 FAutoConsoleCommand CmdCompositionGraphDebug(
 	TEXT("r.CompositionGraphDebug"),
@@ -286,6 +293,33 @@ void FRenderingCompositePassContext::Process(const TArray<FRenderingCompositePas
 	}
 }
 
+ERenderTargetLoadAction FRenderingCompositePassContext::GetLoadActionForRenderTarget(const FSceneRenderTargetItem& DestRenderTarget) const
+{
+	ERenderTargetLoadAction LoadAction = ERenderTargetLoadAction::ENoAction;
+
+	if (IsViewFamilyRenderTarget(DestRenderTarget))
+	{
+		const bool bForceLoad = !!CVarCompositionForceRenderTargetLoad.GetValueOnAnyThread();
+		if (bForceLoad)
+		{
+			LoadAction = ERenderTargetLoadAction::ELoad;
+		}
+		else
+		{
+			// If rendering the final view family's render target, must clear first view, and load subsequent views.
+			LoadAction = (&View != View.Family->Views[0]) ? ERenderTargetLoadAction::ELoad : ERenderTargetLoadAction::EClear;
+		}
+	}
+	else if (HasHmdMesh())
+	{
+		// Clears render target because going to have unrendered pixels inside view rect.
+		LoadAction = ERenderTargetLoadAction::EClear;
+	}
+
+	return LoadAction;
+}
+
+
 // --------------------------------------------------------------------------
 
 FRenderingCompositionGraph::FRenderingCompositionGraph()
diff --git a/Engine/Source/Runtime/Renderer/Private/PostProcess/RenderingCompositionGraph.h b/Engine/Source/Runtime/Renderer/Private/PostProcess/RenderingCompositionGraph.h
--- a/Engine/Source/Runtime/Renderer/Private/PostProcess/RenderingCompositionGraph.h
+++ b/Engine/Source/Runtime/Renderer/Private/PostProcess/RenderingCompositionGraph.h
@@ -175,23 +175,7 @@ struct FRenderingCompositePassContext
 	}
 
 	/** Returns the LoadAction that should be use for a given render target. */
-	ERenderTargetLoadAction GetLoadActionForRenderTarget(const FSceneRenderTargetItem& DestRenderTarget) const
-	{
-		ERenderTargetLoadAction LoadAction = ERenderTargetLoadAction::ENoAction;
-		
-		if (IsViewFamilyRenderTarget(DestRenderTarget))
-		{
-			// If rendering the final view family's render target, must clear first view, and load subsequent views.
-			LoadAction = (&View != View.Family->Views[0]) ? ERenderTargetLoadAction::ELoad : ERenderTargetLoadAction::EClear;
-		}
-		else if (HasHmdMesh())
-		{
-			// Clears render target because going to have unrendered pixels inside view rect.
-			LoadAction = ERenderTargetLoadAction::EClear;
-		}
-
-		return LoadAction;
-	}
+	ERenderTargetLoadAction GetLoadActionForRenderTarget(const FSceneRenderTargetItem& DestRenderTarget) const;
 
 	ERHIFeatureLevel::Type GetFeatureLevel() const { return FeatureLevel; }
 	EShaderPlatform GetShaderPlatform() const { return GShaderPlatformForFeatureLevel[FeatureLevel]; }
diff --git a/Engine/Source/Runtime/Renderer/Private/SceneRendering.cpp b/Engine/Source/Runtime/Renderer/Private/SceneRendering.cpp
--- a/Engine/Source/Runtime/Renderer/Private/SceneRendering.cpp
+++ b/Engine/Source/Runtime/Renderer/Private/SceneRendering.cpp
@@ -1291,7 +1291,7 @@ void FViewInfo::SetupUniformBufferParameters(
 	{
 		// If rendering in stereo, the other stereo passes uses the left eye's translucency lighting volume.
 		const FViewInfo* PrimaryView = this;
-		if (IStereoRendering::IsASecondaryView(StereoPass))
+		if (IStereoRendering::IsASecondaryView(StereoPass, GEngine->StereoRenderingDevice))
 		{
 			if (Family->Views.IsValidIndex(0))
 			{
@@ -1641,7 +1641,7 @@ FSceneViewState* FViewInfo::GetEffectiveViewState() const
 	FSceneViewState* EffectiveViewState = ViewState;
 
 	// When rendering in stereo we want to use the same exposure for both eyes.
-	if (IStereoRendering::IsASecondaryView(StereoPass))
+	if (IStereoRendering::IsASecondaryView(StereoPass, GEngine->StereoRenderingDevice))
 	{
 		int32 ViewIndex = Family->Views.Find(this);
 		if (Family->Views.IsValidIndex(ViewIndex))
@@ -2365,7 +2365,7 @@ void FSceneRenderer::ComputeFamilySize()
 		MaxFamilyX = FMath::Max(MaxFamilyX, FinalViewMaxX);
 		MaxFamilyY = FMath::Max(MaxFamilyY, FinalViewMaxY);
 
-		if (!IStereoRendering::IsAnAdditionalView(View.StereoPass))
+		if (!IStereoRendering::IsAnAdditionalView(View.StereoPass, GEngine->StereoRenderingDevice))
 		{
 		InstancedStereoWidth = FPlatformMath::Max(InstancedStereoWidth, static_cast<uint32>(View.ViewRect.Max.X));
 	}
diff --git a/Engine/Source/Runtime/Renderer/Private/ShadowSetup.cpp b/Engine/Source/Runtime/Renderer/Private/ShadowSetup.cpp
--- a/Engine/Source/Runtime/Renderer/Private/ShadowSetup.cpp
+++ b/Engine/Source/Runtime/Renderer/Private/ShadowSetup.cpp
@@ -3339,15 +3339,15 @@ void FSceneRenderer::AddViewDependentWholeSceneShadowsForView(
 		FadeAlphas.Init(0.0f, Views.Num());
 		FadeAlphas[ViewIndex] = 1.0f;
 
-		if (View.StereoPass == eSSP_LEFT_EYE
+		if (IStereoRendering::IsAPrimaryView(View.StereoPass, GEngine->StereoRenderingDevice)
 			&& Views.IsValidIndex(ViewIndex + 1)
-			&& Views[ViewIndex + 1].StereoPass == eSSP_RIGHT_EYE)
+			&& IStereoRendering::IsASecondaryView(View.StereoPass, GEngine->StereoRenderingDevice))
 		{
 			FadeAlphas[ViewIndex + 1] = 1.0f;
 		}		
-
+		
 		// If rendering in stereo mode we render shadow depths only for the left eye, but project for both eyes!
-		if (View.StereoPass != eSSP_RIGHT_EYE)
+		if (IStereoRendering::IsAPrimaryView(View.StereoPass, GEngine->StereoRenderingDevice))
 		{
 			const bool bExtraDistanceFieldCascade = LightSceneInfo.Proxy->ShouldCreateRayTracedCascade(View.GetFeatureLevel(), LightSceneInfo.IsPrecomputedLightingValid(), View.MaxShadowCascades);
 
diff --git a/Templates/TP_nDisplayBP/Config/DefaultEngine.ini b/Templates/TP_nDisplayBP/Config/DefaultEngine.ini
--- a/Templates/TP_nDisplayBP/Config/DefaultEngine.ini
+++ b/Templates/TP_nDisplayBP/Config/DefaultEngine.ini
@@ -12,6 +12,7 @@ GlobalDefaultGameMode=/Script/DisplayCluster.DisplayClusterGameModeDefault
 [/Script/Engine.Engine]
 GameEngine=/Script/DisplayCluster.DisplayClusterGameEngine
 UnrealEdEngine=/Script/DisplayClusterEditor.DisplayClusterEditorEngine
+GameViewportClientClassName=/Script/DisplayCluster.DisplayClusterViewportClient
 
 [/Script/DisplayClusterEditor.DisplayClusterEditorSettings]
 bEnabled=True
@@ -81,3 +82,4 @@ OcclusionPlugin=
 r.DefaultFeature.Bloom=False
 r.DefaultFeature.AutoExposure=False
 r.DefaultFeature.MotionBlur=False
+r.SceneRenderTargetResizeMethod=2
