Fixed crash during shutdown when using DisplayCluster message interceptor

diff --git a/Engine/Plugins/Runtime/nDisplay/Source/DisplayClusterMessageInterceptor/Private/DisplayClusterMessageInterceptionModule.cpp b/Engine/Plugins/Runtime/nDisplay/Source/DisplayClusterMessageInterceptor/Private/DisplayClusterMessageInterceptionModule.cpp
--- a/Engine/Plugins/Runtime/nDisplay/Source/DisplayClusterMessageInterceptor/Private/DisplayClusterMessageInterceptionModule.cpp
+++ b/Engine/Plugins/Runtime/nDisplay/Source/DisplayClusterMessageInterceptor/Private/DisplayClusterMessageInterceptionModule.cpp
@@ -131,6 +131,8 @@ private:
 		if (Interceptor)
 		{
 			Interceptor->Stop();
+			// remove internal reference to message bus
+			Interceptor->Setup(nullptr, nullptr);
 		}
 	}
 
diff --git a/Engine/Plugins/Runtime/nDisplay/Source/DisplayClusterMessageInterceptor/Private/DisplayClusterMessageInterceptor.cpp b/Engine/Plugins/Runtime/nDisplay/Source/DisplayClusterMessageInterceptor/Private/DisplayClusterMessageInterceptor.cpp
--- a/Engine/Plugins/Runtime/nDisplay/Source/DisplayClusterMessageInterceptor/Private/DisplayClusterMessageInterceptor.cpp
+++ b/Engine/Plugins/Runtime/nDisplay/Source/DisplayClusterMessageInterceptor/Private/DisplayClusterMessageInterceptor.cpp
@@ -17,7 +17,7 @@ FDisplayClusterMessageInterceptor::FDisplayClusterMessageInterceptor()
 	, Address(FMessageAddress::NewAddress())
 {}
 
-void FDisplayClusterMessageInterceptor::Setup(IDisplayClusterClusterManager* InClusterManager, TSharedRef<IMessageBus, ESPMode::ThreadSafe> InBus)
+void FDisplayClusterMessageInterceptor::Setup(IDisplayClusterClusterManager* InClusterManager, TSharedPtr<IMessageBus, ESPMode::ThreadSafe> InBus)
 {
 	ClusterManager = InClusterManager;
 	InterceptedBus = InBus;
@@ -34,18 +34,20 @@ void FDisplayClusterMessageInterceptor::Purge()
 		}
 		ContextMap.Empty();
 	}
-	for (const auto& Context : ContextToForward)
+	if (InterceptedBus)
 	{
-		InterceptedBus->Forward(Context.ToSharedRef(), Context->GetRecipients(), FTimespan::Zero(), AsShared());
+		for (const auto& Context : ContextToForward)
+		{
+			InterceptedBus->Forward(Context.ToSharedRef(), Context->GetRecipients(), FTimespan::Zero(), AsShared());
+		}
 	}
 }
 
 void FDisplayClusterMessageInterceptor::Start()
 {
 	const UDisplayClusterMessageInterceptionSettings* InterceptionSettings = GetDefault<UDisplayClusterMessageInterceptionSettings>();
-	if (!bIsIntercepting && InterceptionSettings->bIsEnabled)
+	if (!bIsIntercepting && InterceptionSettings->bIsEnabled && InterceptedBus)
 	{
-		check(InterceptedBus);
 		for (const FName& MessageType : InterceptionSettings->MessageTypes)
 		{
 			InterceptedBus->Intercept(AsShared(), MessageType);
@@ -57,9 +59,8 @@ void FDisplayClusterMessageInterceptor::Start()
 
 void FDisplayClusterMessageInterceptor::Stop()
 {
-	if (bIsIntercepting)
+	if (bIsIntercepting && InterceptedBus)
 	{
-		check(InterceptedBus);
 		InterceptedBus->Unintercept(AsShared(), NAME_All);
 		bIsIntercepting = false;
 		Purge();
@@ -123,7 +124,7 @@ void FDisplayClusterMessageInterceptor::SyncMessages()
 void FDisplayClusterMessageInterceptor::HandleClusterEvent(const FDisplayClusterClusterEvent& InEvent)
 {
 	TArray<TSharedPtr<IMessageContext, ESPMode::ThreadSafe>> ContextToForward;
-	if (InEvent.Category == TEXT("nDCI"))
+	if (InEvent.Category == TEXT("nDCI") && ClusterManager)
 	{
 		FScopeLock Lock(&ContextQueueCS);
 		if (FContextSync* ContextSync = ContextMap.Find(InEvent.Type))
diff --git a/Engine/Plugins/Runtime/nDisplay/Source/DisplayClusterMessageInterceptor/Private/DisplayClusterMessageInterceptor.h b/Engine/Plugins/Runtime/nDisplay/Source/DisplayClusterMessageInterceptor/Private/DisplayClusterMessageInterceptor.h
--- a/Engine/Plugins/Runtime/nDisplay/Source/DisplayClusterMessageInterceptor/Private/DisplayClusterMessageInterceptor.h
+++ b/Engine/Plugins/Runtime/nDisplay/Source/DisplayClusterMessageInterceptor/Private/DisplayClusterMessageInterceptor.h
@@ -22,7 +22,7 @@ class FDisplayClusterMessageInterceptor : public TSharedFromThis<FDisplayCluster
 public:
 	FDisplayClusterMessageInterceptor();
 
-	void Setup(IDisplayClusterClusterManager* InClusterManager, TSharedRef<IMessageBus, ESPMode::ThreadSafe> InBus);
+	void Setup(IDisplayClusterClusterManager* InClusterManager, TSharedPtr<IMessageBus, ESPMode::ThreadSafe> InBus);
 	
 	void Start();
 	void Stop();
