Fix for EasyBlend frustums so that .ol and .pol internal pipelines provide the same view rotation

diff --git a/Engine/Plugins/Runtime/nDisplay/Source/DisplayClusterProjection/Private/Policy/EasyBlend/DX11/DisplayClusterProjectionEasyBlendViewAdapterDX11.cpp b/Engine/Plugins/Runtime/nDisplay/Source/DisplayClusterProjection/Private/Policy/EasyBlend/DX11/DisplayClusterProjectionEasyBlendViewAdapterDX11.cpp
--- a/Engine/Plugins/Runtime/nDisplay/Source/DisplayClusterProjection/Private/Policy/EasyBlend/DX11/DisplayClusterProjectionEasyBlendViewAdapterDX11.cpp
+++ b/Engine/Plugins/Runtime/nDisplay/Source/DisplayClusterProjection/Private/Policy/EasyBlend/DX11/DisplayClusterProjectionEasyBlendViewAdapterDX11.cpp
@@ -22,7 +22,6 @@ FDisplayClusterProjectionEasyBlendViewAdapterDX11::FDisplayClusterProjectionEasy
 	: FDisplayClusterProjectionEasyBlendViewAdapterBase(InitParams)
 	, bIsRenderResourcesInitialized(false)
 {
-	check(InitParams.ViewportSize.GetMax() > InitParams.ViewportSize.GetMin());
 	check(InitParams.NumViews > 0)
 
 	Views.AddDefaulted(InitParams.NumViews);
@@ -61,12 +60,16 @@ bool FDisplayClusterProjectionEasyBlendViewAdapterDX11::Initialize(const FString
 	for (auto& It : Views)
 	{
 		// Initialize the mesh data
-		check(DisplayClusterProjectionEasyBlendLibraryDX11::EasyBlendInitializeFunc);
-		EasyBlendSDKDXError err = DisplayClusterProjectionEasyBlendLibraryDX11::EasyBlendInitializeFunc(FileName, It.EasyBlendMeshData.Get());
-		if (!EasyBlendSDKDX_SUCCEEDED(err))
 		{
-			UE_LOG(LogDisplayClusterProjectionEasyBlend, Error, TEXT("Couldn't initialize EasyBlend internals"));
-			return false;
+			FScopeLock lock(&DllAccessCS);
+
+			check(DisplayClusterProjectionEasyBlendLibraryDX11::EasyBlendInitializeFunc);
+			const EasyBlendSDKDXError Result = DisplayClusterProjectionEasyBlendLibraryDX11::EasyBlendInitializeFunc(FileName, It.EasyBlendMeshData.Get());
+			if (!EasyBlendSDKDX_SUCCEEDED(Result))
+			{
+				UE_LOG(LogDisplayClusterProjectionEasyBlend, Error, TEXT("Couldn't initialize EasyBlend internals"));
+				return false;
+			}
 		}
 
 		// EasyBlendMeshData has been initialized
@@ -83,7 +86,7 @@ bool FDisplayClusterProjectionEasyBlendViewAdapterDX11::Initialize(const FString
 	return true;
 }
 
-
+// Location/Rotation inside the function is in EasyBlend space with a scale applied
 bool FDisplayClusterProjectionEasyBlendViewAdapterDX11::CalculateView(const uint32 ViewIdx, FVector& InOutViewLocation, FRotator& InOutViewRotation, const FVector& ViewOffset, const float WorldToMeters, const float NCP, const float FCP)
 {
 	check(Views.Num() > (int)ViewIdx);
@@ -91,26 +94,31 @@ bool FDisplayClusterProjectionEasyBlendViewAdapterDX11::CalculateView(const uint
 	ZNear = NCP;
 	ZFar  = FCP;
 
-	static const FMatrix Game2EasyBlend(
-		FPlane(0, 0, -1, 0),
-		FPlane(-1, 0, 0, 0),
-		FPlane(0, 1, 0, 0),
-		FPlane(0, 0, 0, 1));
+	// Convert to EasyBlend coordinate system
+	FVector EasyBlendEyeLocation;
+	EasyBlendEyeLocation.X =  InOutViewLocation.Y;
+	EasyBlendEyeLocation.Y = -InOutViewLocation.Z;
+	EasyBlendEyeLocation.Z =  InOutViewLocation.X;
 
-	// Calculate view location in the EasyBlend coordinate system
-	FVector EasyBlendEyeLocation = Game2EasyBlend.TransformPosition(InOutViewLocation);
+	// View rotation
+	double Yaw   = 0.l;
+	double Pitch = 0.l;
+	double Roll  = 0.l;
 
-	// Update EasyBlend state
 	{
 		FScopeLock lock(&DllAccessCS);
+
+		// Update view location
 		check(DisplayClusterProjectionEasyBlendLibraryDX11::EasyBlendSetEyepointFunc);
 		DisplayClusterProjectionEasyBlendLibraryDX11::EasyBlendSetEyepointFunc(Views[ViewIdx].EasyBlendMeshData.Get(), EasyBlendEyeLocation.X, EasyBlendEyeLocation.Y, EasyBlendEyeLocation.Z);
+
+		// Get actual view rotation
+		check(DisplayClusterProjectionEasyBlendLibraryDX11::EasyBlendSDK_GetHeadingPitchRollFunc);
+		DisplayClusterProjectionEasyBlendLibraryDX11::EasyBlendSDK_GetHeadingPitchRollFunc(Yaw, Pitch, Roll, Views[ViewIdx].EasyBlendMeshData.Get());
 	}
 
-	// Get rotation from the frustum that has been updated already
-	InOutViewRotation.Pitch = float(-Views[ViewIdx].EasyBlendMeshData->Frustum.ViewAngleC - 90);
-	InOutViewRotation.Yaw   = float(Views[ViewIdx].EasyBlendMeshData->Frustum.ViewAngleB);
-	InOutViewRotation.Roll  = float(Views[ViewIdx].EasyBlendMeshData->Frustum.ViewAngleA);
+	// Forward view rotation to a caller
+	InOutViewRotation = FRotator(-(float)Pitch, (float)Yaw, (float)Roll);
 
 	return true;
 }
diff --git a/Engine/Plugins/Runtime/nDisplay/Source/DisplayClusterProjection/Private/Policy/EasyBlend/DX11/DisplayClusterProjectionEasyBlendViewAdapterDX11.h b/Engine/Plugins/Runtime/nDisplay/Source/DisplayClusterProjection/Private/Policy/EasyBlend/DX11/DisplayClusterProjectionEasyBlendViewAdapterDX11.h
--- a/Engine/Plugins/Runtime/nDisplay/Source/DisplayClusterProjection/Private/Policy/EasyBlend/DX11/DisplayClusterProjectionEasyBlendViewAdapterDX11.h
+++ b/Engine/Plugins/Runtime/nDisplay/Source/DisplayClusterProjection/Private/Policy/EasyBlend/DX11/DisplayClusterProjectionEasyBlendViewAdapterDX11.h
@@ -41,8 +41,6 @@ private:
 
 		bool bIsMeshInitialized = false;
 		
-		FVector EyeLocation = FVector::ZeroVector;
-
 		ViewData()
 		{
 			EasyBlendMeshData.Reset(new EasyBlendSDKDX_Mesh);
diff --git a/Engine/Plugins/Runtime/nDisplay/Source/DisplayClusterProjection/Private/Policy/EasyBlend/DisplayClusterProjectionEasyBlendPolicyBase.cpp b/Engine/Plugins/Runtime/nDisplay/Source/DisplayClusterProjection/Private/Policy/EasyBlend/DisplayClusterProjectionEasyBlendPolicyBase.cpp
--- a/Engine/Plugins/Runtime/nDisplay/Source/DisplayClusterProjection/Private/Policy/EasyBlend/DisplayClusterProjectionEasyBlendPolicyBase.cpp
+++ b/Engine/Plugins/Runtime/nDisplay/Source/DisplayClusterProjection/Private/Policy/EasyBlend/DisplayClusterProjectionEasyBlendPolicyBase.cpp
@@ -66,9 +66,9 @@ bool FDisplayClusterProjectionEasyBlendPolicyBase::HandleAddViewport(const FIntP
 
 	// Create and store nDisplay-to-EasyBlend viewport adapter
 	TSharedPtr<FDisplayClusterProjectionEasyBlendViewAdapterBase> NewViewAdapter = CreateViewAdapter(FDisplayClusterProjectionEasyBlendViewAdapterBase::FInitParams{ ViewportSize, ViewsAmount });
-	if (!NewViewAdapter->Initialize(FullFilePath))
+	if (!NewViewAdapter || !NewViewAdapter->Initialize(FullFilePath))
 	{
-		UE_LOG(LogDisplayClusterProjectionEasyBlend, Error, TEXT("Couldn't initialize EasyBlend viewport adapter"));
+		UE_LOG(LogDisplayClusterProjectionEasyBlend, Error, TEXT("An error occurred during EasyBlend viewport adapter initialization"));
 		return false;
 	}
 
@@ -89,7 +89,7 @@ void FDisplayClusterProjectionEasyBlendPolicyBase::HandleRemoveViewport()
 
 bool FDisplayClusterProjectionEasyBlendPolicyBase::CalculateView(const uint32 ViewIdx, FVector& InOutViewLocation, FRotator& InOutViewRotation, const FVector& ViewOffset, const float WorldToMeters, const float NCP, const float FCP)
 {
-	check(IsInGameThread())
+	check(IsInGameThread());
 
 	if (!ViewAdapter.IsValid())
 	{
@@ -98,24 +98,29 @@ bool FDisplayClusterProjectionEasyBlendPolicyBase::CalculateView(const uint32 Vi
 
 	const float WorldScale = WorldToMeters / 100.f;
 
-	// Get world-2-local matrix
+	// Get origin component
 	const USceneComponent* const OriginComp = GetOriginComp();
+	check(OriginComp);
+
+	// Get world-origin matrix
 	const FTransform& World2LocalTransform = (OriginComp != nullptr ? OriginComp->GetComponentTransform() : FTransform::Identity);
 
 	// Calculate view location in origin space
-	const FVector LocalViewLoc = World2LocalTransform.InverseTransformPosition(InOutViewLocation);
+	FVector  OriginSpaceViewLocation = World2LocalTransform.InverseTransformPosition(InOutViewLocation);
+
+	// Apply EasyBlend scale depending on the measurement units used in calibration
+	OriginSpaceViewLocation = (OriginSpaceViewLocation / 100 / EasyBlendScale);
 
-	// Prepare non-const variables to get output
-	FVector  EasyBlendViewLocation = LocalViewLoc * EasyBlendScale;
-	FRotator EasyBlendViewRotation = FRotator::ZeroRotator;
-	if (!ViewAdapter->CalculateView(ViewIdx, EasyBlendViewLocation, EasyBlendViewRotation, ViewOffset, WorldScale, NCP, FCP))
+	// Forward data to the RHI dependend EasyBlend implementation
+	FRotator OriginSpaceViewRotation = FRotator::ZeroRotator;
+	if (!ViewAdapter->CalculateView(ViewIdx, OriginSpaceViewLocation, OriginSpaceViewRotation, ViewOffset, WorldScale, NCP, FCP))
 	{
 		UE_LOG(LogDisplayClusterProjectionEasyBlend, Warning, TEXT("Couldn't compute view info for <%s> viewport"), *GetViewportId());
 		return false;
 	}
 
-	// Convert rotation from local to world space
-	InOutViewRotation = World2LocalTransform.TransformRotation(EasyBlendViewRotation.Quaternion()).Rotator();
+	// Convert rotation back from origin to world space
+	InOutViewRotation = World2LocalTransform.TransformRotation(OriginSpaceViewRotation.Quaternion()).Rotator();
 
 	return true;
 }
