Fixed custom postprocess (including output remap)

diff --git a/Engine/Plugins/Runtime/nDisplay/Shaders/Private/OutputRemapShaders.usf b/Engine/Plugins/Runtime/nDisplay/Shaders/Private/OutputRemapShaders.usf
--- a/Engine/Plugins/Runtime/nDisplay/Shaders/Private/OutputRemapShaders.usf
+++ b/Engine/Plugins/Runtime/nDisplay/Shaders/Private/OutputRemapShaders.usf
@@ -26,8 +26,7 @@ OutputVS OutputRemap_VS(InputVS IN)
 {
 	float2 Pos = IN.Position.xy;
 
-	Pos.xy = Pos.xy * 2 - 1.0f;
-	Pos.xy *= float2(1, -1);
+	Pos.xy = Pos.xy * 2 - 1.0f;	
 
 	OutputVS Out;
 	Out.Position = float4(Pos, 0, 1);
diff --git a/Engine/Plugins/Runtime/nDisplay/Source/DisplayCluster/Private/Render/Device/DisplayClusterDeviceBase.cpp b/Engine/Plugins/Runtime/nDisplay/Source/DisplayCluster/Private/Render/Device/DisplayClusterDeviceBase.cpp
--- a/Engine/Plugins/Runtime/nDisplay/Source/DisplayCluster/Private/Render/Device/DisplayClusterDeviceBase.cpp
+++ b/Engine/Plugins/Runtime/nDisplay/Source/DisplayCluster/Private/Render/Device/DisplayClusterDeviceBase.cpp
@@ -60,7 +60,7 @@ static TAutoConsoleVariable<int32> CVarCustomPPEnabled(
 
 FDisplayClusterDeviceBase::FDisplayClusterDeviceBase(uint32 ViewsPerViewport)
 	: FRHICustomPresent()
-	, FDisplayClusterDeviceBase_PostProcess(RenderViewports, ViewsPerViewport)
+	, FDisplayClusterDeviceBase_PostProcess(RenderViewports, ViewsPerViewport, EyeRegions)
 	, ViewsAmountPerViewport(ViewsPerViewport)
 {
 	DISPLAY_CLUSTER_FUNC_TRACE(LogDisplayClusterRender);
@@ -541,6 +541,9 @@ void FDisplayClusterDeviceBase::CalculateRenderTargetSize(const class FViewport&
 		InOutSizeY = FMath::Max(InOutSizeY, (uint32)Item.GetArea().Max.Y);
 	}
 
+	// Store eye region
+	EyeRegions[0] = FIntRect(FIntPoint(0, 0), FIntPoint(InOutSizeX, InOutSizeY));
+
 	UE_LOG(LogDisplayClusterRender, Verbose, TEXT("Render target size: [%d x %d]"), InOutSizeX, InOutSizeY);
 
 	check(InOutSizeX > 0 && InOutSizeY > 0);
@@ -558,7 +561,7 @@ bool FDisplayClusterDeviceBase::NeedReAllocateViewportRenderTarget(const class F
 	// Get desired RT size
 	uint32 newSizeX = 0;
 	uint32 newSizeY = 0;
-	CalculateRenderTargetSize(Viewport, newSizeX, newSizeY);	
+	CalculateRenderTargetSize(Viewport, newSizeX, newSizeY);
 
 	// Here we conclude if need to re-allocate
 	const bool Result = (newSizeX != rtSize.X || newSizeY != rtSize.Y);
diff --git a/Engine/Plugins/Runtime/nDisplay/Source/DisplayCluster/Private/Render/Device/DisplayClusterDeviceBase.h b/Engine/Plugins/Runtime/nDisplay/Source/DisplayCluster/Private/Render/Device/DisplayClusterDeviceBase.h
--- a/Engine/Plugins/Runtime/nDisplay/Source/DisplayCluster/Private/Render/Device/DisplayClusterDeviceBase.h
+++ b/Engine/Plugins/Runtime/nDisplay/Source/DisplayCluster/Private/Render/Device/DisplayClusterDeviceBase.h
@@ -156,8 +156,12 @@ protected:
 	uint32 ViewsAmountPerViewport = 0;
 	// UE4 main viewport
 	FViewport* MainViewport = nullptr;
-	// custom post processing settings
-	TMap<int, FPostProcessSettings> ViewportStartPostProcessingSettings;	 
+
+	// Per-eye regions
+	FIntRect EyeRegions[2];
+
+	// Custom post processing settings
+	TMap<int, FPostProcessSettings> ViewportStartPostProcessingSettings;
 	TMap<int, FOverridePostProcessingSettings> ViewportOverridePostProcessingSettings;
 	TMap<int, FPostProcessSettings> ViewportFinalPostProcessingSettings;
 
diff --git a/Engine/Plugins/Runtime/nDisplay/Source/DisplayCluster/Private/Render/Device/DisplayClusterDeviceBase_PostProcess.cpp b/Engine/Plugins/Runtime/nDisplay/Source/DisplayCluster/Private/Render/Device/DisplayClusterDeviceBase_PostProcess.cpp
--- a/Engine/Plugins/Runtime/nDisplay/Source/DisplayCluster/Private/Render/Device/DisplayClusterDeviceBase_PostProcess.cpp
+++ b/Engine/Plugins/Runtime/nDisplay/Source/DisplayCluster/Private/Render/Device/DisplayClusterDeviceBase_PostProcess.cpp
@@ -73,13 +73,10 @@ void FDisplayClusterDeviceBase_PostProcess::PerformPostProcessFrameBeforeWarpBle
 		{
 			if (CurPP.Operation->IsPostProcessFrameBeforeWarpBlendRequired((uint32)ViewsPerViewport))
 			{
-				for (int ViewportIdx = 0; ViewportIdx < RenderViewportsRef.Num(); ++ViewportIdx)
+				for (int ViewIdx = 0; ViewIdx < ViewsPerViewport; ++ViewIdx)
 				{
-					for (int ViewIdx = 0; ViewIdx < ViewsPerViewport; ++ViewIdx)
-					{
-						UE_LOG(LogDisplayClusterRender, VeryVerbose, TEXT("Postprocess FRAME before WarpBlend - Viewport %s, ViewportIdx %d, ViewIdx"), *RenderViewportsRef[ViewportIdx].GetId(), ViewportIdx, ViewIdx);
-						CurPP.Operation->PerformPostProcessFrameBeforeWarpBlend_RenderThread(RHICmdList, SrcTexture, RenderViewportsRef[ViewportIdx].GetContext(ViewIdx).RenderTargetRect);
-					}
+					UE_LOG(LogDisplayClusterRender, VeryVerbose, TEXT("Postprocess FRAME before WarpBlend - ViewIdx %d"), ViewIdx);
+					CurPP.Operation->PerformPostProcessFrameBeforeWarpBlend_RenderThread(RHICmdList, SrcTexture, EyeRegions[ViewIdx]);
 				}
 			}
 		}
@@ -112,14 +109,8 @@ void FDisplayClusterDeviceBase_PostProcess::PerformPostProcessRenderTargetBefore
 		{
 			if (CurPP.Operation->IsPostProcessRenderTargetBeforeWarpBlendRequired())
 			{
-				for (int ViewportIdx = 0; ViewportIdx < RenderViewportsRef.Num(); ++ViewportIdx)
-				{
-					for (int ViewIdx = 0; ViewIdx < ViewsPerViewport; ++ViewIdx)
-					{
-						UE_LOG(LogDisplayClusterRender, VeryVerbose, TEXT("Postprocess TARGET before WarpBlend - Viewport %s, ViewportIdx %d, ViewIdx"), *RenderViewportsRef[ViewportIdx].GetId(), ViewportIdx, ViewIdx);
-						CurPP.Operation->PerformPostProcessRenderTargetBeforeWarpBlend_RenderThread(RHICmdList, SrcTexture);
-					}
-				}
+				UE_LOG(LogDisplayClusterRender, VeryVerbose, TEXT("Postprocess TARGET before WarpBlend"));
+				CurPP.Operation->PerformPostProcessRenderTargetBeforeWarpBlend_RenderThread(RHICmdList, SrcTexture);
 			}
 		}
 	}
@@ -190,13 +181,10 @@ void FDisplayClusterDeviceBase_PostProcess::PerformPostProcessFrameAfterWarpBlen
 		{
 			if (CurPP.Operation->IsPostProcessFrameAfterWarpBlendRequired((uint32)ViewsPerViewport))
 			{
-				for (int ViewportIdx = 0; ViewportIdx < RenderViewportsRef.Num(); ++ViewportIdx)
+				for (int ViewIdx = 0; ViewIdx < ViewsPerViewport; ++ViewIdx)
 				{
-					for (int ViewIdx = 0; ViewIdx < ViewsPerViewport; ++ViewIdx)
-					{
-						UE_LOG(LogDisplayClusterRender, VeryVerbose, TEXT("Postprocess FRAME after WarpBlend - Viewport %s, ViewportIdx %d, ViewIdx"), *RenderViewportsRef[ViewportIdx].GetId(), ViewportIdx, ViewIdx);
-						CurPP.Operation->PerformPostProcessFrameAfterWarpBlend_RenderThread(RHICmdList, SrcTexture, RenderViewportsRef[ViewportIdx].GetContext(ViewIdx).RenderTargetRect);
-					}
+					UE_LOG(LogDisplayClusterRender, Verbose, TEXT("Postprocess FRAME after WarpBlend - ViewIdx %d"), ViewIdx);
+					CurPP.Operation->PerformPostProcessFrameAfterWarpBlend_RenderThread(RHICmdList, SrcTexture, EyeRegions[ViewIdx]);
 				}
 			}
 		}
@@ -229,14 +217,8 @@ void FDisplayClusterDeviceBase_PostProcess::PerformPostProcessRenderTargetAfterW
 		{
 			if (CurPP.Operation->IsPostProcessRenderTargetAfterWarpBlendRequired())
 			{
-				for (int ViewportIdx = 0; ViewportIdx < RenderViewportsRef.Num(); ++ViewportIdx)
-				{
-					for (int ViewIdx = 0; ViewIdx < ViewsPerViewport; ++ViewIdx)
-					{
-						UE_LOG(LogDisplayClusterRender, VeryVerbose, TEXT("Postprocess TARGET after WarpBlend - Viewport %s, ViewportIdx %d, ViewIdx"), *RenderViewportsRef[ViewportIdx].GetId(), ViewportIdx, ViewIdx);
-						CurPP.Operation->PerformPostProcessRenderTargetAfterWarpBlend_RenderThread(RHICmdList, SrcTexture);
-					}
-				}
+				UE_LOG(LogDisplayClusterRender, Verbose, TEXT("Postprocess TARGET after WarpBlend"));
+				CurPP.Operation->PerformPostProcessRenderTargetAfterWarpBlend_RenderThread(RHICmdList, SrcTexture);
 			}
 		}
 	}
diff --git a/Engine/Plugins/Runtime/nDisplay/Source/DisplayCluster/Private/Render/Device/DisplayClusterDeviceBase_PostProcess.h b/Engine/Plugins/Runtime/nDisplay/Source/DisplayCluster/Private/Render/Device/DisplayClusterDeviceBase_PostProcess.h
--- a/Engine/Plugins/Runtime/nDisplay/Source/DisplayCluster/Private/Render/Device/DisplayClusterDeviceBase_PostProcess.h
+++ b/Engine/Plugins/Runtime/nDisplay/Source/DisplayCluster/Private/Render/Device/DisplayClusterDeviceBase_PostProcess.h
@@ -14,11 +14,13 @@ class FDisplayClusterDeviceBase_PostProcess
 	: public IDisplayClusterPostProcess
 {
 public:
-	FDisplayClusterDeviceBase_PostProcess(TArray<FDisplayClusterRenderViewport>& InRenderViewports, int InViewsPerViewport)
+	FDisplayClusterDeviceBase_PostProcess(TArray<FDisplayClusterRenderViewport>& InRenderViewports, int InViewsPerViewport, const FIntRect* const InEyeRegions)
 		: RenderViewportsRef(InRenderViewports)
 		, ViewsPerViewport(InViewsPerViewport)
+		, EyeRegions(InEyeRegions)
 	{
 		check(ViewsPerViewport > 0);
+		check(EyeRegions);
 	}
 
 	virtual ~FDisplayClusterDeviceBase_PostProcess() = default;
@@ -37,6 +39,7 @@ protected:
 private:
 	TArray<FDisplayClusterRenderViewport>& RenderViewportsRef;
 	const int ViewsPerViewport;
+	const FIntRect* EyeRegions;
 
 private:
 	virtual bool IsPostProcessViewBeforeWarpBlendRequired() override final
diff --git a/Engine/Plugins/Runtime/nDisplay/Source/DisplayCluster/Private/Render/Device/QuadBufferStereo/DisplayClusterDeviceQuadBufferStereoBase.cpp b/Engine/Plugins/Runtime/nDisplay/Source/DisplayCluster/Private/Render/Device/QuadBufferStereo/DisplayClusterDeviceQuadBufferStereoBase.cpp
--- a/Engine/Plugins/Runtime/nDisplay/Source/DisplayCluster/Private/Render/Device/QuadBufferStereo/DisplayClusterDeviceQuadBufferStereoBase.cpp
+++ b/Engine/Plugins/Runtime/nDisplay/Source/DisplayCluster/Private/Render/Device/QuadBufferStereo/DisplayClusterDeviceQuadBufferStereoBase.cpp
@@ -24,9 +24,16 @@ void FDisplayClusterDeviceQuadBufferStereoBase::CalculateRenderTargetSize(const
 	
 	check(IsInGameThread());
 
+	// Call base to calculate single-view RT size
+	FDisplayClusterDeviceBase::CalculateRenderTargetSize(Viewport, InOutSizeX, InOutSizeY);
+	
+	// And make it twice bigger for second eye
 	InOutSizeX = Viewport.GetSizeXY().X * 2;
 	InOutSizeY = Viewport.GetSizeXY().Y;
 
+	// Store right eye region
+	EyeRegions[1] = FIntRect(FIntPoint(InOutSizeX / 2, 0), FIntPoint(InOutSizeX, InOutSizeY));
+
 	UE_LOG(LogDisplayClusterRender, Verbose, TEXT("Render target size: [%d x %d]"), InOutSizeX, InOutSizeY);
 
 	check(InOutSizeX > 0 && InOutSizeY > 0);
@@ -71,21 +78,22 @@ void FDisplayClusterDeviceQuadBufferStereoBase::CopyTextureToBackBuffer_RenderTh
 	check(IsInRenderingThread());
 
 	const FIntPoint SrcSize = SrcTexture->GetSizeXY();
-
+	const FIntPoint BBSize  = BackBuffer->GetSizeXY();
+	
 	// Calculate sub regions to copy
-	const int halfSizeX = SrcSize.X / 2;
+	const int HalfSrcSizeX = SrcSize.X / 2;
 
 	FResolveParams copyParamsLeft;
 	copyParamsLeft.DestArrayIndex = 0;
 	copyParamsLeft.SourceArrayIndex = 0;
 	copyParamsLeft.Rect.X1 = 0;
 	copyParamsLeft.Rect.Y1 = 0;
-	copyParamsLeft.Rect.X2 = halfSizeX;
-	copyParamsLeft.Rect.Y2 = SrcSize.Y;
+	copyParamsLeft.Rect.X2 = HalfSrcSizeX;
+	copyParamsLeft.Rect.Y2 = BBSize.Y;
 	copyParamsLeft.DestRect.X1 = 0;
 	copyParamsLeft.DestRect.Y1 = 0;
-	copyParamsLeft.DestRect.X2 = halfSizeX;
-	copyParamsLeft.DestRect.Y2 = SrcSize.Y;
+	copyParamsLeft.DestRect.X2 = HalfSrcSizeX;
+	copyParamsLeft.DestRect.Y2 = BBSize.Y;
 
 	UE_LOG(LogDisplayClusterRender, Verbose, TEXT("CopyToResolveTarget [L]: [%d,%d - %d,%d] -> [%d,%d - %d,%d]"),
 		copyParamsLeft.Rect.X1, copyParamsLeft.Rect.Y1, copyParamsLeft.Rect.X2, copyParamsLeft.Rect.Y2,
@@ -98,7 +106,7 @@ void FDisplayClusterDeviceQuadBufferStereoBase::CopyTextureToBackBuffer_RenderTh
 	copyParamsRight.SourceArrayIndex = 0;
 
 	copyParamsRight.Rect = copyParamsLeft.Rect;
-	copyParamsRight.Rect.X1 = halfSizeX;
+	copyParamsRight.Rect.X1 = HalfSrcSizeX;
 	copyParamsRight.Rect.X2 = SrcSize.X;
 
 	copyParamsRight.DestRect = copyParamsLeft.DestRect;
diff --git a/Engine/Plugins/Runtime/nDisplay/Source/OutputRemap/Private/OutputRemapMesh.cpp b/Engine/Plugins/Runtime/nDisplay/Source/OutputRemap/Private/OutputRemapMesh.cpp
--- a/Engine/Plugins/Runtime/nDisplay/Source/OutputRemap/Private/OutputRemapMesh.cpp
+++ b/Engine/Plugins/Runtime/nDisplay/Source/OutputRemap/Private/OutputRemapMesh.cpp
@@ -20,13 +20,43 @@
 static int VarOutputRemapReloadChangedExtFilesFrame = 0;
 static TAutoConsoleVariable<int32> CVarOutputRemapReloadChangedExtFiles(
 	TEXT("nDisplay.reload.output_remap"),
-	(int) 0, // Disabled by default
+	(int) 0,
 	TEXT("Changed external files reload period for output remap:\n")
 	TEXT("0 : disable\n"),
 	ECVF_RenderThreadSafe
 );
 
 
+namespace ObjMeshStrings
+{
+	static constexpr auto Vertex = TEXT("v ");
+	static constexpr auto UV = TEXT("vt ");
+	static constexpr auto Face = TEXT("f ");
+
+	namespace delims
+	{
+		static constexpr auto Values = TEXT(" ");
+		static constexpr auto Face = TEXT("/");
+	}
+};
+
+
+FOutputRemapMesh::FOutputRemapMesh(const FString& InExtFileName)
+	: NumVertices(0)
+	, NumIndices(0)
+	, NumTriangles(0)
+	, bIsValid(false)
+	, FileName(InExtFileName)
+	, bIsExtFilesReloadEnabled(false)
+	, bForceExtFilesReload(false)
+{
+	CreateMesh();
+}
+
+FOutputRemapMesh::~FOutputRemapMesh()
+{
+	ReleaseMesh();
+}
 
 void FOutputRemapMesh::DrawMesh(FRHICommandList& RHICmdList)
 {
@@ -46,22 +76,9 @@ void FOutputRemapMesh::DrawMesh(FRHICommandList& RHICmdList)
 	}
 }
 
-namespace ObjMeshStrings
-{
-	static constexpr auto Vertex = TEXT("v ");
-	static constexpr auto UV = TEXT("vt ");
-	static constexpr auto Face = TEXT("f ");
-
-	namespace delims
-	{
-		static constexpr auto Values = TEXT(" ");
-		static constexpr auto Face = TEXT("/");
-	}
-};
-
 bool FOutputRemapMesh::FMeshData::CreateFromFile(const FString& FullPathFileName)
 {	
-	//@load from obj file at runtime:
+	// Load from obj file at runtime:
 	if (FPaths::FileExists(FullPathFileName))
 	{
 		TArray<FString> data;
@@ -70,18 +87,20 @@ bool FOutputRemapMesh::FMeshData::CreateFromFile(const FString& FullPathFileName
 			bool bResult = true;
 
 			// Parse each line from config
-			for (auto line : data)
+			for (FString Line : data)
 			{
-				line.TrimStartAndEndInline();
-				if (!ParseLine(line))
+				Line.TrimStartAndEndInline();
+				if (!ParseLine(Line))
 				{
 					// Handle error
 					bResult = false;
 				}
 			}
+
 			return bResult;
 		}
 	}
+
 	return false;
 }
 
@@ -91,49 +110,54 @@ bool FOutputRemapMesh::FMeshData::ParseLine(const FString& Line)
 	{
 		return ExtractVertex(Line);
 	}
-
-	if (Line.StartsWith(ObjMeshStrings::UV))
+	else if (Line.StartsWith(ObjMeshStrings::UV))
 	{
 		return ExtractUV(Line);
 	}
-
-	if (Line.StartsWith(ObjMeshStrings::Face))
+	else if (Line.StartsWith(ObjMeshStrings::Face))
 	{
 		return ExtractFace(Line);
 	}
 
-	//skip
 	return true;
 }
 
 bool FOutputRemapMesh::FMeshData::ExtractVertex(const FString& Line)
 {
 	TArray<FString> Data;
-	if (Line.ParseIntoArray(Data, ObjMeshStrings::delims::Values)==4)
+
+	if (Line.ParseIntoArray(Data, ObjMeshStrings::delims::Values) == 4)
 	{
 		float X = FCString::Atof(*Data[1]);
 		float Y = FCString::Atof(*Data[2]);
 		float Z = FCString::Atof(*Data[3]);
+
 		InVertex.Add(FVector(X, Y, Z));
+
 		return true;
 	}
-	//! handle error
+
 	return false;
 }
+
 bool FOutputRemapMesh::FMeshData::ExtractUV(const FString& Line)
 {
 	TArray<FString> Data;
+
 	if (Line.ParseIntoArray(Data, ObjMeshStrings::delims::Values) == 4)
 	{
 		float U = FCString::Atof(*Data[1]);
 		float V = FCString::Atof(*Data[2]);
 		float W = FCString::Atof(*Data[3]);
+
 		InUV.Add(FVector(U, V, W));
+
 		return true;
 	}
-	//! handle error
+
 	return false;
 }
+
 bool FOutputRemapMesh::FMeshData::ExtractFace(const FString& Line)
 {
 	TArray<FString> Data;
@@ -141,90 +165,94 @@ bool FOutputRemapMesh::FMeshData::ExtractFace(const FString& Line)
 
 	if (Data.Num() > 3)
 	{
-		int IndexBase = Vertex.Num();
-		if (ExtractFaceVertex(Data[1])
-			&& ExtractFaceVertex(Data[2])
-			&& ExtractFaceVertex(Data[3])
-			)
+		TArray<int> FaceIndices;
+		for (int i = 1; i < Data.Num(); ++i)
 		{
-			// Triangle
-			Index.Add(IndexBase + 0);
-			Index.Add(IndexBase + 1);
-			Index.Add(IndexBase + 2);
+			if (ExtractFaceVertex(Data[i]))
+			{				
+				const int IdxC = Vertices.Num() - 1;
 
-			if (Data.Num() > 4)
-			{
-				// Quad face
-				if (!ExtractFaceVertex(Data[4]))
+				if (FaceIndices.Num() > 2)
 				{
-					//! handle error
-					return false;
-				}
+					const int IdxA = FaceIndices[0];
+					const int IdxB = FaceIndices.Last();
 
-				Index.Add(IndexBase + 3);
-				Index.Add(IndexBase + 0);
-				Index.Add(IndexBase + 2);
+					FaceIndices.Add(IdxA);
+					FaceIndices.Add(IdxB);
+					FaceIndices.Add(IdxC);
+				}
+				else
+				{
+					FaceIndices.Add(IdxC);
+				}
+			}
+			else
+			{
+				return false;
 			}
-			return true;
 		}
+
+		Indices.Append(FaceIndices);
+
+		return true;
 	}
-	//! handle error
+
 	return false;
 }
+
 bool FOutputRemapMesh::FMeshData::ExtractFaceVertex(const FString& Line)
 {
 	TArray<FString> Data;
 	if (Line.ParseIntoArray(Data, ObjMeshStrings::delims::Face) > 1)
 	{
-		int InVertexIndex = FCString::Atoi(*Data[0]) - 1;
-		int InUVIndex = FCString::Atoi(*Data[1]) - 1;
+		const int InVertexIndex = FCString::Atoi(*Data[0]) - 1;
+		const int InUVIndex = FCString::Atoi(*Data[1]) - 1;
 
-		Vertex.Add(InVertex[InVertexIndex]);
-		UV.Add(InUV[InUVIndex]);
+		FMeshData::FVertexData VertexData(InVertex[InVertexIndex], InUV[InUVIndex]);
+		Vertices.Add(VertexData);
+		
 		return true;
 	}
 
-	//! handle error
 	return false;
 }
 
 bool FOutputRemapMesh::FMeshData::CreatePassthrough()
 {
-	//@create dummy Passthrough mesh (test purpose)
-	
-	Vertex.Add(FVector(2.6f, 0, 1));
-	Vertex.Add(FVector(2.6f, 0, 2));
-	Vertex.Add(FVector(5.5f, 0, 2));
-	Vertex.Add(FVector(5.5f, 0, 1));
+	Vertices.Add(FMeshData::FVertexData(FVector(2.6f, 1, 0), FVector(0, 1, 0)));
+	Vertices.Add(FMeshData::FVertexData(FVector(2.6f, 2, 0), FVector(0, 0, 0)));
+	Vertices.Add(FMeshData::FVertexData(FVector(5.5f, 2, 0), FVector(1, 0, 0)));
+	Vertices.Add(FMeshData::FVertexData(FVector(5.5f, 1, 0), FVector(1, 1, 0)));
 
-	UV.Add(FVector(0, 1, 0));
-	UV.Add(FVector(0, 0, 0));
-	UV.Add(FVector(1, 0, 0));
-	UV.Add(FVector(1, 1, 0));
+	for (const auto& It : Vertices)
+	{
+		InVertex.Add(It.Pos);
+	}
 
-	Index.Add(0);
-	Index.Add(1);
-	Index.Add(2);
+	Indices.Add(0);
+	Indices.Add(1);
+	Indices.Add(2);
 
-	Index.Add(3);
-	Index.Add(0);
-	Index.Add(2);
+	Indices.Add(3);
+	Indices.Add(0);
+	Indices.Add(2);
 
 	return true;
 }
+
 void FOutputRemapMesh::FMeshData::Normalize()
 {
 	FBox AABBox(FVector(FLT_MAX, FLT_MAX, FLT_MAX), FVector(-FLT_MAX, -FLT_MAX, -FLT_MAX));
 	
-	for (const auto& Pts: Vertex)
+	for (const FVector& It: InVertex)
 	{
-		AABBox.Min.X = FMath::Min(AABBox.Min.X, Pts.X);
-		AABBox.Min.Y = FMath::Min(AABBox.Min.Y, Pts.Y);
-		AABBox.Min.Z = FMath::Min(AABBox.Min.Z, Pts.Z);
+		AABBox.Min.X = FMath::Min(AABBox.Min.X, It.X);
+		AABBox.Min.Y = FMath::Min(AABBox.Min.Y, It.Y);
+		AABBox.Min.Z = FMath::Min(AABBox.Min.Z, It.Z);
 
-		AABBox.Max.X = FMath::Max(AABBox.Max.X, Pts.X);
-		AABBox.Max.Y = FMath::Max(AABBox.Max.Y, Pts.Y);
-		AABBox.Max.Z = FMath::Max(AABBox.Max.Z, Pts.Z);
+		AABBox.Max.X = FMath::Max(AABBox.Max.X, It.X);
+		AABBox.Max.Y = FMath::Max(AABBox.Max.Y, It.Y);
+		AABBox.Max.Z = FMath::Max(AABBox.Max.Z, It.Z);
 	}
 
 	//Normalize
@@ -234,49 +262,69 @@ void FOutputRemapMesh::FMeshData::Normalize()
 		(AABBox.Max.Z - AABBox.Min.Z)
 	);
 
+	// Detect axis aligned plane
+	const bool bHelperSwapYZ = fabs(Size.Y) > fabs(Size.Z);
+	
 	if (Size.X == 0) { Size.X = 1; }
 	if (Size.Y == 0) { Size.Y = 1; }
 	if (Size.Z == 0) { Size.Z = 1; }
 
 	FVector Scale(1/Size.X, 1 / Size.Y, 1 / Size.Z);
 	
-	for (auto& Pts : Vertex)
+	for (auto& It : Vertices)
 	{
-		float X = (Pts.X - AABBox.Min.X) * Scale.X;
-		float Y = (Pts.Y - AABBox.Min.Y) * Scale.Y;
-		float Z = (Pts.Z - AABBox.Min.Z) * Scale.Z;
+		float X = (It.Pos.X - AABBox.Min.X) * Scale.X;
+		float Y = (It.Pos.Y - AABBox.Min.Y) * Scale.Y;
+		float Z = (It.Pos.Z - AABBox.Min.Z) * Scale.Z;
+
+		if (bHelperSwapYZ)
+		{
+			It.Pos = FVector(X, Y, Z);
+		}
+		else
+		{
+			It.Pos = FVector(X, Z, Y);
+		}
 
-		Pts = FVector(X,Y,Z);
+		// Apply axis wrap
+		It.UV.Y = 1 - It.UV.Y;
 	}
 }
+
 void FOutputRemapMesh::FMeshData::RemoveInvisibleFaces()
 {
-	int FacesNum = Index.Num() / 3;
-	for (int Face = 0; Face < FacesNum; Face++)
+	const int FacesNum = Indices.Num() / 3;
+	for (int Face = 0; Face < FacesNum; ++Face)
 	{
-		if (!IsVisibleFace(Face))
+		if (!IsFaceVisible(Face))
 		{
-			Index.RemoveAt(Face * 3, 3, true);
-			Face--;
+			Indices.RemoveAt(Face * 3, 3, false);
 		}
 	}
+
+	Indices.Shrink();
 }
-bool FOutputRemapMesh::FMeshData::IsVisibleFace(int Face)
+
+bool FOutputRemapMesh::FMeshData::IsFaceVisible(int Face)
 {
-	int idx0 = Index[Face * 3 + 0];
-	int idx1 = Index[Face * 3 + 1];
-	int idx2 = Index[Face * 3 + 2];
+	const int idx0 = Indices[Face * 3 + 0];
+	const int idx1 = Indices[Face * 3 + 1];
+	const int idx2 = Indices[Face * 3 + 2];
 
-	return IsVisibleUV(idx0) && IsVisibleUV(idx1) && IsVisibleUV(idx2);
+	return IsUVVisible(idx0) && IsUVVisible(idx1) && IsUVVisible(idx2);
 }
 
-bool FOutputRemapMesh::FMeshData::IsVisibleUV(int UVIndex)
+bool FOutputRemapMesh::FMeshData::IsUVVisible(int UVIndex)
 {
-	if (UV[UVIndex].X < 0 || UV[UVIndex].X > 1)
+	if (Vertices[UVIndex].UV.X < 0 || Vertices[UVIndex].UV.X > 1)
+	{
 		return false;
+	}
 
-	if (UV[UVIndex].Y < 0 || UV[UVIndex].Y > 1)
+	if (Vertices[UVIndex].UV.Y < 0 || Vertices[UVIndex].UV.Y > 1)
+	{
 		return false;
+	}
 
 	return true;
 }
@@ -295,12 +343,11 @@ void FOutputRemapMesh::ReleaseMesh()
 
 void FOutputRemapMesh::CreateMesh()
 {
-
 	FMeshData* Mesh = new FMeshData();
 
 	if (FileName == "Passthrough")
 	{
-		//autogenerate test mesh:
+		// Auto-generate test mesh
 		bIsValid = Mesh->CreatePassthrough();
 	}
 	else
@@ -324,6 +371,7 @@ void FOutputRemapMesh::CreateMesh()
 	else
 	{
 		delete Mesh;
+		Mesh = nullptr;
 	}
 }
 
@@ -346,48 +394,30 @@ void FOutputRemapMesh::ReloadMesh()
 		}
 
 		if(bForceExtFilesReload)
-		{				
+		{
 			// Check for file data is modified
-			FString FullPathFileName = DisplayClusterHelpers::config::GetFullPath(FileName);
-			FDateTime CurrentDateTime = IFileManager::Get().GetAccessTimeStamp(*FullPathFileName);
+			const FString FullPathFileName = DisplayClusterHelpers::config::GetFullPath(FileName);
+			const FDateTime CurrentDateTime = IFileManager::Get().GetAccessTimeStamp(*FullPathFileName);
+
 			if (CurrentDateTime != FileLastAccessDateTime)
 			{
-				//! Add log. Reload file *FullPathFileName
 				ReleaseMesh();
 				CreateMesh();
 			}
 		}
 	}
-	bForceExtFilesReload = false;
-}
-
-
-FOutputRemapMesh::FOutputRemapMesh(const FString& InExtFileName)	
-	: bIsValid(false)
-	, FileName(InExtFileName)
-	, bIsExtFilesReloadEnabled(false)
-	, bForceExtFilesReload(false)
-	, NumVertices(0)
-	, NumIndices(0)
-	, NumTriangles(0)
-{
-	CreateMesh();
-}
 
-FOutputRemapMesh::~FOutputRemapMesh()
-{
-	ReleaseMesh();
+	bForceExtFilesReload = false;
 }
 
 static void BuildMesh_RenderThread(FOutputRemapMesh& Dst, const FOutputRemapMesh::FMeshData* Mesh)
 {
-	Dst.NumVertices = Mesh->Vertex.Num();
-	Dst.NumTriangles = Mesh->Index.Num() / 3;
-	Dst.NumIndices = Mesh->Index.Num();
+	Dst.NumVertices  = Mesh->Vertices.Num();
+	Dst.NumTriangles = Mesh->Indices.Num() / 3;
+	Dst.NumIndices   = Mesh->Indices.Num();
 
 	check(Dst.NumVertices > 2 && Dst.NumTriangles > 0);
 
-
 	FRHIResourceCreateInfo CreateInfo;
 	Dst.VertexBufferRHI = RHICreateVertexBuffer(sizeof(FFilterVertex) * Dst.NumVertices, BUF_Static, CreateInfo);
 	void* VoidPtr = RHILockVertexBuffer(Dst.VertexBufferRHI, 0, sizeof(FFilterVertex) * Dst.NumVertices, RLM_WriteOnly);
@@ -396,31 +426,25 @@ static void BuildMesh_RenderThread(FOutputRemapMesh& Dst, const FOutputRemapMesh
 	Dst.IndexBufferRHI = RHICreateIndexBuffer(sizeof(uint16), sizeof(uint16) * Dst.NumIndices, BUF_Static, CreateInfo);
 	void* VoidPtr2 = RHILockIndexBuffer(Dst.IndexBufferRHI, 0, sizeof(uint16) * Dst.NumIndices, RLM_WriteOnly);
 	uint16* pIndices = reinterpret_cast<uint16*>(VoidPtr2);
-
-	//Copy Vertex
+	
+	// Copy vertices
 	int DataIndex = 0;
-	for (auto& It : Mesh->Vertex)
+	for (const FOutputRemapMesh::FMeshData::FVertexData& It : Mesh->Vertices)
 	{
 		FFilterVertex& Vertex = pVertices[DataIndex++];
-		// Remap from to NDC space [0 1] -> [-1 1]
-		Vertex.Position.X = It.X;
-		Vertex.Position.Y = 1-It.Z;
+
+		Vertex.Position.X = It.Pos.X;
+		Vertex.Position.Y = It.Pos.Y;
 		Vertex.Position.Z = 1.0f;
 		Vertex.Position.W = 1.0f;
-	}
 
-	//Copy UV
-	DataIndex = 0;
-	for (auto& It : Mesh->UV)
-	{
-		FFilterVertex& Vertex = pVertices[DataIndex++];		
-		Vertex.UV.X = It.X;
-		Vertex.UV.Y = It.Y;
+		Vertex.UV.X = It.UV.X;
+		Vertex.UV.Y = It.UV.Y;
 	}
-
-	//Copy faces Index
+	
+	// Copy face indices
 	DataIndex = 0;
-	for (auto& It : Mesh->Index)
+	for (const auto& It : Mesh->Indices)
 	{
 		pIndices[DataIndex++] = (uint16)It;
 	}
@@ -428,7 +452,9 @@ static void BuildMesh_RenderThread(FOutputRemapMesh& Dst, const FOutputRemapMesh
 	RHIUnlockVertexBuffer(Dst.VertexBufferRHI);
 	RHIUnlockIndexBuffer(Dst.IndexBufferRHI);
 
-	delete Mesh; // Now free tempropary mesh data
+	// Now free temporary mesh data
+	delete Mesh;
+	Mesh = nullptr;
 }
 
 void FOutputRemapMesh::BuildMesh(const FMeshData* Mesh)
diff --git a/Engine/Plugins/Runtime/nDisplay/Source/OutputRemap/Private/OutputRemapMesh.h b/Engine/Plugins/Runtime/nDisplay/Source/OutputRemap/Private/OutputRemapMesh.h
--- a/Engine/Plugins/Runtime/nDisplay/Source/OutputRemap/Private/OutputRemapMesh.h
+++ b/Engine/Plugins/Runtime/nDisplay/Source/OutputRemap/Private/OutputRemapMesh.h
@@ -6,24 +6,34 @@
 #include "RHI.h"
 #include "RHICommandList.h"
 
+
 class FOutputRemapMesh
 {
 public:
 	struct FMeshData
 	{
-		TArray<FVector> Vertex; 
-		TArray<FVector> UV;
-		TArray<int>     Index;
+		struct FVertexData
+		{
+			FVector Pos;
+			FVector UV;
+			FVertexData(const FVector &InPos, const FVector &InUV)
+				: Pos(InPos), UV(InUV)
+			{ }
+		};
+
+		TArray<FVertexData> Vertices;
+		TArray<int>         Indices;
 
 		bool CreateFromFile(const FString& FullPathFileName);
 		bool CreatePassthrough();
 
 		void Normalize();
 		void RemoveInvisibleFaces();
+
 	private:
 		// Mesh post-op's
-		bool IsVisibleFace(int Face);
-		bool IsVisibleUV(int UVIndex);
+		bool IsFaceVisible(int Face);
+		bool IsUVVisible(int UVIndex);
 
 		// Obj file parser:
 		bool ParseLine(const FString& Line);
@@ -40,35 +50,23 @@ public:
 	FOutputRemapMesh(const FString& ExtFileName);
 	~FOutputRemapMesh();
 	
+public:
 	void DrawMesh(FRHICommandList& RHICmdList);
 
-	const FString& GetFileName() const 
-	{ 
-		return FileName; 
-	}
-	bool IsValid() const 
-	{ 
-		return bIsValid; 
-	}
+	const FString& GetFileName() const
+	{ return FileName; }
+
+	bool IsValid() const
+	{ return bIsValid; }
 
 	void ReloadChangedExtFiles()
-	{
-		bForceExtFilesReload = true;
-	}
+	{ bForceExtFilesReload = true; }
 
 private:
 	void BuildMesh(const FMeshData* Mesh);
-
 	void CreateMesh();
 	void ReleaseMesh();
 	void ReloadMesh();
-private:
-	bool       bIsValid;
-
-	FString    FileName;
-	FDateTime  FileLastAccessDateTime;
-	bool       bIsExtFilesReloadEnabled;
-	bool       bForceExtFilesReload;
 
 public:
 	FVertexBufferRHIRef VertexBufferRHI;
@@ -77,6 +75,12 @@ public:
 	uint32  NumVertices;
 	uint32  NumIndices;
 	uint32  NumTriangles;
-};
 
+private:
+	bool       bIsValid;
 
+	FString    FileName;
+	FDateTime  FileLastAccessDateTime;
+	bool       bIsExtFilesReloadEnabled;
+	bool       bForceExtFilesReload;
+};
diff --git a/Engine/Plugins/Runtime/nDisplay/Source/OutputRemap/Private/OutputRemapShader.cpp b/Engine/Plugins/Runtime/nDisplay/Source/OutputRemap/Private/OutputRemapShader.cpp
--- a/Engine/Plugins/Runtime/nDisplay/Source/OutputRemap/Private/OutputRemapShader.cpp
+++ b/Engine/Plugins/Runtime/nDisplay/Source/OutputRemap/Private/OutputRemapShader.cpp
@@ -17,7 +17,6 @@
 
 #define OutputRemapShaderFileName TEXT("/Plugin/nDisplay/Private/OutputRemapShaders.usf")
 
-
 // Select output remap shader
 enum class EVarOutputRemapShaderType : uint8
 {
@@ -41,16 +40,15 @@ class FOutputRemapVS : public FGlobalShader
 {
 	DECLARE_SHADER_TYPE(FOutputRemapVS, Global);
 
-	static bool ShouldCompilePermutation(const FGlobalShaderPermutationParameters& Parameters)
-	{
-		return true;
-	}
-
+public:
 	/** Default constructor. */
-	FOutputRemapVS() 
-	{
-	}
+	FOutputRemapVS()
+	{ }
+
 public:
+	static bool ShouldCompilePermutation(const FGlobalShaderPermutationParameters& Parameters)
+	{ return true; }
+
 	/** Initialization constructor. */
 	FOutputRemapVS(const ShaderMetaType::CompiledShaderInitializerType& Initializer)
 		: FGlobalShader(Initializer)
@@ -64,8 +62,7 @@ class FOutputRemapPS : public FGlobalShader
 
 public:
 	FOutputRemapPS()
-	{ 
-	}
+	{ }
 
 	FOutputRemapPS(const ShaderMetaType::CompiledShaderInitializerType& Initializer)
 		: FGlobalShader(Initializer)
@@ -113,75 +110,76 @@ private:
 IMPLEMENT_SHADER_TYPE(, FOutputRemapVS, OutputRemapShaderFileName, TEXT("OutputRemap_VS"), SF_Vertex);
 IMPLEMENT_SHADER_TYPE(, FOutputRemapPS, OutputRemapShaderFileName, TEXT("OutputRemap_PS"), SF_Pixel);
 
+DECLARE_GPU_STAT_NAMED(DisplayClusterOutputRemap, TEXT("DisplayCluster PP OutputRemap"));
+
 bool FOutputRemapShader::ApplyOutputRemap_RenderThread(FRHICommandListImmediate& RHICmdList, FRHITexture2D* ShaderResourceTexture, FRHITexture2D* TargetableTexture, FOutputRemapMesh* MeshData)
 {
 	check(IsInRenderingThread());
 
 	if (MeshData == nullptr)
 	{
-		//Handle error
 		return false;
 	}
 
-
 	const EVarOutputRemapShaderType ShaderType = (EVarOutputRemapShaderType)CVarOutputRemapShaderType.GetValueOnAnyThread();
 	switch (ShaderType)
 	{
-	case EVarOutputRemapShaderType::Passthrough:
-	{
-		// Use simple 1:1 test mesh for shader forwarding
-		static FOutputRemapMesh TestMesh("Passthrough");
-		MeshData = &TestMesh;
-	}
-	case EVarOutputRemapShaderType::Default:
-		break;
-	case EVarOutputRemapShaderType::Disable:
-		return false;
+		case EVarOutputRemapShaderType::Passthrough:
+		{
+			// Use simple 1:1 test mesh for shader forwarding
+			static FOutputRemapMesh TestMesh("Passthrough");
+			MeshData = &TestMesh;
+			break;
+		}
 
-	default:
-		return false;
+		case EVarOutputRemapShaderType::Default:
+			break;
+
+		case EVarOutputRemapShaderType::Disable:
+			return false;
+
+		default:
+			return false;
 	};
 
+	RHICmdList.ImmediateFlush(EImmediateFlushType::FlushRHIThreadFlushResources);
 
-	{// Do remap single render pass		
-		FRHIRenderPassInfo RPInfo(TargetableTexture, ERenderTargetActions::Load_Store);
-		RHICmdList.BeginRenderPass(RPInfo, TEXT("DisplayClusterOutputRemapShader"));
-		{
-			{
-				FIntRect DstRect(FIntPoint(0,0), TargetableTexture->GetSizeXY());
-				//Clear viewport before render
-				RHICmdList.SetViewport(DstRect.Min.X, DstRect.Min.Y, 0.0f, DstRect.Max.X, DstRect.Max.Y, 1.0f);
-				//DrawClearQuad(RHICmdList, FLinearColor::Black);
-			}
+	FIntRect DstRect(FIntPoint(0, 0), TargetableTexture->GetSizeXY());
+
+	SCOPED_GPU_STAT(RHICmdList, DisplayClusterOutputRemap);
+	SCOPED_DRAW_EVENTF(RHICmdList, DisplayClusterOutputRemap, TEXT("DisplayClusterOutputRemap"));
 
-			// Set the graphic pipeline state.
-			FGraphicsPipelineStateInitializer GraphicsPSOInit;
-			RHICmdList.ApplyCachedRenderTargets(GraphicsPSOInit);
+	// Single render pass remap
+	FRHIRenderPassInfo RPInfo(TargetableTexture, ERenderTargetActions::Clear_Store);
+	RHICmdList.BeginRenderPass(RPInfo, TEXT("DisplayClusterOutputRemap"));
+	{
+		// Set the graphic pipeline state.
+		FGraphicsPipelineStateInitializer GraphicsPSOInit;
+		RHICmdList.ApplyCachedRenderTargets(GraphicsPSOInit);
 
-			GraphicsPSOInit.DepthStencilState = TStaticDepthStencilState<false, CF_Never>::GetRHI();
-			GraphicsPSOInit.BlendState = TStaticBlendState <>::GetRHI();
-			GraphicsPSOInit.RasterizerState = TStaticRasterizerState<>::GetRHI();
+		RHICmdList.SetViewport(0, 0, 0.0f, DstRect.Max.X, DstRect.Max.Y, 1.0f);
 
-			TShaderMap<FGlobalShaderType>* ShaderMap = GetGlobalShaderMap(GMaxRHIFeatureLevel);
-			TShaderMapRef<FOutputRemapVS> VertexShader(ShaderMap);
-			TShaderMapRef<FOutputRemapPS> PixelShader(ShaderMap);
+		GraphicsPSOInit.DepthStencilState = TStaticDepthStencilState<false, CF_Never>::GetRHI();
+		GraphicsPSOInit.BlendState = TStaticBlendState <>::GetRHI();
+		GraphicsPSOInit.RasterizerState = TStaticRasterizerState<>::GetRHI();
 
-			GraphicsPSOInit.BoundShaderState.VertexDeclarationRHI = GFilterVertexDeclaration.VertexDeclarationRHI;// GetVertexDeclarationFVector4();
-			GraphicsPSOInit.BoundShaderState.VertexShaderRHI = GETSAFERHISHADER_VERTEX(*VertexShader);
-			GraphicsPSOInit.BoundShaderState.PixelShaderRHI = GETSAFERHISHADER_PIXEL(*PixelShader);
-			GraphicsPSOInit.PrimitiveType = PT_TriangleList;
+		// Grab shaders
+		TShaderMap<FGlobalShaderType>* ShaderMap = GetGlobalShaderMap(GMaxRHIFeatureLevel);
+		TShaderMapRef<FOutputRemapVS> VertexShader(ShaderMap);
+		TShaderMapRef<FOutputRemapPS> PixelShader(ShaderMap);
 
+		GraphicsPSOInit.BoundShaderState.VertexDeclarationRHI = GFilterVertexDeclaration.VertexDeclarationRHI;
+		GraphicsPSOInit.BoundShaderState.VertexShaderRHI = GETSAFERHISHADER_VERTEX(*VertexShader);
+		GraphicsPSOInit.BoundShaderState.PixelShaderRHI = GETSAFERHISHADER_PIXEL(*PixelShader);
+		GraphicsPSOInit.PrimitiveType = PT_TriangleList;
 
-			SetGraphicsPipelineState(RHICmdList, GraphicsPSOInit);
+		SetGraphicsPipelineState(RHICmdList, GraphicsPSOInit);
 
-			PixelShader->SetParameters(RHICmdList, PixelShader->GetPixelShader(), ShaderResourceTexture);
-			MeshData->DrawMesh(RHICmdList);
-		}
-		RHICmdList.EndRenderPass();
+		PixelShader->SetParameters(RHICmdList, PixelShader->GetPixelShader(), ShaderResourceTexture);
+		MeshData->DrawMesh(RHICmdList);
 	}
 
-	// Render pass failed, handle error
+	RHICmdList.EndRenderPass();
+
 	return true;
 }
-
-
