Fixing missed variable name change in RoboMerge resolve

diff --git a/Engine/Plugins/Runtime/nDisplay/Source/DisplayCluster/Private/Render/Devices/DisplayClusterDeviceBase.cpp b/Engine/Plugins/Runtime/nDisplay/Source/DisplayCluster/Private/Render/Devices/DisplayClusterDeviceBase.cpp
--- a/Engine/Plugins/Runtime/nDisplay/Source/DisplayCluster/Private/Render/Devices/DisplayClusterDeviceBase.cpp
+++ b/Engine/Plugins/Runtime/nDisplay/Source/DisplayCluster/Private/Render/Devices/DisplayClusterDeviceBase.cpp
@@ -417,7 +417,7 @@ uint32 FDisplayClusterDeviceBase::GetViewIndexForPass(EStereoscopicPass StereoPa
 		break;
 
 	default:
-		DecodedViewIndex = (int(StereoPassType) - int(EStereoscopicPass::eSSP_MONOSCOPIC_EYE) + 1);
+		DecodedViewIndex = (int(StereoPassType) - int(EStereoscopicPass::eSSP_RIGHT_EYE_SIDE) + 1);
 		break;
 	}
 
