Fixed more deprecation warnings to keep CIS happy

diff --git a/Engine/Plugins/Runtime/nDisplay/Source/PicpMPCDI/Private/PicpBlurPostProcess.h b/Engine/Plugins/Runtime/nDisplay/Source/PicpMPCDI/Private/PicpBlurPostProcess.h
--- a/Engine/Plugins/Runtime/nDisplay/Source/PicpMPCDI/Private/PicpBlurPostProcess.h
+++ b/Engine/Plugins/Runtime/nDisplay/Source/PicpMPCDI/Private/PicpBlurPostProcess.h
@@ -103,7 +103,7 @@ public:
 
 	void SetParameters(FRHICommandListImmediate& RHICmdList, FTexture2DRHIParamRef SourceTexture, FVector2D SampleOffset, int KernelRadius)
 	{
-		const FPixelShaderRHIParamRef ShaderRHI = GetPixelShader();
+		FPixelShaderRHIParamRef* ShaderRHI = GetPixelShader();
 
 		SetSamplerParameter(RHICmdList, ShaderRHI, BilinearClampTextureSamplerParameter, TStaticSamplerState<SF_Bilinear, AM_Clamp, AM_Clamp, AM_Clamp>::GetRHI());
 		SetTextureParameter(RHICmdList, ShaderRHI, SrcTextureParameter, SourceTexture);
@@ -156,7 +156,7 @@ public:
 
 	void SetParameters(FRHICommandListImmediate& RHICmdList, FTextureRHIRef SourceTexture)
 	{
-		const FPixelShaderRHIParamRef ShaderRHI = GetPixelShader();
+		FPixelShaderRHIParamRef* ShaderRHI = GetPixelShader();
 
 		SetSamplerParameter(RHICmdList, ShaderRHI, BilinearClampTextureSamplerParameter, TStaticSamplerState<SF_Bilinear, AM_Clamp, AM_Clamp, AM_Clamp>::GetRHI());
 		SetTextureParameter(RHICmdList, ShaderRHI, SrcTextureParameter, SourceTexture);
