Get rid of DisplayCluster GameMode and Pawns

diff --git a/Engine/Plugins/Runtime/nDisplay/Source/DisplayCluster/Private/Blueprints/DisplayClusterBlueprintAPIImpl.cpp b/Engine/Plugins/Runtime/nDisplay/Source/DisplayCluster/Private/Blueprints/DisplayClusterBlueprintAPIImpl.cpp
--- a/Engine/Plugins/Runtime/nDisplay/Source/DisplayCluster/Private/Blueprints/DisplayClusterBlueprintAPIImpl.cpp
+++ b/Engine/Plugins/Runtime/nDisplay/Source/DisplayCluster/Private/Blueprints/DisplayClusterBlueprintAPIImpl.cpp
@@ -12,6 +12,8 @@
 #include "Input/IDisplayClusterInputManager.h"
 #include "Render/IDisplayClusterRenderManager.h"
 
+#include "DisplayClusterRootComponent.h"
+
 #include "Config/DisplayClusterConfigTypes.h"
 #include "Misc/DisplayClusterHelpers.h"
 
@@ -167,7 +169,7 @@ void UDisplayClusterBlueprintAPIImpl::GetLocalViewports(bool IsRTT, TArray<FStri
 // Game API
 //////////////////////////////////////////////////////////////////////////////////////////////
 // Root
-ADisplayClusterPawn* UDisplayClusterBlueprintAPIImpl::GetRoot()
+UDisplayClusterRootComponent* UDisplayClusterBlueprintAPIImpl::GetRoot()
 {
 	DISPLAY_CLUSTER_FUNC_TRACE(LogDisplayClusterBlueprint);
 
@@ -273,17 +275,6 @@ UDisplayClusterCameraComponent* UDisplayClusterBlueprintAPIImpl::GetDefaultCamer
 	return nullptr;
 }
 
-void UDisplayClusterBlueprintAPIImpl::SetDefaultCameraByIndex(int32 Index)
-{
-	DISPLAY_CLUSTER_FUNC_TRACE(LogDisplayClusterBlueprint);
-
-	IDisplayClusterGameManager* const Manager = IDisplayCluster::Get().GetGameMgr();
-	if (Manager)
-	{
-		return Manager->SetDefaultCamera(Index);
-	}
-}
-
 void UDisplayClusterBlueprintAPIImpl::SetDefaultCameraById(const FString& id)
 {
 	DISPLAY_CLUSTER_FUNC_TRACE(LogDisplayClusterBlueprint);
@@ -323,77 +314,6 @@ TArray<UDisplayClusterSceneComponent*> UDisplayClusterBlueprintAPIImpl::GetAllNo
 	return TArray<UDisplayClusterSceneComponent*>();
 }
 
-// Navigation
-USceneComponent* UDisplayClusterBlueprintAPIImpl::GetTranslationDirectionComponent()
-{
-	DISPLAY_CLUSTER_FUNC_TRACE(LogDisplayClusterBlueprint);
-
-	IDisplayClusterGameManager* const Manager = IDisplayCluster::Get().GetGameMgr();
-	if (Manager)
-	{
-		return Manager->GetTranslationDirectionComponent();
-	}
-
-	return nullptr;
-}
-
-void UDisplayClusterBlueprintAPIImpl::SetTranslationDirectionComponent(USceneComponent* const pComp)
-{
-	DISPLAY_CLUSTER_FUNC_TRACE(LogDisplayClusterBlueprint);
-
-	IDisplayClusterGameManager* const Manager = IDisplayCluster::Get().GetGameMgr();
-	if (Manager)
-	{
-		Manager->SetTranslationDirectionComponent(pComp);
-	}
-}
-
-void UDisplayClusterBlueprintAPIImpl::SetTranslationDirectionComponentId(const FString& id)
-{
-	DISPLAY_CLUSTER_FUNC_TRACE(LogDisplayClusterBlueprint);
-
-	IDisplayClusterGameManager* const Manager = IDisplayCluster::Get().GetGameMgr();
-	if (Manager)
-	{
-		Manager->SetTranslationDirectionComponent(id);
-	}
-}
-
-USceneComponent* UDisplayClusterBlueprintAPIImpl::GetRotateAroundComponent()
-{
-	DISPLAY_CLUSTER_FUNC_TRACE(LogDisplayClusterBlueprint);
-
-	IDisplayClusterGameManager* const Manager = IDisplayCluster::Get().GetGameMgr();
-	if (Manager)
-	{
-		return Manager->GetRotateAroundComponent();
-	}
-
-	return nullptr;
-}
-
-void UDisplayClusterBlueprintAPIImpl::SetRotateAroundComponent(USceneComponent* const pComp)
-{
-	DISPLAY_CLUSTER_FUNC_TRACE(LogDisplayClusterBlueprint);
-
-	IDisplayClusterGameManager* const Manager = IDisplayCluster::Get().GetGameMgr();
-	if (Manager)
-	{
-		Manager->SetRotateAroundComponent(pComp);
-	}
-}
-
-void UDisplayClusterBlueprintAPIImpl::SetRotateAroundComponentId(const FString& id)
-{
-	DISPLAY_CLUSTER_FUNC_TRACE(LogDisplayClusterBlueprint);
-
-	IDisplayClusterGameManager* const Manager = IDisplayCluster::Get().GetGameMgr();
-	if (Manager)
-	{
-		Manager->SetRotateAroundComponent(id);
-	}
-}
-
 
 //////////////////////////////////////////////////////////////////////////////////////////////
 // Input API
@@ -729,77 +649,3 @@ bool UDisplayClusterBlueprintAPIImpl::ToggleEyesSwap(const FString& CameraId)
 
 	return false;
 }
-
-float UDisplayClusterBlueprintAPIImpl::GetNearCullingDistance(const FString& CameraId) const
-{
-	DISPLAY_CLUSTER_FUNC_TRACE(LogDisplayClusterBlueprint);
-
-	IDisplayClusterRenderManager* const Manager = IDisplayCluster::Get().GetRenderMgr();
-	if (Manager)
-	{
-		return Manager->GetNearCullingDistance(CameraId);
-	}
-
-	return 0.f;
-}
-
-void UDisplayClusterBlueprintAPIImpl::SetNearCullingDistance(const FString& CameraId, float NearDistance)
-{
-	DISPLAY_CLUSTER_FUNC_TRACE(LogDisplayClusterBlueprint);
-
-	IDisplayClusterRenderManager* const Manager = IDisplayCluster::Get().GetRenderMgr();
-	if (Manager)
-	{
-		Manager->SetNearCullingDistance(CameraId, NearDistance);
-	}
-}
-
-float UDisplayClusterBlueprintAPIImpl::GetFarCullingDistance(const FString& CameraId) const
-{
-	DISPLAY_CLUSTER_FUNC_TRACE(LogDisplayClusterBlueprint);
-
-	IDisplayClusterRenderManager* const Manager = IDisplayCluster::Get().GetRenderMgr();
-	if (Manager)
-	{
-		return Manager->GetFarCullingDistance(CameraId);
-	}
-
-	return 0.f;
-}
-
-void UDisplayClusterBlueprintAPIImpl::SetFarCullingDistance(const FString& CameraId, float FarDistance)
-{
-	DISPLAY_CLUSTER_FUNC_TRACE(LogDisplayClusterBlueprint);
-
-	IDisplayClusterRenderManager* const Manager = IDisplayCluster::Get().GetRenderMgr();
-	if (Manager)
-	{
-		Manager->SetFarCullingDistance(CameraId, FarDistance);
-	}
-}
-
-void UDisplayClusterBlueprintAPIImpl::GetCullingDistance(const FString& CameraId, float& NearDistance, float& FarDistance)
-{
-	DISPLAY_CLUSTER_FUNC_TRACE(LogDisplayClusterBlueprint);
-
-	IDisplayClusterRenderManager* const Manager = IDisplayCluster::Get().GetRenderMgr();
-	if (Manager)
-	{
-		return Manager->GetCullingDistance(CameraId, NearDistance, FarDistance);
-	}
-
-	return;
-}
-
-void UDisplayClusterBlueprintAPIImpl::SetCullingDistance(const FString& CameraId, float NearDistance, float FarDistance)
-{
-	DISPLAY_CLUSTER_FUNC_TRACE(LogDisplayClusterBlueprint);
-
-	IDisplayClusterRenderManager* const Manager = IDisplayCluster::Get().GetRenderMgr();
-	if (Manager)
-	{
-		return Manager->SetCullingDistance(CameraId, NearDistance, FarDistance);
-	}
-
-	return;
-}
diff --git a/Engine/Plugins/Runtime/nDisplay/Source/DisplayCluster/Private/Blueprints/DisplayClusterBlueprintAPIImpl.h b/Engine/Plugins/Runtime/nDisplay/Source/DisplayCluster/Private/Blueprints/DisplayClusterBlueprintAPIImpl.h
--- a/Engine/Plugins/Runtime/nDisplay/Source/DisplayCluster/Private/Blueprints/DisplayClusterBlueprintAPIImpl.h
+++ b/Engine/Plugins/Runtime/nDisplay/Source/DisplayCluster/Private/Blueprints/DisplayClusterBlueprintAPIImpl.h
@@ -6,6 +6,8 @@
 #include "Blueprints/IDisplayClusterBlueprintAPI.h"
 #include "DisplayClusterBlueprintAPIImpl.generated.h"
 
+class UDisplayClusterRootComponent;
+
 struct FDisplayClusterClusterEvent;
 struct FPostProcessSettings;
 
@@ -78,7 +80,7 @@ public:
 	//////////////////////////////////////////////////////////////////////////////////////////////
 	// Root
 	UFUNCTION(BlueprintCallable, meta = (DisplayName = "Get root"), Category = "DisplayCluster|Game")
-	virtual ADisplayClusterPawn* GetRoot() override;
+	virtual UDisplayClusterRootComponent* GetRoot() override;
 
 	// Screens
 	UFUNCTION(BlueprintCallable, meta = (DisplayName = "Get screen by ID"), Category = "DisplayCluster|Game")
@@ -103,9 +105,6 @@ public:
 	UFUNCTION(BlueprintCallable, meta = (DisplayName = "Get default camera"), Category = "DisplayCluster|Game")
 	virtual UDisplayClusterCameraComponent* GetDefaultCamera() override;
 
-	UFUNCTION(BlueprintCallable, meta = (DisplayName = "Set default camera by index"), Category = "DisplayCluster|Game")
-	virtual void SetDefaultCameraByIndex(int32 Index) override;
-
 	UFUNCTION(BlueprintCallable, meta = (DisplayName = "Set default camera by ID"), Category = "DisplayCluster|Game")
 	virtual void SetDefaultCameraById(const FString& id) override;
 
@@ -116,25 +115,6 @@ public:
 	UFUNCTION(BlueprintCallable, meta = (DisplayName = "Get all nodes"), Category = "DisplayCluster|Game")
 	virtual TArray<UDisplayClusterSceneComponent*> GetAllNodes() override;
 
-	// Navigation
-	UFUNCTION(BlueprintCallable, meta = (DisplayName = "Get translation direction component"), Category = "DisplayCluster|Game")
-	virtual USceneComponent* GetTranslationDirectionComponent() override;
-
-	UFUNCTION(BlueprintCallable, meta = (DisplayName = "Set translation direction component"), Category = "DisplayCluster|Game")
-	virtual void SetTranslationDirectionComponent(USceneComponent* pComp) override;
-
-	UFUNCTION(BlueprintCallable, meta = (DisplayName = "Set translation direction component by ID"), Category = "DisplayCluster|Game")
-	virtual void SetTranslationDirectionComponentId(const FString& id) override;
-
-	UFUNCTION(BlueprintCallable, meta = (DisplayName = "Get rotate around component"), Category = "DisplayCluster|Game")
-	virtual USceneComponent* GetRotateAroundComponent() override;
-
-	UFUNCTION(BlueprintCallable, meta = (DisplayName = "Set rotate around component"), Category = "DisplayCluster|Game")
-	virtual void SetRotateAroundComponent(USceneComponent* pComp) override;
-
-	UFUNCTION(BlueprintCallable, meta = (DisplayName = "Set rotate around component by ID"), Category = "DisplayCluster|Game")
-	virtual void SetRotateAroundComponentId(const FString& id) override;
-
 public:
 	//////////////////////////////////////////////////////////////////////////////////////////////
 	// Input API
@@ -233,28 +213,4 @@ public:
 	/** Toggle current eye swap state. */
 	UFUNCTION(BlueprintCallable, meta = (DisplayName = "Toggle eye swap"), Category = "DisplayCluster|Render|Camera")
 	virtual bool ToggleEyesSwap(const FString& CameraId) override;
-
-	/** Return near culling distance of specified camera. */
-	UFUNCTION(BlueprintCallable, meta = (DisplayName = "Get near culling distance"), Category = "DisplayCluster|Render|Camera")
-	virtual float GetNearCullingDistance(const FString& CameraId) const override;
-
-	/** Set near culling distance of specified camera. */
-	UFUNCTION(BlueprintCallable, meta = (DisplayName = "Set near culling distance"), Category = "DisplayCluster|Render|Camera")
-	virtual void SetNearCullingDistance(const FString& CameraId, float NearDistance) override;
-
-	/** Get far culling distance of specified camera. */
-	UFUNCTION(BlueprintCallable, meta = (DisplayName = "Get far culling distance"), Category = "DisplayCluster|Render|Camera")
-	virtual float GetFarCullingDistance(const FString& CameraId) const override;
-
-	/** Set far culling distance of specified camera. */
-	UFUNCTION(BlueprintCallable, meta = (DisplayName = "Set far culling distance"), Category = "DisplayCluster|Render|Camera")
-	virtual void SetFarCullingDistance(const FString& CameraId, float FarDistance) override;
-
-	/** Return near and far plane clip plane distances. */
-	UFUNCTION(BlueprintCallable, meta = (DisplayName = "Get near and far clipping distance"), Category = "DisplayCluster|Render|Camera")
-	virtual void GetCullingDistance(const FString& CameraId, float& NearDistance, float& FarDistance) override;
-
-	/** Set near and far plane clip plane distances. */
-	UFUNCTION(BlueprintCallable, meta = (DisplayName = "Set near and far clipping distance"), Category = "DisplayCluster|Render|Camera")
-	virtual void SetCullingDistance(const FString& CameraId, float NearDistance, float FarDistance) override;
 };
diff --git a/Engine/Plugins/Runtime/nDisplay/Source/DisplayCluster/Private/Cluster/Controller/DisplayClusterClusterNodeCtrlMaster.cpp b/Engine/Plugins/Runtime/nDisplay/Source/DisplayCluster/Private/Cluster/Controller/DisplayClusterClusterNodeCtrlMaster.cpp
--- a/Engine/Plugins/Runtime/nDisplay/Source/DisplayCluster/Private/Cluster/Controller/DisplayClusterClusterNodeCtrlMaster.cpp
+++ b/Engine/Plugins/Runtime/nDisplay/Source/DisplayCluster/Private/Cluster/Controller/DisplayClusterClusterNodeCtrlMaster.cpp
@@ -2,13 +2,22 @@
 
 #include "Cluster/Controller/DisplayClusterClusterNodeCtrlMaster.h"
 
+#include "Cluster/IPDisplayClusterClusterManager.h"
+#include "Game/IPDisplayClusterGameManager.h"
 #include "Config/IPDisplayClusterConfigManager.h"
+#include "Input/IPDisplayClusterInputManager.h"
+
 #include "Network/Service/ClusterSync/DisplayClusterClusterSyncService.h"
 #include "Network/Service/SwapSync/DisplayClusterSwapSyncService.h"
 #include "Network/Service/ClusterEvents/DisplayClusterClusterEventsService.h"
 
+#include "Engine/World.h"
+#include "HAL/Event.h"
+
 #include "Misc/App.h"
 
+#include "DisplayClusterPlayerInput.h"
+
 #include "DisplayClusterGlobals.h"
 #include "DisplayClusterLog.h"
 
@@ -16,6 +25,13 @@
 FDisplayClusterClusterNodeCtrlMaster::FDisplayClusterClusterNodeCtrlMaster(const FString& ctrlName, const FString& nodeName) :
 	FDisplayClusterClusterNodeCtrlSlave(ctrlName, nodeName)
 {
+	IsDataCached_GetSyncData.Emplace(EDisplayClusterSyncGroup::PreTick,  false);
+	IsDataCached_GetSyncData.Emplace(EDisplayClusterSyncGroup::Tick,     false);
+	IsDataCached_GetSyncData.Emplace(EDisplayClusterSyncGroup::PostTick, false);
+
+	CachedSyncData_GetSyncData.Emplace(EDisplayClusterSyncGroup::PreTick);
+	CachedSyncData_GetSyncData.Emplace(EDisplayClusterSyncGroup::Tick);
+	CachedSyncData_GetSyncData.Emplace(EDisplayClusterSyncGroup::PostTick);
 }
 
 FDisplayClusterClusterNodeCtrlMaster::~FDisplayClusterClusterNodeCtrlMaster()
@@ -26,11 +42,166 @@ FDisplayClusterClusterNodeCtrlMaster::~FDisplayClusterClusterNodeCtrlMaster()
 //////////////////////////////////////////////////////////////////////////////////////////////
 // IPDisplayClusterClusterSyncProtocol
 //////////////////////////////////////////////////////////////////////////////////////////////
-void FDisplayClusterClusterNodeCtrlMaster::GetTimecode(FTimecode& timecode, FFrameRate& frameRate)
+void FDisplayClusterClusterNodeCtrlMaster::GetDeltaTime(float& DeltaSeconds)
 {
-	// This values are updated in UEngine::UpdateTimeAndHandleMaxTickRate (via UpdateTimecode).
-	timecode = FApp::GetTimecode();
-	frameRate = FApp::GetTimecodeFrameRate();
+	FScopeLock lock(&InternalsSyncScope);
+
+	// Cache data so it will be the same for all requests within current frame
+	if (!bIsDataCached_GetDeltaTime)
+	{
+		bIsDataCached_GetDeltaTime = true;
+
+		CachedDeltaTime_GetDeltaTime = FApp::GetDeltaTime();
+
+		UE_LOG(LogDisplayClusterCluster, Verbose, TEXT("GetDeltaTime cached values: DeltaSeconds %f"), CachedDeltaTime_GetDeltaTime);
+	}
+
+	// Return cached value
+	DeltaSeconds = CachedDeltaTime_GetDeltaTime;
+}
+
+void FDisplayClusterClusterNodeCtrlMaster::GetTimecode(FTimecode& Timecode, FFrameRate& FrameRate)
+{
+	FScopeLock lock(&InternalsSyncScope);
+
+	// Cache data so it will be the same for all requests within current frame
+	if (!bIsDataCached_GetTimecode)
+	{
+		bIsDataCached_GetTimecode = true;
+
+		// This values are updated in UEngine::UpdateTimeAndHandleMaxTickRate (via UpdateTimecode).
+		CachedTimecode_GetTimecode  = FApp::GetTimecode();
+		CachedFramerate_GetTimecode = FApp::GetTimecodeFrameRate();
+
+		UE_LOG(LogDisplayClusterCluster, Verbose, TEXT("GetDeltaTime cached values: DeltaSeconds %f"), CachedDeltaTime_GetDeltaTime);
+	}
+
+	// Return cached value
+	Timecode  = CachedTimecode_GetTimecode;
+	FrameRate = CachedFramerate_GetTimecode;
+}
+
+void FDisplayClusterClusterNodeCtrlMaster::GetSyncData(FDisplayClusterMessage::DataType& SyncData, EDisplayClusterSyncGroup SyncGroup)
+{
+	FScopeLock lock(&InternalsSyncScope);
+
+	static IPDisplayClusterClusterManager* const ClusterMgr = GDisplayCluster->GetPrivateClusterMgr();
+
+	// Cache data so it will be the same for all requests within current frame
+	if (!IsDataCached_GetSyncData[SyncGroup])
+	{
+		IsDataCached_GetSyncData[SyncGroup] = true;
+
+		FDisplayClusterMessage::DataType SyncDataToCache;
+		ClusterMgr->ExportSyncData(SyncDataToCache, SyncGroup);
+		CachedSyncData_GetSyncData.Emplace(SyncGroup, SyncDataToCache);
+
+		UE_LOG(LogDisplayClusterCluster, Verbose, TEXT("GetSyncData cached values amount: %d"), SyncDataToCache.Num());
+
+		int i = 0;
+		for (auto it = SyncDataToCache.CreateConstIterator(); it; ++it)
+		{
+			UE_LOG(LogDisplayClusterCluster, Verbose, TEXT("GetSyncData cached value %d: %s - %s"), i++, *it->Key, *it->Value);
+		}
+	}
+
+	UE_LOG(LogDisplayClusterCluster, Verbose, TEXT("SyncData has %d records"), CachedSyncData_GetSyncData[SyncGroup].Num());
+
+	// Return cached value
+	SyncData = CachedSyncData_GetSyncData[SyncGroup];
+}
+
+void FDisplayClusterClusterNodeCtrlMaster::GetInputData(FDisplayClusterMessage::DataType& InputData)
+{
+	FScopeLock lock(&InternalsSyncScope);
+
+	static IPDisplayClusterInputManager* const InputMgr = GDisplayCluster->GetPrivateInputMgr();
+
+	// Cache data so it will be the same for all requests within current frame
+	if (!bIsDataCached_GetInputData)
+	{
+		bIsDataCached_GetInputData = true;
+
+		InputMgr->ExportInputData(CachedInputData_GetInputData);
+
+		UE_LOG(LogDisplayClusterCluster, Verbose, TEXT("GetInputData cached values amount: %d"), CachedInputData_GetInputData.Num());
+
+		int i = 0;
+		for (auto it = CachedInputData_GetInputData.CreateConstIterator(); it; ++it)
+		{
+			UE_LOG(LogDisplayClusterCluster, Verbose, TEXT("GetInputData cached value %d: %s - %s"), i++, *it->Key, *it->Value);
+		}
+	}
+
+	UE_LOG(LogDisplayClusterCluster, Verbose, TEXT("InputData has %d records"), CachedInputData_GetInputData.Num());
+
+	// Return cached value
+	InputData = CachedInputData_GetInputData;
+}
+
+void FDisplayClusterClusterNodeCtrlMaster::GetEventsData(FDisplayClusterMessage::DataType& EventsData)
+{
+	FScopeLock lock(&InternalsSyncScope);
+
+	static IPDisplayClusterClusterManager* const ClusterMgr = GDisplayCluster->GetPrivateClusterMgr();
+
+	// Cache data so it will be the same for all requests within current frame
+	if (!bIsDataCached_GetEventsData)
+	{
+		bIsDataCached_GetEventsData = true;
+
+		ClusterMgr->ExportEventsData(CachedEventsData_GetEventsData);
+
+		UE_LOG(LogDisplayClusterCluster, Verbose, TEXT("GetEventsData cached values amount: %d"), CachedEventsData_GetEventsData.Num());
+
+		int i = 0;
+		for (auto it = CachedEventsData_GetEventsData.CreateConstIterator(); it; ++it)
+		{
+			UE_LOG(LogDisplayClusterCluster, Verbose, TEXT("GetEventsData cached value %d: %s - %s"), i++, *it->Key, *it->Value);
+		}
+	}
+
+	UE_LOG(LogDisplayClusterCluster, Verbose, TEXT("EventsData has %d records"), CachedEventsData_GetEventsData.Num());
+
+	// Return cached value
+	EventsData = CachedEventsData_GetEventsData;
+}
+
+void FDisplayClusterClusterNodeCtrlMaster::GetNativeInputData(FDisplayClusterMessage::DataType& NativeInputData)
+{
+	static IPDisplayClusterClusterManager* const ClusterMgr = GDisplayCluster->GetPrivateClusterMgr();
+	ClusterMgr->SyncNativeInput(NativeInputData);
+}
+
+
+//////////////////////////////////////////////////////////////////////////////////////////////
+// IPDisplayClusterNodeController
+//////////////////////////////////////////////////////////////////////////////////////////////
+void FDisplayClusterClusterNodeCtrlMaster::ClearCache()
+{
+	FScopeLock lock(&InternalsSyncScope);
+
+	// Reset all cache flags
+	bIsDataCached_GetDeltaTime  = false;
+	bIsDataCached_GetTimecode   = false;
+	bIsDataCached_GetInputData  = false;
+	bIsDataCached_GetEventsData = false;
+
+	// Reset cache containers
+	CachedInputData_GetInputData.Reset();
+	CachedEventsData_GetEventsData.Reset();
+
+	for (auto& it : IsDataCached_GetSyncData)
+	{
+		it.Value = false;
+	}
+
+	for (auto& it : CachedSyncData_GetSyncData)
+	{
+		it.Value.Reset();
+	}
+
+	UE_LOG(LogDisplayClusterCluster, Verbose, TEXT("FDisplayClusterClusterNodeCtrlMaster has cleaned the cache"));
 }
 
 
diff --git a/Engine/Plugins/Runtime/nDisplay/Source/DisplayCluster/Private/Cluster/Controller/DisplayClusterClusterNodeCtrlMaster.h b/Engine/Plugins/Runtime/nDisplay/Source/DisplayCluster/Private/Cluster/Controller/DisplayClusterClusterNodeCtrlMaster.h
--- a/Engine/Plugins/Runtime/nDisplay/Source/DisplayCluster/Private/Cluster/Controller/DisplayClusterClusterNodeCtrlMaster.h
+++ b/Engine/Plugins/Runtime/nDisplay/Source/DisplayCluster/Private/Cluster/Controller/DisplayClusterClusterNodeCtrlMaster.h
@@ -6,6 +6,7 @@
 
 #include "Network/DisplayClusterMessage.h"
 
+
 class FDisplayClusterClusterSyncService;
 class FDisplayClusterSwapSyncService;
 class FDisplayClusterClusterEventsService;
@@ -25,12 +26,19 @@ public:
 	//////////////////////////////////////////////////////////////////////////////////////////////
 	// IPDisplayClusterClusterSyncProtocol
 	//////////////////////////////////////////////////////////////////////////////////////////////
-	virtual void GetTimecode(FTimecode& timecode, FFrameRate& frameRate) override;
+	virtual void GetDeltaTime(float& DeltaSeconds) override;
+	virtual void GetTimecode(FTimecode& Timecode, FFrameRate& FrameRate) override;
+	virtual void GetSyncData(FDisplayClusterMessage::DataType& SyncData, EDisplayClusterSyncGroup SyncGroup) override;
+	virtual void GetInputData(FDisplayClusterMessage::DataType& InputData) override;
+	virtual void GetEventsData(FDisplayClusterMessage::DataType& EventsData) override;
+	virtual void GetNativeInputData(FDisplayClusterMessage::DataType& EventsData) override;
 
 public:
 	//////////////////////////////////////////////////////////////////////////////////////////////
 	// IPDisplayClusterNodeController
 	//////////////////////////////////////////////////////////////////////////////////////////////
+	virtual void ClearCache() override;
+
 	virtual bool IsSlave() const override final
 	{ return false; }
 
@@ -51,5 +59,30 @@ private:
 	TUniquePtr<FDisplayClusterClusterSyncService>   ClusterSyncServer;
 	TUniquePtr<FDisplayClusterSwapSyncService>      SwapSyncServer;
 	TUniquePtr<FDisplayClusterClusterEventsService> ClusterEventsServer;
+
+private:
+	// GetDeltaTime internals
+	bool  bIsDataCached_GetDeltaTime = false;
+	float CachedDeltaTime_GetDeltaTime = 0.f;
+
+	// GetTimecode internals
+	bool       bIsDataCached_GetTimecode = false;
+	FTimecode  CachedTimecode_GetTimecode;
+	FFrameRate CachedFramerate_GetTimecode;
+
+	// GetSyncData internals
+	TMap<EDisplayClusterSyncGroup, bool> IsDataCached_GetSyncData;
+	TMap<EDisplayClusterSyncGroup, FDisplayClusterMessage::DataType> CachedSyncData_GetSyncData;
+
+	// GetInputData internals
+	bool bIsDataCached_GetInputData = false;
+	FDisplayClusterMessage::DataType CachedInputData_GetInputData;
+
+	// GetEventsData internals
+	bool bIsDataCached_GetEventsData = false;
+	FDisplayClusterMessage::DataType CachedEventsData_GetEventsData;
+
+private:
+	mutable FCriticalSection InternalsSyncScope;
 };
 
diff --git a/Engine/Plugins/Runtime/nDisplay/Source/DisplayCluster/Private/Cluster/Controller/DisplayClusterClusterNodeCtrlSlave.cpp b/Engine/Plugins/Runtime/nDisplay/Source/DisplayCluster/Private/Cluster/Controller/DisplayClusterClusterNodeCtrlSlave.cpp
--- a/Engine/Plugins/Runtime/nDisplay/Source/DisplayCluster/Private/Cluster/Controller/DisplayClusterClusterNodeCtrlSlave.cpp
+++ b/Engine/Plugins/Runtime/nDisplay/Source/DisplayCluster/Private/Cluster/Controller/DisplayClusterClusterNodeCtrlSlave.cpp
@@ -44,29 +44,34 @@ void FDisplayClusterClusterNodeCtrlSlave::WaitForTickEnd()
 	ClusterSyncClient->WaitForTickEnd();
 }
 
-void FDisplayClusterClusterNodeCtrlSlave::GetDeltaTime(float& deltaTime)
+void FDisplayClusterClusterNodeCtrlSlave::GetDeltaTime(float& DeltaSeconds)
 {
-	ClusterSyncClient->GetDeltaTime(deltaTime);
+	ClusterSyncClient->GetDeltaTime(DeltaSeconds);
 }
 
-void FDisplayClusterClusterNodeCtrlSlave::GetTimecode(FTimecode& timecode, FFrameRate& frameRate)
+void FDisplayClusterClusterNodeCtrlSlave::GetTimecode(FTimecode& Timecode, FFrameRate& FrameRate)
 {
-	ClusterSyncClient->GetTimecode(timecode, frameRate);
+	ClusterSyncClient->GetTimecode(Timecode, FrameRate);
 }
 
-void FDisplayClusterClusterNodeCtrlSlave::GetSyncData(FDisplayClusterMessage::DataType& data)
+void FDisplayClusterClusterNodeCtrlSlave::GetSyncData(FDisplayClusterMessage::DataType& SyncData, EDisplayClusterSyncGroup SyncGroup)
 {
-	ClusterSyncClient->GetSyncData(data);
+	ClusterSyncClient->GetSyncData(SyncData, SyncGroup);
 }
 
-void FDisplayClusterClusterNodeCtrlSlave::GetInputData(FDisplayClusterMessage::DataType& data)
+void FDisplayClusterClusterNodeCtrlSlave::GetInputData(FDisplayClusterMessage::DataType& InputData)
 {
-	ClusterSyncClient->GetInputData(data);
+	ClusterSyncClient->GetInputData(InputData);
 }
 
-void FDisplayClusterClusterNodeCtrlSlave::GetEventsData(FDisplayClusterMessage::DataType& data)
+void FDisplayClusterClusterNodeCtrlSlave::GetEventsData(FDisplayClusterMessage::DataType& EventsData)
 {
-	ClusterSyncClient->GetEventsData(data);
+	ClusterSyncClient->GetEventsData(EventsData);
+}
+
+void FDisplayClusterClusterNodeCtrlSlave::GetNativeInputData(FDisplayClusterMessage::DataType& NativeInputData)
+{
+	ClusterSyncClient->GetNativeInputData(NativeInputData);
 }
 
 
diff --git a/Engine/Plugins/Runtime/nDisplay/Source/DisplayCluster/Private/Cluster/Controller/DisplayClusterClusterNodeCtrlSlave.h b/Engine/Plugins/Runtime/nDisplay/Source/DisplayCluster/Private/Cluster/Controller/DisplayClusterClusterNodeCtrlSlave.h
--- a/Engine/Plugins/Runtime/nDisplay/Source/DisplayCluster/Private/Cluster/Controller/DisplayClusterClusterNodeCtrlSlave.h
+++ b/Engine/Plugins/Runtime/nDisplay/Source/DisplayCluster/Private/Cluster/Controller/DisplayClusterClusterNodeCtrlSlave.h
@@ -36,11 +36,12 @@ public:
 	virtual void WaitForFrameStart() override final;
 	virtual void WaitForFrameEnd()   override final;
 	virtual void WaitForTickEnd()    override final;
-	virtual void GetDeltaTime(float& deltaTime) override final;
-	virtual void GetTimecode(FTimecode& timecode, FFrameRate& frameRate) override;
-	virtual void GetSyncData(FDisplayClusterMessage::DataType& data)  override;
-	virtual void GetInputData(FDisplayClusterMessage::DataType& data) override;
-	virtual void GetEventsData(FDisplayClusterMessage::DataType& data) override;
+	virtual void GetDeltaTime(float& DeltaSeconds) override;
+	virtual void GetTimecode(FTimecode& Timecode, FFrameRate& FrameRate) override;
+	virtual void GetSyncData(FDisplayClusterMessage::DataType& SyncData, EDisplayClusterSyncGroup SyncGroup) override;
+	virtual void GetInputData(FDisplayClusterMessage::DataType& InputData) override;
+	virtual void GetEventsData(FDisplayClusterMessage::DataType& EventsData) override;
+	virtual void GetNativeInputData(FDisplayClusterMessage::DataType& EventsData) override;
 
 public:
 	//////////////////////////////////////////////////////////////////////////////////////////////
diff --git a/Engine/Plugins/Runtime/nDisplay/Source/DisplayCluster/Private/Cluster/Controller/DisplayClusterNodeCtrlBase.h b/Engine/Plugins/Runtime/nDisplay/Source/DisplayCluster/Private/Cluster/Controller/DisplayClusterNodeCtrlBase.h
--- a/Engine/Plugins/Runtime/nDisplay/Source/DisplayCluster/Private/Cluster/Controller/DisplayClusterNodeCtrlBase.h
+++ b/Engine/Plugins/Runtime/nDisplay/Source/DisplayCluster/Private/Cluster/Controller/DisplayClusterNodeCtrlBase.h
@@ -67,19 +67,22 @@ public:
 	virtual void WaitForTickEnd() override
 	{ }
 
-	virtual void GetDeltaTime(float& deltaTime) override
+	virtual void GetDeltaTime(float& DeltaSeconds) override
 	{ }
 
-	virtual void GetTimecode(FTimecode& timecode, FFrameRate& frameRate) override
+	virtual void GetTimecode(FTimecode& Timecode, FFrameRate& FrameRate) override
 	{ }
 
-	virtual void GetSyncData(FDisplayClusterMessage::DataType& data) override
+	virtual void GetSyncData(FDisplayClusterMessage::DataType& SyncData, EDisplayClusterSyncGroup SyncGroup) override
 	{ }
 
-	virtual void GetInputData(FDisplayClusterMessage::DataType& data) override
+	virtual void GetInputData(FDisplayClusterMessage::DataType& InputData) override
 	{ }
 
-	virtual void GetEventsData(FDisplayClusterMessage::DataType& data) override
+	virtual void GetEventsData(FDisplayClusterMessage::DataType& EventsData) override
+	{ }
+
+	virtual void GetNativeInputData(FDisplayClusterMessage::DataType& EventsData) override
 	{ }
 
 public:
diff --git a/Engine/Plugins/Runtime/nDisplay/Source/DisplayCluster/Private/Cluster/Controller/IPDisplayClusterNodeController.h b/Engine/Plugins/Runtime/nDisplay/Source/DisplayCluster/Private/Cluster/Controller/IPDisplayClusterNodeController.h
--- a/Engine/Plugins/Runtime/nDisplay/Source/DisplayCluster/Private/Cluster/Controller/IPDisplayClusterNodeController.h
+++ b/Engine/Plugins/Runtime/nDisplay/Source/DisplayCluster/Private/Cluster/Controller/IPDisplayClusterNodeController.h
@@ -22,6 +22,9 @@ public:
 	virtual bool Initialize() = 0;
 	virtual void Release() = 0;
 
+	virtual void ClearCache()
+	{ }
+
 	virtual bool IsMaster() const = 0;
 	virtual bool IsSlave() const = 0;
 	virtual bool IsStandalone() const = 0;
diff --git a/Engine/Plugins/Runtime/nDisplay/Source/DisplayCluster/Private/Cluster/DisplayClusterClusterManager.cpp b/Engine/Plugins/Runtime/nDisplay/Source/DisplayCluster/Private/Cluster/DisplayClusterClusterManager.cpp
--- a/Engine/Plugins/Runtime/nDisplay/Source/DisplayCluster/Private/Cluster/DisplayClusterClusterManager.cpp
+++ b/Engine/Plugins/Runtime/nDisplay/Source/DisplayCluster/Private/Cluster/DisplayClusterClusterManager.cpp
@@ -34,15 +34,20 @@ FDisplayClusterClusterManager::FDisplayClusterClusterManager()
 {
 	DISPLAY_CLUSTER_FUNC_TRACE(LogDisplayClusterCluster);
 
-	ObjectsToSync.Reserve(64);
+	ObjectsToSync.Emplace(EDisplayClusterSyncGroup::PreTick).Reserve(64);
+	ObjectsToSync.Emplace(EDisplayClusterSyncGroup::Tick).Reserve(64);
+	ObjectsToSync.Emplace(EDisplayClusterSyncGroup::PostTick).Reserve(64);
 
 	// Set main event handler. This is an entry point for any incoming cluster events.
 	OnClusterEvent.AddRaw(this, &FDisplayClusterClusterManager::OnClusterEventHandler);
+
+	NativeInputDataAvailableEvent = FPlatformProcess::CreateSynchEvent(true);
 }
 
 FDisplayClusterClusterManager::~FDisplayClusterClusterManager()
 {
 	DISPLAY_CLUSTER_FUNC_TRACE(LogDisplayClusterCluster);
+	delete NativeInputDataAvailableEvent;
 }
 
 
@@ -131,7 +136,6 @@ bool FDisplayClusterClusterManager::StartSession(const FString& configPath, cons
 
 	// Instantiate node controller
 	Controller = CreateController();
-
 	if (!Controller)
 	{
 		UE_LOG(LogDisplayClusterCluster, Error, TEXT("Couldn't create a controller."));
@@ -180,30 +184,60 @@ void FDisplayClusterClusterManager::EndScene()
 
 	{
 		FScopeLock lock(&ObjectsToSyncCritSec);
-		ObjectsToSync.Reset();
+		for (auto& SyncGroupPair: ObjectsToSync)
+		{
+			SyncGroupPair.Value.Reset();
+		}
+	}
+}
+
+void FDisplayClusterClusterManager::EndFrame(uint64 FrameNum)
+{
+	DISPLAY_CLUSTER_FUNC_TRACE(LogDisplayClusterCluster);
+
+	if (Controller)
+	{
+		Controller->ClearCache();
 	}
+
+	NativeInputDataAvailableEvent->Reset();
 }
 
 void FDisplayClusterClusterManager::PreTick(float DeltaSeconds)
 {
 	DISPLAY_CLUSTER_FUNC_TRACE(LogDisplayClusterCluster);
 
-	// Moved clearing of the sync object cache to the end of the frame
-	//// Clear cached data from previous game frame
-	//{
-	//	FScopeLock lock(&ObjectsToSyncCritSec);
-	//	SyncObjectsCache.Empty(SyncObjectsCache.Num() | 0x07);
-	//}
+	// Update input state in the cluster
+	SyncInput();
+
+	// Sync cluster objects (PreTick)
+	SyncObjects(EDisplayClusterSyncGroup::PreTick);
+
+	// Sync cluster events
+	SyncEvents();
 
 	// Move cluster events from the primary pool to the output pool. These will be synchronized on the current frame.
 	{
 		FScopeLock lock(& ClusterEventsCritSec);
+
 		ClusterEventsPoolOut = MoveTemp(ClusterEventsPoolMain);
 		ClusterEventsPoolMain.Empty(ClusterEventsPoolOut.Num() | 0x07);
 		ClusterEventsCacheOut.Empty(ClusterEventsPoolOut.Num() | 0x07);
 	}
 }
 
+void FDisplayClusterClusterManager::Tick(float DeltaSeconds)
+{
+	// Sync cluster objects (Tick)
+	SyncObjects(EDisplayClusterSyncGroup::Tick);
+}
+
+void FDisplayClusterClusterManager::PostTick(float DeltaSeconds)
+{
+	// Sync cluster objects (PostTick)
+	SyncObjects(EDisplayClusterSyncGroup::PostTick);
+}
+
 
 //////////////////////////////////////////////////////////////////////////////////////////////
 // IDisplayClusterClusterManager
@@ -234,6 +268,34 @@ bool FDisplayClusterClusterManager::IsCluster() const
 	return Controller ? Controller->IsCluster() : false;
 }
 
+void FDisplayClusterClusterManager::RegisterSyncObject(IDisplayClusterClusterSyncObject* pSyncObj, EDisplayClusterSyncGroup SyncGroup)
+{
+	DISPLAY_CLUSTER_FUNC_TRACE(LogDisplayClusterCluster);
+
+	{
+		FScopeLock lock(&ObjectsToSyncCritSec);
+		ObjectsToSync[SyncGroup].Add(pSyncObj);
+	}
+
+	UE_LOG(LogDisplayClusterCluster, Log, TEXT("Registered sync object: %s"), *pSyncObj->GetSyncId());
+}
+
+void FDisplayClusterClusterManager::UnregisterSyncObject(IDisplayClusterClusterSyncObject* pSyncObj)
+{
+	DISPLAY_CLUSTER_FUNC_TRACE(LogDisplayClusterCluster);
+
+	{
+		FScopeLock lock(&ObjectsToSyncCritSec);
+
+		for (auto& GroupPair : ObjectsToSync)
+		{
+			GroupPair.Value.Remove(pSyncObj);
+		}
+	}
+
+	UE_LOG(LogDisplayClusterCluster, Log, TEXT("Unregistered sync object: %s"), *pSyncObj->GetSyncId());
+}
+
 void FDisplayClusterClusterManager::AddClusterEventListener(TScriptInterface<IDisplayClusterClusterEventListener> Listener)
 {
 	DISPLAY_CLUSTER_FUNC_TRACE(LogDisplayClusterCluster);
@@ -322,79 +384,51 @@ void FDisplayClusterClusterManager::EmitClusterEvent(const FDisplayClusterCluste
 //////////////////////////////////////////////////////////////////////////////////////////////
 // IPDisplayClusterClusterManager
 //////////////////////////////////////////////////////////////////////////////////////////////
-void FDisplayClusterClusterManager::RegisterSyncObject(IDisplayClusterClusterSyncObject* pSyncObj)
+void FDisplayClusterClusterManager::ExportSyncData(FDisplayClusterMessage::DataType& SyncData, EDisplayClusterSyncGroup SyncGroup) const
 {
 	DISPLAY_CLUSTER_FUNC_TRACE(LogDisplayClusterCluster);
 
-	{
-		FScopeLock lock(&ObjectsToSyncCritSec);
-		ObjectsToSync.Add(pSyncObj);
-	}
-
-	UE_LOG(LogDisplayClusterCluster, Log, TEXT("Registered sync object: %s"), *pSyncObj->GetSyncId());
-}
-
-void FDisplayClusterClusterManager::UnregisterSyncObject(IDisplayClusterClusterSyncObject* pSyncObj)
-{
-	DISPLAY_CLUSTER_FUNC_TRACE(LogDisplayClusterCluster);
+	UE_LOG(LogDisplayClusterCluster, Verbose, TEXT("Exporting sync data for sync group: %d, items to sync: %d"), (int)SyncGroup, ObjectsToSync[SyncGroup].Num());
 
 	{
 		FScopeLock lock(&ObjectsToSyncCritSec);
-		ObjectsToSync.Remove(pSyncObj);
-	}
 
-	UE_LOG(LogDisplayClusterCluster, Log, TEXT("Unregistered sync object: %s"), *pSyncObj->GetSyncId());
-}
+		SyncData.Empty(SyncData.Num() | 0x7);
 
-void FDisplayClusterClusterManager::ExportSyncData(FDisplayClusterMessage::DataType& data) const
-{
-	DISPLAY_CLUSTER_FUNC_TRACE(LogDisplayClusterCluster);
-
-	{
-		FScopeLock lock(&ObjectsToSyncCritSec);
-
-		// Cache the data for current frame.
-		// There is no check for ObjectsToSync emptiness because we always have at least one
-		// shared transform which is AFDisplayClusterPawn.
-		if (SyncObjectsCache.Num() == 0)
+		for (IDisplayClusterClusterSyncObject* Obj : ObjectsToSync[SyncGroup])
 		{
-			for (auto obj : ObjectsToSync)
+			if (Obj && Obj->IsDirty())
 			{
-				if (obj->IsDirty())
-				{
-					UE_LOG(LogDisplayClusterCluster, Verbose, TEXT("Adding object to sync: %s"), *obj->GetSyncId());
-					SyncObjectsCache.Add(obj->GetSyncId(), obj->SerializeToString());
-					obj->ClearDirty();
-				}
+				UE_LOG(LogDisplayClusterCluster, Verbose, TEXT("Adding object to sync: %s"), *Obj->GetSyncId());
+				SyncData.Add(Obj->GetSyncId(), Obj->SerializeToString());
+				Obj->ClearDirty();
 			}
 		}
-
-		data = SyncObjectsCache;
 	}
 }
 
-void FDisplayClusterClusterManager::ImportSyncData(const FDisplayClusterMessage::DataType& data)
+void FDisplayClusterClusterManager::ImportSyncData(const FDisplayClusterMessage::DataType& SyncData, EDisplayClusterSyncGroup SyncGroup)
 {
 	DISPLAY_CLUSTER_FUNC_TRACE(LogDisplayClusterCluster);
 
-	if (data.Num() > 0)
+	if (SyncData.Num() > 0)
 	{
-		for (auto it = data.CreateConstIterator(); it; ++it)
+		for (auto it = SyncData.CreateConstIterator(); it; ++it)
 		{
 			UE_LOG(LogDisplayClusterCluster, Verbose, TEXT("sync-data: %s=%s"), *it->Key, *it->Value);
 		}
 
-		for (auto obj : ObjectsToSync)
+		for (auto obj : ObjectsToSync[SyncGroup])
 		{
 			const FString syncId = obj->GetSyncId();
-			if (!data.Contains(syncId))
+			if (!SyncData.Contains(syncId))
 			{
 				UE_LOG(LogDisplayClusterCluster, Verbose, TEXT("%s has nothing to update"), *syncId);
 				continue;
 			}
 
 			UE_LOG(LogDisplayClusterCluster, Verbose, TEXT("Found %s in sync data. Applying..."), *syncId);
-			if (!obj->DeserializeFromString(data[syncId]))
+			if (!obj->DeserializeFromString(SyncData[syncId]))
 			{
 				UE_LOG(LogDisplayClusterCluster, Error, TEXT("Couldn't apply sync data for sync object %s"), *syncId);
 			}
@@ -402,7 +436,7 @@ void FDisplayClusterClusterManager::ImportSyncData(const FDisplayClusterMessage:
 	}
 }
 
-void FDisplayClusterClusterManager::ExportEventsData(FDisplayClusterMessage::DataType& data) const
+void FDisplayClusterClusterManager::ExportEventsData(FDisplayClusterMessage::DataType& EventsData) const
 {
 	DISPLAY_CLUSTER_FUNC_TRACE(LogDisplayClusterCluster);
 
@@ -429,7 +463,7 @@ void FDisplayClusterClusterManager::ExportEventsData(FDisplayClusterMessage::Dat
 			ClusterEventsPoolOut.Empty(ClusterEventsPoolOut.Num() | 0x07);
 		}
 
-		data = ClusterEventsCacheOut;
+		EventsData = ClusterEventsCacheOut;
 	}
 }
 
@@ -458,25 +492,21 @@ void FDisplayClusterClusterManager::ImportEventsData(const FDisplayClusterMessag
 	}
 }
 
-void FDisplayClusterClusterManager::SyncObjects()
+void FDisplayClusterClusterManager::SyncObjects(EDisplayClusterSyncGroup SyncGroup)
 {
 	DISPLAY_CLUSTER_FUNC_TRACE(LogDisplayClusterCluster);
 
-	//@note:
-	// Don't use FScopeLock lock(&ObjectsToSyncCritSec) here because
-	// a) There are no race conditions at this point. We're in single-threaded mode right now (see UDisplayClusterGameEngine::Tick())
-	// b) Performance
+	TMap<FString, FString> SyncData;
 
-	// No need to do the sync for master
+	UE_LOG(LogDisplayClusterCluster, Verbose, TEXT("Downloading synchronization data (objects)..."));
+	Controller->GetSyncData(SyncData, SyncGroup);
+	UE_LOG(LogDisplayClusterCluster, Verbose, TEXT("Downloading finished. Available %d records (objects)."), SyncData.Num());
+
+	// We don't have to import data here unless sync data provider is located on master node
 	if (IsSlave())
 	{
-		UE_LOG(LogDisplayClusterCluster, Verbose, TEXT("Downloading synchronization data (objects)..."));
-		TMap<FString, FString> data;
-		Controller->GetSyncData(data);
-		UE_LOG(LogDisplayClusterCluster, Verbose, TEXT("Downloading finished. Available %d records (objects)."), data.Num());
-
 		// Perform data load (objects state update)
-		ImportSyncData(data);
+		ImportSyncData(SyncData, SyncGroup);
 	}
 }
 
@@ -484,16 +514,18 @@ void FDisplayClusterClusterManager::SyncInput()
 {
 	DISPLAY_CLUSTER_FUNC_TRACE(LogDisplayClusterCluster);
 
-	// No need to do the sync for master
+	TMap<FString, FString> InputData;
+
+	// Get input data from a provider
+	UE_LOG(LogDisplayClusterCluster, Verbose, TEXT("Downloading synchronization data (input)..."));
+	Controller->GetInputData(InputData);
+	UE_LOG(LogDisplayClusterCluster, Verbose, TEXT("Downloading finished. Available %d records (input)."), InputData.Num());
+
+	// We don't have to import data here unless input data provider is located on master node
 	if (IsSlave())
 	{
-		UE_LOG(LogDisplayClusterCluster, Verbose, TEXT("Downloading synchronization data (input)..."));
-		TMap<FString, FString> data;
-		Controller->GetInputData(data);
-		UE_LOG(LogDisplayClusterCluster, Verbose, TEXT("Downloading finished. Available %d records (input)."), data.Num());
-
 		// Perform data load (objects state update)
-		GDisplayCluster->GetPrivateInputMgr()->ImportInputData(data);
+		GDisplayCluster->GetPrivateInputMgr()->ImportInputData(InputData);
 	}
 }
 
@@ -505,17 +537,17 @@ void FDisplayClusterClusterManager::SyncEvents()
 	if (IsSlave())
 	{
 		UE_LOG(LogDisplayClusterCluster, Verbose, TEXT("Downloading synchronization data (events)..."));
-		TMap<FString, FString> data;
+		TMap<FString, FString> EventsData;
 
 		{
 			FScopeLock lock(&ClusterEventsCritSec);
-			Controller->GetEventsData(data);
+			Controller->GetEventsData(EventsData);
 		}
 
-		UE_LOG(LogDisplayClusterCluster, Verbose, TEXT("Downloading finished. Available %d records (events)."), data.Num());
+		UE_LOG(LogDisplayClusterCluster, Verbose, TEXT("Downloading finished. Available %d records (events)."), EventsData.Num());
 
 		// Perform data load (objects state update)
-		GDisplayCluster->GetPrivateClusterMgr()->ImportEventsData(data);
+		GDisplayCluster->GetPrivateClusterMgr()->ImportEventsData(EventsData);
 	}
 	else
 	{
@@ -525,12 +557,36 @@ void FDisplayClusterClusterManager::SyncEvents()
 	}
 }
 
-void FDisplayClusterClusterManager::ClearSyncObjects()
+void FDisplayClusterClusterManager::ProvideNativeInputData(const TMap<FString, FString>& NativeInputData)
 {
-	FScopeLock lock(&ObjectsToSyncCritSec);
-	SyncObjectsCache.Empty(SyncObjectsCache.Num() | 0x07);
+	DISPLAY_CLUSTER_FUNC_TRACE(LogDisplayClusterCluster);
+
+	UE_LOG(LogDisplayClusterCluster, Verbose, TEXT("SyncNativeInput - data available trigger. NativeInput records amount %d"), NativeInputData.Num());
+
+	NativeInputDataCache = NativeInputData;
+	NativeInputDataAvailableEvent->Trigger();
 }
 
+void FDisplayClusterClusterManager::SyncNativeInput(TMap<FString, FString>& NativeInputData)
+{
+	DISPLAY_CLUSTER_FUNC_TRACE(LogDisplayClusterCluster);
+
+	if (IsMaster())
+	{
+		UE_LOG(LogDisplayClusterCluster, Verbose, TEXT("Returning native input data, records amount %d"), NativeInputDataCache.Num());
+
+		NativeInputDataAvailableEvent->Wait();
+
+		NativeInputData = NativeInputDataCache;
+	}
+	else
+	{
+		UE_LOG(LogDisplayClusterCluster, Verbose, TEXT("Downloading native input data..."));
+		Controller->GetNativeInputData(NativeInputData);
+	}
+}
+
+
 //////////////////////////////////////////////////////////////////////////////////////////////
 // FDisplayClusterClusterManager
 //////////////////////////////////////////////////////////////////////////////////////////////
diff --git a/Engine/Plugins/Runtime/nDisplay/Source/DisplayCluster/Private/Cluster/DisplayClusterClusterManager.h b/Engine/Plugins/Runtime/nDisplay/Source/DisplayCluster/Private/Cluster/DisplayClusterClusterManager.h
--- a/Engine/Plugins/Runtime/nDisplay/Source/DisplayCluster/Private/Cluster/DisplayClusterClusterManager.h
+++ b/Engine/Plugins/Runtime/nDisplay/Source/DisplayCluster/Private/Cluster/DisplayClusterClusterManager.h
@@ -10,6 +10,7 @@
 class ADisplayClusterGameMode;
 class ADisplayClusterSettings;
 class FJsonObject;
+class FEvent;
 
 
 /**
@@ -32,7 +33,10 @@ public:
 	virtual void EndSession() override;
 	virtual bool StartScene(UWorld* pWorld) override;
 	virtual void EndScene() override;
+	virtual void EndFrame(uint64 FrameNum) override;
 	virtual void PreTick(float DeltaSeconds) override;
+	virtual void Tick(float DeltaSeconds) override;
+	virtual void PostTick(float DeltaSeconds) override;
 
 public:
 	//////////////////////////////////////////////////////////////////////////////////////////////
@@ -49,6 +53,9 @@ public:
 	virtual uint32 GetNodesAmount() const override
 	{ return NodesAmount; }
 
+	virtual void RegisterSyncObject(IDisplayClusterClusterSyncObject* pSyncObj, EDisplayClusterSyncGroup SyncGroup) override;
+	virtual void UnregisterSyncObject(IDisplayClusterClusterSyncObject* pSyncObj) override;
+
 	virtual void AddClusterEventListener(TScriptInterface<IDisplayClusterClusterEventListener>) override;
 	virtual void RemoveClusterEventListener(TScriptInterface<IDisplayClusterClusterEventListener>) override;
 
@@ -63,33 +70,18 @@ public:
 	//////////////////////////////////////////////////////////////////////////////////////////////
 	virtual IPDisplayClusterNodeController* GetController() const override;
 
-	virtual float GetDeltaTime() const override
-	{ return DeltaTime; }
-
-	virtual void  SetDeltaTime(float deltaTime) override
-	{ DeltaTime = deltaTime; }
-
-	virtual void GetTimecode(FTimecode& timecode, FFrameRate& frameRate) const override
-	{ timecode = FApp::GetTimecode(); frameRate = FApp::GetTimecodeFrameRate(); }
+	virtual void ExportSyncData(FDisplayClusterMessage::DataType& SyncData, EDisplayClusterSyncGroup SyncGroup) const override;
+	virtual void ImportSyncData(const FDisplayClusterMessage::DataType& SyncData, EDisplayClusterSyncGroup SyncGroup) override;
 
-	virtual void SetTimecode(const FTimecode& timecode, const FFrameRate& frameRate) override
-	{ FApp::SetTimecodeAndFrameRate(timecode, frameRate); }
-	
-	virtual void RegisterSyncObject(IDisplayClusterClusterSyncObject* pSyncObj) override;
-	virtual void UnregisterSyncObject(IDisplayClusterClusterSyncObject* pSyncObj) override;
-
-	virtual void ExportSyncData(FDisplayClusterMessage::DataType& data) const override;
-	virtual void ImportSyncData(const FDisplayClusterMessage::DataType& data) override;
+	virtual void ExportEventsData(FDisplayClusterMessage::DataType& EventsData) const override;
+	virtual void ImportEventsData(const FDisplayClusterMessage::DataType& EventsData) override;
 
-	virtual void ExportEventsData(FDisplayClusterMessage::DataType& data) const override;
-	virtual void ImportEventsData(const FDisplayClusterMessage::DataType& data) override;
-
-	virtual void SyncObjects() override;
+	virtual void SyncObjects(EDisplayClusterSyncGroup SyncGroup) override;
 	virtual void SyncInput()   override;
 	virtual void SyncEvents()  override;
 
-	virtual void ClearSyncObjects() override;
-
+	virtual void ProvideNativeInputData(const TMap<FString, FString>& NativeInputData) override;
+	virtual void SyncNativeInput(TMap<FString, FString>& NativeInputData) override;
 
 private:
 	bool GetResolvedNodeId(FString& id) const;
@@ -106,8 +98,6 @@ private:
 	TController Controller;
 	// Cluster/node props
 	uint32 NodesAmount = 0;
-	// Current time delta for sync
-	float DeltaTime = 0.f;
 
 	// Current operation mode
 	EDisplayClusterOperationMode CurrentOperationMode;
@@ -119,8 +109,7 @@ private:
 	UWorld* CurrentWorld;
 
 	// Sync transforms
-	TSet<IDisplayClusterClusterSyncObject*>      ObjectsToSync;
-	mutable FDisplayClusterMessage::DataType     SyncObjectsCache;
+	TMap<EDisplayClusterSyncGroup, TSet<IDisplayClusterClusterSyncObject*>> ObjectsToSync;
 	mutable FCriticalSection                     ObjectsToSyncCritSec;
 
 	// Sync events - types
@@ -135,6 +124,9 @@ private:
 	mutable FCriticalSection                     ClusterEventsCritSec;
 	FOnClusterEvent                              OnClusterEvent;
 	TArray<TScriptInterface<IDisplayClusterClusterEventListener>> ClusterEventListeners;
+	// Sync native input
+	FEvent* NativeInputDataAvailableEvent = nullptr;
+	TMap<FString, FString> NativeInputDataCache;
 
 	mutable FCriticalSection InternalsSyncScope;
 };
diff --git a/Engine/Plugins/Runtime/nDisplay/Source/DisplayCluster/Private/Cluster/IPDisplayClusterClusterManager.h b/Engine/Plugins/Runtime/nDisplay/Source/DisplayCluster/Private/Cluster/IPDisplayClusterClusterManager.h
--- a/Engine/Plugins/Runtime/nDisplay/Source/DisplayCluster/Private/Cluster/IPDisplayClusterClusterManager.h
+++ b/Engine/Plugins/Runtime/nDisplay/Source/DisplayCluster/Private/Cluster/IPDisplayClusterClusterManager.h
@@ -10,7 +10,6 @@
 #include "Network/DisplayClusterMessage.h"
 
 class IPDisplayClusterNodeController;
-class IDisplayClusterClusterSyncObject;
 class FJsonObject;
 
 
@@ -26,25 +25,17 @@ public:
 	{ }
 
 	virtual IPDisplayClusterNodeController* GetController() const = 0;
-	
-	virtual float GetDeltaTime() const = 0;
-	virtual void  SetDeltaTime(float deltaTime) = 0;
-
-	virtual void GetTimecode(FTimecode& timecode, FFrameRate& frameRate) const = 0;
-	virtual void SetTimecode(const FTimecode& timecode, const FFrameRate& frameRate) = 0;
-	
-	virtual void RegisterSyncObject  (IDisplayClusterClusterSyncObject* pSyncObj) = 0;
-	virtual void UnregisterSyncObject(IDisplayClusterClusterSyncObject* pSyncObj) = 0;
 
-	virtual void ExportSyncData(FDisplayClusterMessage::DataType& data) const = 0;
-	virtual void ImportSyncData(const FDisplayClusterMessage::DataType& data) = 0;
+	virtual void ExportSyncData(FDisplayClusterMessage::DataType& SyncData, EDisplayClusterSyncGroup SyncGroup) const = 0;
+	virtual void ImportSyncData(const FDisplayClusterMessage::DataType& SyncData, EDisplayClusterSyncGroup SyncGroup) = 0;
 
-	virtual void ExportEventsData(FDisplayClusterMessage::DataType& data) const = 0;
-	virtual void ImportEventsData(const FDisplayClusterMessage::DataType& data) = 0;
+	virtual void ExportEventsData(FDisplayClusterMessage::DataType& EventsData) const = 0;
+	virtual void ImportEventsData(const FDisplayClusterMessage::DataType& EventsData) = 0;
 
-	virtual void SyncObjects() = 0;
+	virtual void SyncObjects(EDisplayClusterSyncGroup SyncGroup) = 0;
 	virtual void SyncInput()   = 0;
 	virtual void SyncEvents()  = 0;
-
-	virtual void ClearSyncObjects() = 0;
+	
+	virtual void ProvideNativeInputData(const TMap<FString, FString>& NativeInputData) = 0;
+	virtual void SyncNativeInput(TMap<FString, FString>& NativeInputData) = 0;
 };
diff --git a/Engine/Plugins/Runtime/nDisplay/Source/DisplayCluster/Private/Config/DisplayClusterConfigManager.cpp b/Engine/Plugins/Runtime/nDisplay/Source/DisplayCluster/Private/Config/DisplayClusterConfigManager.cpp
--- a/Engine/Plugins/Runtime/nDisplay/Source/DisplayCluster/Private/Config/DisplayClusterConfigManager.cpp
+++ b/Engine/Plugins/Runtime/nDisplay/Source/DisplayCluster/Private/Config/DisplayClusterConfigManager.cpp
@@ -78,11 +78,6 @@ int32 FDisplayClusterConfigManager::GetClusterNodesAmount() const
 	return CfgClusterNodes.Num();
 }
 
-bool FDisplayClusterConfigManager::GetClusterNode(int32 idx, FDisplayClusterConfigClusterNode& node) const
-{
-	return GetItem(CfgClusterNodes, idx, node, FString("GetNode"));
-}
-
 bool FDisplayClusterConfigManager::GetClusterNode(const FString& id, FDisplayClusterConfigClusterNode& node) const
 {
 	return GetItem(CfgClusterNodes, id, node, FString("GetNode"));
@@ -154,11 +149,6 @@ int32 FDisplayClusterConfigManager::GetScreensAmount() const
 	return CfgScreens.Num();
 }
 
-bool FDisplayClusterConfigManager::GetScreen(int32 idx, FDisplayClusterConfigScreen& screen) const
-{
-	return GetItem(CfgScreens, idx, screen, FString("GetScreen"));
-}
-
 bool FDisplayClusterConfigManager::GetScreen(const FString& id, FDisplayClusterConfigScreen& screen) const
 {
 	return GetItem(CfgScreens, id, screen, FString("GetScreen"));
@@ -176,11 +166,6 @@ int32 FDisplayClusterConfigManager::GetCamerasAmount() const
 	return CfgCameras.Num();
 }
 
-bool FDisplayClusterConfigManager::GetCamera(int32 idx, FDisplayClusterConfigCamera& camera) const
-{
-	return GetItem(CfgCameras, idx, camera, FString("GetCamera"));
-}
-
 bool FDisplayClusterConfigManager::GetCamera(const FString& id, FDisplayClusterConfigCamera& camera) const
 {
 	return GetItem(CfgCameras, id, camera, FString("GetCamera"));
@@ -198,11 +183,6 @@ int32 FDisplayClusterConfigManager::GetViewportsAmount() const
 	return CfgViewports.Num();
 }
 
-bool FDisplayClusterConfigManager::GetViewport(int32 idx, FDisplayClusterConfigViewport& viewport) const
-{
-	return GetItem(CfgViewports, idx, viewport, FString("GetViewport"));
-}
-
 bool FDisplayClusterConfigManager::GetViewport(const FString& id, FDisplayClusterConfigViewport& viewport) const
 {
 	return GetItem(CfgViewports, id, viewport, FString("GetViewport"));
@@ -220,11 +200,6 @@ int32 FDisplayClusterConfigManager::GetPostprocessAmount() const
 	return CfgPostprocess.Num();
 }
 
-bool FDisplayClusterConfigManager::GetPostprocess(int32 idx, FDisplayClusterConfigPostprocess& postprocess) const
-{
-	return GetItem(CfgPostprocess, idx, postprocess, FString("GetPostprocess"));
-}
-
 bool FDisplayClusterConfigManager::GetPostprocess(const FString& id, FDisplayClusterConfigPostprocess& postprocess) const
 {
 	return GetItem(CfgPostprocess, id, postprocess, FString("GetPostprocess"));
@@ -242,11 +217,6 @@ int32 FDisplayClusterConfigManager::GetSceneNodesAmount() const
 	return CfgSceneNodes.Num();
 }
 
-bool FDisplayClusterConfigManager::GetSceneNode(int32 idx, FDisplayClusterConfigSceneNode& actor) const
-{
-	return GetItem(CfgSceneNodes, idx, actor, FString("GetActor"));
-}
-
 bool FDisplayClusterConfigManager::GetSceneNode(const FString& id, FDisplayClusterConfigSceneNode& actor) const
 {
 	return GetItem(CfgSceneNodes, id, actor, FString("GetActor"));
@@ -264,11 +234,6 @@ int32 FDisplayClusterConfigManager::GetInputDevicesAmount() const
 	return CfgInputDevices.Num();
 }
 
-bool FDisplayClusterConfigManager::GetInputDevice(int32 idx, FDisplayClusterConfigInput& input) const
-{
-	return GetItem(CfgInputDevices, idx, input, FString("GetInputDevice"));
-}
-
 bool FDisplayClusterConfigManager::GetInputDevice(const FString& id, FDisplayClusterConfigInput& input) const
 {
 	return GetItem(CfgInputDevices, id, input, FString("GetInputDevice"));
@@ -524,19 +489,6 @@ void FDisplayClusterConfigManager::ResetConfigData()
 	CfgCustom  = FDisplayClusterConfigCustom();
 }
 
-template <typename DataType>
-bool FDisplayClusterConfigManager::GetItem(const TArray<DataType>& container, uint32 idx, DataType& item, const FString& logHeader) const
-{
-	if (idx >= static_cast<uint32>(container.Num()))
-	{
-		UE_LOG(LogDisplayClusterConfig, Error, TEXT("%s: index is out of bound <%d>"), *logHeader, idx);
-		return false;
-	}
-
-	item = container[static_cast<int32>(idx)];
-	return true;
-}
-
 template <typename DataType>
 bool FDisplayClusterConfigManager::GetItem(const TArray<DataType>& container, const FString& id, DataType& item, const FString& logHeader) const
 {
diff --git a/Engine/Plugins/Runtime/nDisplay/Source/DisplayCluster/Private/Config/DisplayClusterConfigManager.h b/Engine/Plugins/Runtime/nDisplay/Source/DisplayCluster/Private/Config/DisplayClusterConfigManager.h
--- a/Engine/Plugins/Runtime/nDisplay/Source/DisplayCluster/Private/Config/DisplayClusterConfigManager.h
+++ b/Engine/Plugins/Runtime/nDisplay/Source/DisplayCluster/Private/Config/DisplayClusterConfigManager.h
@@ -39,7 +39,6 @@ public:
 	//////////////////////////////////////////////////////////////////////////////////////////////
 	virtual TArray<FDisplayClusterConfigClusterNode> GetClusterNodes() const override;
 	virtual int32 GetClusterNodesAmount() const override;
-	virtual bool GetClusterNode(int32 idx, FDisplayClusterConfigClusterNode& node) const override;
 	virtual bool GetClusterNode(const FString& id, FDisplayClusterConfigClusterNode& node) const override;
 	virtual bool GetMasterClusterNode(FDisplayClusterConfigClusterNode& node) const override;
 
@@ -50,32 +49,26 @@ public:
 
 	virtual TArray<FDisplayClusterConfigScreen> GetScreens() const override;
 	virtual int32 GetScreensAmount() const override;
-	virtual bool GetScreen(int32 idx, FDisplayClusterConfigScreen& screen) const override;
 	virtual bool GetScreen(const FString& id, FDisplayClusterConfigScreen& screen) const override;
 
 	virtual TArray<FDisplayClusterConfigCamera> GetCameras() const override;
 	virtual int32 GetCamerasAmount() const override;
-	virtual bool GetCamera(int32 idx, FDisplayClusterConfigCamera& camera) const override;
 	virtual bool GetCamera(const FString& id, FDisplayClusterConfigCamera& camera) const override;
 
 	virtual TArray<FDisplayClusterConfigViewport> GetViewports() const override;
 	virtual int32 GetViewportsAmount() const override;
-	virtual bool GetViewport(int32 idx, FDisplayClusterConfigViewport& viewport) const override;
 	virtual bool GetViewport(const FString& id, FDisplayClusterConfigViewport& viewport) const override;
 
 	virtual TArray<FDisplayClusterConfigPostprocess> GetPostprocess() const override;
 	virtual int32 GetPostprocessAmount() const override;
-	virtual bool GetPostprocess(int32 idx, FDisplayClusterConfigPostprocess& postprocess) const override;
 	virtual bool GetPostprocess(const FString& id, FDisplayClusterConfigPostprocess& postprocess) const override;
 
 	virtual TArray<FDisplayClusterConfigSceneNode> GetSceneNodes() const override;
 	virtual int32 GetSceneNodesAmount() const override;
-	virtual bool GetSceneNode(int32 idx, FDisplayClusterConfigSceneNode& actor) const override;
 	virtual bool GetSceneNode(const FString& id, FDisplayClusterConfigSceneNode& actor) const override;
 
 	virtual TArray<FDisplayClusterConfigInput> GetInputDevices() const override;
 	virtual int32 GetInputDevicesAmount() const override;
-	virtual bool GetInputDevice(int32 idx, FDisplayClusterConfigInput& input) const override;
 	virtual bool GetInputDevice(const FString& id, FDisplayClusterConfigInput& input) const override;
 
 	virtual TArray<FDisplayClusterConfigInputSetup> GetInputSetupRecords() const override;
@@ -146,9 +139,6 @@ private:
 		Xml
 	};
 
-	template <typename DataType>
-	bool GetItem(const TArray<DataType>& container, uint32 idx, DataType& item, const FString& logHeader) const;
-
 	template <typename DataType>
 	bool GetItem(const TArray<DataType>& container, const FString& id, DataType& item, const FString& logHeader) const;
 
diff --git a/Engine/Plugins/Runtime/nDisplay/Source/DisplayCluster/Private/Config/DisplayClusterConfigTypes.cpp b/Engine/Plugins/Runtime/nDisplay/Source/DisplayCluster/Private/Config/DisplayClusterConfigTypes.cpp
--- a/Engine/Plugins/Runtime/nDisplay/Source/DisplayCluster/Private/Config/DisplayClusterConfigTypes.cpp
+++ b/Engine/Plugins/Runtime/nDisplay/Source/DisplayCluster/Private/Config/DisplayClusterConfigTypes.cpp
@@ -37,11 +37,11 @@ FString FDisplayClusterConfigClusterNode::ToString() const
 		DisplayClusterStrings::cfg::data::Id,               *Id,
 		DisplayClusterStrings::cfg::data::cluster::Window,  *WindowId,
 		DisplayClusterStrings::cfg::data::cluster::Addr,    *Addr,
-		DisplayClusterStrings::cfg::data::cluster::Master,  DisplayClusterHelpers::str::BoolToStr(IsMaster),
+		DisplayClusterStrings::cfg::data::cluster::Master,  *DisplayClusterHelpers::str::BoolToStr(IsMaster),
 		DisplayClusterStrings::cfg::data::cluster::PortCS,  Port_CS,
 		DisplayClusterStrings::cfg::data::cluster::PortSS,  Port_SS,
 		DisplayClusterStrings::cfg::data::cluster::PortCE,  Port_CE,
-		DisplayClusterStrings::cfg::data::cluster::Sound,   DisplayClusterHelpers::str::BoolToStr(SoundEnabled)
+		DisplayClusterStrings::cfg::data::cluster::Sound,   *DisplayClusterHelpers::str::BoolToStr(SoundEnabled)
 	);
 }
 
@@ -70,7 +70,7 @@ FString FDisplayClusterConfigWindow::ToString() const
 		DisplayClusterStrings::cfg::data::Id,                  *Id,
 		DisplayClusterStrings::cfg::data::window::Viewports,   *DisplayClusterHelpers::str::ArrayToStr(ViewportIds),
 		DisplayClusterStrings::cfg::data::window::Postprocess, *DisplayClusterHelpers::str::ArrayToStr(PostprocessIds),
-		DisplayClusterStrings::cfg::data::window::Fullscreen,  DisplayClusterHelpers::str::BoolToStr(IsFullscreen),
+		DisplayClusterStrings::cfg::data::window::Fullscreen,  *DisplayClusterHelpers::str::BoolToStr(IsFullscreen),
 		DisplayClusterStrings::cfg::data::window::WinX,        WinX,
 		DisplayClusterStrings::cfg::data::window::WinY,        WinY,
 		DisplayClusterStrings::cfg::data::window::ResX,        ResX,
@@ -107,7 +107,7 @@ FString FDisplayClusterConfigViewport::ToString() const
 		DisplayClusterStrings::cfg::data::viewport::PosY,        Loc.Y,
 		DisplayClusterStrings::cfg::data::viewport::Width,       Size.X,
 		DisplayClusterStrings::cfg::data::viewport::Height,      Size.Y,
-		DisplayClusterStrings::cfg::data::viewport::RTT,         DisplayClusterHelpers::str::BoolToStr(IsRTT),
+		DisplayClusterStrings::cfg::data::viewport::RTT,         *DisplayClusterHelpers::str::BoolToStr(IsRTT),
 		DisplayClusterStrings::cfg::data::viewport::BufferRatio, BufferRatio
 	);
 }
@@ -202,7 +202,7 @@ FString FDisplayClusterConfigCamera::ToString() const
 	return FString::Printf(TEXT("[%s + %s=%f, %s=%s, %s=%d]"),
 		*FDisplayClusterConfigSceneNode::ToString(),
 		DisplayClusterStrings::cfg::data::camera::EyeDist,     EyeDist,
-		DisplayClusterStrings::cfg::data::camera::EyeSwap,     DisplayClusterHelpers::str::BoolToStr(EyeSwap),
+		DisplayClusterStrings::cfg::data::camera::EyeSwap,     *DisplayClusterHelpers::str::BoolToStr(EyeSwap),
 		DisplayClusterStrings::cfg::data::camera::ForceOffset, ForceOffset);
 }
 
@@ -265,14 +265,16 @@ bool FDisplayClusterConfigInputSetup::DeserializeFromString(const FString& line)
 //////////////////////////////////////////////////////////////////////////////////////////////
 FString FDisplayClusterConfigGeneral::ToString() const
 {
-	return FString::Printf(TEXT("[%s + %s=%d]"),
+	return FString::Printf(TEXT("[%s + %s=%d, %s=%d]"),
 		*FDisplayClusterConfigBase::ToString(),
-		DisplayClusterStrings::cfg::data::general::SwapSyncPolicy, SwapSyncPolicy);
+		DisplayClusterStrings::cfg::data::general::SwapSyncPolicy, SwapSyncPolicy,
+		DisplayClusterStrings::cfg::data::general::UnrealInputSyncPolicy, NativeInputSyncPolicy);
 }
 
 bool FDisplayClusterConfigGeneral::DeserializeFromString(const FString& line)
 {
-	DisplayClusterHelpers::str::ExtractValue(line, FString(DisplayClusterStrings::cfg::data::general::SwapSyncPolicy), SwapSyncPolicy);
+	DisplayClusterHelpers::str::ExtractValue(line, FString(DisplayClusterStrings::cfg::data::general::SwapSyncPolicy),        SwapSyncPolicy);
+	DisplayClusterHelpers::str::ExtractValue(line, FString(DisplayClusterStrings::cfg::data::general::UnrealInputSyncPolicy), NativeInputSyncPolicy);
 	return FDisplayClusterConfigBase::DeserializeFromString(line);
 }
 
@@ -339,9 +341,9 @@ FString FDisplayClusterConfigDebug::ToString() const
 {
 	return FString::Printf(TEXT("[%s + %s=%s, %s=%s, %s=%f]"),
 		*FDisplayClusterConfigBase::ToString(),
-		DisplayClusterStrings::cfg::data::debug::DrawStats, DisplayClusterHelpers::str::BoolToStr(DrawStats),
-		DisplayClusterStrings::cfg::data::debug::LagSim,  DisplayClusterHelpers::str::BoolToStr(LagSimulateEnabled),
-		DisplayClusterStrings::cfg::data::debug::LagTime, LagMaxTime);
+		DisplayClusterStrings::cfg::data::debug::DrawStats, *DisplayClusterHelpers::str::BoolToStr(DrawStats),
+		DisplayClusterStrings::cfg::data::debug::LagSim,    *DisplayClusterHelpers::str::BoolToStr(LagSimulateEnabled),
+		DisplayClusterStrings::cfg::data::debug::LagTime,   LagMaxTime);
 }
 
 bool FDisplayClusterConfigDebug::DeserializeFromString(const FString& line)
diff --git a/Engine/Plugins/Runtime/nDisplay/Source/DisplayCluster/Private/DisplayClusterHelpers.h b/Engine/Plugins/Runtime/nDisplay/Source/DisplayCluster/Private/DisplayClusterHelpers.h
new file mode 100644
--- /dev/null
+++ b/Engine/Plugins/Runtime/nDisplay/Source/DisplayCluster/Private/DisplayClusterHelpers.h
@@ -0,0 +1,12 @@
+// Copyright 1998-2019 Epic Games, Inc. All Rights Reserved.
+
+#pragma once
+
+#include "CoreMinimal.h"
+#include "DisplayClusterBuildConfig.h"
+#include "DisplayClusterUtils\DisplayClusterCommonHelpers.h"
+
+
+namespace DisplayClusterHelpers
+{
+};
\ No newline at end of file
diff --git a/Engine/Plugins/Runtime/nDisplay/Source/DisplayCluster/Private/DisplayClusterModule.cpp b/Engine/Plugins/Runtime/nDisplay/Source/DisplayCluster/Private/DisplayClusterModule.cpp
--- a/Engine/Plugins/Runtime/nDisplay/Source/DisplayCluster/Private/DisplayClusterModule.cpp
+++ b/Engine/Plugins/Runtime/nDisplay/Source/DisplayCluster/Private/DisplayClusterModule.cpp
@@ -192,6 +192,34 @@ void FDisplayClusterModule::EndScene()
 	}
 }
 
+void FDisplayClusterModule::StartFrame(uint64 FrameNum)
+{
+	DISPLAY_CLUSTER_FUNC_TRACE(LogDisplayClusterModule);
+
+	UE_LOG(LogDisplayClusterModule, Verbose, TEXT("StartFrame: frame num - %llu"), FrameNum);
+
+	for (auto pMgr : Managers)
+	{
+		pMgr->StartFrame(FrameNum);
+	}
+
+	DisplayClusterStartFrameEvent.Broadcast(FrameNum);
+}
+
+void FDisplayClusterModule::EndFrame(uint64 FrameNum)
+{
+	DISPLAY_CLUSTER_FUNC_TRACE(LogDisplayClusterModule);
+
+	UE_LOG(LogDisplayClusterModule, Verbose, TEXT("EndFrame: frame num - %llu"), FrameNum);
+
+	for (auto pMgr : Managers)
+	{
+		pMgr->EndFrame(FrameNum);
+	}
+
+	DisplayClusterEndFrameEvent.Broadcast(FrameNum);
+}
+
 void FDisplayClusterModule::PreTick(float DeltaSeconds)
 {
 	DISPLAY_CLUSTER_FUNC_TRACE(LogDisplayClusterModule);
@@ -206,4 +234,32 @@ void FDisplayClusterModule::PreTick(float DeltaSeconds)
 	DisplayClusterPreTickEvent.Broadcast();
 }
 
+void FDisplayClusterModule::Tick(float DeltaSeconds)
+{
+	DISPLAY_CLUSTER_FUNC_TRACE(LogDisplayClusterModule);
+
+	UE_LOG(LogDisplayClusterModule, Verbose, TEXT("Tick: delta time - %f"), DeltaSeconds);
+
+	for (auto pMgr : Managers)
+	{
+		pMgr->Tick(DeltaSeconds);
+	}
+
+	DisplayClusterTickEvent.Broadcast();
+}
+
+void FDisplayClusterModule::PostTick(float DeltaSeconds)
+{
+	DISPLAY_CLUSTER_FUNC_TRACE(LogDisplayClusterModule);
+
+	UE_LOG(LogDisplayClusterModule, Verbose, TEXT("PostTick: delta time - %f"), DeltaSeconds);
+
+	for (auto pMgr : Managers)
+	{
+		pMgr->PostTick(DeltaSeconds);
+	}
+
+	DisplayClusterPostTickEvent.Broadcast();
+}
+
 IMPLEMENT_MODULE(FDisplayClusterModule, DisplayCluster)
diff --git a/Engine/Plugins/Runtime/nDisplay/Source/DisplayCluster/Private/DisplayClusterModule.h b/Engine/Plugins/Runtime/nDisplay/Source/DisplayCluster/Private/DisplayClusterModule.h
--- a/Engine/Plugins/Runtime/nDisplay/Source/DisplayCluster/Private/DisplayClusterModule.h
+++ b/Engine/Plugins/Runtime/nDisplay/Source/DisplayCluster/Private/DisplayClusterModule.h
@@ -57,8 +57,11 @@ public:
 	virtual void EndSession() override;
 	virtual bool StartScene(UWorld* pWorld) override;
 	virtual void EndScene() override;
+	virtual void StartFrame(uint64 FrameNum) override;
 	virtual void PreTick(float DeltaSeconds) override;
-
+	virtual void Tick(float DeltaSeconds) override;
+	virtual void PostTick(float DeltaSeconds) override;
+	virtual void EndFrame(uint64 FrameNum) override;
 
 public:
 	virtual FDisplayClusterStartSessionEvent& OnDisplayClusterStartSession() override
@@ -67,13 +70,29 @@ public:
 	virtual FDisplayClusterEndSessionEvent& OnDisplayClusterEndSession() override
 	{ return DisplayClusterEndSessionEvent; }
 
+	virtual FDisplayClusterStartFrameEvent& OnDisplayClusterStartFrame(uint64 FrameNum) override
+	{ return DisplayClusterStartFrameEvent; }
+
+	virtual FDisplayClusterEndFrameEvent& OnDisplayClusterEndFrame(uint64 FrameNum) override
+	{ return DisplayClusterEndFrameEvent; }
+
 	virtual FDisplayClusterPreTickEvent& OnDisplayClusterPreTick() override
 	{ return DisplayClusterPreTickEvent; }
 
+	virtual FDisplayClusterTickEvent& OnDisplayClusterTick() override
+	{ return DisplayClusterTickEvent; }
+
+	virtual FDisplayClusterPostTickEvent& OnDisplayClusterPostTick() override
+	{ return DisplayClusterPostTickEvent; }
+
 private:
 	FDisplayClusterStartSessionEvent         DisplayClusterStartSessionEvent;
 	FDisplayClusterEndSessionEvent           DisplayClusterEndSessionEvent;
+	FDisplayClusterStartFrameEvent           DisplayClusterStartFrameEvent;
+	FDisplayClusterEndFrameEvent             DisplayClusterEndFrameEvent;
 	FDisplayClusterPreTickEvent              DisplayClusterPreTickEvent;
+	FDisplayClusterTickEvent                 DisplayClusterTickEvent;
+	FDisplayClusterPostTickEvent             DisplayClusterPostTickEvent;
 
 private:
 	//////////////////////////////////////////////////////////////////////////////////////////////
diff --git a/Engine/Plugins/Runtime/nDisplay/Source/DisplayCluster/Private/Game/Classes/Basics/DisplayClusterGameEngine.cpp b/Engine/Plugins/Runtime/nDisplay/Source/DisplayCluster/Private/Game/Classes/Basics/DisplayClusterGameEngine.cpp
--- a/Engine/Plugins/Runtime/nDisplay/Source/DisplayCluster/Private/Game/Classes/Basics/DisplayClusterGameEngine.cpp
+++ b/Engine/Plugins/Runtime/nDisplay/Source/DisplayCluster/Private/Game/Classes/Basics/DisplayClusterGameEngine.cpp
@@ -15,23 +15,30 @@
 #include "Misc/DisplayClusterHelpers.h"
 
 #include "DisplayClusterBuildConfig.h"
+#include "DisplayClusterEnums.h"
 #include "DisplayClusterGlobals.h"
 #include "DisplayClusterLog.h"
 #include "DisplayClusterStrings.h"
 
 #include "Stats/Stats.h"
 
+#include "Misc/CoreDelegates.h"
+
 
 void UDisplayClusterGameEngine::Init(class IEngineLoop* InEngineLoop)
 {
 	DISPLAY_CLUSTER_FUNC_TRACE(LogDisplayClusterEngine);
 
-	// Instantiate the tickable helper
-	TickableHelper = NewObject<UDisplayClusterGameEngineTickableHelper>();
-
 	// Detect requested operation mode
 	OperationMode = DetectOperationMode();
 
+	if (OperationMode == EDisplayClusterOperationMode::Cluster ||
+		OperationMode == EDisplayClusterOperationMode::Standalone)
+	{
+		// Instantiate the tickable helper
+		TickableHelper = NewObject<UDisplayClusterGameEngineTickableHelper>();
+	}
+
 	// Initialize Display Cluster
 	if (!GDisplayCluster->Init(OperationMode))
 	{
@@ -152,41 +159,33 @@ bool UDisplayClusterGameEngine::LoadMap(FWorldContext& WorldContext, FURL URL, c
 {
 	DISPLAY_CLUSTER_FUNC_TRACE(LogDisplayClusterEngine);
 
-	// Perform map loading
-	if (!Super::LoadMap(WorldContext, URL, Pending, Error))
-	{
-		return false;
-	}
-
 	if (OperationMode == EDisplayClusterOperationMode::Cluster ||
 		OperationMode == EDisplayClusterOperationMode::Standalone)
 	{
-		// Game start barrier
-		NodeController->WaitForGameStart();
-	}
+		// Finish previous scene
+		GDisplayCluster->EndScene();
 
-	return true;
-}
+		// Perform map loading
+		if (!Super::LoadMap(WorldContext, URL, Pending, Error))
+		{
+			return false;
+		}
 
-void UDisplayClusterGameEngine::RegisterSyncObject(IDisplayClusterClusterSyncObject* SyncObj)
-{
-	if (ClusterMgr)
-	{
-		ClusterMgr->RegisterSyncObject(SyncObj);
-	}
-}
+		// Start new scene
+		GDisplayCluster->StartScene(WorldContext.World());
 
-void UDisplayClusterGameEngine::UnregisterSyncObject(IDisplayClusterClusterSyncObject* SyncObj)
-{
-	if (ClusterMgr)
+		// Game start barrier
+		if (NodeController)
+		{
+			NodeController->WaitForGameStart();
+		}
+	}
+	else
 	{
-		ClusterMgr->UnregisterSyncObject(SyncObj);
+		return Super::LoadMap(WorldContext, URL, Pending, Error);
 	}
-}
 
-bool UDisplayClusterGameEngine::IsMaster() const
-{
-	return NodeController && NodeController->IsMaster();
+	return true;
 }
 
 void UDisplayClusterGameEngine::Tick(float DeltaSeconds, bool bIdleMode)
@@ -199,45 +198,36 @@ void UDisplayClusterGameEngine::Tick(float DeltaSeconds, bool bIdleMode)
 		FTimecode Timecode;
 		FFrameRate FrameRate;
 
-		// Update input device state (master only)
-		InputMgr->Update();
-
-		// Update delta time. Cluster slaves will get this value from the master few steps later
-		ClusterMgr->SetDeltaTime(DeltaSeconds);
-
 		//////////////////////////////////////////////////////////////////////////////////////////////
 		// Frame start barrier
 		NodeController->WaitForFrameStart();
 		UE_LOG(LogDisplayClusterEngine, Verbose, TEXT("Sync frame start"));
 
+		// Perform StartFrame notification
+		GDisplayCluster->StartFrame(GFrameCounter);
 
-		// Get DisplayCluster time delta
+		// Sync DeltaSeconds
 		NodeController->GetDeltaTime(DeltaSeconds);
-		NodeController->GetTimecode(Timecode, FrameRate);
-		UE_LOG(LogDisplayClusterEngine, Verbose, TEXT("DisplayCluster delta time (seconds): %f"), DeltaSeconds);
-		UE_LOG(LogDisplayClusterEngine, Verbose, TEXT("DisplayCluster Timecode: %s | %s"), *Timecode.ToString(), *FrameRate.ToPrettyText().ToString());
-
-		// Update delta time in the application
 		FApp::SetDeltaTime(DeltaSeconds);
-		FApp::SetTimecodeAndFrameRate(Timecode, FrameRate);
+		UE_LOG(LogDisplayClusterEngine, Verbose, TEXT("DisplayCluster delta seconds: %f"), DeltaSeconds);
 
-		// Update input state in the cluster
-		ClusterMgr->SyncInput();
+		// Sync timecode and framerate
+		NodeController->GetTimecode(Timecode, FrameRate);
+		FApp::SetTimecodeAndFrameRate(Timecode, FrameRate);
+		UE_LOG(LogDisplayClusterEngine, Verbose, TEXT("DisplayCluster timecode: %s | %s"), *Timecode.ToString(), *FrameRate.ToPrettyText().ToString());
 
 		// Perform PreTick for DisplayCluster module
 		UE_LOG(LogDisplayClusterEngine, Verbose, TEXT("Perform PreTick()"));
 		GDisplayCluster->PreTick(DeltaSeconds);
 
-		// Sync cluster objects
-		ClusterMgr->SyncObjects();
-
-		// Sync cluster events
-		ClusterMgr->SyncEvents();
-
-		// Perform Tick() calls for scene actors
-		UE_LOG(LogDisplayClusterEngine, Verbose, TEXT("Perform Tick()"));
+		// Perform UGameEngine::Tick() calls for scene actors
+		UE_LOG(LogDisplayClusterEngine, Verbose, TEXT("Perform UGameEngine::Tick()"));
 		Super::Tick(DeltaSeconds, bIdleMode);
 
+		// Perform PostTick for DisplayCluster module
+		UE_LOG(LogDisplayClusterEngine, Verbose, TEXT("Perform PostTick()"));
+		GDisplayCluster->PostTick(DeltaSeconds);
+
 		if (CfgDebug.LagSimulateEnabled)
 		{
 			const float lag = CfgDebug.LagMaxTime;
@@ -258,7 +248,9 @@ void UDisplayClusterGameEngine::Tick(float DeltaSeconds, bool bIdleMode)
 		//////////////////////////////////////////////////////////////////////////////////////////////
 		// Frame end barrier
 		NodeController->WaitForFrameEnd();
-		ClusterMgr->ClearSyncObjects();
+
+		// Perform EndFrame notification
+		GDisplayCluster->EndFrame(GFrameCounter);
 
 		UE_LOG(LogDisplayClusterEngine, Verbose, TEXT("Sync frame end"));
 	}
@@ -268,18 +260,19 @@ void UDisplayClusterGameEngine::Tick(float DeltaSeconds, bool bIdleMode)
 	}
 }
 
-
 TStatId UDisplayClusterGameEngineTickableHelper::GetStatId() const
 {
 	RETURN_QUICK_DECLARE_CYCLE_STAT(UDisplayClusterGameEngineTickableHelper, STATGROUP_Tickables);
 }
 
-void UDisplayClusterGameEngineTickableHelper::Tick(float DeltaTime)
+void UDisplayClusterGameEngineTickableHelper::Tick(float DeltaSeconds)
 {
-	// Sync cluster objects
-	static IPDisplayClusterClusterManager* const ClusterMgr = GDisplayCluster->GetPrivateClusterMgr();
-	if (ClusterMgr)
+	static const EDisplayClusterOperationMode OperationMode = GDisplayCluster->GetOperationMode();
+
+	if (OperationMode == EDisplayClusterOperationMode::Cluster ||
+		OperationMode == EDisplayClusterOperationMode::Standalone)
 	{
-		ClusterMgr->SyncObjects();
+		UE_LOG(LogDisplayClusterEngine, Verbose, TEXT("Perform Tick()"));
+		GDisplayCluster->Tick(DeltaSeconds);
 	}
 }
diff --git a/Engine/Plugins/Runtime/nDisplay/Source/DisplayCluster/Private/Game/Classes/Basics/DisplayClusterPlayerInput.cpp b/Engine/Plugins/Runtime/nDisplay/Source/DisplayCluster/Private/Game/Classes/Basics/DisplayClusterPlayerInput.cpp
new file mode 100644
--- /dev/null
+++ b/Engine/Plugins/Runtime/nDisplay/Source/DisplayCluster/Private/Game/Classes/Basics/DisplayClusterPlayerInput.cpp
@@ -0,0 +1,392 @@
+// Copyright 1998-2019 Epic Games, Inc. All Rights Reserved.
+
+#include "DisplayClusterPlayerInput.h"
+#include "Components/InputComponent.h"
+
+#include "Cluster/IPDisplayClusterClusterManager.h"
+#include "DisplayClusterGlobals.h"
+#include "DisplayClusterHelpers.h"
+#include "DisplayClusterLog.h"
+#include "DisplayClusterUtils/DisplayClusterTypesConverter.h"
+
+
+UDisplayClusterPlayerInput::UDisplayClusterPlayerInput()
+	: Super()
+{
+}
+
+void UDisplayClusterPlayerInput::ProcessInputStack(const TArray<UInputComponent*>& InputComponentStack, const float DeltaTime, const bool bGamePaused)
+{
+	UE_LOG(LogDisplayClusterGame, Verbose, TEXT("Processing input stack..."));
+
+	static IPDisplayClusterClusterManager* const ClusterMgr = GDisplayCluster->GetPrivateClusterMgr();
+	if (ClusterMgr && ClusterMgr->IsCluster())
+	{
+		if (ClusterMgr->IsMaster())
+		{
+			// Build master auxiliary data
+			BuildDelegatesMap(InputComponentStack, true);
+			// Start procssing the stack
+			Super::ProcessInputStack(InputComponentStack, DeltaTime, bGamePaused);
+		}
+		else
+		{
+			// Build slave auxiliary data
+			BuildDelegatesMap(InputComponentStack, false);
+			// Get native input data from master
+			TMap<FString, FString> NativeInputData;
+			ClusterMgr->SyncNativeInput(NativeInputData);
+			// Process native input data got from master
+			ImportNativeInputData(NativeInputData, InputComponentStack);
+		}
+	}
+}
+
+void UDisplayClusterPlayerInput::ProcessDelegatesFilter(const TArray<UInputComponent*>& InputComponentStack, TArray<FAxisDelegateDetails>& AxisDelegates, TArray<FVectorAxisDelegateDetails>& VectorAxisDelegates, TArray<FDelegateDispatchDetails>& NonAxisDelegates)
+{
+	UE_LOG(LogDisplayClusterGame, Verbose, TEXT("Processing input delegates..."));
+
+	static IPDisplayClusterClusterManager* const ClusterMgr = GDisplayCluster->GetPrivateClusterMgr();
+
+	// Serialize native input data
+	TMap<FString, FString> NativeInputData;
+	ExportNativeInputData(NativeInputData, InputComponentStack, AxisDelegates, VectorAxisDelegates, NonAxisDelegates);
+	// Provide data to slave nodes
+	ClusterMgr->ProvideNativeInputData(NativeInputData);
+}
+
+void UDisplayClusterPlayerInput::BuildDelegatesMap(const TArray<UInputComponent*>& InputComponentStack, bool IsMaster)
+{
+	CleanDelegateMap(IsMaster);
+
+	// Master related logic
+	if (IsMaster)
+	{
+		for (UInputComponent* IC : InputComponentStack)
+		{
+			// Cache axis delegates
+			for (const FInputAxisBinding& AxisBinding : IC->AxisBindings)
+			{
+				void* Addr1 = (void*)&AxisBinding.AxisDelegate.GetDelegate();
+				void* Addr2 = (void*)&AxisBinding.AxisDelegate.GetDynamicDelegate();
+
+				if (!AxisDelegateNames.Contains(Addr1))
+				{
+					AxisDelegateNames.Emplace(Addr1);
+				}
+
+				AxisDelegateNames[Addr1].Emplace(Addr2, AxisBinding.AxisName);
+			}
+
+			// Cache gesture delegates
+			for (const FInputGestureBinding& GestureBinding : IC->GestureBindings)
+			{
+				void* Addr1 = (void*)&GestureBinding.GestureDelegate.GetDelegate();
+				void* Addr2 = (void*)&GestureBinding.GestureDelegate.GetDynamicDelegate();
+
+				if (!GestureDelegateNames.Contains(Addr1))
+				{
+					GestureDelegateNames.Emplace(Addr1);
+				}
+
+				GestureDelegateNames[Addr1].Emplace(Addr2, GestureBinding.GestureKey.GetFName());
+			}
+
+			// Cache vector axis delegates
+			for (const FInputVectorAxisBinding& VectorAxisBinding : IC->VectorAxisBindings)
+			{
+				void* Addr1 = (void*)&VectorAxisBinding.AxisDelegate.GetDelegate();
+				void* Addr2 = (void*)&VectorAxisBinding.AxisDelegate.GetDynamicDelegate();
+
+				if (!VectorAxisDelegateNames.Contains(Addr1))
+				{
+					VectorAxisDelegateNames.Emplace(Addr1);
+				}
+
+				VectorAxisDelegateNames[Addr1].Emplace(Addr2, VectorAxisBinding.AxisKey.GetFName());
+			}
+		}
+	}
+	// Slave related logic
+	else
+	{
+		for (UInputComponent* IC : InputComponentStack)
+		{
+			// Cache axis delegates
+			for (const FInputAxisBinding& AxisBinding : IC->AxisBindings)
+			{
+				AxisDelegateInstances.Emplace(AxisBinding.AxisName, &AxisBinding.AxisDelegate);
+			}
+
+			// Cache action delegates
+			const int32 ActionBindingsAmount = IC->GetNumActionBindings();
+			for (int i = 0; i < ActionBindingsAmount; ++i)
+			{
+				const FInputActionBinding& ActionBinding = IC->GetActionBinding(i);
+				const FName ActionName = ActionBinding.GetActionName();
+				if (!ActionDelegateInstances.Contains(ActionName))
+				{
+					ActionDelegateInstances.Emplace(ActionName);
+				}
+
+				const EInputEvent InputEvent = ActionBinding.KeyEvent.GetValue();
+				if (!ActionDelegateInstances[ActionName].Contains(InputEvent))
+				{
+					ActionDelegateInstances[ActionName].Emplace(InputEvent);
+				}
+
+				ActionDelegateInstances[ActionName][InputEvent] = &ActionBinding.ActionDelegate;
+			}
+
+			// Cache gesture delegates
+			for (const FInputGestureBinding& GestureBinding : IC->GestureBindings)
+			{
+				const FName GestureKey = GestureBinding.GestureKey.GetFName();
+				GestureDelegateInstances.Emplace(GestureKey, &GestureBinding.GestureDelegate);
+			}
+
+			// Cache vector axis delegates
+			for (const FInputVectorAxisBinding& VectorAxisBinding : IC->VectorAxisBindings)
+			{
+				const FName AxisKey = VectorAxisBinding.AxisKey.GetFName();
+				VectorAxisDelegateInstances.Emplace(AxisKey, &VectorAxisBinding.AxisDelegate);
+			}
+		}
+	}
+
+	// For any node type
+	for (UInputComponent* IC : InputComponentStack)
+	{
+		// Cache touch delegates
+		for (const FInputTouchBinding& TouchBinding : IC->TouchBindings)
+		{
+			const EInputEvent InputEvent = TouchBinding.KeyEvent.GetValue();
+			TouchDelegateInstances.Emplace(InputEvent, &TouchBinding.TouchDelegate);
+		}
+	}
+}
+
+void UDisplayClusterPlayerInput::CleanDelegateMap(bool IsMaster)
+{
+	if (IsMaster)
+	{
+		AxisDelegateNames.Reset();
+		GestureDelegateNames.Reset();
+		VectorAxisDelegateNames.Reset();
+	}
+	else
+	{
+		AxisDelegateInstances.Reset();
+		ActionDelegateInstances.Reset();
+		GestureDelegateInstances.Reset();
+		VectorAxisDelegateInstances.Reset();
+	}
+
+	TouchDelegateInstances.Reset();
+}
+
+void UDisplayClusterPlayerInput::ExportNativeInputData(TMap<FString, FString>& InputData, const TArray<UInputComponent*>& InputComponentStack, TArray<FAxisDelegateDetails>& AxisDelegates, TArray<FVectorAxisDelegateDetails>& VectorAxisDelegates, TArray<FDelegateDispatchDetails>& NonAxisDelegates)
+{
+	UE_LOG(LogDisplayClusterGame, Verbose, TEXT("Exporting native input data..."));
+
+	InputData.Reset();
+
+	InputData.Emplace(FString("action"), SerializeActionData(InputComponentStack, NonAxisDelegates));
+	InputData.Emplace(FString("axis"),   SerializeAxisData(InputComponentStack,   AxisDelegates));
+	InputData.Emplace(FString("vector"), SerializeVectorAxisData(InputComponentStack, VectorAxisDelegates));
+
+	UE_LOG(LogDisplayClusterGame, Verbose, TEXT("[EXP] Native output action: %s"), *InputData[FString("action")]);
+	UE_LOG(LogDisplayClusterGame, Verbose, TEXT("[EXP] Native output axis  : %s"), *InputData[FString("axis")]);
+	UE_LOG(LogDisplayClusterGame, Verbose, TEXT("[EXP] Native output vector: %s"), *InputData[FString("vector")]);
+}
+
+void UDisplayClusterPlayerInput::ImportNativeInputData(const TMap<FString, FString>& InputData, const TArray<UInputComponent*>& InputComponentStack)
+{
+	UE_LOG(LogDisplayClusterGame, Verbose, TEXT("Importing native input data..."));
+
+	check(InputData.Num() == 3);
+
+	UE_LOG(LogDisplayClusterGame, Verbose, TEXT("[IMP] Native output action: %s"), *InputData[FString("action")]);
+	UE_LOG(LogDisplayClusterGame, Verbose, TEXT("[IMP] Native output axis  : %s"), *InputData[FString("axis")]);
+	UE_LOG(LogDisplayClusterGame, Verbose, TEXT("[IMP] Native output vector: %s"), *InputData[FString("vector")]);
+
+	DeserializeAndProcessActionInput    (InputData[FString("action")]);
+	DeserializeAndProcessAxisInput      (InputData[FString("axis")]);
+	DeserializeAndProcessVectorAxisInput(InputData[FString("vector")]);
+}
+
+FString UDisplayClusterPlayerInput::SerializeAxisData(const TArray<UInputComponent*>& InputComponentStack, const TArray<FAxisDelegateDetails>& AxisDelegates)
+{
+	FString OutData;
+	OutData.Reserve(512);
+
+	for (const FAxisDelegateDetails& Details : AxisDelegates)
+	{
+		const void* const Addr1 = (void*)&Details.Delegate.GetDelegate();
+		const void* const Addr2 = (void*)&Details.Delegate.GetDynamicDelegate();
+
+		if (AxisDelegateNames.Contains(Addr1) && AxisDelegateNames[Addr1].Contains(Addr2))
+		{
+			OutData += FString::Printf(TEXT("%s=%f;"), *AxisDelegateNames[Addr1][Addr2].ToString(), Details.Value);
+		}
+	}
+
+	return OutData;
+}
+
+FString UDisplayClusterPlayerInput::SerializeVectorAxisData(const TArray<UInputComponent*>& InputComponentStack, const TArray<FVectorAxisDelegateDetails>& VectorAxisDelegates)
+{
+	FString OutData;
+	OutData.Reserve(512);
+
+	for (const FVectorAxisDelegateDetails& Details : VectorAxisDelegates)
+	{
+		void* Addr1 = (void*)&Details.Delegate.GetDelegate();
+		void* Addr2 = (void*)&Details.Delegate.GetDynamicDelegate();
+
+		if (VectorAxisDelegateNames.Contains(Addr1) && VectorAxisDelegateNames[Addr1].Contains(Addr2))
+		{
+			OutData += FString::Printf(TEXT("%s=%s;"), *VectorAxisDelegateNames[Addr1][Addr2].ToString(), *FDisplayClusterTypesConverter::ToString(Details.Value));
+		}
+	}
+
+	return OutData;
+}
+
+FString UDisplayClusterPlayerInput::SerializeActionData(const TArray<UInputComponent*>& InputComponentStack, const TArray<FDelegateDispatchDetails>& NonAxisDelegates)
+{
+	FString OutData;
+	OutData.Reserve(512);
+
+	for (const FDelegateDispatchDetails& Details : NonAxisDelegates)
+	{
+		if (Details.SourceAction)
+		{
+			const FString ActionName = Details.SourceAction->GetActionName().ToString();
+			if (Details.ActionDelegate.IsBound())
+			{
+				OutData += FString::Printf(TEXT("key=%s|%s|%d;"), *ActionName, *Details.Chord.Key.GetFName().ToString(), (int)Details.SourceAction->KeyEvent.GetValue());
+			}
+			else if (Details.TouchDelegate.IsBound())
+			{
+				OutData += FString::Printf(TEXT("touch=%d|%s|%d;"), (int)Details.KeyEvent.GetValue(), *FDisplayClusterTypesConverter::ToString(Details.TouchLocation), (int)Details.FingerIndex);
+			}
+			else if (Details.GestureDelegate.IsBound())
+			{
+				const void* const Addr1 = (void*)&Details.GestureDelegate.GetDelegate();
+				const void* const Addr2 = (void*)&Details.GestureDelegate.GetDynamicDelegate();
+
+				if (GestureDelegateNames.Contains(Addr1) && GestureDelegateNames[Addr1].Contains(Addr2))
+				{
+					OutData += FString::Printf(TEXT("gesture=%s|%f;"), *GestureDelegateNames[Addr1][Addr2].ToString(), Details.GestureValue);
+				}
+			}
+		}
+	}
+
+	return OutData;
+}
+
+void UDisplayClusterPlayerInput::DeserializeAndProcessAxisInput(const FString& Data)
+{
+	const TArray<FString> AxisRecords = DisplayClusterHelpers::str::StrToArray<FString>(Data, FString(";"));
+	for (const FString& AxisRecord : AxisRecords)
+	{
+		if (AxisRecord.Len() > 0)
+		{
+			const TArray<FString> AxisParams = DisplayClusterHelpers::str::StrToArray<FString>(AxisRecord, FString(TEXT("=")));
+			check(AxisParams.Num() == 2);
+
+			const FName AxisName = *AxisParams[0];
+			if (AxisDelegateInstances.Contains(AxisName))
+			{
+				const float AxisValue = FDisplayClusterTypesConverter::FromString<float>(AxisParams[1]);
+				UE_LOG(LogDisplayClusterGame, Verbose, TEXT("Native input - axis %s, value %f"), *AxisName.ToString(), AxisValue);
+				AxisDelegateInstances[AxisName]->Execute(AxisValue);
+			}
+		}
+	}
+}
+
+void UDisplayClusterPlayerInput::DeserializeAndProcessVectorAxisInput(const FString& Data)
+{
+	const TArray<FString> VectorRecords = DisplayClusterHelpers::str::StrToArray<FString>(Data, FString(";"));
+	for (const FString& VectorRecord : VectorRecords)
+	{
+		if (VectorRecord.Len() > 0)
+		{
+			const TArray<FString> VectorParams = DisplayClusterHelpers::str::StrToArray<FString>(VectorRecord, FString(TEXT("=")));
+			check(VectorParams.Num() == 2);
+
+			FName VectorAxisName = *VectorParams[0];
+			if (VectorAxisDelegateInstances.Contains(VectorAxisName))
+			{
+				FVector VectorValue = FDisplayClusterTypesConverter::FromString<FVector>(VectorParams[1]);
+				UE_LOG(LogDisplayClusterGame, Verbose, TEXT("Native input - vector axis %s, value %f"), *VectorAxisName.ToString(), *VectorValue.ToString());
+				VectorAxisDelegateInstances[VectorAxisName]->Execute(VectorValue);
+			}
+		}
+	}
+}
+
+void UDisplayClusterPlayerInput::DeserializeAndProcessActionInput(const FString& Data)
+{
+	const TArray<FString> ActionRecords = DisplayClusterHelpers::str::StrToArray<FString>(Data, FString(";"));
+	for (const FString& ActionRecord : ActionRecords)
+	{
+		if(ActionRecord.Len() > 0)
+		{
+			const TArray<FString> ActionParams = DisplayClusterHelpers::str::StrToArray<FString>(ActionRecord, FString(TEXT("=")));
+			check(ActionParams.Num() == 2);
+			const TArray<FString> ActionParamsEx = DisplayClusterHelpers::str::StrToArray<FString>(ActionParams[1], FString(TEXT("|")));
+
+			if (ActionParams[0].Equals(FString("key"), ESearchCase::IgnoreCase))
+			{
+				check(ActionParamsEx.Num() == 3);
+
+				const FName ActionName = *ActionParamsEx[0];
+				if (ActionDelegateInstances.Contains(ActionName))
+				{
+					const EInputEvent InputEvent = (EInputEvent)FDisplayClusterTypesConverter::FromString<int>(ActionParamsEx[2]);
+					if (ActionDelegateInstances[ActionName].Contains(InputEvent))
+					{
+						const FKey ActionKey(*ActionParamsEx[1]);
+
+						UE_LOG(LogDisplayClusterGame, Verbose, TEXT("Native input - key action %s, key %s, input event %d"), *ActionName.ToString(), *ActionKey.ToString(), (int)InputEvent);
+
+						ActionDelegateInstances[ActionName][InputEvent]->Execute(ActionKey);
+					}
+				}
+			}
+			else if (ActionParams[0].Equals(FString("touch"), ESearchCase::IgnoreCase))
+			{
+				check(ActionParamsEx.Num() == 3);
+
+				const EInputEvent InputEvent = (EInputEvent)FDisplayClusterTypesConverter::FromString<int>(ActionParamsEx[0]);
+				if (TouchDelegateInstances.Contains(InputEvent))
+				{
+					const FVector           TouchLocation = FDisplayClusterTypesConverter::FromString<FVector>(ActionParamsEx[1]);
+					const ETouchIndex::Type FingerIndex   = (ETouchIndex::Type)FDisplayClusterTypesConverter::FromString<int>(ActionParamsEx[2]);
+
+					UE_LOG(LogDisplayClusterGame, Verbose, TEXT("Native input - touch action, finger index %d, touch location %s, input event %d"), (int)FingerIndex, *TouchLocation.ToString(), (int)InputEvent);
+
+					TouchDelegateInstances[InputEvent]->Execute(FingerIndex, TouchLocation);
+				}
+			}
+			else if (ActionParams[0].Equals(FString("gesture"), ESearchCase::IgnoreCase))
+			{
+				check(ActionParamsEx.Num() == 2);
+
+				const FName GestureKey = *ActionParamsEx[0];
+				if (GestureDelegateInstances.Contains(GestureKey))
+				{
+					const float GestureValue = FDisplayClusterTypesConverter::FromString<float>(ActionParamsEx[1]);
+
+					UE_LOG(LogDisplayClusterGame, Verbose, TEXT("Native input - gesture action, key %s, value %f"), *GestureKey.ToString(), GestureValue);
+
+					GestureDelegateInstances[GestureKey]->Execute(GestureValue);
+				}
+			}
+		}
+	}
+}
diff --git a/Engine/Plugins/Runtime/nDisplay/Source/DisplayCluster/Private/Game/Classes/Scene/DisplayClusterCameraComponent.cpp b/Engine/Plugins/Runtime/nDisplay/Source/DisplayCluster/Private/Game/Classes/Scene/DisplayClusterCameraComponent.cpp
--- a/Engine/Plugins/Runtime/nDisplay/Source/DisplayCluster/Private/Game/Classes/Scene/DisplayClusterCameraComponent.cpp
+++ b/Engine/Plugins/Runtime/nDisplay/Source/DisplayCluster/Private/Game/Classes/Scene/DisplayClusterCameraComponent.cpp
@@ -3,6 +3,8 @@
 #include "DisplayClusterCameraComponent.h"
 #include "Config/DisplayClusterConfigTypes.h"
 
+#include "Camera/CameraComponent.h"
+
 #include "CoreGlobals.h"
 
 
@@ -11,26 +13,14 @@ UDisplayClusterCameraComponent::UDisplayClusterCameraComponent(const FObjectInit
 	, EyeDist(0.064f)
 	, bEyeSwap(false)
 	, ForceEyeOffset(0)
-	, NearClipPlane(GNearClippingPlane)
-	, FarClipPlane(2000000.f)
 {
-	PrimaryComponentTick.bCanEverTick = true;
+	PrimaryComponentTick.bCanEverTick = false;
 }
 
 
 void UDisplayClusterCameraComponent::BeginPlay()
 {
 	Super::BeginPlay();
-	
-	// ...
-	
-}
-
-void UDisplayClusterCameraComponent::TickComponent( float DeltaTime, ELevelTick TickType, FActorComponentTickFunction* ThisTickFunction )
-{
-	Super::TickComponent( DeltaTime, TickType, ThisTickFunction );
-
-	// ...
 }
 
 void UDisplayClusterCameraComponent::SetSettings(const FDisplayClusterConfigSceneNode* ConfigData)
@@ -51,7 +41,5 @@ void UDisplayClusterCameraComponent::SetSettings(const FDisplayClusterConfigScen
 
 bool UDisplayClusterCameraComponent::ApplySettings()
 {
-
-
 	return Super::ApplySettings();
 }
diff --git a/Engine/Plugins/Runtime/nDisplay/Source/DisplayCluster/Private/Game/Classes/Scene/DisplayClusterSceneComponentSync.cpp b/Engine/Plugins/Runtime/nDisplay/Source/DisplayCluster/Private/Game/Classes/Scene/DisplayClusterSceneComponentSync.cpp
--- a/Engine/Plugins/Runtime/nDisplay/Source/DisplayCluster/Private/Game/Classes/Scene/DisplayClusterSceneComponentSync.cpp
+++ b/Engine/Plugins/Runtime/nDisplay/Source/DisplayCluster/Private/Game/Classes/Scene/DisplayClusterSceneComponentSync.cpp
@@ -2,6 +2,8 @@
 
 #include "DisplayClusterSceneComponentSync.h"
 
+#include "GameFramework/Actor.h"
+
 #include "Cluster/IPDisplayClusterClusterManager.h"
 #include "Game/IPDisplayClusterGameManager.h"
 
@@ -35,7 +37,7 @@ void UDisplayClusterSceneComponentSync::BeginPlay()
 		if (ClusterMgr)
 		{
 			UE_LOG(LogDisplayClusterGame, Log, TEXT("Registering sync object %s..."), *SyncId);
-			ClusterMgr->RegisterSyncObject(this);
+			ClusterMgr->RegisterSyncObject(this, EDisplayClusterSyncGroup::Tick);
 		}
 		else
 		{
diff --git a/Engine/Plugins/Runtime/nDisplay/Source/DisplayCluster/Private/Game/Classes/Scene/DisplayClusterSceneComponentSyncParent.cpp b/Engine/Plugins/Runtime/nDisplay/Source/DisplayCluster/Private/Game/Classes/Scene/DisplayClusterSceneComponentSyncParent.cpp
--- a/Engine/Plugins/Runtime/nDisplay/Source/DisplayCluster/Private/Game/Classes/Scene/DisplayClusterSceneComponentSyncParent.cpp
+++ b/Engine/Plugins/Runtime/nDisplay/Source/DisplayCluster/Private/Game/Classes/Scene/DisplayClusterSceneComponentSyncParent.cpp
@@ -3,6 +3,9 @@
 #include "DisplayClusterSceneComponentSyncParent.h"
 #include "GameFramework/Actor.h"
 
+#include "DisplayClusterHelpers.h"
+#include "DisplayClusterLog.h"
+
 
 UDisplayClusterSceneComponentSyncParent::UDisplayClusterSceneComponentSyncParent(const FObjectInitializer& ObjectInitializer) :
 	UDisplayClusterSceneComponentSync(ObjectInitializer)
@@ -44,15 +47,25 @@ FString UDisplayClusterSceneComponentSyncParent::GetSyncId() const
 bool UDisplayClusterSceneComponentSyncParent::IsDirty() const
 {
 	USceneComponent* const pParent = GetAttachParent();
-	return (LastSyncLoc != pParent->GetRelativeLocation() || LastSyncRot != pParent->GetRelativeRotation() || LastSyncScale != pParent->GetRelativeScale3D());
+	if (pParent && !pParent->IsPendingKill())
+	{
+		const bool bIsDirty = (LastSyncLoc != pParent->GetRelativeLocation() || LastSyncRot != pParent->GetRelativeRotation() || LastSyncScale != pParent->GetRelativeScale3D());
+		UE_LOG(LogDisplayClusterGame, Verbose, TEXT("SYNC_PARENT: %s dirty state is %s"), *GetSyncId(), *DisplayClusterHelpers::str::BoolToStr(bIsDirty));
+		return bIsDirty;
+	}
+
+	return false;
 }
 
 void UDisplayClusterSceneComponentSyncParent::ClearDirty()
 {
 	USceneComponent* const pParent = GetAttachParent();
-	LastSyncLoc = pParent->GetRelativeLocation();
-	LastSyncRot = pParent->GetRelativeRotation();
-	LastSyncScale = pParent->GetRelativeScale3D();
+	if (pParent && !pParent->IsPendingKill())
+	{
+		LastSyncLoc   = pParent->GetRelativeLocation();
+		LastSyncRot   = pParent->GetRelativeRotation();
+		LastSyncScale = pParent->GetRelativeScale3D();
+	}
 }
 
 
diff --git a/Engine/Plugins/Runtime/nDisplay/Source/DisplayCluster/Private/Game/Classes/Scene/DisplayClusterScreenComponent.cpp b/Engine/Plugins/Runtime/nDisplay/Source/DisplayCluster/Private/Game/Classes/Scene/DisplayClusterScreenComponent.cpp
--- a/Engine/Plugins/Runtime/nDisplay/Source/DisplayCluster/Private/Game/Classes/Scene/DisplayClusterScreenComponent.cpp
+++ b/Engine/Plugins/Runtime/nDisplay/Source/DisplayCluster/Private/Game/Classes/Scene/DisplayClusterScreenComponent.cpp
@@ -1,7 +1,6 @@
 // Copyright 1998-2019 Epic Games, Inc. All Rights Reserved.
 
 #include "DisplayClusterScreenComponent.h"
-#include "DisplayClusterSettings.h"
 #include "Components/StaticMeshComponent.h"
 #include "Engine/GameEngine.h"
 #include "Engine/StaticMesh.h"
@@ -9,6 +8,8 @@
 #include "Materials/Material.h"
 #include "UObject/ConstructorHelpers.h"
 
+#include "DisplayClusterRootComponent.h"
+
 #include "Game/IPDisplayClusterGameManager.h"
 #include "DisplayClusterGlobals.h"
 #include "EngineDefines.h"
@@ -19,49 +20,47 @@ UDisplayClusterScreenComponent::UDisplayClusterScreenComponent(const FObjectInit
 {
 	PrimaryComponentTick.bCanEverTick = true;
 
-#if WITH_EDITOR
-	if (GEngine && GEngine->IsEditor())
+	ScreenGeometryComponent = CreateDefaultSubobject<UStaticMeshComponent>(FName(*(GetName() + FString("_impl"))));
+	check(ScreenGeometryComponent);
+
+	if (ScreenGeometryComponent)
 	{
-		const IPDisplayClusterGameManager* const GameMgr = GDisplayCluster->GetPrivateGameMgr();
-		if (GameMgr)
-		{
-			const ADisplayClusterSettings* const pDisplayClusterSettings = GameMgr->GetDisplayClusterSceneSettings();
-			if (pDisplayClusterSettings && pDisplayClusterSettings->bEditorShowProjectionScreens)
-			{
-				ScreenGeometryComponent = CreateDefaultSubobject<UStaticMeshComponent>(FName(*(GetName() + FString("_impl"))));
-				check(ScreenGeometryComponent);
-
-				if (ScreenGeometryComponent)
-				{
-					static ConstructorHelpers::FObjectFinder<UStaticMesh> screenMesh(TEXT("StaticMesh'/Engine/BasicShapes/Cube.Cube'"));
-					static ConstructorHelpers::FObjectFinder<UMaterial>   screenMat(TEXT("Material'/Engine/Engine_MI_Shaders/M_Shader_SimpleTranslucent.M_Shader_SimpleTranslucent'"));
-
-					ScreenGeometryComponent->AttachToComponent(this, FAttachmentTransformRules(EAttachmentRule::KeepRelative, false));
-					ScreenGeometryComponent->SetStaticMesh(screenMesh.Object);
-					ScreenGeometryComponent->SetMobility(EComponentMobility::Movable);
-					ScreenGeometryComponent->SetMaterial(0, screenMat.Object);
-					ScreenGeometryComponent->SetCollisionEnabled(ECollisionEnabled::NoCollision);
-				}
-			}
-		}
+		static ConstructorHelpers::FObjectFinder<UStaticMesh> ScreenMesh(TEXT("StaticMesh'/Engine/BasicShapes/Cube.Cube'"));
+
+		ScreenGeometryComponent->AttachToComponent(this, FAttachmentTransformRules(EAttachmentRule::KeepRelative, false));
+		ScreenGeometryComponent->SetStaticMesh(ScreenMesh.Object);
+		ScreenGeometryComponent->SetMobility(EComponentMobility::Movable);
+		ScreenGeometryComponent->SetCollisionEnabled(ECollisionEnabled::NoCollision);
+		ScreenGeometryComponent->SetVisibility(false);
 	}
-#endif
 }
 
 
 void UDisplayClusterScreenComponent::BeginPlay()
 {
 	Super::BeginPlay();
-
-	// ...
 }
 
 
 void UDisplayClusterScreenComponent::TickComponent( float DeltaTime, ELevelTick TickType, FActorComponentTickFunction* ThisTickFunction )
 {
-	Super::TickComponent(DeltaTime, TickType, ThisTickFunction);
+	static const IPDisplayClusterGameManager* const GameMgr = GDisplayCluster->GetPrivateGameMgr();
+	if (GameMgr)
+	{
+		UDisplayClusterRootComponent* const RootComp = GameMgr->GetRoot();
+		if (RootComp)
+		{
+			if (RootComp->GetProjectionScreenMaterial())
+			{
+				ScreenGeometryComponent->SetMaterial(0, RootComp->GetProjectionScreenMaterial());
+			}
 
-	// ...
+			const bool bScreenVisible = RootComp->GetShowProjectionScreens();
+			ScreenGeometryComponent->SetVisibility(bScreenVisible);
+		}
+	}
+
+	Super::TickComponent(DeltaTime, TickType, ThisTickFunction);
 }
 
 void UDisplayClusterScreenComponent::SetSettings(const FDisplayClusterConfigSceneNode* pConfig)
