Get rid of DisplayCluster GameMode and Pawns

diff --git a/Engine/Plugins/Runtime/nDisplay/Source/DisplayCluster/Private/Game/DisplayClusterGameManager.cpp b/Engine/Plugins/Runtime/nDisplay/Source/DisplayCluster/Private/Game/DisplayClusterGameManager.cpp
--- a/Engine/Plugins/Runtime/nDisplay/Source/DisplayCluster/Private/Game/DisplayClusterGameManager.cpp
+++ b/Engine/Plugins/Runtime/nDisplay/Source/DisplayCluster/Private/Game/DisplayClusterGameManager.cpp
@@ -4,9 +4,6 @@
 
 #include "Config/IPDisplayClusterConfigManager.h"
 
-#include "DisplayClusterGameMode.h"
-#include "DisplayClusterSettings.h"
-
 #include "Misc/DisplayClusterHelpers.h"
 
 #include "Kismet/GameplayStatics.h"
@@ -17,6 +14,7 @@
 #include "Camera/CameraComponent.h"
 #include "Components/SceneComponent.h"
 #include "DisplayClusterCameraComponent.h"
+#include "DisplayClusterRootComponent.h"
 #include "DisplayClusterSceneComponent.h"
 #include "DisplayClusterScreenComponent.h"
 
@@ -24,6 +22,8 @@
 #include "DisplayClusterLog.h"
 #include "DisplayClusterStrings.h"
 
+#include "GameFramework/Actor.h"
+
 
 FDisplayClusterGameManager::FDisplayClusterGameManager()
 {
@@ -68,29 +68,34 @@ void FDisplayClusterGameManager::EndSession()
 	DISPLAY_CLUSTER_FUNC_TRACE(LogDisplayClusterGame);
 }
 
-bool FDisplayClusterGameManager::StartScene(UWorld* pWorld)
+bool FDisplayClusterGameManager::StartScene(UWorld* InWorld)
 {
 	DISPLAY_CLUSTER_FUNC_TRACE(LogDisplayClusterGame);
 
-	check(pWorld);
-	CurrentWorld = pWorld;
-
-	VRRootActor = nullptr;
-	DefaultCameraComponent = nullptr;
+	check(InWorld);
+	CurrentWorld = InWorld;
 
-	// Clean containers. We store only pointers so there is no need to do any additional
-	// operations. All components will be destroyed by the engine.
-	ScreenComponents.Reset();
-	CameraComponents.Reset();
-	SceneNodeComponents.Reset();
+	// Look for existing VR root components
+	for (AActor* Actor : InWorld->PersistentLevel->Actors)
+	{
+		if (Actor && !Actor->IsPendingKill())
+		{
+			CurrentRoot = Actor->FindComponentByClass<UDisplayClusterRootComponent>();
+			if (CurrentRoot)
+			{
+				UE_LOG(LogDisplayClusterGame, Log, TEXT("Found root component - %s"), *CurrentRoot->GetName());
+				break;
+			}
+		}
+	}
 
-	if (IsDisplayClusterActive())
+	APlayerController* const PlayerController = GetWorld()->GetFirstPlayerController();
+	if (!CurrentRoot && PlayerController)
 	{
-		//@todo: move initialization to DisplayClusterRoot side
-		if (!InitializeDisplayClusterActor())
+		APawn* CurPawn = PlayerController->GetPawn();
+		if (CurPawn)
 		{
-			UE_LOG(LogDisplayClusterGame, Error, TEXT("Couldn't initialize DisplayCluster hierarchy"));
-			return false;
+			SpawnRootComponent(CurPawn);
 		}
 	}
 
@@ -102,383 +107,161 @@ void FDisplayClusterGameManager::EndScene()
 	DISPLAY_CLUSTER_FUNC_TRACE(LogDisplayClusterGame);
 	FScopeLock lock(&InternalsSyncScope);
 
-	VRRootActor = nullptr;
-	DefaultCameraComponent = nullptr;
-
-	// Clean containers. We store only pointers so there is no need to do any additional
-	// operations. All components will be destroyed by the engine.
-	ScreenComponents.Reset();
-	CameraComponents.Reset();
-	SceneNodeComponents.Reset();
+	CurrentRoot = nullptr;
 }
 
 
 //////////////////////////////////////////////////////////////////////////////////////////////
 // IDisplayClusterGameManager
 //////////////////////////////////////////////////////////////////////////////////////////////
-ADisplayClusterPawn* FDisplayClusterGameManager::GetRoot() const
-{
-	FScopeLock lock(&InternalsSyncScope);
-	return VRRootActor;
-}
-
-TArray<UDisplayClusterScreenComponent*> FDisplayClusterGameManager::GetAllScreens() const
-{
-	FScopeLock lock(&InternalsSyncScope);
-	return GetMapValues<UDisplayClusterScreenComponent>(ScreenComponents);
-}
-
-UDisplayClusterScreenComponent* FDisplayClusterGameManager::GetScreenById(const FString& id) const
+UDisplayClusterRootComponent* FDisplayClusterGameManager::GetRoot() const
 {
 	FScopeLock lock(&InternalsSyncScope);
-	return GetItem<UDisplayClusterScreenComponent>(ScreenComponents, id, FString("GetScreenById"));
+	return CurrentRoot;
 }
 
-int32 FDisplayClusterGameManager::GetScreensAmount() const
+void FDisplayClusterGameManager::SetRoot(UDisplayClusterRootComponent* InRoot)
 {
 	FScopeLock lock(&InternalsSyncScope);
-	return ScreenComponents.Num();
+	CurrentRoot = InRoot;
 }
 
-UDisplayClusterCameraComponent* FDisplayClusterGameManager::GetCameraById(const FString& id) const
+TArray<UDisplayClusterScreenComponent*> FDisplayClusterGameManager::GetAllScreens() const
 {
 	FScopeLock lock(&InternalsSyncScope);
-	return GetItem<UDisplayClusterCameraComponent>(CameraComponents, id, FString("GetCameraById"));
-}
 
-TArray<UDisplayClusterCameraComponent*> FDisplayClusterGameManager::GetAllCameras() const
-{
-	FScopeLock lock(&InternalsSyncScope);
-	return GetMapValues<UDisplayClusterCameraComponent>(CameraComponents);
-}
+	if (!CurrentRoot)
+	{
+		return TArray<UDisplayClusterScreenComponent*>();
+	}
 
-int32 FDisplayClusterGameManager::GetCamerasAmount() const
-{
-	FScopeLock lock(&InternalsSyncScope);
-	return CameraComponents.Num();
+	return CurrentRoot->GetAllScreens();
 }
 
-UDisplayClusterCameraComponent* FDisplayClusterGameManager::GetDefaultCamera() const
+UDisplayClusterScreenComponent* FDisplayClusterGameManager::GetScreenById(const FString& ScreenId) const
 {
 	FScopeLock lock(&InternalsSyncScope);
-	return DefaultCameraComponent;
-}
-
-void FDisplayClusterGameManager::SetDefaultCamera(int32 idx)
-{
-	DISPLAY_CLUSTER_FUNC_TRACE(LogDisplayClusterGame);
-
-	if (!IsDisplayClusterActive())
-	{
-		return;
-	}
 
-	FDisplayClusterConfigCamera cam;
-	if (!GDisplayCluster->GetPrivateConfigMgr()->GetCamera(idx, cam))
+	if (!CurrentRoot)
 	{
-		UE_LOG(LogDisplayClusterGame, Error, TEXT("Camera not found (idx=%d)"), idx);
-		return;
+		return nullptr;
 	}
 
-	return SetDefaultCamera(cam.Id);
+	return CurrentRoot->GetScreenById(ScreenId);
 }
 
-void FDisplayClusterGameManager::SetDefaultCamera(const FString& id)
+int32 FDisplayClusterGameManager::GetScreensAmount() const
 {
-	DISPLAY_CLUSTER_FUNC_TRACE(LogDisplayClusterGame);
-
-	if (!IsDisplayClusterActive())
-	{
-		return;
-	}
-
 	FScopeLock lock(&InternalsSyncScope);
 
-	if (!CameraComponents.Contains(id))
+	if (!CurrentRoot)
 	{
-		UE_LOG(LogDisplayClusterGame, Error, TEXT("Couldn't switch camera. No such node id: %s"), *id);
-		return;
+		return 0;
 	}
 
-	DefaultCameraComponent = CameraComponents[id];
-	VRRootActor->GetCameraComponent()->AttachToComponent(DefaultCameraComponent, FAttachmentTransformRules(EAttachmentRule::KeepRelative, false));
-	VRRootActor->GetCameraComponent()->SetRelativeLocation(FVector::ZeroVector);
-	VRRootActor->GetCameraComponent()->SetRelativeRotation(FRotator::ZeroRotator);
-
-	// Update 'rotate around' component
-	SetRotateAroundComponent(DefaultCameraComponent);
-
-	UE_LOG(LogDisplayClusterGame, Log, TEXT("Default camera: %s"), *DefaultCameraComponent->GetId());
-}
-
-UDisplayClusterSceneComponent* FDisplayClusterGameManager::GetNodeById(const FString& id) const
-{
-	FScopeLock lock(&InternalsSyncScope);
-	return GetItem<UDisplayClusterSceneComponent>(SceneNodeComponents, id, FString("GetNodeById"));
+	return CurrentRoot->GetScreensAmount();
 }
 
-TArray<UDisplayClusterSceneComponent*> FDisplayClusterGameManager::GetAllNodes() const
+UDisplayClusterCameraComponent* FDisplayClusterGameManager::GetCameraById(const FString& CameraId) const
 {
 	FScopeLock lock(&InternalsSyncScope);
-	return GetMapValues<UDisplayClusterSceneComponent>(SceneNodeComponents);
-}
 
-USceneComponent* FDisplayClusterGameManager::GetTranslationDirectionComponent() const
-{
-	if (!IsDisplayClusterActive())
+	if (!CurrentRoot)
 	{
 		return nullptr;
 	}
 
-	if (VRRootActor == nullptr)
-	{
-		return nullptr;
-	}
-
-	FScopeLock lock(&InternalsSyncScope);
-	UE_LOG(LogDisplayClusterGame, Verbose, TEXT("GetTranslationDirectionComponent: %s"), (VRRootActor->TranslationDirection ? *VRRootActor->TranslationDirection->GetName() : TEXT("nullptr")));
-	return VRRootActor->TranslationDirection;
+	return CurrentRoot->GetCameraById(CameraId);
 }
 
-void FDisplayClusterGameManager::SetTranslationDirectionComponent(USceneComponent* pComp)
+TArray<UDisplayClusterCameraComponent*> FDisplayClusterGameManager::GetAllCameras() const
 {
-	DISPLAY_CLUSTER_FUNC_TRACE(LogDisplayClusterGame);
+	FScopeLock lock(&InternalsSyncScope);
 
-	if (!IsDisplayClusterActive())
+	if (!CurrentRoot)
 	{
-		return;
+		return TArray<UDisplayClusterCameraComponent*>();
 	}
 
-	if (VRRootActor == nullptr)
-	{
-		return;
-	}
-
-	FScopeLock lock(&InternalsSyncScope);
-	UE_LOG(LogDisplayClusterGame, Log, TEXT("New translation direction component set: %s"), (pComp ? *pComp->GetName() : TEXT("nullptr")));
-	VRRootActor->TranslationDirection = pComp;
+	return CurrentRoot->GetAllCameras();
 }
 
-void FDisplayClusterGameManager::SetTranslationDirectionComponent(const FString& id)
+int32 FDisplayClusterGameManager::GetCamerasAmount() const
 {
-	DISPLAY_CLUSTER_FUNC_TRACE(LogDisplayClusterGame);
+	FScopeLock lock(&InternalsSyncScope);
 
-	if (!IsDisplayClusterActive())
+	if (!CurrentRoot)
 	{
-		return;
+		return 0;
 	}
 
-	UE_LOG(LogDisplayClusterGame, Log, TEXT("New translation direction node id requested: %s"), *id);
-	SetTranslationDirectionComponent(GetNodeById(id));
+	return CurrentRoot->GetCamerasAmount();
 }
 
-USceneComponent* FDisplayClusterGameManager::GetRotateAroundComponent() const
+UDisplayClusterCameraComponent* FDisplayClusterGameManager::GetDefaultCamera() const
 {
-	if (!IsDisplayClusterActive())
-	{
-		return nullptr;
-	}
+	FScopeLock lock(&InternalsSyncScope);
 
-	if (VRRootActor == nullptr)
+	if (!CurrentRoot)
 	{
 		return nullptr;
 	}
 
-	FScopeLock lock(&InternalsSyncScope);
-	UE_LOG(LogDisplayClusterGame, Verbose, TEXT("GetRotateAroundComponent: %s"), (VRRootActor->RotationAround ? *VRRootActor->RotationAround->GetName() : TEXT("nullptr")));
-	return VRRootActor->RotationAround;
+	return CurrentRoot->GetDefaultCamera();
 }
 
-void FDisplayClusterGameManager::SetRotateAroundComponent(USceneComponent* pComp)
+void FDisplayClusterGameManager::SetDefaultCamera(const FString& id)
 {
 	DISPLAY_CLUSTER_FUNC_TRACE(LogDisplayClusterGame);
 
-	if (!IsDisplayClusterActive())
-	{
-		return;
-	}
-
-	if (VRRootActor == nullptr)
+	if(!CurrentRoot)
 	{
 		return;
 	}
 
-	FScopeLock lock(&InternalsSyncScope);
-	UE_LOG(LogDisplayClusterGame, Log, TEXT("New rotate around component set: %s"), (pComp ? *pComp->GetName() : TEXT("nullptr")));
-	VRRootActor->RotationAround = pComp;
+	CurrentRoot->SetDefaultCamera(id);
 }
 
-void FDisplayClusterGameManager::SetRotateAroundComponent(const FString& id)
+UDisplayClusterSceneComponent* FDisplayClusterGameManager::GetNodeById(const FString& NodeId) const
 {
-	DISPLAY_CLUSTER_FUNC_TRACE(LogDisplayClusterGame);
-
-	if (!IsDisplayClusterActive())
-	{
-		return;
-	}
-
-	if (VRRootActor == nullptr)
-	{
-		return;
-	}
-
 	FScopeLock lock(&InternalsSyncScope);
-	UE_LOG(LogDisplayClusterGame, Log, TEXT("New rotate around node id requested: %s"), *id);
-	VRRootActor->RotationAround = GetNodeById(id);
-}
-
-
-//////////////////////////////////////////////////////////////////////////////////////////////
-// FDisplayClusterGameManager
-//////////////////////////////////////////////////////////////////////////////////////////////
-bool FDisplayClusterGameManager::InitializeDisplayClusterActor()
-{
-	DISPLAY_CLUSTER_FUNC_TRACE(LogDisplayClusterGame);
-
-	APlayerController* pController = UGameplayStatics::GetPlayerController(CurrentWorld, 0);
-	check(pController);
-	
-	VRRootActor = Cast<ADisplayClusterPawn>(pController->GetPawn());
-	if (!VRRootActor)
-	{
-		// Seems the DisplayCluster features has been disabled
-		UE_LOG(LogDisplayClusterGame, Warning, TEXT("No DisplayCluster root found"));
-		return false;
-	}
-
-	if (!(CreateCameras() && CreateScreens() && CreateNodes()))
-	{
-		UE_LOG(LogDisplayClusterGame, Error, TEXT("An error occurred during DisplayCluster root initialization"));
-		return false;
-	}
-
-	// Let DisplayCluster nodes initialize ourselves
-	for (auto it = SceneNodeComponents.CreateIterator(); it; ++it)
-	{
-		if (it->Value->ApplySettings() == false)
-		{
-			UE_LOG(LogDisplayClusterGame, Warning, TEXT("Coulnd't initialize DisplayCluster node: ID=%s"), *it->Key);
-		}
-	}
 
-	// Set the first camera active by default
-	SetDefaultCamera(DefaultCameraComponent->GetId());
-
-	// Check if default camera was specified in command line arguments
-	FString camId;
-	if (FParse::Value(FCommandLine::Get(), DisplayClusterStrings::args::Camera, camId))
+	if (!CurrentRoot)
 	{
-		DisplayClusterHelpers::str::TrimStringValue(camId);
-		UE_LOG(LogDisplayClusterGame, Log, TEXT("Default camera from command line arguments: %s"), *camId);
-		if (CameraComponents.Contains(camId))
-		{
-			SetDefaultCamera(camId);
-		}
+		return nullptr;
 	}
 
-	return true;
+	return CurrentRoot->GetNodeById(NodeId);
 }
 
-bool FDisplayClusterGameManager::CreateScreens()
+TArray<UDisplayClusterSceneComponent*> FDisplayClusterGameManager::GetAllNodes() const
 {
-	// Create screens
-	const IPDisplayClusterConfigManager* const ConfigMgr = GDisplayCluster->GetPrivateConfigMgr();
-	if (!ConfigMgr)
-	{
-		UE_LOG(LogDisplayClusterGame, Error, TEXT("Couldn't get config manager interface"));
-		return false;
-	}
+	FScopeLock lock(&InternalsSyncScope);
 
-	const TArray<FDisplayClusterConfigScreen> AllScreens = ConfigMgr->GetScreens();
-	for (const auto& Screen : AllScreens)
+	if (!CurrentRoot)
 	{
-		// Create screen component
-		UDisplayClusterScreenComponent* ScreenComp = NewObject<UDisplayClusterScreenComponent>(VRRootActor, FName(*Screen.Id), RF_Transient);
-		check(ScreenComp);
-
-		ScreenComp->AttachToComponent(VRRootActor->GetCollisionOffsetComponent(), FAttachmentTransformRules(EAttachmentRule::KeepRelative, false));
-		ScreenComp->RegisterComponent();
-
-		// Pass settings
-		ScreenComp->SetSettings(&Screen);
-
-		// Store the screen
-		ScreenComponents.Add(Screen.Id, ScreenComp);
-		SceneNodeComponents.Add(Screen.Id, ScreenComp);
+		return TArray<UDisplayClusterSceneComponent*>();
 	}
 
-	return true;
+	return CurrentRoot->GetAllNodes();
 }
 
-bool FDisplayClusterGameManager::CreateNodes()
+bool FDisplayClusterGameManager::SpawnRootComponent(AActor* Actor)
 {
-	// Create other nodes
-	const TArray<FDisplayClusterConfigSceneNode> nodes = GDisplayCluster->GetPrivateConfigMgr()->GetSceneNodes();
-	for (const auto& node : nodes)
-	{
-		UDisplayClusterSceneComponent* pNode = NewObject<UDisplayClusterSceneComponent>(VRRootActor, FName(*node.Id), RF_Transient);
-		check(pNode);
+	check(Actor);
 
-		pNode->AttachToComponent(VRRootActor->GetCollisionOffsetComponent(), FAttachmentTransformRules(EAttachmentRule::KeepRelative, false));
-		pNode->RegisterComponent();
-
-		pNode->SetSettings(&node);
-		SceneNodeComponents.Add(node.Id, pNode);
-	}
-
-	return true;
-}
-
-bool FDisplayClusterGameManager::CreateCameras()
-{
-	const TArray<FDisplayClusterConfigCamera> cams = GDisplayCluster->GetPrivateConfigMgr()->GetCameras();
-	for (const auto& cam : cams)
-	{
-		UDisplayClusterCameraComponent* pCam = NewObject<UDisplayClusterCameraComponent>(VRRootActor, FName(*cam.Id), RF_Transient);
-		check(pCam);
-
-		pCam->AttachToComponent(VRRootActor->GetCollisionOffsetComponent(), FAttachmentTransformRules(EAttachmentRule::KeepRelative, false));
-		pCam->RegisterComponent();
-
-		pCam->SetSettings(&cam);
-		
-		CameraComponents.Add(cam.Id, pCam);
-		SceneNodeComponents.Add(cam.Id, pCam);
-
-		if (DefaultCameraComponent == nullptr)
-		{
-			DefaultCameraComponent = pCam;
-		}
-	}
+	FScopeLock lock(&InternalsSyncScope);
 
-	// At least one camera must be set up
-	if (!DefaultCameraComponent)
+	// Create new root component
+	UDisplayClusterRootComponent* NewRoot = NewObject<UDisplayClusterRootComponent>(Actor, FName(TEXT("nDisplay_Root")));
+	if (!NewRoot)
 	{
-		UE_LOG(LogDisplayClusterGame, Warning, TEXT("No camera found"));
 		return false;
 	}
 
-	return CameraComponents.Num() > 0;
-}
+	// Set up
+	NewRoot->AttachToComponent(Actor->GetRootComponent(), FAttachmentTransformRules(EAttachmentRule::KeepRelative, false));
+	NewRoot->RegisterComponent();
 
-// Extracts array of values from a map
-template <typename ObjType>
-TArray<ObjType*> FDisplayClusterGameManager::GetMapValues(const TMap<FString, ObjType*>& container) const
-{
-	TArray<ObjType*> items;
-	container.GenerateValueArray(items);
-	return items;
-}
-
-// Gets item by id. Performs checks and logging.
-template <typename DataType>
-DataType* FDisplayClusterGameManager::GetItem(const TMap<FString, DataType*>& container, const FString& id, const FString& logHeader) const
-{
-	if (container.Contains(id))
-	{
-		return container[id];
-	}
-
-	UE_LOG(LogDisplayClusterGame, Warning, TEXT("%s: ID not found <%s>. Return nullptr."), *logHeader, *id);
-	return nullptr;
+	return true;
 }
diff --git a/Engine/Plugins/Runtime/nDisplay/Source/DisplayCluster/Private/Game/DisplayClusterGameManager.h b/Engine/Plugins/Runtime/nDisplay/Source/DisplayCluster/Private/Game/DisplayClusterGameManager.h
--- a/Engine/Plugins/Runtime/nDisplay/Source/DisplayCluster/Private/Game/DisplayClusterGameManager.h
+++ b/Engine/Plugins/Runtime/nDisplay/Source/DisplayCluster/Private/Game/DisplayClusterGameManager.h
@@ -6,19 +6,21 @@
 
 #include "Game/IPDisplayClusterGameManager.h"
 
-#include "DisplayClusterOperationMode.h"
-#include "DisplayClusterGameMode.h"
-#include "DisplayClusterPawn.h"
-#include "DisplayClusterSettings.h"
-#include "DisplayClusterScreenComponent.h"
-#include "DisplayClusterCameraComponent.h"
+
+class AActor;
+class USceneComponent;
+class UDisplayClusterCameraComponent;
+class UDisplayClusterRootComponent;
+class UDisplayClusterSceneComponent;
+class UDisplayClusterScreenComponent;
+
 
 
 /**
  * Game manager. Responsible for building VR object hierarchy from a config file. Implements some in-game logic.
  */
 class FDisplayClusterGameManager
-	: public    IPDisplayClusterGameManager
+	: public IPDisplayClusterGameManager
 {
 public:
 	FDisplayClusterGameManager();
@@ -39,7 +41,8 @@ public:
 	//////////////////////////////////////////////////////////////////////////////////////////////
 	// IDisplayClusterGameManager
 	//////////////////////////////////////////////////////////////////////////////////////////////
-	virtual ADisplayClusterPawn*                    GetRoot() const override;
+	virtual UDisplayClusterRootComponent*           GetRoot() const override;
+	virtual void                                    SetRoot(UDisplayClusterRootComponent* InRoot) override;
 
 	virtual TArray<UDisplayClusterScreenComponent*> GetAllScreens() const override;
 	virtual UDisplayClusterScreenComponent*         GetScreenById(const FString& id) const override;
@@ -49,76 +52,32 @@ public:
 	virtual UDisplayClusterCameraComponent*         GetCameraById(const FString& id) const override;
 	virtual int32                                   GetCamerasAmount() const override;
 	virtual UDisplayClusterCameraComponent*         GetDefaultCamera() const override;
-	virtual void                                    SetDefaultCamera(int32 idx) override;
 	virtual void                                    SetDefaultCamera(const FString& id) override;
 
 	virtual TArray<UDisplayClusterSceneComponent*>  GetAllNodes() const override;
 	virtual UDisplayClusterSceneComponent*          GetNodeById(const FString& id) const override;
 
-	virtual USceneComponent*             GetTranslationDirectionComponent() const override;
-	virtual void                         SetTranslationDirectionComponent(USceneComponent* pComp) override;
-	virtual void                         SetTranslationDirectionComponent(const FString& id) override;
-
-	virtual USceneComponent*             GetRotateAroundComponent() const override;
-	virtual void                         SetRotateAroundComponent(USceneComponent* pComp) override;
-	virtual void                         SetRotateAroundComponent(const FString& id) override;
+	virtual UWorld* GetWorld() const override
+	{ return CurrentWorld; }
 
 public:
 	//////////////////////////////////////////////////////////////////////////////////////////////
 	// IPDisplayClusterGameManager
 	//////////////////////////////////////////////////////////////////////////////////////////////
 	virtual bool IsDisplayClusterActive() const override
-	{ return ((CurrentOperationMode != EDisplayClusterOperationMode::Disabled) && (CurrentGameMode ? CurrentGameMode->IsDisplayClusterActive() : false)); }
-	
-	virtual void SetDisplayClusterGameMode(ADisplayClusterGameMode* pGameMode) override
-	{ CurrentGameMode = pGameMode; }
-
-	virtual ADisplayClusterGameMode* GetDisplayClusterGameMode() const override
-	{ return CurrentGameMode; }
-
-	virtual void SetDisplayClusterSceneSettings(ADisplayClusterSettings* pSceneSettings) override
-	{ CurrentSceneSettings = pSceneSettings; }
-
-	virtual ADisplayClusterSettings* GetDisplayClusterSceneSettings() const override
-	{ return CurrentSceneSettings; }
+	{ return CurrentOperationMode != EDisplayClusterOperationMode::Disabled; }
 
-private:
-	// Creates DisplayCluster actor and fulfills with components hierarchy
-	bool InitializeDisplayClusterActor();
-
-	bool CreateScreens();
-	bool CreateNodes();
-	bool CreateCameras();
-
-	// Extracts array of values from a map
-	template <typename ObjType>
-	TArray<ObjType*> GetMapValues(const TMap<FString, ObjType*>& container) const;
-
-	// Gets item by id. Performs checks and logging.
-	template <typename DataType>
-	DataType* GetItem(const TMap<FString, DataType*>& container, const FString& id, const FString& logHeader) const;
+protected:
+	bool SpawnRootComponent(AActor* Actor);
 
 private:
-	// DisplayCluster root actor
-	ADisplayClusterPawn* VRRootActor = nullptr;
-	// Default camera (joint component)
-	UDisplayClusterCameraComponent* DefaultCameraComponent = nullptr;
-
-	// Available screens (from config file)
-	TMap<FString, UDisplayClusterScreenComponent*> ScreenComponents;
-	// Available cameras (from config file)
-	TMap<FString, UDisplayClusterCameraComponent*> CameraComponents;
-	// All available DisplayCluster nodes in hierarchy
-	TMap<FString, UDisplayClusterSceneComponent*>  SceneNodeComponents;
+	// Active DisplayCluster root
+	UDisplayClusterRootComponent* CurrentRoot = nullptr;
 
 	EDisplayClusterOperationMode CurrentOperationMode;
 	FString ConfigPath;
 	FString ClusterNodeId;
 	UWorld* CurrentWorld;
 
-	ADisplayClusterSettings* CurrentSceneSettings = nullptr;
-	ADisplayClusterGameMode* CurrentGameMode = nullptr;
-
 	mutable FCriticalSection InternalsSyncScope;
 };
-
diff --git a/Engine/Plugins/Runtime/nDisplay/Source/DisplayCluster/Private/Game/IPDisplayClusterGameManager.h b/Engine/Plugins/Runtime/nDisplay/Source/DisplayCluster/Private/Game/IPDisplayClusterGameManager.h
--- a/Engine/Plugins/Runtime/nDisplay/Source/DisplayCluster/Private/Game/IPDisplayClusterGameManager.h
+++ b/Engine/Plugins/Runtime/nDisplay/Source/DisplayCluster/Private/Game/IPDisplayClusterGameManager.h
@@ -5,9 +5,6 @@
 #include "Game/IDisplayClusterGameManager.h"
 #include "IPDisplayClusterManager.h"
 
-class ADisplayClusterGameMode;
-class ADisplayClusterSettings;
-
 
 /**
  * Game manager private interface
@@ -20,11 +17,6 @@ public:
 	virtual ~IPDisplayClusterGameManager()
 	{ }
 
+public:
 	virtual bool IsDisplayClusterActive() const = 0;
-
-	virtual void SetDisplayClusterGameMode(ADisplayClusterGameMode* pGameMode) = 0;
-	virtual ADisplayClusterGameMode* GetDisplayClusterGameMode() const = 0;
-
-	virtual void SetDisplayClusterSceneSettings(ADisplayClusterSettings* pSceneSettings) = 0;
-	virtual ADisplayClusterSettings* GetDisplayClusterSceneSettings() const = 0;
 };
diff --git a/Engine/Plugins/Runtime/nDisplay/Source/DisplayCluster/Private/Input/DisplayClusterInputManager.cpp b/Engine/Plugins/Runtime/nDisplay/Source/DisplayCluster/Private/Input/DisplayClusterInputManager.cpp
--- a/Engine/Plugins/Runtime/nDisplay/Source/DisplayCluster/Private/Input/DisplayClusterInputManager.cpp
+++ b/Engine/Plugins/Runtime/nDisplay/Source/DisplayCluster/Private/Input/DisplayClusterInputManager.cpp
@@ -10,8 +10,6 @@
 #include "Devices/VRPN/Tracker/DisplayClusterVrpnTrackerInputDevice.h"
 #include "Devices/VRPN/Keyboard/DisplayClusterVrpnKeyboardInputDevice.h"
 
-#include "DisplayClusterGameMode.h"
-
 #include "DisplayClusterGlobals.h"
 #include "DisplayClusterLog.h"
 #include "DisplayClusterStrings.h"
@@ -81,6 +79,14 @@ void FDisplayClusterInputManager::EndScene()
 	DISPLAY_CLUSTER_FUNC_TRACE(LogDisplayClusterInput);
 }
 
+void FDisplayClusterInputManager::StartFrame(uint64 FrameNum)
+{
+	DISPLAY_CLUSTER_FUNC_TRACE(LogDisplayClusterInput);
+
+	// Update input state
+	Update();
+}
+
 void FDisplayClusterInputManager::PreTick(float DeltaSeconds)
 {
 	DISPLAY_CLUSTER_FUNC_TRACE(LogDisplayClusterInput);
@@ -365,50 +371,58 @@ void FDisplayClusterInputManager::Update()
 				}
 			}
 		}
+
 		UE_LOG(LogDisplayClusterInput, Verbose, TEXT("Input update finished"));
-	
-		// Update input data cache for slave nodes
-		UpdateInputDataCache();
 	}
 }
 
-void FDisplayClusterInputManager::ExportInputData(FDisplayClusterMessage::DataType& data) const
+void FDisplayClusterInputManager::ExportInputData(FDisplayClusterMessage::DataType& InputData) const
 {
 	DISPLAY_CLUSTER_FUNC_TRACE(LogDisplayClusterInput);
 
 	FScopeLock ScopeLock(&InternalsSyncScope);
 
-	// Get data from cache
-	data = PackedTransferData;
+	// Clear container before put some data
+	InputData.Empty(InputData.Num() | 0x07);
+
+	for (auto DevClassIt = Devices.CreateConstIterator(); DevClassIt; ++DevClassIt)
+	{
+		for (auto DevIt = DevClassIt->Value.CreateConstIterator(); DevIt; ++DevIt)
+		{
+			const FString key = FString::Printf(TEXT("%d%s%s"), DevClassIt->Key, SerializationDeviceTypeNameDelimiter, *DevIt->Key);
+			const FString val = DevIt->Value->SerializeToString();
+			UE_LOG(LogDisplayClusterInput, VeryVerbose, TEXT("Input device %d:%s serialized: <%s, %s>"), DevClassIt->Key, *DevIt->Key, *key, *val);
+			InputData.Add(key, val);
+		}
+	}
 }
 
-void FDisplayClusterInputManager::ImportInputData(const FDisplayClusterMessage::DataType& data)
+void FDisplayClusterInputManager::ImportInputData(const FDisplayClusterMessage::DataType& InputData)
 {
 	DISPLAY_CLUSTER_FUNC_TRACE(LogDisplayClusterInput);
 
 	FScopeLock ScopeLock(&InternalsSyncScope);
 
-	for (auto rec : data)
+	for (const auto& DataItem : InputData)
 	{
-		FString strClassId;
-		FString strDevId;
-		if (rec.Key.Split(FString(SerializationDeviceTypeNameDelimiter), &strClassId, &strDevId))
+		FString DevClassId;
+		FString DevId;
+		if (DataItem.Key.Split(FString(SerializationDeviceTypeNameDelimiter), &DevClassId, &DevId))
 		{
-			UE_LOG(LogDisplayClusterInput, VeryVerbose, TEXT("Deserializing input device: <%s, %s>"), *rec.Key, *rec.Value);
+			UE_LOG(LogDisplayClusterInput, VeryVerbose, TEXT("Deserializing input device: <%s, %s>"), *DataItem.Key, *DataItem.Value);
 
-			int classId = FCString::Atoi(*strClassId);
+			int classId = FCString::Atoi(*DevClassId);
 			if (Devices.Contains(classId))
 			{
-				if (Devices[classId].Contains(strDevId))
+				if (Devices[classId].Contains(DevId))
 				{
-					Devices[classId][strDevId]->DeserializeFromString(rec.Value);
+					Devices[classId][DevId]->DeserializeFromString(DataItem.Value);
 				}
 			}
 		}
 	}
 }
 
-
 bool FDisplayClusterInputManager::GetAxisData(const FString& devId, const uint8 channel, FDisplayClusterVrpnAnalogChannelData&  data) const
 {
 	FScopeLock ScopeLock(&InternalsSyncScope);
@@ -585,24 +599,3 @@ void FDisplayClusterInputManager::ReleaseDevices()
 	UE_LOG(LogDisplayClusterInput, Log, TEXT("Releasing input devices..."));
 	Devices.Empty();
 }
-
-void FDisplayClusterInputManager::UpdateInputDataCache()
-{
-	DISPLAY_CLUSTER_FUNC_TRACE(LogDisplayClusterInput);
-
-	FScopeLock ScopeLock(&InternalsSyncScope);
-
-	// Clear previously cached data
-	PackedTransferData.Empty(PackedTransferData.Num() | 0x07);
-
-	for (auto classIt = Devices.CreateConstIterator(); classIt; ++classIt)
-	{
-		for (auto devIt = classIt->Value.CreateConstIterator(); devIt; ++devIt)
-		{
-			const FString key = FString::Printf(TEXT("%d%s%s"), classIt->Key, SerializationDeviceTypeNameDelimiter, *devIt->Key);
-			const FString val = devIt->Value->SerializeToString();
-			UE_LOG(LogDisplayClusterInput, VeryVerbose, TEXT("Input device %d:%s serialized: <%s, %s>"), classIt->Key, *devIt->Key, *key, *val);
-			PackedTransferData.Add(key, val);
-		}
-	}
-}
diff --git a/Engine/Plugins/Runtime/nDisplay/Source/DisplayCluster/Private/Input/DisplayClusterInputManager.h b/Engine/Plugins/Runtime/nDisplay/Source/DisplayCluster/Private/Input/DisplayClusterInputManager.h
--- a/Engine/Plugins/Runtime/nDisplay/Source/DisplayCluster/Private/Input/DisplayClusterInputManager.h
+++ b/Engine/Plugins/Runtime/nDisplay/Source/DisplayCluster/Private/Input/DisplayClusterInputManager.h
@@ -34,6 +34,7 @@ public:
 	virtual void EndSession() override;
 	virtual bool StartScene(UWorld* pWorld) override;
 	virtual void EndScene() override;
+	virtual void StartFrame(uint64 FrameNum) override;
 	virtual void PreTick(float DeltaSeconds);
 
 public:
@@ -82,8 +83,8 @@ public:
 	//////////////////////////////////////////////////////////////////////////////////////////////
 	virtual void Update() override;
 
-	virtual void ExportInputData(FDisplayClusterMessage::DataType& data) const override;
-	virtual void ImportInputData(const FDisplayClusterMessage::DataType& data) override;
+	virtual void ExportInputData(FDisplayClusterMessage::DataType& InputData) const override;
+	virtual void ImportInputData(const FDisplayClusterMessage::DataType& InputData) override;
 
 private:
 	typedef TUniquePtr<IDisplayClusterInputDevice>    TDevice;
@@ -103,8 +104,6 @@ private:
 private:
 	// Input devices
 	TDeviceMap Devices;
-	// Input state data cache
-	FDisplayClusterMessage::DataType PackedTransferData;
 	// Current config path
 	FString ConfigPath;
 	// Current cluster node ID
diff --git a/Engine/Plugins/Runtime/nDisplay/Source/DisplayCluster/Private/Input/IPDisplayClusterInputManager.h b/Engine/Plugins/Runtime/nDisplay/Source/DisplayCluster/Private/Input/IPDisplayClusterInputManager.h
--- a/Engine/Plugins/Runtime/nDisplay/Source/DisplayCluster/Private/Input/IPDisplayClusterInputManager.h
+++ b/Engine/Plugins/Runtime/nDisplay/Source/DisplayCluster/Private/Input/IPDisplayClusterInputManager.h
@@ -21,6 +21,6 @@ public:
 
 	virtual void Update() = 0;
 
-	virtual void ExportInputData(FDisplayClusterMessage::DataType& data) const = 0;
-	virtual void ImportInputData(const FDisplayClusterMessage::DataType& data) = 0;
+	virtual void ExportInputData(FDisplayClusterMessage::DataType& InputData) const = 0;
+	virtual void ImportInputData(const FDisplayClusterMessage::DataType& InputData) = 0;
 };
diff --git a/Engine/Plugins/Runtime/nDisplay/Source/DisplayCluster/Private/Network/DisplayClusterClient.cpp b/Engine/Plugins/Runtime/nDisplay/Source/DisplayCluster/Private/Network/DisplayClusterClient.cpp
--- a/Engine/Plugins/Runtime/nDisplay/Source/DisplayCluster/Private/Network/DisplayClusterClient.cpp
+++ b/Engine/Plugins/Runtime/nDisplay/Source/DisplayCluster/Private/Network/DisplayClusterClient.cpp
@@ -41,8 +41,8 @@ bool FDisplayClusterClient::Connect(const FString& InAddr, const int32 InPort, c
 	int32 TryIdx = 0;
 	while(GetSocket()->Connect(*InternetAddr) == false)
 	{
-		UE_LOG(LogDisplayClusterNetwork, Log, TEXT("%s couldn't connect to the server %s [%d]"), *GetName(), *(InternetAddr->ToString(true)), TryIdx++);
-		if (TriesAmount > 0 && TryIdx >= TriesAmount)
+		UE_LOG(LogDisplayClusterNetwork, Log, TEXT("%s couldn't connect to the server %s [%d]"), *GetName(), *(InternetAddr->ToString(true)), TryIdx);
+		if (TriesAmount > 0 && ++TryIdx >= TriesAmount)
 		{
 			UE_LOG(LogDisplayClusterNetwork, Error, TEXT("%s connection attempts limit reached"), *GetName());
 			break;
diff --git a/Engine/Plugins/Runtime/nDisplay/Source/DisplayCluster/Private/Network/Protocol/IPDisplayClusterClusterSyncProtocol.h b/Engine/Plugins/Runtime/nDisplay/Source/DisplayCluster/Private/Network/Protocol/IPDisplayClusterClusterSyncProtocol.h
--- a/Engine/Plugins/Runtime/nDisplay/Source/DisplayCluster/Private/Network/Protocol/IPDisplayClusterClusterSyncProtocol.h
+++ b/Engine/Plugins/Runtime/nDisplay/Source/DisplayCluster/Private/Network/Protocol/IPDisplayClusterClusterSyncProtocol.h
@@ -3,6 +3,7 @@
 #pragma once
 
 #include "Network/DisplayClusterMessage.h"
+#include "DisplayClusterEnums.h"
 
 struct FTimecode;
 struct FFrameRate;
@@ -27,18 +28,20 @@ public:
 	virtual void WaitForTickEnd() = 0;
 
 	// Provides with time delta for current frame
-	virtual void GetDeltaTime(float& deltaTime) = 0;
+	virtual void GetDeltaTime(float& DeltaSeconds) = 0;
 
 	// Get the Timecode value for the current frame.
-	virtual void GetTimecode(FTimecode& timecode, FFrameRate& frameRate) = 0;
+	virtual void GetTimecode(FTimecode& Timecode, FFrameRate& FrameRate) = 0;
 
 	// Sync objects
-	virtual void GetSyncData(FDisplayClusterMessage::DataType& data) = 0;
+	virtual void GetSyncData(FDisplayClusterMessage::DataType& SyncData, EDisplayClusterSyncGroup SyncGroup) = 0;
 
 	// Sync input
-	virtual void GetInputData(FDisplayClusterMessage::DataType& data) = 0;
+	virtual void GetInputData(FDisplayClusterMessage::DataType& InputData) = 0;
 
 	// Sync events
-	virtual void GetEventsData(FDisplayClusterMessage::DataType& data) = 0;
-};
+	virtual void GetEventsData(FDisplayClusterMessage::DataType& EventsData) = 0;
 
+	// Sync native UE4 input
+	virtual void GetNativeInputData(FDisplayClusterMessage::DataType& NativeInputData) = 0;
+};
diff --git a/Engine/Plugins/Runtime/nDisplay/Source/DisplayCluster/Private/Network/Service/ClusterSync/DisplayClusterClusterSyncClient.cpp b/Engine/Plugins/Runtime/nDisplay/Source/DisplayCluster/Private/Network/Service/ClusterSync/DisplayClusterClusterSyncClient.cpp
--- a/Engine/Plugins/Runtime/nDisplay/Source/DisplayCluster/Private/Network/Service/ClusterSync/DisplayClusterClusterSyncClient.cpp
+++ b/Engine/Plugins/Runtime/nDisplay/Source/DisplayCluster/Private/Network/Service/ClusterSync/DisplayClusterClusterSyncClient.cpp
@@ -48,7 +48,7 @@ void FDisplayClusterClusterSyncClient::WaitForTickEnd()
 	TSharedPtr<FDisplayClusterMessage> Response = SendRecvMsg(Request);
 }
 
-void FDisplayClusterClusterSyncClient::GetDeltaTime(float& DeltaTime)
+void FDisplayClusterClusterSyncClient::GetDeltaTime(float& DeltaSeconds)
 {
 	static const TSharedPtr<FDisplayClusterMessage> Request(new FDisplayClusterMessage(FDisplayClusterClusterSyncMsg::GetDeltaTime::name, FDisplayClusterClusterSyncMsg::TypeRequest, FDisplayClusterClusterSyncMsg::ProtocolName));
 	TSharedPtr<FDisplayClusterMessage> Response = SendRecvMsg(Request);
@@ -59,9 +59,9 @@ void FDisplayClusterClusterSyncClient::GetDeltaTime(float& DeltaTime)
 	}
 
 	// Extract sync data from response message
-	if (Response->GetArg(FDisplayClusterClusterSyncMsg::GetDeltaTime::argDeltaTime, DeltaTime) == false)
+	if (Response->GetArg(FDisplayClusterClusterSyncMsg::GetDeltaTime::argDeltaSeconds, DeltaSeconds) == false)
 	{
-		UE_LOG(LogDisplayClusterNetworkMsg, Error, TEXT("Couldn't extract an argument: %s"), FDisplayClusterClusterSyncMsg::GetDeltaTime::argDeltaTime);
+		UE_LOG(LogDisplayClusterNetworkMsg, Error, TEXT("Couldn't extract an argument: %s"), FDisplayClusterClusterSyncMsg::GetDeltaTime::argDeltaSeconds);
 	}
 }
 
@@ -86,9 +86,12 @@ void FDisplayClusterClusterSyncClient::GetTimecode(FTimecode& Timecode, FFrameRa
 	}
 }
 
-void FDisplayClusterClusterSyncClient::GetSyncData(FDisplayClusterMessage::DataType& Data)
+void FDisplayClusterClusterSyncClient::GetSyncData(FDisplayClusterMessage::DataType& SyncData, EDisplayClusterSyncGroup SyncGroup)
 {
-	static const TSharedPtr<FDisplayClusterMessage> Request(new FDisplayClusterMessage(FDisplayClusterClusterSyncMsg::GetSyncData::name, FDisplayClusterClusterSyncMsg::TypeRequest, FDisplayClusterClusterSyncMsg::ProtocolName));
+	static TSharedPtr<FDisplayClusterMessage> Request(new FDisplayClusterMessage(FDisplayClusterClusterSyncMsg::GetSyncData::name, FDisplayClusterClusterSyncMsg::TypeRequest, FDisplayClusterClusterSyncMsg::ProtocolName));
+	
+	Request->SetArg(FDisplayClusterClusterSyncMsg::GetSyncData::argSyncGroup, (int)SyncGroup);
+	
 	TSharedPtr<FDisplayClusterMessage> Response = SendRecvMsg(Request);
 
 	if (!Response.IsValid())
@@ -96,11 +99,11 @@ void FDisplayClusterClusterSyncClient::GetSyncData(FDisplayClusterMessage::DataT
 		return;
 	}
 
-	// Extract sync data from response message
-	Data = Response->GetArgs();
+	// Extract data from response message
+	SyncData = Response->GetArgs();
 }
 
-void FDisplayClusterClusterSyncClient::GetInputData(FDisplayClusterMessage::DataType& Data)
+void FDisplayClusterClusterSyncClient::GetInputData(FDisplayClusterMessage::DataType& InputData)
 {
 	static const TSharedPtr<FDisplayClusterMessage> Request(new FDisplayClusterMessage(FDisplayClusterClusterSyncMsg::GetInputData::name, FDisplayClusterClusterSyncMsg::TypeRequest, FDisplayClusterClusterSyncMsg::ProtocolName));
 	TSharedPtr<FDisplayClusterMessage> Response = SendRecvMsg(Request);
@@ -110,11 +113,11 @@ void FDisplayClusterClusterSyncClient::GetInputData(FDisplayClusterMessage::Data
 		return;
 	}
 
-	// Extract sync data from response message
-	Data = Response->GetArgs();
+	// Extract data from response message
+	InputData = Response->GetArgs();
 }
 
-void FDisplayClusterClusterSyncClient::GetEventsData(FDisplayClusterMessage::DataType& Data)
+void FDisplayClusterClusterSyncClient::GetEventsData(FDisplayClusterMessage::DataType& EventsData)
 {
 	static const TSharedPtr<FDisplayClusterMessage> Request(new FDisplayClusterMessage(FDisplayClusterClusterSyncMsg::GetEventsData::name, FDisplayClusterClusterSyncMsg::TypeRequest, FDisplayClusterClusterSyncMsg::ProtocolName));
 	TSharedPtr<FDisplayClusterMessage> Response = SendRecvMsg(Request);
@@ -124,6 +127,20 @@ void FDisplayClusterClusterSyncClient::GetEventsData(FDisplayClusterMessage::Dat
 		return;
 	}
 
-	// Extract sync data from response message
-	Data = Response->GetArgs();
+	// Extract data from response message
+	EventsData = Response->GetArgs();
+}
+
+void FDisplayClusterClusterSyncClient::GetNativeInputData(FDisplayClusterMessage::DataType& NativeInputData)
+{
+	static const TSharedPtr<FDisplayClusterMessage> Request(new FDisplayClusterMessage(FDisplayClusterClusterSyncMsg::GetNativeInputData::name, FDisplayClusterClusterSyncMsg::TypeRequest, FDisplayClusterClusterSyncMsg::ProtocolName));
+	TSharedPtr<FDisplayClusterMessage> Response = SendRecvMsg(Request);
+
+	if (!Response.IsValid())
+	{
+		return;
+	}
+
+	// Extract data from response message
+	NativeInputData = Response->GetArgs();
 }
diff --git a/Engine/Plugins/Runtime/nDisplay/Source/DisplayCluster/Private/Network/Service/ClusterSync/DisplayClusterClusterSyncClient.h b/Engine/Plugins/Runtime/nDisplay/Source/DisplayCluster/Private/Network/Service/ClusterSync/DisplayClusterClusterSyncClient.h
--- a/Engine/Plugins/Runtime/nDisplay/Source/DisplayCluster/Private/Network/Service/ClusterSync/DisplayClusterClusterSyncClient.h
+++ b/Engine/Plugins/Runtime/nDisplay/Source/DisplayCluster/Private/Network/Service/ClusterSync/DisplayClusterClusterSyncClient.h
@@ -26,10 +26,10 @@ public:
 	virtual void WaitForFrameStart() override;
 	virtual void WaitForFrameEnd() override;
 	virtual void WaitForTickEnd() override;
-	virtual void GetDeltaTime(float& DeltaTime) override;
+	virtual void GetDeltaTime(float& DeltaSeconds) override;
 	virtual void GetTimecode(FTimecode& Timecode, FFrameRate& FrameRate) override;
-	virtual void GetSyncData(FDisplayClusterMessage::DataType& Data) override;
-	virtual void GetInputData(FDisplayClusterMessage::DataType& Data) override;
-	virtual void GetEventsData(FDisplayClusterMessage::DataType& Data) override;
+	virtual void GetSyncData(FDisplayClusterMessage::DataType& SyncData, EDisplayClusterSyncGroup SyncGroup) override;
+	virtual void GetInputData(FDisplayClusterMessage::DataType& InputData) override;
+	virtual void GetEventsData(FDisplayClusterMessage::DataType& EventsData) override;
+	virtual void GetNativeInputData(FDisplayClusterMessage::DataType& NativeInputData) override;
 };
-
diff --git a/Engine/Plugins/Runtime/nDisplay/Source/DisplayCluster/Private/Network/Service/ClusterSync/DisplayClusterClusterSyncMsg.h b/Engine/Plugins/Runtime/nDisplay/Source/DisplayCluster/Private/Network/Service/ClusterSync/DisplayClusterClusterSyncMsg.h
--- a/Engine/Plugins/Runtime/nDisplay/Source/DisplayCluster/Private/Network/Service/ClusterSync/DisplayClusterClusterSyncMsg.h
+++ b/Engine/Plugins/Runtime/nDisplay/Source/DisplayCluster/Private/Network/Service/ClusterSync/DisplayClusterClusterSyncMsg.h
@@ -36,8 +36,8 @@ namespace FDisplayClusterClusterSyncMsg
 
 	namespace GetDeltaTime
 	{
-		constexpr static auto name         = "GetDeltaTime";
-		constexpr static auto argDeltaTime = "DeltaTime";
+		constexpr static auto name            = "GetDeltaTime";
+		constexpr static auto argDeltaSeconds = "DeltaSeconds";
 	};
 
 	namespace GetTimecode
@@ -49,7 +49,8 @@ namespace FDisplayClusterClusterSyncMsg
 
 	namespace GetSyncData
 	{
-		constexpr static auto name = "GetSyncData";
+		constexpr static auto name         = "GetSyncData";
+		constexpr static auto argSyncGroup = "SyncGroup";
 	};
 
 	namespace GetInputData
@@ -61,4 +62,9 @@ namespace FDisplayClusterClusterSyncMsg
 	{
 		constexpr static auto name = "GetEventsData";
 	}
+
+	namespace GetNativeInputData
+	{
+		constexpr static auto name = "GetNativeInputData";
+	}
 };
diff --git a/Engine/Plugins/Runtime/nDisplay/Source/DisplayCluster/Private/Network/Service/ClusterSync/DisplayClusterClusterSyncService.cpp b/Engine/Plugins/Runtime/nDisplay/Source/DisplayCluster/Private/Network/Service/ClusterSync/DisplayClusterClusterSyncService.cpp
--- a/Engine/Plugins/Runtime/nDisplay/Source/DisplayCluster/Private/Network/Service/ClusterSync/DisplayClusterClusterSyncService.cpp
+++ b/Engine/Plugins/Runtime/nDisplay/Source/DisplayCluster/Private/Network/Service/ClusterSync/DisplayClusterClusterSyncService.cpp
@@ -4,13 +4,17 @@
 #include "Network/Service/ClusterSync/DisplayClusterClusterSyncMsg.h"
 
 #include "Cluster/IPDisplayClusterClusterManager.h"
+#include "Game/IPDisplayClusterGameManager.h"
 #include "Config/IPDisplayClusterConfigManager.h"
 #include "Input/IPDisplayClusterInputManager.h"
 
+#include "Cluster/Controller/IPDisplayClusterNodeController.h"
+
 #include "Network/Session/DisplayClusterSessionInternal.h"
 
 #include "Misc/DisplayClusterAppExit.h"
 
+#include "DisplayClusterEnums.h"
 #include "DisplayClusterGlobals.h"
 #include "DisplayClusterLog.h"
 
@@ -122,9 +126,9 @@ TSharedPtr<FDisplayClusterMessage> FDisplayClusterClusterSyncService::ProcessMes
 	}
 	else if (ReqName == FDisplayClusterClusterSyncMsg::GetDeltaTime::name)
 	{
-		float deltaTime = 0.0f;
-		GetDeltaTime(deltaTime);
-		Response->SetArg(FDisplayClusterClusterSyncMsg::GetDeltaTime::argDeltaTime, deltaTime);
+		float DeltaSeconds = 0.0f;
+		GetDeltaTime(DeltaSeconds);
+		Response->SetArg(FDisplayClusterClusterSyncMsg::GetDeltaTime::argDeltaSeconds, DeltaSeconds);
 		return Response;
 	}
 	else if (ReqName == FDisplayClusterClusterSyncMsg::GetTimecode::name)
@@ -138,26 +142,38 @@ TSharedPtr<FDisplayClusterMessage> FDisplayClusterClusterSyncService::ProcessMes
 	}
 	else if (ReqName == FDisplayClusterClusterSyncMsg::GetSyncData::name)
 	{
-		FDisplayClusterMessage::DataType data;
-		GetSyncData(data);
+		FDisplayClusterMessage::DataType SyncData;
+		int SyncGroupNum = 0;
+		Request->GetArg(FDisplayClusterClusterSyncMsg::GetSyncData::argSyncGroup, SyncGroupNum);
+		EDisplayClusterSyncGroup SyncGroup = (EDisplayClusterSyncGroup)SyncGroupNum;
+
+		GetSyncData(SyncData, SyncGroup);
 
-		Response->SetArgs(data);
+		Response->SetArgs(SyncData);
 		return Response;
 	}
 	else if (ReqName == FDisplayClusterClusterSyncMsg::GetInputData::name)
 	{
-		FDisplayClusterMessage::DataType data;
-		GetInputData(data);
+		FDisplayClusterMessage::DataType InputData;
+		GetInputData(InputData);
 
-		Response->SetArgs(data);
+		Response->SetArgs(InputData);
 		return Response;
 	}
 	else if (ReqName == FDisplayClusterClusterSyncMsg::GetEventsData::name)
 	{
-		FDisplayClusterMessage::DataType data;
-		GetEventsData(data);
+		FDisplayClusterMessage::DataType EventsData;
+		GetEventsData(EventsData);
 
-		Response->SetArgs(data);
+		Response->SetArgs(EventsData);
+		return Response;
+	}
+	else if (ReqName == FDisplayClusterClusterSyncMsg::GetNativeInputData::name)
+	{
+		FDisplayClusterMessage::DataType NativeInputData;
+		GetNativeInputData(NativeInputData);
+
+		Response->SetArgs(NativeInputData);
 		return Response;
 	}
 
@@ -202,27 +218,37 @@ void FDisplayClusterClusterSyncService::WaitForTickEnd()
 	}
 }
 
-void FDisplayClusterClusterSyncService::GetDeltaTime(float& DeltaTime)
+void FDisplayClusterClusterSyncService::GetDeltaTime(float& DeltaSeconds)
 {
-	DeltaTime = GDisplayCluster->GetPrivateClusterMgr()->GetDeltaTime();
+	static IPDisplayClusterNodeController* const NodeController = GDisplayCluster->GetPrivateClusterMgr()->GetController();
+	return NodeController->GetDeltaTime(DeltaSeconds);
 }
 
 void FDisplayClusterClusterSyncService::GetTimecode(FTimecode& Timecode, FFrameRate& FrameRate)
 {
-	GDisplayCluster->GetPrivateClusterMgr()->GetTimecode(Timecode, FrameRate);
+	static IPDisplayClusterNodeController* const NodeController = GDisplayCluster->GetPrivateClusterMgr()->GetController();
+	return NodeController->GetTimecode(Timecode, FrameRate);
+}
+
+void FDisplayClusterClusterSyncService::GetSyncData(FDisplayClusterMessage::DataType& SyncData, EDisplayClusterSyncGroup SyncGroup)
+{
+	static IPDisplayClusterNodeController* const NodeController = GDisplayCluster->GetPrivateClusterMgr()->GetController();
+	NodeController->GetSyncData(SyncData, SyncGroup);
 }
 
-void FDisplayClusterClusterSyncService::GetSyncData(FDisplayClusterMessage::DataType& Data)
+void FDisplayClusterClusterSyncService::GetInputData(FDisplayClusterMessage::DataType& InputData)
 {
-	GDisplayCluster->GetPrivateClusterMgr()->ExportSyncData(Data);
+	static IPDisplayClusterNodeController* const NodeController = GDisplayCluster->GetPrivateClusterMgr()->GetController();
+	return NodeController->GetInputData(InputData);
 }
 
-void FDisplayClusterClusterSyncService::GetInputData(FDisplayClusterMessage::DataType& Data)
+void FDisplayClusterClusterSyncService::GetEventsData(FDisplayClusterMessage::DataType& EventsData)
 {
-	GDisplayCluster->GetPrivateInputMgr()->ExportInputData(Data);
+	GDisplayCluster->GetPrivateClusterMgr()->ExportEventsData(EventsData);
 }
 
-void FDisplayClusterClusterSyncService::GetEventsData(FDisplayClusterMessage::DataType& Data)
+void FDisplayClusterClusterSyncService::GetNativeInputData(FDisplayClusterMessage::DataType& NativeInputData)
 {
-	GDisplayCluster->GetPrivateClusterMgr()->ExportEventsData(Data);
+	static IPDisplayClusterNodeController* const NodeController = GDisplayCluster->GetPrivateClusterMgr()->GetController();
+	return NodeController->GetNativeInputData(NativeInputData);
 }
diff --git a/Engine/Plugins/Runtime/nDisplay/Source/DisplayCluster/Private/Network/Service/ClusterSync/DisplayClusterClusterSyncService.h b/Engine/Plugins/Runtime/nDisplay/Source/DisplayCluster/Private/Network/Service/ClusterSync/DisplayClusterClusterSyncService.h
--- a/Engine/Plugins/Runtime/nDisplay/Source/DisplayCluster/Private/Network/Service/ClusterSync/DisplayClusterClusterSyncService.h
+++ b/Engine/Plugins/Runtime/nDisplay/Source/DisplayCluster/Private/Network/Service/ClusterSync/DisplayClusterClusterSyncService.h
@@ -43,11 +43,12 @@ private:
 	virtual void WaitForFrameStart() override;
 	virtual void WaitForFrameEnd() override;
 	virtual void WaitForTickEnd() override;
-	virtual void GetDeltaTime(float& DeltaTime) override;
+	virtual void GetDeltaTime(float& DeltaSeconds) override;
 	virtual void GetTimecode(FTimecode& Timecode, FFrameRate& FrameRate) override;
-	virtual void GetSyncData(FDisplayClusterMessage::DataType& Data)  override;
-	virtual void GetInputData(FDisplayClusterMessage::DataType& Data) override;
-	virtual void GetEventsData(FDisplayClusterMessage::DataType& Data) override;
+	virtual void GetSyncData(FDisplayClusterMessage::DataType& SyncData, EDisplayClusterSyncGroup SyncGroup)  override;
+	virtual void GetInputData(FDisplayClusterMessage::DataType& InputData) override;
+	virtual void GetEventsData(FDisplayClusterMessage::DataType& EventsData) override;
+	virtual void GetNativeInputData(FDisplayClusterMessage::DataType& NativeInputData) override;
 
 private:
 	// Game start sync barrier
diff --git a/Engine/Plugins/Runtime/nDisplay/Source/DisplayCluster/Private/Render/Device/DisplayClusterDeviceBase.cpp b/Engine/Plugins/Runtime/nDisplay/Source/DisplayCluster/Private/Render/Device/DisplayClusterDeviceBase.cpp
--- a/Engine/Plugins/Runtime/nDisplay/Source/DisplayCluster/Private/Render/Device/DisplayClusterDeviceBase.cpp
+++ b/Engine/Plugins/Runtime/nDisplay/Source/DisplayCluster/Private/Render/Device/DisplayClusterDeviceBase.cpp
@@ -19,6 +19,7 @@
 
 #include "Render/Device/DisplayClusterRenderViewport.h"
 #include "Render/PostProcess/IDisplayClusterPostProcess.h"
+#include "Render/Presentation/DisplayClusterPresentationBase.h"
 #include "Render/Projection/IDisplayClusterProjectionPolicy.h"
 #include "Render/Projection/IDisplayClusterProjectionPolicyFactory.h"
 #include "Render/Synchronization/IDisplayClusterRenderSyncPolicy.h"
@@ -29,14 +30,6 @@
 #include <utility>
 
 
-// Custom VSync interval control
-static TAutoConsoleVariable<int32>  CVarVSyncInterval(
-	TEXT("nDisplay.render.VSyncInterval"),
-	1,
-	TEXT("VSync interval"),
-	ECVF_RenderThreadSafe
-);
-
 // Enable/disable warp&blend
 static TAutoConsoleVariable<int32> CVarWarpBlendEnabled(
 	TEXT("nDisplay.render.WarpBlendEnabled"),
@@ -59,8 +52,7 @@ static TAutoConsoleVariable<int32> CVarCustomPPEnabled(
 
 
 FDisplayClusterDeviceBase::FDisplayClusterDeviceBase(uint32 ViewsPerViewport)
-	: FRHICustomPresent()
-	, FDisplayClusterDeviceBase_PostProcess(RenderViewports, ViewsPerViewport, EyeRegions)
+	: FDisplayClusterDeviceBase_PostProcess(RenderViewports, ViewsPerViewport, EyeRegions)
 	, ViewsAmountPerViewport(ViewsPerViewport)
 {
 	DISPLAY_CLUSTER_FUNC_TRACE(LogDisplayClusterRender);
@@ -84,8 +76,6 @@ bool FDisplayClusterDeviceBase::Initialize()
 		return false;
 	}
 
-	UE_LOG(LogDisplayClusterRender, Log, TEXT("Use swap interval: %d"), GetSwapInt());
-
 	const IPDisplayClusterConfigManager* const ConfigMgr = GDisplayCluster->GetPrivateConfigMgr();
 	if (!ConfigMgr)
 	{
@@ -296,7 +286,7 @@ bool FDisplayClusterDeviceBase::GetBufferRatio(const FString& InViewportID, floa
 	}
 
 	OutBufferRatio = DesiredViewport->GetBufferRatio();
-	UE_LOG(LogDisplayClusterRender, Log, TEXT("Viewport '%s' has buffer ratio %f"), *InViewportID, OutBufferRatio);
+	UE_LOG(LogDisplayClusterRender, Verbose, TEXT("Viewport '%s' has buffer ratio %f"), *InViewportID, OutBufferRatio);
 	return true;
 }
 
@@ -308,7 +298,7 @@ bool FDisplayClusterDeviceBase::GetBufferRatio(int32 ViewIdx, float& OutBufferRa
 	}
 
 	OutBufferRatio = RenderViewports[ViewIdx].GetBufferRatio();
-	UE_LOG(LogDisplayClusterRender, Log, TEXT("Viewport '%s' has buffer ratio %f"), *RenderViewports[ViewIdx].GetId(), OutBufferRatio);
+	UE_LOG(LogDisplayClusterRender, Verbose, TEXT("Viewport '%s' has buffer ratio %f"), *RenderViewports[ViewIdx].GetId(), OutBufferRatio);
 	return true;
 }
 
@@ -379,7 +369,10 @@ void FDisplayClusterDeviceBase::CalculateStereoViewOffset(const enum EStereoscop
 		return;
 	}
 
-	UE_LOG(LogDisplayClusterRender, Verbose, TEXT("Viewport '%s' has assigned camera '%s'"), *Viewport.GetId(), *CameraId);
+	if (CameraId.Len() > 0)
+	{
+		UE_LOG(LogDisplayClusterRender, Verbose, TEXT("Viewport '%s' has assigned camera '%s'"), *Viewport.GetId(), *CameraId);
+	}
 
 	// Get the actual camera settings
 	const float CfgEyeDist     = ViewCamera->GetInterpupillaryDistance();
@@ -438,9 +431,8 @@ void FDisplayClusterDeviceBase::CalculateStereoViewOffset(const enum EStereoscop
 	ViewContext.ViewLocation  = ViewLocation;
 	ViewContext.ViewRotation  = ViewRotation;
 	ViewContext.WorldToMeters = WorldToMeters;
-	ViewContext.ClipPlaneNear = CfgNCP;
 
-	UE_LOG(LogDisplayClusterRender, VeryVerbose, TEXT("NEW ViewLoc: %s, ViewRot: %s"), *ViewLocation.ToString(), *ViewRotation.ToString());
+	UE_LOG(LogDisplayClusterRender, VeryVerbose, TEXT("ViewLoc: %s, ViewRot: %s"), *ViewLocation.ToString(), *ViewRotation.ToString());
 }
 
 FMatrix FDisplayClusterDeviceBase::GetStereoProjectionMatrix(const enum EStereoscopicPass StereoPassType) const
@@ -546,6 +538,7 @@ uint32 FDisplayClusterDeviceBase::GetViewIndexForPass(EStereoscopicPass StereoPa
 
 	switch (StereoPassType)
 	{
+	case EStereoscopicPass::eSSP_FULL:
 	case EStereoscopicPass::eSSP_LEFT_EYE:
 		DecodedViewIndex = 0;
 		break;
@@ -562,6 +555,7 @@ uint32 FDisplayClusterDeviceBase::GetViewIndexForPass(EStereoscopicPass StereoPa
 	return DecodedViewIndex;
 }
 
+
 //////////////////////////////////////////////////////////////////////////////////////////////
 // IStereoRenderTargetManager
 //////////////////////////////////////////////////////////////////////////////////////////////
@@ -573,8 +567,17 @@ void FDisplayClusterDeviceBase::UpdateViewport(bool bUseSeparateRenderTarget, co
 	// Store viewport
 	if (!MainViewport)
 	{
+		// UE viewport
 		MainViewport = (FViewport*)&Viewport;
-		Viewport.GetViewportRHI()->SetCustomPresent(this);
+		// Current sync policy
+		TSharedPtr<IDisplayClusterRenderSyncPolicy> SyncPolicy = GDisplayCluster->GetRenderMgr()->GetCurrentSynchronizationPolicy();
+
+		// Create present handler
+		FDisplayClusterPresentationBase* const CustomPresentHandler = CreatePresentationObject(MainViewport, SyncPolicy);
+		check(CustomPresentHandler);
+
+		// Set up our new present handler
+		Viewport.GetViewportRHI()->SetCustomPresent(CustomPresentHandler);
 	}
 }
 
@@ -628,43 +631,6 @@ bool FDisplayClusterDeviceBase::NeedReAllocateViewportRenderTarget(const class F
 }
 
 
-//////////////////////////////////////////////////////////////////////////////////////////////
-// FRHICustomPresent
-//////////////////////////////////////////////////////////////////////////////////////////////
-void FDisplayClusterDeviceBase::OnBackBufferResize()
-{
-	DISPLAY_CLUSTER_FUNC_TRACE(LogDisplayClusterRender);
-
-	//@todo: see comment below
-	// if we are in the middle of rendering: prevent from calling EndFrame
-	//if (RenderContext.IsValid())
-	//{
-	//	RenderContext->bFrameBegun = false;
-	//}
-}
-
-bool FDisplayClusterDeviceBase::Present(int32& InOutSyncInterval)
-{
-	DISPLAY_CLUSTER_FUNC_TRACE(LogDisplayClusterRender);
-
-	// Update sync value with nDisplay value
-	InOutSyncInterval = GetSwapInt();
-
-	// Get sync policy instance
-	TSharedPtr<IDisplayClusterRenderSyncPolicy> SyncPolicy = GDisplayCluster->GetRenderMgr()->GetCurrentSynchronizationPolicy();
-	if (SyncPolicy.IsValid())
-	{
-		// False results means we don't need to present current frame, the sync object already presented it
-		if (!SyncPolicy->SynchronizeClusterRendering(InOutSyncInterval))
-		{
-			return false;
-		}
-	}
-
-	return true;
-}
-
-
 //////////////////////////////////////////////////////////////////////////////////////////////
 // FDisplayClusterDeviceBase
 //////////////////////////////////////////////////////////////////////////////////////////////
@@ -727,6 +693,7 @@ int FDisplayClusterDeviceBase::DecodeViewportIndex(const enum EStereoscopicPass
 
 	UE_LOG(LogDisplayClusterRender, VeryVerbose, TEXT("DecodeViewportIndex: %d -> %d"), int(StereoPassType), DecodedViewportIndex);
 
+	check(DecodedViewportIndex >= 0);
 	return DecodedViewportIndex;
 }
 
@@ -767,12 +734,6 @@ uint32 FDisplayClusterDeviceBase::DecodeViewIndex(const enum EStereoscopicPass S
 	return DecodedViewIndex;
 }
 
-uint32 FDisplayClusterDeviceBase::GetSwapInt() const
-{
-	const uint32 SyncInterval = static_cast<uint32>(CVarVSyncInterval.GetValueOnAnyThread());
-	return (SyncInterval);
-}
-
 void FDisplayClusterDeviceBase::AddViewport(const FString& InViewportId, const FIntPoint& InViewportLocation, const FIntPoint& InViewportSize, TSharedPtr<IDisplayClusterProjectionPolicy> InProjPolicy, const FString& InCameraId, float InBufferRatio /* = 1.f */, bool IsRTT /*= false*/)
 {
 	DISPLAY_CLUSTER_FUNC_TRACE(LogDisplayClusterRender);
diff --git a/Engine/Plugins/Runtime/nDisplay/Source/DisplayCluster/Private/Render/Device/DisplayClusterDeviceBase.h b/Engine/Plugins/Runtime/nDisplay/Source/DisplayCluster/Private/Render/Device/DisplayClusterDeviceBase.h
--- a/Engine/Plugins/Runtime/nDisplay/Source/DisplayCluster/Private/Render/Device/DisplayClusterDeviceBase.h
+++ b/Engine/Plugins/Runtime/nDisplay/Source/DisplayCluster/Private/Render/Device/DisplayClusterDeviceBase.h
@@ -16,6 +16,7 @@
 #include "Containers/Queue.h"
 
 class IDisplayClusterPostProcess;
+class FDisplayClusterPresentationBase;
 
 
 /**
@@ -24,7 +25,6 @@ class IDisplayClusterPostProcess;
 class FDisplayClusterDeviceBase
 	: public IStereoRenderTargetManager
 	, public IDisplayClusterRenderDevice
-	, public FRHICustomPresent
 	, protected FDisplayClusterDeviceBase_PostProcess
 {
 public:
@@ -98,17 +98,6 @@ protected:
 	virtual bool DeviceIsAnAdditionalView(const FSceneView& View) override
 	{ return false; }
 
-protected:
-	//////////////////////////////////////////////////////////////////////////////////////////////
-	// FRHICustomPresent
-	//////////////////////////////////////////////////////////////////////////////////////////////
-	virtual void OnBackBufferResize() override;
-
-	virtual bool NeedsNativePresent() override
-	{ return true; }
-
-	virtual bool Present(int32& InOutSyncInterval) override;
-
 protected:
 	//////////////////////////////////////////////////////////////////////////////////////////////
 	// FDisplayClusterDeviceBase
@@ -133,14 +122,13 @@ protected:
 	// Decodes view index for a viewport (i.e. left=0, right=1)
 	uint32 DecodeViewIndex(const enum EStereoscopicPass StereoPassType) const;
 
-	// Returns swap interval
-	uint32 GetSwapInt() const;
-
 	// Adds a new viewport with specified parameters and projection policy object
 	void AddViewport(const FString& InViewportId, const FIntPoint& InViewportLocation, const FIntPoint& InViewportSize, TSharedPtr<IDisplayClusterProjectionPolicy> InProjPolicy, const FString& InCameraId, float InBufferRatio = 1.f, bool IsRTT = false);
 	// Performs copying of render target data to the back buffer
 	virtual void CopyTextureToBackBuffer_RenderThread(FRHICommandListImmediate& RHICmdList, FRHITexture2D* BackBuffer, FRHITexture2D* SrcTexture, FVector2D WindowSize) const;
-	
+	// Factory method to instantiate an output presentation class implementation
+	virtual FDisplayClusterPresentationBase* CreatePresentationObject(FViewport* const Viewport, TSharedPtr<IDisplayClusterRenderSyncPolicy>& SyncPolicy) = 0;
+
 	// Checks if custom post processing settings is assigned for specific viewport and assign them to be used
 	virtual void StartFinalPostprocessSettings(struct FPostProcessSettings* StartPostProcessingSettings, const enum EStereoscopicPass StereoPassType) override;
 	virtual bool OverrideFinalPostprocessSettings(struct FPostProcessSettings* OverridePostProcessingSettings, const enum EStereoscopicPass StereoPassType, float& BlendWeight) override;
diff --git a/Engine/Plugins/Runtime/nDisplay/Source/DisplayCluster/Private/Render/Device/DisplayClusterRenderViewContext.h b/Engine/Plugins/Runtime/nDisplay/Source/DisplayCluster/Private/Render/Device/DisplayClusterRenderViewContext.h
--- a/Engine/Plugins/Runtime/nDisplay/Source/DisplayCluster/Private/Render/Device/DisplayClusterRenderViewContext.h
+++ b/Engine/Plugins/Runtime/nDisplay/Source/DisplayCluster/Private/Render/Device/DisplayClusterRenderViewContext.h
@@ -7,10 +7,6 @@
 
 struct FDisplayClusterRenderViewContext
 {
-	// Near/far clipping
-	float ClipPlaneNear = GNearClippingPlane;
-	float ClipPlaneFar  = 2000000.f;
-
 	// Camera location and orientation
 	FVector  ViewLocation = FVector::ZeroVector;
 	FRotator ViewRotation = FRotator::ZeroRotator;
diff --git a/Engine/Plugins/Runtime/nDisplay/Source/DisplayCluster/Private/Render/Device/Monoscopic/DisplayClusterDeviceMonoscopicDX11.cpp b/Engine/Plugins/Runtime/nDisplay/Source/DisplayCluster/Private/Render/Device/Monoscopic/DisplayClusterDeviceMonoscopicDX11.cpp
--- a/Engine/Plugins/Runtime/nDisplay/Source/DisplayCluster/Private/Render/Device/Monoscopic/DisplayClusterDeviceMonoscopicDX11.cpp
+++ b/Engine/Plugins/Runtime/nDisplay/Source/DisplayCluster/Private/Render/Device/Monoscopic/DisplayClusterDeviceMonoscopicDX11.cpp
@@ -1,6 +1,8 @@
 // Copyright 1998-2019 Epic Games, Inc. All Rights Reserved.
 
 #include "Render/Device/Monoscopic/DisplayClusterDeviceMonoscopicDX11.h"
+#include "Render/Presentation/DisplayClusterPresentationDX11.h"
+
 #include "DisplayClusterLog.h"
 
 
@@ -14,16 +16,7 @@ FDisplayClusterDeviceMonoscopicDX11::~FDisplayClusterDeviceMonoscopicDX11()
 	DISPLAY_CLUSTER_FUNC_TRACE(LogDisplayClusterRender);
 }
 
-bool FDisplayClusterDeviceMonoscopicDX11::Present(int32& InOutSyncInterval)
+FDisplayClusterPresentationBase* FDisplayClusterDeviceMonoscopicDX11::CreatePresentationObject(FViewport* const Viewport, TSharedPtr<IDisplayClusterRenderSyncPolicy>& SyncPolicy)
 {
-	DISPLAY_CLUSTER_FUNC_TRACE(LogDisplayClusterRender);
-
-	// Perform abstract synchronization on a higher level
-	if (!FDisplayClusterDeviceBase::Present(InOutSyncInterval))
-	{
-		return false;
-	}
-
-	// A sync policy hasn't presented current frame so we do it ourselves
-	return PresentImpl(MainViewport->GetViewportRHI(), InOutSyncInterval);
+	return new FDisplayClusterPresentationDX11(Viewport, SyncPolicy);
 }
diff --git a/Engine/Plugins/Runtime/nDisplay/Source/DisplayCluster/Private/Render/Device/Monoscopic/DisplayClusterDeviceMonoscopicDX11.h b/Engine/Plugins/Runtime/nDisplay/Source/DisplayCluster/Private/Render/Device/Monoscopic/DisplayClusterDeviceMonoscopicDX11.h
--- a/Engine/Plugins/Runtime/nDisplay/Source/DisplayCluster/Private/Render/Device/Monoscopic/DisplayClusterDeviceMonoscopicDX11.h
+++ b/Engine/Plugins/Runtime/nDisplay/Source/DisplayCluster/Private/Render/Device/Monoscopic/DisplayClusterDeviceMonoscopicDX11.h
@@ -3,7 +3,6 @@
 #pragma once
 
 #include "Render/Device/DisplayClusterDeviceMonoscopicBase.h"
-#include "Render/Presentation/DisplayClusterDevicePresentationDX11.h"
 
 
 /**
@@ -11,15 +10,11 @@
  */
 class FDisplayClusterDeviceMonoscopicDX11
 	: public FDisplayClusterDeviceMonoscopicBase
-	, public FDisplayClusterDevicePresentationDX11
 {
 public:
 	FDisplayClusterDeviceMonoscopicDX11();
 	virtual ~FDisplayClusterDeviceMonoscopicDX11();
 
 protected:
-	//////////////////////////////////////////////////////////////////////////////////////////////
-	// FRHICustomPresent
-	//////////////////////////////////////////////////////////////////////////////////////////////
-	virtual bool Present(int32& InOutSyncInterval) override;
+	virtual FDisplayClusterPresentationBase* CreatePresentationObject(FViewport* const Viewport, TSharedPtr<IDisplayClusterRenderSyncPolicy>& SyncPolicy) override;
 };
diff --git a/Engine/Plugins/Runtime/nDisplay/Source/DisplayCluster/Private/Render/Device/Monoscopic/DisplayClusterDeviceMonoscopicDX12.cpp b/Engine/Plugins/Runtime/nDisplay/Source/DisplayCluster/Private/Render/Device/Monoscopic/DisplayClusterDeviceMonoscopicDX12.cpp
--- a/Engine/Plugins/Runtime/nDisplay/Source/DisplayCluster/Private/Render/Device/Monoscopic/DisplayClusterDeviceMonoscopicDX12.cpp
+++ b/Engine/Plugins/Runtime/nDisplay/Source/DisplayCluster/Private/Render/Device/Monoscopic/DisplayClusterDeviceMonoscopicDX12.cpp
@@ -1,6 +1,8 @@
 // Copyright 1998-2019 Epic Games, Inc. All Rights Reserved.
 
 #include "Render/Device/Monoscopic/DisplayClusterDeviceMonoscopicDX12.h"
+#include "Render/Presentation/DisplayClusterPresentationDX12.h"
+
 #include "DisplayClusterLog.h"
 
 
@@ -14,16 +16,7 @@ FDisplayClusterDeviceMonoscopicDX12::~FDisplayClusterDeviceMonoscopicDX12()
 	DISPLAY_CLUSTER_FUNC_TRACE(LogDisplayClusterRender);
 }
 
-bool FDisplayClusterDeviceMonoscopicDX12::Present(int32& InOutSyncInterval)
+FDisplayClusterPresentationBase* FDisplayClusterDeviceMonoscopicDX12::CreatePresentationObject(FViewport* const Viewport, TSharedPtr<IDisplayClusterRenderSyncPolicy>& SyncPolicy)
 {
-	DISPLAY_CLUSTER_FUNC_TRACE(LogDisplayClusterRender);
-
-	// Perform abstract synchronization on a higher level
-	if (!FDisplayClusterDeviceBase::Present(InOutSyncInterval))
-	{
-		return false;
-	}
-
-	// A sync policy hasn't presented current frame so we do it ourselves
-	return PresentImpl(MainViewport->GetViewportRHI(), InOutSyncInterval);
-}
+	return new FDisplayClusterPresentationDX12(Viewport, SyncPolicy);
+}
\ No newline at end of file
diff --git a/Engine/Plugins/Runtime/nDisplay/Source/DisplayCluster/Private/Render/Device/Monoscopic/DisplayClusterDeviceMonoscopicDX12.h b/Engine/Plugins/Runtime/nDisplay/Source/DisplayCluster/Private/Render/Device/Monoscopic/DisplayClusterDeviceMonoscopicDX12.h
--- a/Engine/Plugins/Runtime/nDisplay/Source/DisplayCluster/Private/Render/Device/Monoscopic/DisplayClusterDeviceMonoscopicDX12.h
+++ b/Engine/Plugins/Runtime/nDisplay/Source/DisplayCluster/Private/Render/Device/Monoscopic/DisplayClusterDeviceMonoscopicDX12.h
@@ -4,7 +4,6 @@
 
 #include "CoreMinimal.h"
 #include "Render/Device/DisplayClusterDeviceMonoscopicBase.h"
-#include "Render/Presentation/DisplayClusterDevicePresentationDX12.h"
 
 
 /**
@@ -12,15 +11,11 @@
  */
 class FDisplayClusterDeviceMonoscopicDX12
 	: public FDisplayClusterDeviceMonoscopicBase
-	, public FDisplayClusterDevicePresentationDX12
 {
 public:
 	FDisplayClusterDeviceMonoscopicDX12();
 	virtual ~FDisplayClusterDeviceMonoscopicDX12();
 
 protected:
-	//////////////////////////////////////////////////////////////////////////////////////////////
-	// FRHICustomPresent
-	//////////////////////////////////////////////////////////////////////////////////////////////
-	virtual bool Present(int32& InOutSyncInterval) override;
+	virtual FDisplayClusterPresentationBase* CreatePresentationObject(FViewport* const Viewport, TSharedPtr<IDisplayClusterRenderSyncPolicy>& SyncPolicy) override;
 };
diff --git a/Engine/Plugins/Runtime/nDisplay/Source/DisplayCluster/Private/Render/Device/QuadBufferStereo/DisplayClusterDeviceQuadBufferStereoDX11.cpp b/Engine/Plugins/Runtime/nDisplay/Source/DisplayCluster/Private/Render/Device/QuadBufferStereo/DisplayClusterDeviceQuadBufferStereoDX11.cpp
--- a/Engine/Plugins/Runtime/nDisplay/Source/DisplayCluster/Private/Render/Device/QuadBufferStereo/DisplayClusterDeviceQuadBufferStereoDX11.cpp
+++ b/Engine/Plugins/Runtime/nDisplay/Source/DisplayCluster/Private/Render/Device/QuadBufferStereo/DisplayClusterDeviceQuadBufferStereoDX11.cpp
@@ -1,6 +1,8 @@
 // Copyright 1998-2019 Epic Games, Inc. All Rights Reserved.
 
 #include "Render/Device/QuadBufferStereo/DisplayClusterDeviceQuadBufferStereoDX11.h"
+#include "Render/Presentation/DisplayClusterPresentationDX11.h"
+
 #include "DisplayClusterLog.h"
 
 
@@ -14,16 +16,7 @@ FDisplayClusterDeviceQuadBufferStereoDX11::~FDisplayClusterDeviceQuadBufferStere
 	DISPLAY_CLUSTER_FUNC_TRACE(LogDisplayClusterRender);
 }
 
-bool FDisplayClusterDeviceQuadBufferStereoDX11::Present(int32& InOutSyncInterval)
+FDisplayClusterPresentationBase* FDisplayClusterDeviceQuadBufferStereoDX11::CreatePresentationObject(FViewport* const Viewport, TSharedPtr<IDisplayClusterRenderSyncPolicy>& SyncPolicy)
 {
-	DISPLAY_CLUSTER_FUNC_TRACE(LogDisplayClusterRender);
-
-	// Perform abstract synchronization on a higher level
-	if (!FDisplayClusterDeviceBase::Present(InOutSyncInterval))
-	{
-		return false;
-	}
-
-	// A sync policy hasn't presented current frame so we do it ourselves
-	return PresentImpl(MainViewport->GetViewportRHI(), InOutSyncInterval);
+	return new FDisplayClusterPresentationDX11(Viewport, SyncPolicy);
 }
diff --git a/Engine/Plugins/Runtime/nDisplay/Source/DisplayCluster/Private/Render/Device/QuadBufferStereo/DisplayClusterDeviceQuadBufferStereoDX11.h b/Engine/Plugins/Runtime/nDisplay/Source/DisplayCluster/Private/Render/Device/QuadBufferStereo/DisplayClusterDeviceQuadBufferStereoDX11.h
--- a/Engine/Plugins/Runtime/nDisplay/Source/DisplayCluster/Private/Render/Device/QuadBufferStereo/DisplayClusterDeviceQuadBufferStereoDX11.h
+++ b/Engine/Plugins/Runtime/nDisplay/Source/DisplayCluster/Private/Render/Device/QuadBufferStereo/DisplayClusterDeviceQuadBufferStereoDX11.h
@@ -4,7 +4,6 @@
 
 #include "CoreMinimal.h"
 #include "Render/Device/QuadBufferStereo/DisplayClusterDeviceQuadBufferStereoBase.h"
-#include "Render/Presentation/DisplayClusterDevicePresentationDX11.h"
 
 
 /**
@@ -12,15 +11,11 @@
  */
 class FDisplayClusterDeviceQuadBufferStereoDX11
 	: public FDisplayClusterDeviceQuadBufferStereoBase
-	, public FDisplayClusterDevicePresentationDX11
 {
 public:
 	FDisplayClusterDeviceQuadBufferStereoDX11();
 	virtual ~FDisplayClusterDeviceQuadBufferStereoDX11();
 
 protected:
-	//////////////////////////////////////////////////////////////////////////////////////////////
-	// FRHICustomPresent
-	//////////////////////////////////////////////////////////////////////////////////////////////
-	virtual bool Present(int32& InOutSyncInterval) override;
+	virtual FDisplayClusterPresentationBase* CreatePresentationObject(FViewport* const Viewport, TSharedPtr<IDisplayClusterRenderSyncPolicy>& SyncPolicy) override;
 };
diff --git a/Engine/Plugins/Runtime/nDisplay/Source/DisplayCluster/Private/Render/Device/QuadBufferStereo/DisplayClusterDeviceQuadBufferStereoDX12.cpp b/Engine/Plugins/Runtime/nDisplay/Source/DisplayCluster/Private/Render/Device/QuadBufferStereo/DisplayClusterDeviceQuadBufferStereoDX12.cpp
--- a/Engine/Plugins/Runtime/nDisplay/Source/DisplayCluster/Private/Render/Device/QuadBufferStereo/DisplayClusterDeviceQuadBufferStereoDX12.cpp
+++ b/Engine/Plugins/Runtime/nDisplay/Source/DisplayCluster/Private/Render/Device/QuadBufferStereo/DisplayClusterDeviceQuadBufferStereoDX12.cpp
@@ -1,6 +1,8 @@
 // Copyright 1998-2019 Epic Games, Inc. All Rights Reserved.
 
 #include "Render/Device/QuadBufferStereo/DisplayClusterDeviceQuadBufferStereoDX12.h"
+#include "Render/Presentation/DisplayClusterPresentationDX12.h"
+
 #include "DisplayClusterLog.h"
 
 
@@ -14,16 +16,7 @@ FDisplayClusterDeviceQuadBufferStereoDX12::~FDisplayClusterDeviceQuadBufferStere
 	DISPLAY_CLUSTER_FUNC_TRACE(LogDisplayClusterRender);
 }
 
-bool FDisplayClusterDeviceQuadBufferStereoDX12::Present(int32& InOutSyncInterval)
+FDisplayClusterPresentationBase* FDisplayClusterDeviceQuadBufferStereoDX12::CreatePresentationObject(FViewport* const Viewport, TSharedPtr<IDisplayClusterRenderSyncPolicy>& SyncPolicy)
 {
-	DISPLAY_CLUSTER_FUNC_TRACE(LogDisplayClusterRender);
-
-	// Perform abstract synchronization on a higher level
-	if (!FDisplayClusterDeviceBase::Present(InOutSyncInterval))
-	{
-		return false;
-	}
-
-	// A sync policy hasn't presented current frame so we do it ourselves
-	return PresentImpl(MainViewport->GetViewportRHI(), InOutSyncInterval);
+	return new FDisplayClusterPresentationDX12(Viewport, SyncPolicy);
 }
diff --git a/Engine/Plugins/Runtime/nDisplay/Source/DisplayCluster/Private/Render/Device/QuadBufferStereo/DisplayClusterDeviceQuadBufferStereoDX12.h b/Engine/Plugins/Runtime/nDisplay/Source/DisplayCluster/Private/Render/Device/QuadBufferStereo/DisplayClusterDeviceQuadBufferStereoDX12.h
--- a/Engine/Plugins/Runtime/nDisplay/Source/DisplayCluster/Private/Render/Device/QuadBufferStereo/DisplayClusterDeviceQuadBufferStereoDX12.h
+++ b/Engine/Plugins/Runtime/nDisplay/Source/DisplayCluster/Private/Render/Device/QuadBufferStereo/DisplayClusterDeviceQuadBufferStereoDX12.h
@@ -4,7 +4,6 @@
 
 #include "CoreMinimal.h"
 #include "Render/Device/QuadBufferStereo/DisplayClusterDeviceQuadBufferStereoBase.h"
-#include "Render/Presentation/DisplayClusterDevicePresentationDX12.h"
 
 
 /**
@@ -12,15 +11,11 @@
  */
 class FDisplayClusterDeviceQuadBufferStereoDX12
 	: public FDisplayClusterDeviceQuadBufferStereoBase
-	, public FDisplayClusterDevicePresentationDX12
 {
 public:
 	FDisplayClusterDeviceQuadBufferStereoDX12();
 	virtual ~FDisplayClusterDeviceQuadBufferStereoDX12();
 
 protected:
-	//////////////////////////////////////////////////////////////////////////////////////////////
-	// FRHICustomPresent
-	//////////////////////////////////////////////////////////////////////////////////////////////
-	virtual bool Present(int32& InOutSyncInterval) override;
+	virtual FDisplayClusterPresentationBase* CreatePresentationObject(FViewport* const Viewport, TSharedPtr<IDisplayClusterRenderSyncPolicy>& SyncPolicy) override;
 };
diff --git a/Engine/Plugins/Runtime/nDisplay/Source/DisplayCluster/Private/Render/Device/SideBySide/DisplayClusterDeviceSideBySideDX11.cpp b/Engine/Plugins/Runtime/nDisplay/Source/DisplayCluster/Private/Render/Device/SideBySide/DisplayClusterDeviceSideBySideDX11.cpp
--- a/Engine/Plugins/Runtime/nDisplay/Source/DisplayCluster/Private/Render/Device/SideBySide/DisplayClusterDeviceSideBySideDX11.cpp
+++ b/Engine/Plugins/Runtime/nDisplay/Source/DisplayCluster/Private/Render/Device/SideBySide/DisplayClusterDeviceSideBySideDX11.cpp
@@ -1,6 +1,8 @@
 // Copyright 1998-2019 Epic Games, Inc. All Rights Reserved.
 
 #include "Render/Device/SidebySide/DisplayClusterDeviceSideBySideDX11.h"
+#include "Render/Presentation/DisplayClusterPresentationDX11.h"
+
 #include "DisplayClusterLog.h"
 
 
@@ -14,17 +16,7 @@ FDisplayClusterDeviceSideBySideDX11::~FDisplayClusterDeviceSideBySideDX11()
 	DISPLAY_CLUSTER_FUNC_TRACE(LogDisplayClusterRender);
 }
 
-
-bool FDisplayClusterDeviceSideBySideDX11::Present(int32& InOutSyncInterval)
+FDisplayClusterPresentationBase* FDisplayClusterDeviceSideBySideDX11::CreatePresentationObject(FViewport* const Viewport, TSharedPtr<IDisplayClusterRenderSyncPolicy>& SyncPolicy)
 {
-	DISPLAY_CLUSTER_FUNC_TRACE(LogDisplayClusterRender);
-
-	// Perform abstract synchronization on a higher level
-	if (!FDisplayClusterDeviceBase::Present(InOutSyncInterval))
-	{
-		return false;
-	}
-
-	// A sync policy hasn't presented current frame so we do it ourselves
-	return PresentImpl(MainViewport->GetViewportRHI(), InOutSyncInterval);
+	return new FDisplayClusterPresentationDX11(Viewport, SyncPolicy);
 }
diff --git a/Engine/Plugins/Runtime/nDisplay/Source/DisplayCluster/Private/Render/Device/SideBySide/DisplayClusterDeviceSideBySideDX12.cpp b/Engine/Plugins/Runtime/nDisplay/Source/DisplayCluster/Private/Render/Device/SideBySide/DisplayClusterDeviceSideBySideDX12.cpp
--- a/Engine/Plugins/Runtime/nDisplay/Source/DisplayCluster/Private/Render/Device/SideBySide/DisplayClusterDeviceSideBySideDX12.cpp
+++ b/Engine/Plugins/Runtime/nDisplay/Source/DisplayCluster/Private/Render/Device/SideBySide/DisplayClusterDeviceSideBySideDX12.cpp
@@ -1,6 +1,8 @@
 // Copyright 1998-2019 Epic Games, Inc. All Rights Reserved.
 
 #include "Render/Device/SidebySide/DisplayClusterDeviceSideBySideDX12.h"
+#include "Render/Presentation/DisplayClusterPresentationDX12.h"
+
 #include "DisplayClusterLog.h"
 
 
@@ -14,17 +16,7 @@ FDisplayClusterDeviceSideBySideDX12::~FDisplayClusterDeviceSideBySideDX12()
 	DISPLAY_CLUSTER_FUNC_TRACE(LogDisplayClusterRender);
 }
 
-
-bool FDisplayClusterDeviceSideBySideDX12::Present(int32& InOutSyncInterval)
+FDisplayClusterPresentationBase* FDisplayClusterDeviceSideBySideDX12::CreatePresentationObject(FViewport* const Viewport, TSharedPtr<IDisplayClusterRenderSyncPolicy>& SyncPolicy)
 {
-	DISPLAY_CLUSTER_FUNC_TRACE(LogDisplayClusterRender);
-
-	// Perform abstract synchronization on a higher level
-	if (!FDisplayClusterDeviceBase::Present(InOutSyncInterval))
-	{
-		return false;
-	}
-
-	// A sync policy hasn't presented current frame so we do it ourselves
-	return PresentImpl(MainViewport->GetViewportRHI(), InOutSyncInterval);
+	return new FDisplayClusterPresentationDX12(Viewport, SyncPolicy);
 }
diff --git a/Engine/Plugins/Runtime/nDisplay/Source/DisplayCluster/Private/Render/Device/TopBottom/DisplayClusterDeviceTopBottomDX11.cpp b/Engine/Plugins/Runtime/nDisplay/Source/DisplayCluster/Private/Render/Device/TopBottom/DisplayClusterDeviceTopBottomDX11.cpp
--- a/Engine/Plugins/Runtime/nDisplay/Source/DisplayCluster/Private/Render/Device/TopBottom/DisplayClusterDeviceTopBottomDX11.cpp
+++ b/Engine/Plugins/Runtime/nDisplay/Source/DisplayCluster/Private/Render/Device/TopBottom/DisplayClusterDeviceTopBottomDX11.cpp
@@ -1,6 +1,8 @@
 // Copyright 1998-2019 Epic Games, Inc. All Rights Reserved.
 
 #include "Render/Device/TopBottom/DisplayClusterDeviceTopBottomDX11.h"
+#include "Render/Presentation/DisplayClusterPresentationDX11.h"
+
 #include "DisplayClusterLog.h"
 
 
@@ -14,17 +16,7 @@ FDisplayClusterDeviceTopBottomDX11::~FDisplayClusterDeviceTopBottomDX11()
 	DISPLAY_CLUSTER_FUNC_TRACE(LogDisplayClusterRender);
 }
 
-
-bool FDisplayClusterDeviceTopBottomDX11::Present(int32& InOutSyncInterval)
+FDisplayClusterPresentationBase* FDisplayClusterDeviceTopBottomDX11::CreatePresentationObject(FViewport* const Viewport, TSharedPtr<IDisplayClusterRenderSyncPolicy>& SyncPolicy)
 {
-	DISPLAY_CLUSTER_FUNC_TRACE(LogDisplayClusterRender);
-
-	// Perform abstract synchronization on a higher level
-	if (!FDisplayClusterDeviceBase::Present(InOutSyncInterval))
-	{
-		return false;
-	}
-
-	// A sync policy hasn't presented current frame so we do it ourselves
-	return PresentImpl(MainViewport->GetViewportRHI(), InOutSyncInterval);
+	return new FDisplayClusterPresentationDX11(Viewport, SyncPolicy);
 }
diff --git a/Engine/Plugins/Runtime/nDisplay/Source/DisplayCluster/Private/Render/Device/TopBottom/DisplayClusterDeviceTopBottomDX12.cpp b/Engine/Plugins/Runtime/nDisplay/Source/DisplayCluster/Private/Render/Device/TopBottom/DisplayClusterDeviceTopBottomDX12.cpp
--- a/Engine/Plugins/Runtime/nDisplay/Source/DisplayCluster/Private/Render/Device/TopBottom/DisplayClusterDeviceTopBottomDX12.cpp
+++ b/Engine/Plugins/Runtime/nDisplay/Source/DisplayCluster/Private/Render/Device/TopBottom/DisplayClusterDeviceTopBottomDX12.cpp
@@ -1,6 +1,8 @@
 // Copyright 1998-2019 Epic Games, Inc. All Rights Reserved.
 
 #include "Render/Device/TopBottom/DisplayClusterDeviceTopBottomDX12.h"
+#include "Render/Presentation/DisplayClusterPresentationDX12.h"
+
 #include "DisplayClusterLog.h"
 
 
@@ -14,17 +16,7 @@ FDisplayClusterDeviceTopBottomDX12::~FDisplayClusterDeviceTopBottomDX12()
 	DISPLAY_CLUSTER_FUNC_TRACE(LogDisplayClusterRender);
 }
 
-
-bool FDisplayClusterDeviceTopBottomDX12::Present(int32& InOutSyncInterval)
+FDisplayClusterPresentationBase* FDisplayClusterDeviceTopBottomDX12::CreatePresentationObject(FViewport* const Viewport, TSharedPtr<IDisplayClusterRenderSyncPolicy>& SyncPolicy)
 {
-	DISPLAY_CLUSTER_FUNC_TRACE(LogDisplayClusterRender);
-
-	// Perform abstract synchronization on a higher level
-	if (!FDisplayClusterDeviceBase::Present(InOutSyncInterval))
-	{
-		return false;
-	}
-
-	// A sync policy hasn't presented current frame so we do it ourselves
-	return PresentImpl(MainViewport->GetViewportRHI(), InOutSyncInterval);
+	return new FDisplayClusterPresentationDX12(Viewport, SyncPolicy);
 }
diff --git a/Engine/Plugins/Runtime/nDisplay/Source/DisplayCluster/Private/Render/DisplayClusterRenderManager.cpp b/Engine/Plugins/Runtime/nDisplay/Source/DisplayCluster/Private/Render/DisplayClusterRenderManager.cpp
--- a/Engine/Plugins/Runtime/nDisplay/Source/DisplayCluster/Private/Render/DisplayClusterRenderManager.cpp
+++ b/Engine/Plugins/Runtime/nDisplay/Source/DisplayCluster/Private/Render/DisplayClusterRenderManager.cpp
@@ -8,16 +8,17 @@
 
 #include "DisplayClusterGlobals.h"
 #include "DisplayClusterLog.h"
-#include "DisplayClusterOperationMode.h"
 #include "DisplayClusterStrings.h"
 
+#include "DisplayClusterCameraComponent.h"
+#include "DisplayClusterViewportClient.h"
+
 #include "Config/DisplayClusterConfigTypes.h"
 #include "Config/IDisplayClusterConfigManager.h"
 
 #include "Game/IDisplayClusterGameManager.h"
 
 #include "Render/Device/DisplayClusterRenderDeviceFactoryInternal.h"
-#include "Render/Device/DisplayClusterDeviceNativePresentHandler.h"
 #include "Render/Device/Monoscopic/DisplayClusterDeviceMonoscopicDX11.h"
 
 #include "Render/Device/IDisplayClusterRenderDeviceFactory.h"
@@ -26,6 +27,8 @@
 #include "Render/Projection/IDisplayClusterProjectionPolicy.h"
 #include "Render/Synchronization/IDisplayClusterRenderSyncPolicyFactory.h"
 
+#include "Render/Presentation/DisplayClusterPresentationNative.h"
+
 #include "Render/Synchronization/DisplayClusterRenderSyncPolicyFactoryInternal.h"
 #include "Render/Synchronization/DisplayClusterRenderSyncPolicyNone.h"
 #include "Render/Synchronization/DisplayClusterRenderSyncPolicySoftwareGeneric.h"
@@ -93,35 +96,54 @@ bool FDisplayClusterRenderManager::StartSession(const FString& configPath, const
 		return true;
 	}
 
+	UDisplayClusterViewportClient* const GameViewport = Cast<UDisplayClusterViewportClient>(GEngine->GameViewport);
+	if (!GameViewport)
+	{
+		UE_LOG(LogDisplayClusterRender, Warning, TEXT("DisplayClusterViewportClient is not set as default GameViewport class"));
+	}
+
 	// Create synchronization object
 	UE_LOG(LogDisplayClusterRender, Log, TEXT("Instantiating synchronization policy object..."));
 	SyncPolicy = CreateRenderSyncPolicy();
 
 	// Instantiate render device
+	TSharedPtr<IDisplayClusterRenderDevice, ESPMode::ThreadSafe> NewRenderDevice;
 	UE_LOG(LogDisplayClusterRender, Log, TEXT("Instantiating stereo device..."));
-	RenderDevice = CreateRenderDevice();
+	NewRenderDevice = CreateRenderDevice();
 
 	// Set new device as the engine's stereoscopic device
-	if (GEngine && RenderDevice.IsValid())
+	if (GEngine && NewRenderDevice.IsValid())
 	{
-		GEngine->StereoRenderingDevice = StaticCastSharedPtr<IStereoRendering>(RenderDevice);
+		GEngine->StereoRenderingDevice = StaticCastSharedPtr<IStereoRendering>(NewRenderDevice);
+		RenderDevicePtr = NewRenderDevice.Get();
 	}
 
 	// When session is starting in Editor the device won't be initialized so we avoid nullptr access here.
 	//@todo Now we always have a device, even for Editor. Change the condition working on the EditorDevice.
-	return (RenderDevice.IsValid() ? RenderDevice->Initialize() : true);
+	return (RenderDevicePtr ? RenderDevicePtr->Initialize() : true);
 }
 
 void FDisplayClusterRenderManager::EndSession()
 {
 	DISPLAY_CLUSTER_FUNC_TRACE(LogDisplayClusterRender);
+
+#if WITH_EDITOR
+	if (GIsEditor)
+	{
+		// Since we can run multiple PIE sessions we have to clean device before the next one.
+		GEngine->StereoRenderingDevice.Reset();
+		RenderDevicePtr = nullptr;
+	}
+#endif
 }
 
 bool FDisplayClusterRenderManager::StartScene(UWorld* InWorld)
 {
-	if (RenderDevice.IsValid())
+	DISPLAY_CLUSTER_FUNC_TRACE(LogDisplayClusterRender);
+
+	if (RenderDevicePtr)
 	{
-		RenderDevice->InitializeWorldContent(InWorld);
+		RenderDevicePtr->InitializeWorldContent(InWorld);
 	}
 
 	return true;
@@ -129,14 +151,7 @@ bool FDisplayClusterRenderManager::StartScene(UWorld* InWorld)
 
 void FDisplayClusterRenderManager::EndScene()
 {
-#if WITH_EDITOR
-	if (GIsEditor)
-	{
-		// Since we can run multiple PIE sessions we have to clean device before the next one.
-		GEngine->StereoRenderingDevice.Reset();
-		RenderDevice.Reset();
-	}
-#endif
+	DISPLAY_CLUSTER_FUNC_TRACE(LogDisplayClusterRender);
 }
 
 void FDisplayClusterRenderManager::PreTick(float DeltaSeconds)
@@ -449,9 +464,9 @@ void FDisplayClusterRenderManager::SetViewportCamera(const FString& InCameraId /
 
 	{
 		FScopeLock lock(&CritSecInternals);
-		if (RenderDevice.IsValid())
+		if (RenderDevicePtr)
 		{
-			RenderDevice->SetViewportCamera(InCameraId, InViewportId);
+			RenderDevicePtr->SetViewportCamera(InCameraId, InViewportId);
 		}
 	}
 }
@@ -461,12 +476,12 @@ bool FDisplayClusterRenderManager::GetViewportRect(const FString& InViewportID,
 	DISPLAY_CLUSTER_FUNC_TRACE(LogDisplayClusterRender);
 	check(IsInGameThread());
 
-	if (!RenderDevice.IsValid())
+	if (!RenderDevicePtr)
 	{
 		return false;
 	}
 
-	return RenderDevice->GetViewportRect(InViewportID, Rect);
+	return RenderDevicePtr->GetViewportRect(InViewportID, Rect);
 }
 
 bool FDisplayClusterRenderManager::SetBufferRatio(const FString& InViewportID, float InBufferRatio)
@@ -474,12 +489,12 @@ bool FDisplayClusterRenderManager::SetBufferRatio(const FString& InViewportID, f
 	DISPLAY_CLUSTER_FUNC_TRACE(LogDisplayClusterRender);
 	check(IsInGameThread());
 
-	if (!RenderDevice.IsValid())
+	if (!RenderDevicePtr)
 	{
 		return false;
 	}
 
-	return RenderDevice->SetBufferRatio(InViewportID, InBufferRatio);
+	return RenderDevicePtr->SetBufferRatio(InViewportID, InBufferRatio);
 }
 
 bool FDisplayClusterRenderManager::GetBufferRatio(const FString& InViewportID, float &OutBufferRatio) const
@@ -487,12 +502,12 @@ bool FDisplayClusterRenderManager::GetBufferRatio(const FString& InViewportID, f
 	DISPLAY_CLUSTER_FUNC_TRACE(LogDisplayClusterRender);
 	check(IsInGameThread());
 
-	if (!RenderDevice.IsValid())
+	if (!RenderDevicePtr)
 	{
 		return false;
 	}
 
-	return RenderDevice->GetBufferRatio(InViewportID, OutBufferRatio);
+	return RenderDevicePtr->GetBufferRatio(InViewportID, OutBufferRatio);
 }
 
 void FDisplayClusterRenderManager::SetStartPostProcessingSettings(const FString& ViewportID, const FPostProcessSettings& StartPostProcessingSettings)
@@ -500,12 +515,12 @@ void FDisplayClusterRenderManager::SetStartPostProcessingSettings(const FString&
 	DISPLAY_CLUSTER_FUNC_TRACE(LogDisplayClusterRender);
 	check(IsInGameThread());
 
-	if (!RenderDevice.IsValid())
+	if (!RenderDevicePtr)
 	{
 		return;
 	}
 
-	RenderDevice->SetStartPostProcessingSettings(ViewportID, StartPostProcessingSettings);
+	RenderDevicePtr->SetStartPostProcessingSettings(ViewportID, StartPostProcessingSettings);
 }
 
 void FDisplayClusterRenderManager::SetOverridePostProcessingSettings(const FString& ViewportID, const FPostProcessSettings& OverridePostProcessingSettings, float BlendWeight)
@@ -513,12 +528,12 @@ void FDisplayClusterRenderManager::SetOverridePostProcessingSettings(const FStri
 	DISPLAY_CLUSTER_FUNC_TRACE(LogDisplayClusterRender);
 	check(IsInGameThread());
 
-	if (!RenderDevice.IsValid())
+	if (!RenderDevicePtr)
 	{
 		return;
 	}
 
-	RenderDevice->SetOverridePostProcessingSettings(ViewportID, OverridePostProcessingSettings, BlendWeight);
+	RenderDevicePtr->SetOverridePostProcessingSettings(ViewportID, OverridePostProcessingSettings, BlendWeight);
 }
 
 void FDisplayClusterRenderManager::SetFinalPostProcessingSettings(const FString& ViewportID, const FPostProcessSettings& FinalPostProcessingSettings)
@@ -526,12 +541,12 @@ void FDisplayClusterRenderManager::SetFinalPostProcessingSettings(const FString&
 	DISPLAY_CLUSTER_FUNC_TRACE(LogDisplayClusterRender);
 	check(IsInGameThread());
 
-	if (!RenderDevice.IsValid())
+	if (!RenderDevicePtr)
 	{
 		return;
 	}
 
-	RenderDevice->SetFinalPostProcessingSettings(ViewportID, FinalPostProcessingSettings);
+	RenderDevicePtr->SetFinalPostProcessingSettings(ViewportID, FinalPostProcessingSettings);
 }
 
 void FDisplayClusterRenderManager::SetInterpupillaryDistance(const FString& CameraId, float EyeDistance)
@@ -585,72 +600,6 @@ bool FDisplayClusterRenderManager::ToggleEyesSwap(const FString& CameraId)
 	return (Camera ? Camera->ToggleEyesSwap() : false);
 }
 
-float FDisplayClusterRenderManager::GetNearCullingDistance(const FString& CameraId) const
-{
-	DISPLAY_CLUSTER_FUNC_TRACE(LogDisplayClusterRender);
-	check(IsInGameThread());
-
-	UDisplayClusterCameraComponent* const Camera = DisplayClusterHelpers::game::GetCamera(CameraId);
-	return (Camera ? Camera->GetNearCullingDistance() : 0.f);
-}
-
-void FDisplayClusterRenderManager::SetNearCullingDistance(const FString& CameraId, float NearDistance)
-{
-	DISPLAY_CLUSTER_FUNC_TRACE(LogDisplayClusterRender);
-	check(IsInGameThread());
-
-	UDisplayClusterCameraComponent* const Camera = DisplayClusterHelpers::game::GetCamera(CameraId);
-	if (Camera)
-	{
-		Camera->SetNearCullingDistance(NearDistance);
-	}
-}
-
-float FDisplayClusterRenderManager::GetFarCullingDistance(const FString& CameraId) const
-{
-	DISPLAY_CLUSTER_FUNC_TRACE(LogDisplayClusterRender);
-	check(IsInGameThread());
-
-	UDisplayClusterCameraComponent* const Camera = DisplayClusterHelpers::game::GetCamera(CameraId);
-	return (Camera ? Camera->GetFarCullingDistance() : 0.f);
-}
-
-void FDisplayClusterRenderManager::SetFarCullingDistance(const FString& CameraId, float FarDistance)
-{
-	DISPLAY_CLUSTER_FUNC_TRACE(LogDisplayClusterRender);
-	check(IsInGameThread());
-
-	UDisplayClusterCameraComponent* const Camera = DisplayClusterHelpers::game::GetCamera(CameraId);
-	if (Camera)
-	{
-		Camera->SetFarCullingDistance(FarDistance);
-	}
-}
-
-void FDisplayClusterRenderManager::GetCullingDistance(const FString& CameraId, float& NearDistance, float& FarDistance) const
-{
-	DISPLAY_CLUSTER_FUNC_TRACE(LogDisplayClusterRender);
-	check(IsInGameThread());
-
-	UDisplayClusterCameraComponent* const Camera = DisplayClusterHelpers::game::GetCamera(CameraId);
-	if (Camera)
-	{
-		Camera->GetCullingDistance(NearDistance, FarDistance);
-	}
-}
-
-void FDisplayClusterRenderManager::SetCullingDistance(const FString& CameraId, float NearDistance, float FarDistance)
-{
-	DISPLAY_CLUSTER_FUNC_TRACE(LogDisplayClusterRender);
-	check(IsInGameThread());
-
-	UDisplayClusterCameraComponent* const Camera = DisplayClusterHelpers::game::GetCamera(CameraId);
-	if (Camera)
-	{
-		Camera->SetCullingDistance(NearDistance, FarDistance);
-	}
-}
-
 
 //////////////////////////////////////////////////////////////////////////////////////////////
 // FDisplayClusterRenderManager
@@ -701,8 +650,10 @@ TSharedPtr<IDisplayClusterRenderDevice, ESPMode::ThreadSafe> FDisplayClusterRend
 	}
 	else if (CurrentOperationMode == EDisplayClusterOperationMode::Editor)
 	{
+#if 0
 		UE_LOG(LogDisplayClusterRender, Log, TEXT("Instantiating DX11 mono device for PIE"));
 		NewRenderDevice = MakeShareable(new FDisplayClusterDeviceMonoscopicDX11());
+#endif
 	}
 	else if (CurrentOperationMode == EDisplayClusterOperationMode::Disabled)
 	{
@@ -792,11 +743,11 @@ void FDisplayClusterRenderManager::OnViewportCreatedHandler()
 
 void FDisplayClusterRenderManager::OnBeginDrawHandler()
 {
-	//@todo: this is fast solution for prototype. We shouldn't use raw handlers to be able to unsubscribe from the event.
 	static bool initialized = false;
 	if (!initialized && GEngine->GameViewport->Viewport->GetViewportRHI().IsValid())
 	{
-		NativePresentHandler = new FDisplayClusterDeviceNativePresentHandler();
+		FDisplayClusterPresentationNative* const NativePresentHandler = new FDisplayClusterPresentationNative(GEngine->GameViewport->Viewport, SyncPolicy);
+		check(NativePresentHandler);
 		GEngine->GameViewport->Viewport->GetViewportRHI().GetReference()->SetCustomPresent(NativePresentHandler);
 		initialized = true;
 	}
diff --git a/Engine/Plugins/Runtime/nDisplay/Source/DisplayCluster/Private/Render/DisplayClusterRenderManager.h b/Engine/Plugins/Runtime/nDisplay/Source/DisplayCluster/Private/Render/DisplayClusterRenderManager.h
--- a/Engine/Plugins/Runtime/nDisplay/Source/DisplayCluster/Private/Render/DisplayClusterRenderManager.h
+++ b/Engine/Plugins/Runtime/nDisplay/Source/DisplayCluster/Private/Render/DisplayClusterRenderManager.h
@@ -6,7 +6,6 @@
 
 #include "Render/IPDisplayClusterRenderManager.h"
 
-class FDisplayClusterDeviceNativePresentHandler;
 class IDisplayClusterPostProcess;
 class IDisplayClusterProjectionPolicy;
 class IDisplayClusterProjectionPolicyFactory;
@@ -77,12 +76,6 @@ public:
 	virtual bool  GetEyesSwap(const FString& CameraId) const override;
 	virtual void  SetEyesSwap(const FString& CameraId, bool EyeSwapped) override;
 	virtual bool  ToggleEyesSwap(const FString& CameraId) override;
-	virtual float GetNearCullingDistance(const FString& CameraId) const override;
-	virtual void  SetNearCullingDistance(const FString& CameraId, float NearDistance) override;
-	virtual float GetFarCullingDistance(const FString& CameraId) const override;
-	virtual void  SetFarCullingDistance(const FString& CameraId, float FarDistance) override;
-	virtual void  GetCullingDistance(const FString& CameraId, float& NearDistance, float& FarDistance) const override;
-	virtual void  SetCullingDistance(const FString& CameraId, float NearDistance, float FarDistance) override;
 
 public:
 	//////////////////////////////////////////////////////////////////////////////////////////////
@@ -103,8 +96,7 @@ private:
 	FString ClusterNodeId;
 
 	// Interface pointer to avoid type casting
-	TSharedPtr<IDisplayClusterRenderDevice, ESPMode::ThreadSafe> RenderDevice;
-	FDisplayClusterDeviceNativePresentHandler* NativePresentHandler;
+	IDisplayClusterRenderDevice* RenderDevicePtr = nullptr;
 	bool bWindowAdjusted = false;
 
 private:
diff --git a/Engine/Plugins/Runtime/nDisplay/Source/DisplayCluster/Private/Render/Presentation/DisplayClusterPresentationBase.cpp b/Engine/Plugins/Runtime/nDisplay/Source/DisplayCluster/Private/Render/Presentation/DisplayClusterPresentationBase.cpp
new file mode 100644
--- /dev/null
+++ b/Engine/Plugins/Runtime/nDisplay/Source/DisplayCluster/Private/Render/Presentation/DisplayClusterPresentationBase.cpp
@@ -0,0 +1,59 @@
+// Copyright 1998-2019 Epic Games, Inc. All Rights Reserved.
+
+#include "Render\Presentation\DisplayClusterPresentationBase.h "
+#include "Render/Synchronization/IDisplayClusterRenderSyncPolicy.h"
+
+#include "DisplayClusterGlobals.h"
+#include "DisplayClusterLog.h"
+
+
+// Custom VSync interval control
+static TAutoConsoleVariable<int32>  CVarVSyncInterval(
+	TEXT("nDisplay.render.VSyncInterval"),
+	1,
+	TEXT("VSync interval"),
+	ECVF_RenderThreadSafe
+);
+
+
+FDisplayClusterPresentationBase::FDisplayClusterPresentationBase(FViewport* const InViewport, TSharedPtr<IDisplayClusterRenderSyncPolicy>& InSyncPolicy)
+	: FRHICustomPresent()
+	, Viewport(InViewport)
+	, SyncPolicy(InSyncPolicy)
+{
+}
+
+FDisplayClusterPresentationBase::~FDisplayClusterPresentationBase()
+{
+}
+
+uint32 FDisplayClusterPresentationBase::GetSwapInt() const
+{
+	const uint32 SyncInterval = static_cast<uint32>(CVarVSyncInterval.GetValueOnAnyThread());
+	return (SyncInterval);
+}
+
+void FDisplayClusterPresentationBase::OnBackBufferResize()
+{
+	DISPLAY_CLUSTER_FUNC_TRACE(LogDisplayClusterRender);
+}
+
+bool FDisplayClusterPresentationBase::Present(int32& InOutSyncInterval)
+{
+	DISPLAY_CLUSTER_FUNC_TRACE(LogDisplayClusterRender);
+
+	// Update sync value with nDisplay value
+	InOutSyncInterval = GetSwapInt();
+
+	// Get sync policy instance
+	if (SyncPolicy.IsValid())
+	{
+		// False results means we don't need to present current frame, the sync object already presented it
+		if (!SyncPolicy->SynchronizeClusterRendering(InOutSyncInterval))
+		{
+			return false;
+		}
+	}
+
+	return true;
+}
diff --git a/Engine/Plugins/Runtime/nDisplay/Source/DisplayCluster/Private/Render/Presentation/DisplayClusterPresentationDX11.cpp b/Engine/Plugins/Runtime/nDisplay/Source/DisplayCluster/Private/Render/Presentation/DisplayClusterPresentationDX11.cpp
new file mode 100644
--- /dev/null
+++ b/Engine/Plugins/Runtime/nDisplay/Source/DisplayCluster/Private/Render/Presentation/DisplayClusterPresentationDX11.cpp
@@ -0,0 +1,22 @@
+// Copyright 1998-2019 Epic Games, Inc. All Rights Reserved.
+
+#include "Render/Presentation/DisplayClusterPresentationDX11.h"
+
+#include "DisplayClusterLog.h"
+
+#include "RHI.h"
+#include "RHIResources.h"
+
+#include "Windows/AllowWindowsPlatformTypes.h"
+#include "DirectX/Include/DXGI.h"
+#include "Windows/HideWindowsPlatformTypes.h"
+
+
+FDisplayClusterPresentationDX11::FDisplayClusterPresentationDX11(FViewport* const InViewport, TSharedPtr<IDisplayClusterRenderSyncPolicy>& InSyncPolicy)
+	: FDisplayClusterPresentationBase(InViewport, InSyncPolicy)
+{
+}
+
+FDisplayClusterPresentationDX11::~FDisplayClusterPresentationDX11()
+{
+}
diff --git a/Engine/Plugins/Runtime/nDisplay/Source/DisplayCluster/Private/Render/Presentation/DisplayClusterPresentationDX11.h b/Engine/Plugins/Runtime/nDisplay/Source/DisplayCluster/Private/Render/Presentation/DisplayClusterPresentationDX11.h
new file mode 100644
--- /dev/null
+++ b/Engine/Plugins/Runtime/nDisplay/Source/DisplayCluster/Private/Render/Presentation/DisplayClusterPresentationDX11.h
@@ -0,0 +1,16 @@
+// Copyright 1998-2019 Epic Games, Inc. All Rights Reserved.
+
+#pragma once
+
+#include "Render/Presentation/DisplayClusterPresentationBase.h"
+
+
+/**
+ * Helper class to encapsulate DX11 frame presentation
+ */
+class FDisplayClusterPresentationDX11 : public FDisplayClusterPresentationBase
+{
+public:
+	FDisplayClusterPresentationDX11(FViewport* const InViewport, TSharedPtr<IDisplayClusterRenderSyncPolicy>& InSyncPolicy);
+	virtual ~FDisplayClusterPresentationDX11();
+};
diff --git a/Engine/Plugins/Runtime/nDisplay/Source/DisplayCluster/Private/Render/Presentation/DisplayClusterPresentationDX12.cpp b/Engine/Plugins/Runtime/nDisplay/Source/DisplayCluster/Private/Render/Presentation/DisplayClusterPresentationDX12.cpp
new file mode 100644
--- /dev/null
+++ b/Engine/Plugins/Runtime/nDisplay/Source/DisplayCluster/Private/Render/Presentation/DisplayClusterPresentationDX12.cpp
@@ -0,0 +1,22 @@
+// Copyright 1998-2019 Epic Games, Inc. All Rights Reserved.
+
+#include "Render/Presentation/DisplayClusterPresentationDX12.h"
+
+#include "DisplayClusterLog.h"
+
+#include "RHI.h"
+#include "RHIResources.h"
+
+#include "Windows/AllowWindowsPlatformTypes.h"
+#include "DirectX/Include/DXGI.h"
+#include "Windows/HideWindowsPlatformTypes.h"
+
+
+FDisplayClusterPresentationDX12::FDisplayClusterPresentationDX12(FViewport* const InViewport, TSharedPtr<IDisplayClusterRenderSyncPolicy>& InSyncPolicy)
+	: FDisplayClusterPresentationBase(InViewport, InSyncPolicy)
+{
+}
+
+FDisplayClusterPresentationDX12::~FDisplayClusterPresentationDX12()
+{
+}
diff --git a/Engine/Plugins/Runtime/nDisplay/Source/DisplayCluster/Private/Render/Presentation/DisplayClusterPresentationDX12.h b/Engine/Plugins/Runtime/nDisplay/Source/DisplayCluster/Private/Render/Presentation/DisplayClusterPresentationDX12.h
new file mode 100644
--- /dev/null
+++ b/Engine/Plugins/Runtime/nDisplay/Source/DisplayCluster/Private/Render/Presentation/DisplayClusterPresentationDX12.h
@@ -0,0 +1,16 @@
+// Copyright 1998-2019 Epic Games, Inc. All Rights Reserved.
+
+#pragma once
+
+#include "Render/Presentation/DisplayClusterPresentationBase.h"
+
+
+/**
+ * Helper class to encapsulate DX12 frame presentation
+ */
+class FDisplayClusterPresentationDX12 : public FDisplayClusterPresentationBase
+{
+public:
+	FDisplayClusterPresentationDX12(FViewport* const InViewport, TSharedPtr<IDisplayClusterRenderSyncPolicy>& InSyncPolicy);
+	virtual ~FDisplayClusterPresentationDX12();
+};
