Get rid of DisplayCluster GameMode and Pawns

diff --git a/Engine/Plugins/Runtime/nDisplay/Source/DisplayCluster/Private/Render/Presentation/DisplayClusterPresentationNative.cpp b/Engine/Plugins/Runtime/nDisplay/Source/DisplayCluster/Private/Render/Presentation/DisplayClusterPresentationNative.cpp
new file mode 100644
--- /dev/null
+++ b/Engine/Plugins/Runtime/nDisplay/Source/DisplayCluster/Private/Render/Presentation/DisplayClusterPresentationNative.cpp
@@ -0,0 +1,36 @@
+// Copyright 1998-2019 Epic Games, Inc. All Rights Reserved.
+
+#include "Render/Presentation/DisplayClusterPresentationNative.h"
+#include "Render/Synchronization/IDisplayClusterRenderSyncPolicy.h"
+
+#include "DisplayClusterGlobals.h"
+#include "DisplayClusterLog.h"
+
+
+FDisplayClusterPresentationNative::FDisplayClusterPresentationNative(FViewport* const Viewport, TSharedPtr<IDisplayClusterRenderSyncPolicy>& SyncPolicy)
+	: FDisplayClusterPresentationBase(Viewport, SyncPolicy)
+{
+	DISPLAY_CLUSTER_FUNC_TRACE(LogDisplayClusterRender);
+}
+
+FDisplayClusterPresentationNative::~FDisplayClusterPresentationNative()
+{
+	DISPLAY_CLUSTER_FUNC_TRACE(LogDisplayClusterRender);
+}
+
+
+bool FDisplayClusterPresentationNative::Present(int32& InOutSyncInterval)
+{
+	DISPLAY_CLUSTER_FUNC_TRACE(LogDisplayClusterRender);
+
+	TSharedPtr<IDisplayClusterRenderSyncPolicy> CurSyncPolicy = GetSyncPolicyObject();
+	if (CurSyncPolicy.IsValid())
+	{
+		// Update swap interval with internal settings
+		InOutSyncInterval = GetSwapInt();
+		// If the synchronization object hasn't presented the frame, let the Engine do it
+		return CurSyncPolicy->SynchronizeClusterRendering(InOutSyncInterval);
+	}
+
+	return true;
+}
diff --git a/Engine/Plugins/Runtime/nDisplay/Source/DisplayCluster/Private/Render/Presentation/DisplayClusterPresentationNative.h b/Engine/Plugins/Runtime/nDisplay/Source/DisplayCluster/Private/Render/Presentation/DisplayClusterPresentationNative.h
new file mode 100644
--- /dev/null
+++ b/Engine/Plugins/Runtime/nDisplay/Source/DisplayCluster/Private/Render/Presentation/DisplayClusterPresentationNative.h
@@ -0,0 +1,22 @@
+// Copyright 1998-2019 Epic Games, Inc. All Rights Reserved.
+
+#pragma once
+
+#include "Render/Presentation/DisplayClusterPresentationBase.h"
+
+
+/**
+ * Present stub to allow to synchronize a cluster with native rendering pipeline (no nDisplay stereo devices used)
+ */
+class FDisplayClusterPresentationNative : public FDisplayClusterPresentationBase
+{
+public:
+	FDisplayClusterPresentationNative(FViewport* const Viewport, TSharedPtr<IDisplayClusterRenderSyncPolicy>& SyncPolicy);
+	virtual ~FDisplayClusterPresentationNative();
+
+protected:
+	//////////////////////////////////////////////////////////////////////////////////////////////
+	// FRHICustomPresent
+	//////////////////////////////////////////////////////////////////////////////////////////////
+	virtual bool Present(int32& InOutSyncInterval) override;
+};
diff --git a/Engine/Plugins/Runtime/nDisplay/Source/DisplayCluster/Private/Render/Synchronization/DisplayClusterRenderSyncPolicySoftwareDX11.cpp b/Engine/Plugins/Runtime/nDisplay/Source/DisplayCluster/Private/Render/Synchronization/DisplayClusterRenderSyncPolicySoftwareDX11.cpp
--- a/Engine/Plugins/Runtime/nDisplay/Source/DisplayCluster/Private/Render/Synchronization/DisplayClusterRenderSyncPolicySoftwareDX11.cpp
+++ b/Engine/Plugins/Runtime/nDisplay/Source/DisplayCluster/Private/Render/Synchronization/DisplayClusterRenderSyncPolicySoftwareDX11.cpp
@@ -154,7 +154,7 @@ bool FDisplayClusterRenderSyncPolicySoftwareDX11::SynchronizeClusterRendering(in
 	check(IsInRenderingThread());
 	if(!(GEngine && GEngine->GameViewport && GEngine->GameViewport->Viewport))
 	{
-		UE_LOG(LogDisplayClusterRenderSync, Error, TEXT("Couldn't get a D3D11 viewport."));
+		UE_LOG(LogDisplayClusterRenderSync, Warning, TEXT("Couldn't get a D3D11 viewport."));
 		// Tell a caller to present the frame by himself
 		return true;
 	}
diff --git a/Engine/Plugins/Runtime/nDisplay/Source/DisplayCluster/Public/Blueprints/IDisplayClusterBlueprintAPI.h b/Engine/Plugins/Runtime/nDisplay/Source/DisplayCluster/Public/Blueprints/IDisplayClusterBlueprintAPI.h
--- a/Engine/Plugins/Runtime/nDisplay/Source/DisplayCluster/Public/Blueprints/IDisplayClusterBlueprintAPI.h
+++ b/Engine/Plugins/Runtime/nDisplay/Source/DisplayCluster/Public/Blueprints/IDisplayClusterBlueprintAPI.h
@@ -5,7 +5,7 @@
 #include "CoreMinimal.h"
 #include "UObject/Interface.h"
 
-#include "DisplayClusterOperationMode.h"
+#include "DisplayClusterEnums.h"
 #include "Cluster/DisplayClusterClusterEvent.h"
 #include "Cluster/IDisplayClusterClusterEventListener.h"
 
@@ -14,10 +14,10 @@
 
 #include "IDisplayClusterBlueprintAPI.generated.h"
 
+class UDisplayClusterRootComponent;
 struct FPostProcessSettings;
 
 
-
 UINTERFACE(meta = (CannotImplementInterfaceInBlueprint))
 class DISPLAYCLUSTER_API UDisplayClusterBlueprintAPI : public UInterface
 {
@@ -91,7 +91,6 @@ public:
 	UFUNCTION(BlueprintCallable, meta = (DisplayName = "Get local viewports"), Category = "DisplayCluster|Config")
 	virtual void GetLocalViewports(bool IsRTT, TArray<FString>& ViewportIDs, TArray<FString>& ViewportTypes, TArray<FIntPoint>& ViewportLocations, TArray<FIntPoint>& ViewportSizes) = 0;
 
-
 public:
 	//////////////////////////////////////////////////////////////////////////////////////////////
 	// Game API
@@ -99,7 +98,7 @@ public:
 	// Root
 	/** Returns Cluster Pawn. */
 	UFUNCTION(BlueprintCallable, meta = (DisplayName = "Get root"), Category = "DisplayCluster|Game")
-	virtual ADisplayClusterPawn* GetRoot() = 0;
+	virtual UDisplayClusterRootComponent* GetRoot() = 0;
 
 	// Screens
 	/** Returns screen reference by id name. */
@@ -131,10 +130,6 @@ public:
 	UFUNCTION(BlueprintCallable, meta = (DisplayName = "Get default camera"), Category = "DisplayCluster|Game")
 	virtual UDisplayClusterCameraComponent* GetDefaultCamera() = 0;
 
-	/** Sets default camera component specified by index. */
-	UFUNCTION(BlueprintCallable, meta = (DisplayName = "Set default camera by index"), Category = "DisplayCluster|Game")
-	virtual void SetDefaultCameraByIndex(int32 Index) = 0;
-
 	/** Sets default camera component specified by ID. */
 	UFUNCTION(BlueprintCallable, meta = (DisplayName = "Set default camera by ID"), Category = "DisplayCluster|Game")
 	virtual void SetDefaultCameraById(const FString& id) = 0;
@@ -148,31 +143,6 @@ public:
 	UFUNCTION(BlueprintCallable, meta = (DisplayName = "Get all nodes"), Category = "DisplayCluster|Game")
 	virtual TArray<UDisplayClusterSceneComponent*> GetAllNodes() = 0;
 
-	// Navigation
-	/** Returns scene component used for default pawn navigation. */
-	UFUNCTION(BlueprintCallable, meta = (DisplayName = "Get translation direction component"), Category = "DisplayCluster|Game")
-	virtual USceneComponent* GetTranslationDirectionComponent() = 0;
-
-	/** Set scene component to be used for default pawn navigation. */
-	UFUNCTION(BlueprintCallable, meta = (DisplayName = "Set translation direction component"), Category = "DisplayCluster|Game")
-	virtual void SetTranslationDirectionComponent(USceneComponent* pComp) = 0;
-
-	/** Set scene component to be used for default pawn navigation by id name. */
-	UFUNCTION(BlueprintCallable, meta = (DisplayName = "Set translation direction component by ID"), Category = "DisplayCluster|Game")
-	virtual void SetTranslationDirectionComponentId(const FString& id) = 0;
-
-	/** Return scene component used as a pivot point for rotation of the scene node hierarchy. */
-	UFUNCTION(BlueprintCallable, meta = (DisplayName = "Get rotate around component"), Category = "DisplayCluster|Game")
-	virtual USceneComponent* GetRotateAroundComponent() = 0;
-
-	/** Set scene component used as a pivot point for rotation of the scene node hierarchy. */
-	UFUNCTION(BlueprintCallable, meta = (DisplayName = "Set rotate around component"), Category = "DisplayCluster|Game")
-	virtual void SetRotateAroundComponent(USceneComponent* pComp) = 0;
-
-	/** Set scene component used as a pivot point for rotation of the scene node hierarchy by id name. */
-	UFUNCTION(BlueprintCallable, meta = (DisplayName = "Set rotate around component by ID"), Category = "DisplayCluster|Game")
-	virtual void SetRotateAroundComponentId(const FString& id) = 0;
-
 public:
 	//////////////////////////////////////////////////////////////////////////////////////////////
 	// Input API
@@ -285,28 +255,4 @@ public:
 	/** Toggle current eye swap state. */
 	UFUNCTION(BlueprintCallable, meta = (DisplayName = "Toggle eye swap"), Category = "DisplayCluster|Render|Camera")
 	virtual bool ToggleEyesSwap(const FString& CameraId) = 0;
-
-	/** Return near culling distance of specified camera. */
-	UFUNCTION(BlueprintCallable, meta = (DisplayName = "Get near culling distance"), Category = "DisplayCluster|Render|Camera")
-	virtual float GetNearCullingDistance(const FString& CameraId) const = 0;
-
-	/** Set near culling distance of specified camera. */
-	UFUNCTION(BlueprintCallable, meta = (DisplayName = "Set near culling distance"), Category = "DisplayCluster|Render|Camera")
-	virtual void SetNearCullingDistance(const FString& CameraId, float NearDistance) = 0;
-
-	/** Get far culling distance of specified camera. */
-	UFUNCTION(BlueprintCallable, meta = (DisplayName = "Get far culling distance"), Category = "DisplayCluster|Render|Camera")
-	virtual float GetFarCullingDistance(const FString& CameraId) const = 0;
-
-	/** Set far culling distance of specified camera. */
-	UFUNCTION(BlueprintCallable, meta = (DisplayName = "Set far culling distance"), Category = "DisplayCluster|Render|Camera")
-	virtual void SetFarCullingDistance(const FString& CameraId, float FarDistance) = 0;
-
-	/** Return near and far plane clip plane distances. */
-	UFUNCTION(BlueprintCallable, meta = (DisplayName = "Get near and far clipping distance"), Category = "DisplayCluster|Render|Camera")
-	virtual void GetCullingDistance(const FString& CameraId, float& NearDistance, float& FarDistance) = 0;
-
-	/** Set near and far plane clip plane distances. */
-	UFUNCTION(BlueprintCallable, meta = (DisplayName = "Set near and far clipping distance"), Category = "DisplayCluster|Render|Camera")
-	virtual void SetCullingDistance(const FString& CameraId, float NearDistance, float FarDistance) = 0;
 };
diff --git a/Engine/Plugins/Runtime/nDisplay/Source/DisplayCluster/Public/Cluster/IDisplayClusterClusterManager.h b/Engine/Plugins/Runtime/nDisplay/Source/DisplayCluster/Public/Cluster/IDisplayClusterClusterManager.h
--- a/Engine/Plugins/Runtime/nDisplay/Source/DisplayCluster/Public/Cluster/IDisplayClusterClusterManager.h
+++ b/Engine/Plugins/Runtime/nDisplay/Source/DisplayCluster/Public/Cluster/IDisplayClusterClusterManager.h
@@ -5,9 +5,12 @@
 #include "CoreMinimal.h"
 #include "Delegates/Delegate.h"
 #include "UObject/ScriptInterface.h"
+#include "DisplayClusterEnums.h"
+
 
 struct FDisplayClusterClusterEvent;
 class IDisplayClusterClusterEventListener;
+class IDisplayClusterClusterSyncObject;
 
 DECLARE_MULTICAST_DELEGATE_OneParam(FOnClusterEvent, const FDisplayClusterClusterEvent& /* Event */);
 typedef FOnClusterEvent::FDelegate FOnClusterEventListener;
@@ -29,6 +32,9 @@ public:
 	virtual FString GetNodeId()     const = 0;
 	virtual uint32 GetNodesAmount() const = 0;
 
+	virtual void RegisterSyncObject(IDisplayClusterClusterSyncObject* pSyncObj, EDisplayClusterSyncGroup SyncGroup) = 0;
+	virtual void UnregisterSyncObject(IDisplayClusterClusterSyncObject* pSyncObj) = 0;
+
 	virtual void AddClusterEventListener(TScriptInterface<IDisplayClusterClusterEventListener> Listener) = 0;
 	virtual void RemoveClusterEventListener(TScriptInterface<IDisplayClusterClusterEventListener> Listener) = 0;
 
diff --git a/Engine/Plugins/Runtime/nDisplay/Source/DisplayCluster/Public/Config/DisplayClusterConfigTypes.h b/Engine/Plugins/Runtime/nDisplay/Source/DisplayCluster/Public/Config/DisplayClusterConfigTypes.h
--- a/Engine/Plugins/Runtime/nDisplay/Source/DisplayCluster/Public/Config/DisplayClusterConfigTypes.h
+++ b/Engine/Plugins/Runtime/nDisplay/Source/DisplayCluster/Public/Config/DisplayClusterConfigTypes.h
@@ -179,6 +179,7 @@ struct DISPLAYCLUSTER_API FDisplayClusterConfigInputSetup : public FDisplayClust
 struct DISPLAYCLUSTER_API FDisplayClusterConfigGeneral : public FDisplayClusterConfigBase
 {
 	int32 SwapSyncPolicy = 0;
+	int32 NativeInputSyncPolicy = 1;
 
 	virtual FString ToString() const override;
 	virtual bool    DeserializeFromString(const FString& line) override;
diff --git a/Engine/Plugins/Runtime/nDisplay/Source/DisplayCluster/Public/Config/IDisplayClusterConfigManager.h b/Engine/Plugins/Runtime/nDisplay/Source/DisplayCluster/Public/Config/IDisplayClusterConfigManager.h
--- a/Engine/Plugins/Runtime/nDisplay/Source/DisplayCluster/Public/Config/IDisplayClusterConfigManager.h
+++ b/Engine/Plugins/Runtime/nDisplay/Source/DisplayCluster/Public/Config/IDisplayClusterConfigManager.h
@@ -16,7 +16,6 @@ public:
 
 	virtual int32 GetClusterNodesAmount() const = 0;
 	virtual TArray<FDisplayClusterConfigClusterNode> GetClusterNodes() const = 0;
-	virtual bool GetClusterNode(int32 idx, FDisplayClusterConfigClusterNode& cnode) const = 0;
 	virtual bool GetClusterNode(const FString& id, FDisplayClusterConfigClusterNode& cnode) const = 0;
 	virtual bool GetMasterClusterNode(FDisplayClusterConfigClusterNode& cnode) const = 0;
 
@@ -27,32 +26,26 @@ public:
 
 	virtual int32 GetScreensAmount() const = 0;
 	virtual TArray<FDisplayClusterConfigScreen> GetScreens() const = 0;
-	virtual bool GetScreen(int32 idx, FDisplayClusterConfigScreen& screen) const = 0;
 	virtual bool GetScreen(const FString& id, FDisplayClusterConfigScreen& screen) const = 0;
 
 	virtual int32 GetCamerasAmount() const = 0;
 	virtual TArray<FDisplayClusterConfigCamera> GetCameras() const = 0;
-	virtual bool GetCamera(int32 idx, FDisplayClusterConfigCamera& camera) const = 0;
 	virtual bool GetCamera(const FString& id, FDisplayClusterConfigCamera& camera) const = 0;
 
 	virtual int32 GetViewportsAmount() const = 0;
 	virtual TArray<FDisplayClusterConfigViewport> GetViewports() const = 0;
-	virtual bool GetViewport(int32 idx, FDisplayClusterConfigViewport& viewport) const = 0;
 	virtual bool GetViewport(const FString& id, FDisplayClusterConfigViewport& viewport) const = 0;
 
 	virtual int32 GetPostprocessAmount() const = 0;
 	virtual TArray<FDisplayClusterConfigPostprocess> GetPostprocess() const = 0;
-	virtual bool GetPostprocess(int32 idx, FDisplayClusterConfigPostprocess& postprocess) const = 0;
 	virtual bool GetPostprocess(const FString& id, FDisplayClusterConfigPostprocess& postprocess) const = 0;
 
 	virtual int32 GetSceneNodesAmount() const = 0;
 	virtual TArray<FDisplayClusterConfigSceneNode> GetSceneNodes() const = 0;
-	virtual bool GetSceneNode(int32 idx, FDisplayClusterConfigSceneNode& snode) const = 0;
 	virtual bool GetSceneNode(const FString& id, FDisplayClusterConfigSceneNode& snode) const = 0;
 
 	virtual int32 GetInputDevicesAmount() const = 0;
 	virtual TArray<FDisplayClusterConfigInput> GetInputDevices() const = 0;
-	virtual bool GetInputDevice(int32 idx, FDisplayClusterConfigInput& input) const = 0;
 	virtual bool GetInputDevice(const FString& id, FDisplayClusterConfigInput& input) const = 0;
 
 	virtual TArray<FDisplayClusterConfigInputSetup> GetInputSetupRecords() const = 0;
diff --git a/Engine/Plugins/Runtime/nDisplay/Source/DisplayCluster/Public/DisplayClusterCameraComponent.h b/Engine/Plugins/Runtime/nDisplay/Source/DisplayCluster/Public/DisplayClusterCameraComponent.h
--- a/Engine/Plugins/Runtime/nDisplay/Source/DisplayCluster/Public/DisplayClusterCameraComponent.h
+++ b/Engine/Plugins/Runtime/nDisplay/Source/DisplayCluster/Public/DisplayClusterCameraComponent.h
@@ -5,6 +5,8 @@
 #include "DisplayClusterSceneComponent.h"
 #include "DisplayClusterCameraComponent.generated.h"
 
+class UCameraComponent;
+
 
 /**
  * Camera component
@@ -59,64 +61,6 @@ public:
 	bool ToggleEyesSwap()
 	{ return (bEyeSwap = !bEyeSwap); }
 
-	/**
-	* Get camera frustum near culling
-	*
-	* @return - near culling plane distance
-	*/
-	float GetNearCullingDistance() const
-	{ return NearClipPlane; }
-
-	/**
-	* Set camera frustum near culling
-	*
-	* @param NCP - near culling plane distance
-	*/
-	void SetNearCullingDistance(float NCP)
-	{
-		NearClipPlane = NCP;
-	}
-
-	/**
-	* Get camera frustum far culling
-	*
-	* @return - far culling plane distance
-	*/
-	float GetFarCullingDistance() const
-	{ return FarClipPlane; }
-
-	/**
-	* Set camera frustum far culling
-	*
-	* @param FCP - far culling plane distance
-	*/
-	void SetFarCullingDistance(float FCP)
-	{ FarClipPlane = FCP; }
-
-	/**
-	* Get camera frustum culling
-	*
-	* @param OutNCP - near culling plane distance
-	* @param OutFCP - far culling plane distance
-	*/
-	void GetCullingDistance(float& OutNCP, float& OutFCP) const
-	{
-		OutNCP = NearClipPlane;
-		OutFCP = FarClipPlane;
-	}
-
-	/**
-	* Set camera frustum culling
-	*
-	* @param NCP - near culling plane distance
-	* @param FCP - far culling plane distance
-	*/
-	void SetCullingDistance(float NCP, float FCP)
-	{
-		NearClipPlane = NCP;
-		FarClipPlane = FCP;
-	}
-
 	/**
 	* Returns force eye offset value
 	*
@@ -139,12 +83,13 @@ public:
 
 public:
 	virtual void BeginPlay() override;
-	virtual void TickComponent(float DeltaTime, ELevelTick TickType, FActorComponentTickFunction* ThisTickFunction) override;
 
 private:
 	float EyeDist;
 	bool  bEyeSwap;
 	int   ForceEyeOffset = 0;
-	float NearClipPlane;
-	float FarClipPlane;
+
+#if 0
+	UCameraComponent* CameraComponent = nullptr;
+#endif
 };
diff --git a/Engine/Plugins/Runtime/nDisplay/Source/DisplayCluster/Public/DisplayClusterGameEngine.h b/Engine/Plugins/Runtime/nDisplay/Source/DisplayCluster/Public/DisplayClusterGameEngine.h
--- a/Engine/Plugins/Runtime/nDisplay/Source/DisplayCluster/Public/DisplayClusterGameEngine.h
+++ b/Engine/Plugins/Runtime/nDisplay/Source/DisplayCluster/Public/DisplayClusterGameEngine.h
@@ -5,7 +5,8 @@
 #include "Engine/GameEngine.h"
 
 #include "Config/DisplayClusterConfigTypes.h"
-#include "DisplayClusterOperationMode.h"
+
+#include "DisplayClusterEnums.h"
 
 #include "Tickable.h"
 #include "Stats/Stats2.h"
@@ -34,7 +35,7 @@ public:
 	{ return ETickableTickType::Always; }
 
 	virtual TStatId GetStatId() const override;
-	virtual void Tick(float DeltaTime) override;
+	virtual void Tick(float DeltaSeconds) override;
 };
 
 
@@ -55,11 +56,6 @@ public:
 
 	EDisplayClusterOperationMode GetOperationMode() const { return OperationMode; }
 
-	void RegisterSyncObject(IDisplayClusterClusterSyncObject* SyncObj);
-	void UnregisterSyncObject(IDisplayClusterClusterSyncObject* SyncObj);
-
-	bool IsMaster() const;
-
 protected:
 	virtual bool InitializeInternals();
 	EDisplayClusterOperationMode DetectOperationMode();
diff --git a/Engine/Plugins/Runtime/nDisplay/Source/DisplayCluster/Public/DisplayClusterPlayerInput.h b/Engine/Plugins/Runtime/nDisplay/Source/DisplayCluster/Public/DisplayClusterPlayerInput.h
new file mode 100644
--- /dev/null
+++ b/Engine/Plugins/Runtime/nDisplay/Source/DisplayCluster/Public/DisplayClusterPlayerInput.h
@@ -0,0 +1,54 @@
+// Copyright 1998-2019 Epic Games, Inc. All Rights Reserved.
+
+#pragma once
+
+#include "GameFramework/PlayerInput.h"
+#include "DisplayClusterPlayerInput.generated.h"
+
+
+/**
+ * Object within PlayerController that processes player input.
+ */
+UCLASS(Within=PlayerController, config=Input, transient)
+class DISPLAYCLUSTER_API UDisplayClusterPlayerInput : public UPlayerInput
+{
+	GENERATED_BODY()
+
+public:
+	UDisplayClusterPlayerInput();
+
+public:
+	/** Process the frame's input events given the current input component stack. */
+	virtual void ProcessInputStack(const TArray<UInputComponent*>& InputComponentStack, const float DeltaTime, const bool bGamePaused) override;
+
+protected:
+	virtual void ProcessDelegatesFilter(const TArray<UInputComponent*>& InputComponentStack, TArray<FAxisDelegateDetails>& AxisDelegates, TArray<FVectorAxisDelegateDetails>& VectorAxisDelegates, TArray<FDelegateDispatchDetails>& NonAxisDelegates) override;
+
+private:
+	void ExportNativeInputData(TMap<FString, FString>& InputData, const TArray<UInputComponent*>& InputComponentStack, TArray<FAxisDelegateDetails>& AxisDelegates, TArray<FVectorAxisDelegateDetails>& VectorAxisDelegates, TArray<FDelegateDispatchDetails>& NonAxisDelegates);
+	void ImportNativeInputData(const TMap<FString, FString>& InputData, const TArray<UInputComponent*>& InputComponentStack);
+
+	void CleanDelegateMap(bool IsMaster);
+	void BuildDelegatesMap(const TArray<UInputComponent*>& InputComponentStack, bool IsMaster);
+
+	FString SerializeAxisData  (const TArray<UInputComponent*>& InputComponentStack, const TArray<FAxisDelegateDetails>& AxisDelegates);
+	FString SerializeVectorAxisData(const TArray<UInputComponent*>& InputComponentStack, const TArray<FVectorAxisDelegateDetails>& VectorAxisDelegates);
+	FString SerializeActionData(const TArray<UInputComponent*>& InputComponentStack, const TArray<FDelegateDispatchDetails>& NonAxisDelegates);
+
+	void DeserializeAndProcessAxisInput  (const FString& Data);
+	void DeserializeAndProcessVectorAxisInput(const FString& Data);
+	void DeserializeAndProcessActionInput(const FString& Data);
+
+private:
+	TMap<void*, TMap<void*, FName>> AxisDelegateNames;
+	TMap<FName, const FInputAxisUnifiedDelegate*> AxisDelegateInstances;
+
+	TMap<FName, TMap<EInputEvent, const FInputActionUnifiedDelegate*>> ActionDelegateInstances;
+	TMap<EInputEvent, const FInputTouchUnifiedDelegate*> TouchDelegateInstances;
+
+	TMap<void*, TMap<void*, FName>> GestureDelegateNames;
+	TMap<FName, const FInputGestureUnifiedDelegate*> GestureDelegateInstances;
+
+	TMap<void*, TMap<void*, FName>> VectorAxisDelegateNames;
+	TMap<FName, const FInputVectorAxisUnifiedDelegate*> VectorAxisDelegateInstances;
+};
diff --git a/Engine/Plugins/Runtime/nDisplay/Source/DisplayCluster/Public/DisplayClusterRootComponent.h b/Engine/Plugins/Runtime/nDisplay/Source/DisplayCluster/Public/DisplayClusterRootComponent.h
new file mode 100644
--- /dev/null
+++ b/Engine/Plugins/Runtime/nDisplay/Source/DisplayCluster/Public/DisplayClusterRootComponent.h
@@ -0,0 +1,121 @@
+// Copyright 1998-2019 Epic Games, Inc. All Rights Reserved.
+
+#pragma once
+
+#include "Components/SceneComponent.h"
+#include "Camera/PlayerCameraManager.h"
+
+#include "DisplayClusterEnums.h"
+
+#include "DisplayClusterRootComponent.generated.h"
+
+
+class UCameraComponent;
+class UDisplayClusterCameraComponent;
+class UDisplayClusterSceneComponent;
+class UDisplayClusterScreenComponent;
+class UMaterial;
+
+
+/**
+ * DisplayCluster root component (nDisplay root)
+ */
+UCLASS( ClassGroup=(Custom), meta=(BlueprintSpawnableComponent) )
+class DISPLAYCLUSTER_API UDisplayClusterRootComponent
+	: public USceneComponent
+{
+	GENERATED_BODY()
+
+public:
+	UDisplayClusterRootComponent(const FObjectInitializer& ObjectInitializer);
+
+public:
+
+public:
+	virtual void BeginPlay() override;
+	virtual void BeginDestroy() override;
+	virtual void TickComponent(float DeltaTime, ELevelTick TickType, FActorComponentTickFunction* ThisTickFunction) override;
+
+public:
+	TArray<UDisplayClusterScreenComponent*> GetAllScreens() const;
+	UDisplayClusterScreenComponent*         GetScreenById(const FString& id) const;
+	int32                                   GetScreensAmount() const;
+
+	TArray<UDisplayClusterCameraComponent*> GetAllCameras() const;
+	UDisplayClusterCameraComponent*         GetCameraById(const FString& id) const;
+	int32                                   GetCamerasAmount() const;
+	UDisplayClusterCameraComponent*         GetDefaultCamera() const;
+	void                                    SetDefaultCamera(const FString& id);
+
+	TArray<UDisplayClusterSceneComponent*>  GetAllNodes() const;
+	UDisplayClusterSceneComponent*          GetNodeById(const FString& id) const;
+
+protected:
+	// Available screens (from config file)
+	TMap<FString, UDisplayClusterScreenComponent*> ScreenComponents;
+	// Available cameras (from config file)
+	TMap<FString, UDisplayClusterCameraComponent*> CameraComponents;
+	// All available DisplayCluster nodes in hierarchy
+	TMap<FString, UDisplayClusterSceneComponent*>  SceneNodeComponents;
+
+	// Default camera for this root
+	UDisplayClusterCameraComponent* DefaultCameraComponent = nullptr;
+
+public:
+	bool GetShowProjectionScreens() const
+	{ return bShowProjectionScreens; }
+
+	UMaterial* GetProjectionScreenMaterial() const
+	{ return ProjectionScreensMaterial; }
+
+protected:
+	UPROPERTY(EditAnywhere, Category = "DisplayCluster", meta = (DisplayName = "Exit when ESC pressed"))
+	bool bExitOnEsc;
+
+	UPROPERTY(EditAnywhere, Category = "DisplayCluster", meta = (DisplayName = "Show projection screens"))
+	bool bShowProjectionScreens;
+
+	UPROPERTY(EditAnywhere, Category = "DisplayCluster", meta = (DisplayName = "Projection screens material"))
+	UMaterial* ProjectionScreensMaterial;
+
+#if WITH_EDITORONLY_DATA
+public:
+	FString GetEditorConfigPath() const
+	{ return EditorConfigPath; }
+
+	FString GetEditorNodeId() const
+	{ return EditorNodeId; }
+
+protected:
+	UPROPERTY(EditAnywhere, Category = "nDisplay (Editor only)", meta = (DisplayName = "Config file"))
+	FString EditorConfigPath;
+
+	UPROPERTY(EditAnywhere, Category = "nDisplay (Editor only)", meta = (DisplayName = "Node ID"))
+	FString EditorNodeId;
+#endif
+
+protected:
+	// Creates all hierarchy objects declared in a config file
+	virtual bool InitializeHierarchy();
+	virtual bool InitializeExtra();
+	virtual bool CreateScreens();
+	virtual bool CreateNodes();
+	virtual bool CreateCameras();
+
+private:
+	// Extracts array of values from a map
+	template <typename ObjType>
+	TArray<ObjType*> GetMapValues(const TMap<FString, ObjType*>& container) const;
+
+	// Gets item by id. Performs checks and logging.
+	template <typename DataType>
+	DataType* GetItem(const TMap<FString, DataType*>& container, const FString& id, const FString& logHeader) const;
+
+private:
+	mutable FCriticalSection InternalsSyncScope;
+
+	// Current operation mode
+	EDisplayClusterOperationMode OperationMode = EDisplayClusterOperationMode::Disabled;
+
+	int32 NativeInputSyncPolicy = 0;
+};
diff --git a/Engine/Plugins/Runtime/nDisplay/Source/DisplayCluster/Public/Game/IDisplayClusterGameManager.h b/Engine/Plugins/Runtime/nDisplay/Source/DisplayCluster/Public/Game/IDisplayClusterGameManager.h
--- a/Engine/Plugins/Runtime/nDisplay/Source/DisplayCluster/Public/Game/IDisplayClusterGameManager.h
+++ b/Engine/Plugins/Runtime/nDisplay/Source/DisplayCluster/Public/Game/IDisplayClusterGameManager.h
@@ -2,9 +2,11 @@
 
 #pragma once
 
-#include "DisplayClusterCameraComponent.h"
-#include "DisplayClusterScreenComponent.h"
-#include "DisplayClusterPawn.h"
+class UDisplayClusterCameraComponent;
+class UDisplayClusterRootComponent;
+class UDisplayClusterSceneComponent;
+class UDisplayClusterScreenComponent;
+class UWorld;
 
 
 /**
@@ -16,7 +18,9 @@ public:
 	virtual ~IDisplayClusterGameManager() = 0
 	{ }
 
-	virtual ADisplayClusterPawn*                    GetRoot() const = 0;
+public:
+	virtual UDisplayClusterRootComponent*           GetRoot() const = 0;
+	virtual void                                    SetRoot(UDisplayClusterRootComponent* InRoot) = 0;
 
 	virtual TArray<UDisplayClusterScreenComponent*> GetAllScreens() const = 0;
 	virtual UDisplayClusterScreenComponent*         GetScreenById(const FString& id) const = 0;
@@ -26,17 +30,10 @@ public:
 	virtual UDisplayClusterCameraComponent*         GetCameraById(const FString& id) const = 0;
 	virtual int32                                   GetCamerasAmount() const = 0;
 	virtual UDisplayClusterCameraComponent*         GetDefaultCamera() const = 0;
-	virtual void                                    SetDefaultCamera(int32 idx) = 0;
 	virtual void                                    SetDefaultCamera(const FString& id) = 0;
 
 	virtual TArray<UDisplayClusterSceneComponent*>  GetAllNodes() const = 0;
 	virtual UDisplayClusterSceneComponent*          GetNodeById(const FString& id) const = 0;
 
-	virtual USceneComponent*                        GetTranslationDirectionComponent() const = 0;
-	virtual void                                    SetTranslationDirectionComponent(USceneComponent* const pComp) = 0;
-	virtual void                                    SetTranslationDirectionComponent(const FString& id) = 0;
-
-	virtual USceneComponent*                        GetRotateAroundComponent() const = 0;
-	virtual void                                    SetRotateAroundComponent(USceneComponent* const pComp) = 0;
-	virtual void                                    SetRotateAroundComponent(const FString& id) = 0;
+	virtual UWorld* GetWorld() const = 0;
 };
diff --git a/Engine/Plugins/Runtime/nDisplay/Source/DisplayCluster/Public/IDisplayCluster.h b/Engine/Plugins/Runtime/nDisplay/Source/DisplayCluster/Public/IDisplayCluster.h
--- a/Engine/Plugins/Runtime/nDisplay/Source/DisplayCluster/Public/IDisplayCluster.h
+++ b/Engine/Plugins/Runtime/nDisplay/Source/DisplayCluster/Public/IDisplayCluster.h
@@ -5,7 +5,7 @@
 #include "Modules/ModuleManager.h"
 #include "Modules/ModuleInterface.h"
 
-#include "DisplayClusterOperationMode.h"
+#include "DisplayClusterEnums.h"
 
 
 class IDisplayClusterRenderManager;
@@ -108,7 +108,23 @@ public:
 	DECLARE_EVENT(IDisplayCluster, FDisplayClusterEndSessionEvent);
 	virtual FDisplayClusterEndSessionEvent& OnDisplayClusterEndSession() = 0;
 
+	/** Called on DisplayCluster StartFrame **/
+	DECLARE_EVENT_OneParam(IDisplayCluster, FDisplayClusterStartFrameEvent, uint64);
+	virtual FDisplayClusterStartFrameEvent& OnDisplayClusterStartFrame(uint64 FrameNum) = 0;
+
+	/** Called on DisplayCluster EndFrame **/
+	DECLARE_EVENT_OneParam(IDisplayCluster, FDisplayClusterEndFrameEvent, uint64);
+	virtual FDisplayClusterEndFrameEvent& OnDisplayClusterEndFrame(uint64 FrameNum) = 0;
+
 	/** Called on DisplayCluster PreTick **/
 	DECLARE_EVENT(IDisplayCluster, FDisplayClusterPreTickEvent);
 	virtual FDisplayClusterPreTickEvent& OnDisplayClusterPreTick() = 0;
+
+	/** Called on DisplayCluster Tick **/
+	DECLARE_EVENT(IDisplayCluster, FDisplayClusterTickEvent);
+	virtual FDisplayClusterTickEvent& OnDisplayClusterTick() = 0;
+
+	/** Called on DisplayCluster PostTick **/
+	DECLARE_EVENT(IDisplayCluster, FDisplayClusterPostTickEvent);
+	virtual FDisplayClusterPostTickEvent& OnDisplayClusterPostTick() = 0;
 };
diff --git a/Engine/Plugins/Runtime/nDisplay/Source/DisplayCluster/Public/Render/IDisplayClusterRenderManager.h b/Engine/Plugins/Runtime/nDisplay/Source/DisplayCluster/Public/Render/IDisplayClusterRenderManager.h
--- a/Engine/Plugins/Runtime/nDisplay/Source/DisplayCluster/Public/Render/IDisplayClusterRenderManager.h
+++ b/Engine/Plugins/Runtime/nDisplay/Source/DisplayCluster/Public/Render/IDisplayClusterRenderManager.h
@@ -215,48 +215,4 @@ public:
 	* @return - new eyes swap state. False - normal eyes left|right, true - swapped eyes right|left
 	*/
 	virtual bool ToggleEyesSwap(const FString& CameraId) = 0;
-
-	/**
-	* Get camera frustum near culling
-	*
-	* @return - near culling plane distance
-	*/
-	virtual float GetNearCullingDistance(const FString& CameraId) const = 0;
-
-	/**
-	* Set camera frustum near culling
-	*
-	* @param NCP - near culling plane distance
-	*/
-	virtual void SetNearCullingDistance(const FString& CameraId, float NearDistance) = 0;
-
-	/**
-	* Get camera frustum far culling
-	*
-	* @return - far culling plane distance
-	*/
-	virtual float GetFarCullingDistance(const FString& CameraId) const = 0;
-
-	/**
-	* Set camera frustum far culling
-	*
-	* @return - far culling plane distance
-	*/
-	virtual void SetFarCullingDistance(const FString& CameraId, float FarDistance) = 0;
-
-	/**
-	* Get camera frustum culling
-	*
-	* @param NearDistance - near culling plane distance
-	* @param FarDistance - far culling plane distance
-	*/
-	virtual void GetCullingDistance(const FString& CameraId, float& NearDistance, float& FarDistance) const = 0;
-
-	/**
-	* Set camera frustum culling
-	*
-	* @param NearDistance - near culling plane distance
-	* @param FarDistance - far culling plane distance
-	*/
-	virtual void SetCullingDistance(const FString& CameraId, float NearDistance, float FarDistance) = 0;
 };
diff --git a/Engine/Plugins/Runtime/nDisplay/Source/DisplayClusterEditor/Private/DisplayClusterEditorEngine.cpp b/Engine/Plugins/Runtime/nDisplay/Source/DisplayClusterEditor/Private/DisplayClusterEditorEngine.cpp
--- a/Engine/Plugins/Runtime/nDisplay/Source/DisplayClusterEditor/Private/DisplayClusterEditorEngine.cpp
+++ b/Engine/Plugins/Runtime/nDisplay/Source/DisplayClusterEditor/Private/DisplayClusterEditorEngine.cpp
@@ -3,8 +3,11 @@
 #include "DisplayClusterEditorEngine.h"
 #include "DisplayClusterEditorLog.h"
 
+#include "DisplayClusterRootComponent.h"
+
 #include "DisplayCluster/Private/IPDisplayCluster.h"
 
+#include "Editor.h"
 
 void UDisplayClusterEditorEngine::Init(IEngineLoop* InEngineLoop)
 {
@@ -12,23 +15,26 @@ void UDisplayClusterEditorEngine::Init(IEngineLoop* InEngineLoop)
 
 	// Initialize DisplayCluster module for editor mode
 	DisplayClusterModule = static_cast<IPDisplayCluster*>(&IDisplayCluster::Get());
-	if (DisplayClusterModule)
+	if (!DisplayClusterModule)
 	{
-		const bool bResult = DisplayClusterModule->Init(EDisplayClusterOperationMode::Editor);
-		if (bResult)
-		{
-			UE_LOG(LogDisplayClusterEditorEngine, Log, TEXT("DisplayCluster module has been initialized"));
-		}
-		else
-		{
-			UE_LOG(LogDisplayClusterEditorEngine, Error, TEXT("An error occured during DisplayCluster initialization"));
-		}
+		UE_LOG(LogDisplayClusterEditorEngine, Error, TEXT("Couldn't initialize DisplayCluster module"));
+		return;
 	}
-	else
+
+	// Initialize DisplayCluster module for operating in Editor mode
+	const bool bResult = DisplayClusterModule->Init(EDisplayClusterOperationMode::Editor);
+	if (!bResult)
 	{
-		UE_LOG(LogDisplayClusterEditorEngine, Error, TEXT("Couldn't initialize DisplayCluster module"));
+		UE_LOG(LogDisplayClusterEditorEngine, Error, TEXT("An error occured during DisplayCluster initialization"));
+		return;
 	}
 
+	UE_LOG(LogDisplayClusterEditorEngine, Log, TEXT("DisplayCluster module has been initialized"));
+
+	// Subscribe to PIE events
+	BeginPIEDelegate = FEditorDelegates::BeginPIE.AddUObject(this, &UDisplayClusterEditorEngine::OnBeginPIE);
+	EndPIEDelegate   = FEditorDelegates::EndPIE.AddUObject(this, &UDisplayClusterEditorEngine::OnEndPIE);
+
 	return Super::Init(InEngineLoop);
 }
 
@@ -43,5 +49,76 @@ void UDisplayClusterEditorEngine::PlayInEditor(UWorld* InWorld, bool bInSimulate
 {
 	UE_LOG(LogDisplayClusterEditorEngine, VeryVerbose, TEXT("UDisplayClusterEditorEngine::PlayInEditor"));
 
+	// Find nDisplay root
+	UDisplayClusterRootComponent* RootComponent = nullptr;
+	for (AActor* Actor : InWorld->PersistentLevel->Actors)
+	{
+		if (Actor && !Actor->IsPendingKill())
+		{
+			RootComponent = Actor->FindComponentByClass<UDisplayClusterRootComponent>();
+			if (RootComponent && !RootComponent->IsPendingKill())
+			{
+				UE_LOG(LogDisplayClusterEditorEngine, Log, TEXT("Found root component - %s"), *RootComponent->GetName());
+				break;
+			}
+			else
+			{
+				RootComponent = nullptr;
+			}
+		}
+	}
+
+	if (DisplayClusterModule)
+	{
+		// If we found a root component, start DisplayCluster PIE session
+		if (RootComponent)
+		{
+			bIsNDisplayPIE = true;
+
+			if (!DisplayClusterModule->StartSession(RootComponent->GetEditorConfigPath(), RootComponent->GetEditorNodeId()))
+			{
+				UE_LOG(LogDisplayClusterEditorEngine, Error, TEXT("Couldn't start DisplayCluster session"));
+
+				// Couldn't start a new session
+				RequestEndPlayMap();
+				return;
+			}
+
+			DisplayClusterModule->StartScene(InWorld);
+		}
+	}
+
 	Super::PlayInEditor(InWorld, bInSimulateInEditor, Overrides);
 }
+
+void UDisplayClusterEditorEngine::Tick(float DeltaSeconds, bool bIdleMode)
+{
+	if (DisplayClusterModule && bIsActivePIE && bIsNDisplayPIE)
+	{
+		DisplayClusterModule->StartFrame(GFrameCounter);
+		DisplayClusterModule->PreTick(DeltaSeconds);
+		DisplayClusterModule->Tick(DeltaSeconds);
+		DisplayClusterModule->PostTick(DeltaSeconds);
+		DisplayClusterModule->EndFrame(GFrameCounter);
+	}
+
+	Super::Tick(DeltaSeconds, bIdleMode);
+}
+
+void UDisplayClusterEditorEngine::OnBeginPIE(const bool bSimulate)
+{
+	UE_LOG(LogDisplayClusterEditorEngine, VeryVerbose, TEXT("UDisplayClusterEditorEngine::OnBeginPIE"));
+
+	bIsActivePIE = true;
+}
+
+void UDisplayClusterEditorEngine::OnEndPIE(const bool bSimulate)
+{
+	UE_LOG(LogDisplayClusterEditorEngine, VeryVerbose, TEXT("UDisplayClusterEditorEngine::OnEndPIE"));
+
+	bIsActivePIE   = false;
+	bIsNDisplayPIE = false;
+
+	DisplayClusterModule->EndScene();
+	DisplayClusterModule->EndSession();
+}
diff --git a/Engine/Plugins/Runtime/nDisplay/Source/DisplayClusterMessageInterceptor/Private/DisplayClusterMessageInterceptor.cpp b/Engine/Plugins/Runtime/nDisplay/Source/DisplayClusterMessageInterceptor/Private/DisplayClusterMessageInterceptor.cpp
--- a/Engine/Plugins/Runtime/nDisplay/Source/DisplayClusterMessageInterceptor/Private/DisplayClusterMessageInterceptor.cpp
+++ b/Engine/Plugins/Runtime/nDisplay/Source/DisplayClusterMessageInterceptor/Private/DisplayClusterMessageInterceptor.cpp
@@ -15,6 +15,7 @@ FDisplayClusterMessageInterceptor::FDisplayClusterMessageInterceptor()
 	: bIsIntercepting(false)
 	, InterceptorId(FGuid::NewGuid())
 	, Address(FMessageAddress::NewAddress())
+	, ClusterManager(nullptr)
 {}
 
 void FDisplayClusterMessageInterceptor::Setup(IDisplayClusterClusterManager* InClusterManager, TSharedPtr<IMessageBus, ESPMode::ThreadSafe> InBus)
diff --git a/Engine/Plugins/Runtime/nDisplay/Source/DisplayClusterProjection/Private/DisplayClusterProjectionLog.cpp b/Engine/Plugins/Runtime/nDisplay/Source/DisplayClusterProjection/Private/DisplayClusterProjectionLog.cpp
--- a/Engine/Plugins/Runtime/nDisplay/Source/DisplayClusterProjection/Private/DisplayClusterProjectionLog.cpp
+++ b/Engine/Plugins/Runtime/nDisplay/Source/DisplayClusterProjection/Private/DisplayClusterProjectionLog.cpp
@@ -8,3 +8,4 @@ DEFINE_LOG_CATEGORY(LogDisplayClusterProjectionCamera);
 DEFINE_LOG_CATEGORY(LogDisplayClusterProjectionSimple);
 DEFINE_LOG_CATEGORY(LogDisplayClusterProjectionMPCDI);
 DEFINE_LOG_CATEGORY(LogDisplayClusterProjectionEasyBlend);
+DEFINE_LOG_CATEGORY(LogDisplayClusterProjectionManual);
diff --git a/Engine/Plugins/Runtime/nDisplay/Source/DisplayClusterProjection/Private/DisplayClusterProjectionLog.h b/Engine/Plugins/Runtime/nDisplay/Source/DisplayClusterProjection/Private/DisplayClusterProjectionLog.h
--- a/Engine/Plugins/Runtime/nDisplay/Source/DisplayClusterProjection/Private/DisplayClusterProjectionLog.h
+++ b/Engine/Plugins/Runtime/nDisplay/Source/DisplayClusterProjection/Private/DisplayClusterProjectionLog.h
@@ -10,10 +10,12 @@ DECLARE_LOG_CATEGORY_EXTERN(LogDisplayClusterProjectionCamera,    Warning, Warni
 DECLARE_LOG_CATEGORY_EXTERN(LogDisplayClusterProjectionSimple,    Warning, Warning);
 DECLARE_LOG_CATEGORY_EXTERN(LogDisplayClusterProjectionMPCDI,     Warning, Warning);
 DECLARE_LOG_CATEGORY_EXTERN(LogDisplayClusterProjectionEasyBlend, Warning, Warning);
+DECLARE_LOG_CATEGORY_EXTERN(LogDisplayClusterProjectionManual,    Warning, Warning);
 #else
 DECLARE_LOG_CATEGORY_EXTERN(LogDisplayClusterProjection,          Log, All);
 DECLARE_LOG_CATEGORY_EXTERN(LogDisplayClusterProjectionCamera,    Log, All);
 DECLARE_LOG_CATEGORY_EXTERN(LogDisplayClusterProjectionSimple,    Log, All);
 DECLARE_LOG_CATEGORY_EXTERN(LogDisplayClusterProjectionMPCDI,     Log, All);
 DECLARE_LOG_CATEGORY_EXTERN(LogDisplayClusterProjectionEasyBlend, Log, All);
+DECLARE_LOG_CATEGORY_EXTERN(LogDisplayClusterProjectionManual,    Log, All);
 #endif
diff --git a/Engine/Plugins/Runtime/nDisplay/Source/DisplayClusterProjection/Private/DisplayClusterProjectionModule.cpp b/Engine/Plugins/Runtime/nDisplay/Source/DisplayClusterProjection/Private/DisplayClusterProjectionModule.cpp
--- a/Engine/Plugins/Runtime/nDisplay/Source/DisplayClusterProjection/Private/DisplayClusterProjectionModule.cpp
+++ b/Engine/Plugins/Runtime/nDisplay/Source/DisplayClusterProjection/Private/DisplayClusterProjectionModule.cpp
@@ -3,11 +3,13 @@
 #include "DisplayClusterProjectionModule.h"
 
 #include "DisplayClusterProjectionLog.h"
+#include "DisplayClusterProjectionStrings.h"
 
 #include "Policy/Camera/DisplayClusterProjectionCameraPolicyFactory.h"
 #include "Policy/EasyBlend/DisplayClusterProjectionEasyBlendPolicyFactory.h"
 #include "Policy/Simple/DisplayClusterProjectionSimplePolicyFactory.h"
 #include "Policy/MPCDI/DisplayClusterProjectionMPCDIPolicyFactory.h"
+#include "Policy/Manual/DisplayClusterProjectionManualPolicyFactory.h"
 
 #include "IDisplayCluster.h"
 #include "Render/IDisplayClusterRenderManager.h"
@@ -33,6 +35,10 @@ FDisplayClusterProjectionModule::FDisplayClusterProjectionModule()
 	Factory = MakeShareable(new FDisplayClusterProjectionEasyBlendPolicyFactory);
 	ProjectionPolicyFactories.Emplace(DisplayClusterStrings::projection::EasyBlend, Factory);
 
+	// Manual projection
+	Factory = MakeShareable(new FDisplayClusterProjectionManualPolicyFactory);
+	ProjectionPolicyFactories.Emplace(DisplayClusterStrings::projection::Manual, Factory);
+
 	UE_LOG(LogDisplayClusterProjection, Log, TEXT("Projection module has been instantiated"));
 }
 
diff --git a/Engine/Plugins/Runtime/nDisplay/Source/DisplayClusterProjection/Private/DisplayClusterProjectionStrings.h b/Engine/Plugins/Runtime/nDisplay/Source/DisplayClusterProjection/Private/DisplayClusterProjectionStrings.h
--- a/Engine/Plugins/Runtime/nDisplay/Source/DisplayClusterProjection/Private/DisplayClusterProjectionStrings.h
+++ b/Engine/Plugins/Runtime/nDisplay/Source/DisplayClusterProjection/Private/DisplayClusterProjectionStrings.h
@@ -25,6 +25,28 @@ namespace DisplayClusterStrings
 					static constexpr auto Origin = TEXT("origin");
 					static constexpr auto Scale  = TEXT("scale");
 				}
+
+				namespace manual
+				{
+					static constexpr auto Rotation     = TEXT("rot");
+
+					static constexpr auto Matrix       = TEXT("matrix");
+					static constexpr auto MatrixLeft   = TEXT("matrix_left");
+					static constexpr auto MatrixRight  = TEXT("matrix_right");
+
+					static constexpr auto Frustum      = TEXT("frustum");
+					static constexpr auto FrustumLeft  = TEXT("frustum_left");
+					static constexpr auto FrustumRight = TEXT("frustum_right");
+
+					static constexpr auto AngleL      = TEXT("l");
+					static constexpr auto AngleLeft   = TEXT("left");
+					static constexpr auto AngleR      = TEXT("r");
+					static constexpr auto AngleRight  = TEXT("right");
+					static constexpr auto AngleT      = TEXT("t");
+					static constexpr auto AngleTop    = TEXT("top");
+					static constexpr auto AngleB      = TEXT("b");
+					static constexpr auto AngleBottom = TEXT("bottom");
+				}
 			}
 		}
 	}
@@ -35,5 +57,6 @@ namespace DisplayClusterStrings
 		static constexpr auto Simple    = TEXT("simple");
 		static constexpr auto MPCDI     = TEXT("mpcdi");
 		static constexpr auto EasyBlend = TEXT("easyblend");
+		static constexpr auto Manual    = TEXT("manual");
 	}
 };
diff --git a/Engine/Plugins/Runtime/nDisplay/Source/DisplayClusterProjection/Private/Policy/Camera/DisplayClusterProjectionCameraPolicy.cpp b/Engine/Plugins/Runtime/nDisplay/Source/DisplayClusterProjection/Private/Policy/Camera/DisplayClusterProjectionCameraPolicy.cpp
--- a/Engine/Plugins/Runtime/nDisplay/Source/DisplayClusterProjection/Private/Policy/Camera/DisplayClusterProjectionCameraPolicy.cpp
+++ b/Engine/Plugins/Runtime/nDisplay/Source/DisplayClusterProjection/Private/Policy/Camera/DisplayClusterProjectionCameraPolicy.cpp
@@ -7,19 +7,17 @@
 
 #include "DisplayClusterUtils/DisplayClusterCommonHelpers.h"
 
-#include "IDisplayCluster.h"
 #include "Config/IDisplayClusterConfigManager.h"
 #include "Game/IDisplayClusterGameManager.h"
-#include "DisplayClusterCameraComponent.h"
 
-#include "DisplayClusterPawn.h"
+#include "DisplayClusterCameraComponent.h"
+#include "DisplayClusterRootComponent.h"
 #include "DisplayClusterScreenComponent.h"
 
 #include "Camera/CameraComponent.h"
 #include "ComposurePostMoves.h"
 
 
-
 FDisplayClusterProjectionCameraPolicy::FDisplayClusterProjectionCameraPolicy(const FString& ViewportId)
 	: FDisplayClusterProjectionPolicyBase(ViewportId)
 {
@@ -36,20 +34,13 @@ FDisplayClusterProjectionCameraPolicy::~FDisplayClusterProjectionCameraPolicy()
 void FDisplayClusterProjectionCameraPolicy::StartScene(UWorld* World)
 {
 	check(IsInGameThread());
-
-	if (!AssignedCamera)
-	{
-		ADisplayClusterPawn* Pawn = IDisplayCluster::Get().GetGameMgr()->GetRoot();
-		if (Pawn)
-		{
-			AssignedCamera = Pawn->GetCameraComponent();
-		}
-	}
 }
 
 void FDisplayClusterProjectionCameraPolicy::EndScene()
 {
 	check(IsInGameThread());
+
+	AssignedCamera = nullptr;
 }
 
 bool FDisplayClusterProjectionCameraPolicy::HandleAddViewport(const FIntPoint& InViewportSize, const uint32 InViewsAmount)
diff --git a/Engine/Plugins/Runtime/nDisplay/Source/DisplayClusterProjection/Private/Policy/Camera/DisplayClusterProjectionCameraPolicyFactory.h b/Engine/Plugins/Runtime/nDisplay/Source/DisplayClusterProjection/Private/Policy/Camera/DisplayClusterProjectionCameraPolicyFactory.h
--- a/Engine/Plugins/Runtime/nDisplay/Source/DisplayClusterProjection/Private/Policy/Camera/DisplayClusterProjectionCameraPolicyFactory.h
+++ b/Engine/Plugins/Runtime/nDisplay/Source/DisplayClusterProjection/Private/Policy/Camera/DisplayClusterProjectionCameraPolicyFactory.h
@@ -3,7 +3,7 @@
 #pragma once
 
 #include "Render/Projection/IDisplayClusterProjectionPolicyFactory.h"
-#include "DisplayClusterProjectionStrings.h"
+
 
 /**
  * Implements projection policy factory for the 'camera' policy
diff --git a/Engine/Plugins/Runtime/nDisplay/Source/DisplayClusterProjection/Private/Policy/DisplayClusterProjectionPolicyBase.cpp b/Engine/Plugins/Runtime/nDisplay/Source/DisplayClusterProjection/Private/Policy/DisplayClusterProjectionPolicyBase.cpp
--- a/Engine/Plugins/Runtime/nDisplay/Source/DisplayClusterProjection/Private/Policy/DisplayClusterProjectionPolicyBase.cpp
+++ b/Engine/Plugins/Runtime/nDisplay/Source/DisplayClusterProjection/Private/Policy/DisplayClusterProjectionPolicyBase.cpp
@@ -8,7 +8,8 @@
 #include "Config/IDisplayClusterConfigManager.h"
 #include "Game/IDisplayClusterGameManager.h"
 
-#include "DisplayClusterPawn.h"
+#include "DisplayClusterSceneComponent.h"
+#include "DisplayClusterRootComponent.h"
 
 
 FDisplayClusterProjectionPolicyBase::FDisplayClusterProjectionPolicyBase(const FString& InViewportId)
@@ -28,7 +29,7 @@ void FDisplayClusterProjectionPolicyBase::InitializeOriginComponent(const FStrin
 	IDisplayClusterGameManager* const GameMgr = IDisplayCluster::Get().GetGameMgr();
 	if (!GameMgr)
 	{
-		UE_LOG(LogDisplayClusterProjectionEasyBlend, Warning, TEXT("No DisplayCluster game manager available"));
+		UE_LOG(LogDisplayClusterProjection, Warning, TEXT("No DisplayCluster game manager available"));
 		return;
 	}
 
@@ -40,8 +41,8 @@ void FDisplayClusterProjectionPolicyBase::InitializeOriginComponent(const FStrin
 
 	if(PolicyOriginComp == nullptr)
 	{
-		UE_LOG(LogDisplayClusterProjectionEasyBlend, Log, TEXT("No custom origin set or component '%s' not found for viewport '%s'. VR root will be used."), *OriginCompId, *PolicyViewportId);
-		PolicyOriginComp = GameMgr->GetRoot()->GetRootComponent();
+		UE_LOG(LogDisplayClusterProjection, Log, TEXT("No custom origin set or component '%s' not found for viewport '%s'. VR root will be used."), *OriginCompId, *PolicyViewportId);
+		PolicyOriginComp = GameMgr->GetRoot();
 	}
 
 	if (!PolicyOriginComp)
diff --git a/Engine/Plugins/Runtime/nDisplay/Source/DisplayClusterProjection/Private/Policy/EasyBlend/DX11/DisplayClusterProjectionEasyBlendViewAdapterDX11.cpp b/Engine/Plugins/Runtime/nDisplay/Source/DisplayClusterProjection/Private/Policy/EasyBlend/DX11/DisplayClusterProjectionEasyBlendViewAdapterDX11.cpp
--- a/Engine/Plugins/Runtime/nDisplay/Source/DisplayClusterProjection/Private/Policy/EasyBlend/DX11/DisplayClusterProjectionEasyBlendViewAdapterDX11.cpp
+++ b/Engine/Plugins/Runtime/nDisplay/Source/DisplayClusterProjection/Private/Policy/EasyBlend/DX11/DisplayClusterProjectionEasyBlendViewAdapterDX11.cpp
@@ -128,27 +128,18 @@ bool FDisplayClusterProjectionEasyBlendViewAdapterDX11::GetProjectionMatrix(cons
 {
 	check(Views.Num() > (int)ViewIdx);
 
-	// Invert Z-axis (UE4 uses Z-inverted LHS)
-	static const FMatrix flipZ = FMatrix(
-		FPlane(1, 0, 0, 0),
-		FPlane(0, 1, 0, 0),
-		FPlane(0, 0, -1, 0),
-		FPlane(0, 0, 1, 1));
-
 	// Build Projection matrix:
 	const float n = ZNear;
 	const float f = ZFar;
 
-	float l = float(ZNear * tan(FMath::DegreesToRadians(Views[ViewIdx].EasyBlendMeshData->Frustum.LeftAngle)));
-	float r = float(ZNear * tan(FMath::DegreesToRadians(Views[ViewIdx].EasyBlendMeshData->Frustum.RightAngle)));
-	float b = float(ZNear * tan(FMath::DegreesToRadians(Views[ViewIdx].EasyBlendMeshData->Frustum.BottomAngle)));
-	float t = float(ZNear * tan(FMath::DegreesToRadians(Views[ViewIdx].EasyBlendMeshData->Frustum.TopAngle)));
-
-	const FMatrix pm = DisplayClusterHelpers::math::GetSafeProjectionMatrix(l, r, t, b, n, f);
+	const float l = Views[ViewIdx].EasyBlendMeshData->Frustum.LeftAngle;
+	const float r = Views[ViewIdx].EasyBlendMeshData->Frustum.RightAngle;
+	const float b = Views[ViewIdx].EasyBlendMeshData->Frustum.BottomAngle;
+	const float t = Views[ViewIdx].EasyBlendMeshData->Frustum.TopAngle;
 
-	const FMatrix result(pm * flipZ);
+	const FMatrix pm = DisplayClusterHelpers::math::GetProjectionMatrixFromAngles(l, r, t, b, n, f);
 
-	OutPrjMatrix = result;
+	OutPrjMatrix = pm;
 
 	return true;
 }
diff --git a/Engine/Plugins/Runtime/nDisplay/Source/DisplayClusterProjection/Private/Policy/EasyBlend/DisplayClusterProjectionEasyBlendPolicyBase.cpp b/Engine/Plugins/Runtime/nDisplay/Source/DisplayClusterProjection/Private/Policy/EasyBlend/DisplayClusterProjectionEasyBlendPolicyBase.cpp
--- a/Engine/Plugins/Runtime/nDisplay/Source/DisplayClusterProjection/Private/Policy/EasyBlend/DisplayClusterProjectionEasyBlendPolicyBase.cpp
+++ b/Engine/Plugins/Runtime/nDisplay/Source/DisplayClusterProjection/Private/Policy/EasyBlend/DisplayClusterProjectionEasyBlendPolicyBase.cpp
@@ -11,7 +11,7 @@
 #include "Config/IDisplayClusterConfigManager.h"
 #include "Game/IDisplayClusterGameManager.h"
 
-#include "DisplayClusterPawn.h"
+#include "DisplayClusterRootComponent.h"
 #include "DisplayClusterScreenComponent.h"
 
 
diff --git a/Engine/Plugins/Runtime/nDisplay/Source/DisplayClusterProjection/Private/Policy/Manual/DisplayClusterProjectionManualPolicy.cpp b/Engine/Plugins/Runtime/nDisplay/Source/DisplayClusterProjection/Private/Policy/Manual/DisplayClusterProjectionManualPolicy.cpp
new file mode 100644
--- /dev/null
+++ b/Engine/Plugins/Runtime/nDisplay/Source/DisplayClusterProjection/Private/Policy/Manual/DisplayClusterProjectionManualPolicy.cpp
@@ -0,0 +1,222 @@
+// Copyright 1998-2019 Epic Games, Inc. All Rights Reserved.
+
+#include "Policy/Manual/DisplayClusterProjectionManualPolicy.h"
+
+#include "DisplayClusterProjectionLog.h"
+#include "DisplayClusterProjectionStrings.h"
+
+#include "DisplayClusterUtils/DisplayClusterCommonHelpers.h"
+
+#include "IDisplayCluster.h"
+#include "Config/IDisplayClusterConfigManager.h"
+#include "Game/IDisplayClusterGameManager.h"
+
+#include "Config/DisplayClusterConfigTypes.h"
+
+#include "DisplayClusterRootComponent.h"
+#include "DisplayClusterScreenComponent.h"
+
+
+FDisplayClusterProjectionManualPolicy::FDisplayClusterProjectionManualPolicy(const FString& ViewportId)
+	: FDisplayClusterProjectionPolicyBase(ViewportId)
+{
+}
+
+FDisplayClusterProjectionManualPolicy::~FDisplayClusterProjectionManualPolicy()
+{
+}
+
+
+//////////////////////////////////////////////////////////////////////////////////////////////
+// IDisplayClusterProjectionPolicy
+//////////////////////////////////////////////////////////////////////////////////////////////
+void FDisplayClusterProjectionManualPolicy::StartScene(UWorld* World)
+{
+	check(IsInGameThread());
+}
+
+void FDisplayClusterProjectionManualPolicy::EndScene()
+{
+	check(IsInGameThread());
+}
+
+bool FDisplayClusterProjectionManualPolicy::HandleAddViewport(const FIntPoint& ViewportSize, const uint32 ViewsAmount)
+{
+	check(IsInGameThread());
+
+	UE_LOG(LogDisplayClusterProjectionManual, Log, TEXT("Initializing internals for the viewport '%s'"), *GetViewportId());
+
+	// Get projection settings of the specified viewport
+	FDisplayClusterConfigProjection CfgProjection;
+	if (!DisplayClusterHelpers::config::GetViewportProjection(GetViewportId(), CfgProjection))
+	{
+		UE_LOG(LogDisplayClusterProjectionManual, Error, TEXT("No projection ID found for viewport '%s'"), *GetViewportId());
+		return false;
+	}
+
+	// Get view rotation
+	if (!DisplayClusterHelpers::str::ExtractValue(CfgProjection.Params, FString(DisplayClusterStrings::cfg::data::projection::manual::Rotation), ViewRotation))
+	{
+		UE_LOG(LogDisplayClusterProjectionManual, Log, TEXT("No rotation specified for projection policy of viewport '%s'"), *GetViewportId());
+	}
+
+	// Get matrix data
+	bool bDataTypeDetermined = false;
+	if (DisplayClusterHelpers::str::ExtractValue(CfgProjection.Params, FString(DisplayClusterStrings::cfg::data::projection::manual::Matrix), ProjectionMatrix[0]) ||
+		DisplayClusterHelpers::str::ExtractValue(CfgProjection.Params, FString(DisplayClusterStrings::cfg::data::projection::manual::MatrixLeft), ProjectionMatrix[0]))
+	{
+		if (ViewsAmount == 2)
+		{
+			if (DisplayClusterHelpers::str::ExtractValue(CfgProjection.Params, FString(DisplayClusterStrings::cfg::data::projection::manual::MatrixRight), ProjectionMatrix[1]))
+			{
+				bDataTypeDetermined = true;
+				DataType = EManualDataType::Matrix;
+			}
+		}
+		else
+		{
+			bDataTypeDetermined = true;
+			DataType = EManualDataType::Matrix;
+		}
+	}
+
+	if (!bDataTypeDetermined)
+	{
+		FString AnglesLeft;
+
+		if (DisplayClusterHelpers::str::ExtractValue(CfgProjection.Params, FString(DisplayClusterStrings::cfg::data::projection::manual::Frustum), AnglesLeft) ||
+			DisplayClusterHelpers::str::ExtractValue(CfgProjection.Params, FString(DisplayClusterStrings::cfg::data::projection::manual::FrustumLeft), AnglesLeft))
+		{
+			if (!ExtractAngles(AnglesLeft, FrustumAngles[0]))
+			{
+				UE_LOG(LogDisplayClusterProjectionManual, Log, TEXT("Couldn't extract frustum angles from value '%s'"), *AnglesLeft);
+				return false;
+			}
+
+			if (ViewsAmount == 2)
+			{
+				FString AnglesRight;
+
+				if (DisplayClusterHelpers::str::ExtractValue(CfgProjection.Params, FString(DisplayClusterStrings::cfg::data::projection::manual::FrustumRight), AnglesRight))
+				{
+					if (!ExtractAngles(AnglesRight, FrustumAngles[1]))
+					{
+						UE_LOG(LogDisplayClusterProjectionManual, Log, TEXT("Couldn't extract frustum angles from value '%s'"), *AnglesRight);
+						return false;
+					}
+
+					bDataTypeDetermined = true;
+					DataType = EManualDataType::FrustumAngles;
+				}
+			}
+			else
+			{
+				bDataTypeDetermined = true;
+				DataType = EManualDataType::FrustumAngles;
+			}
+		}
+
+	}
+
+	if (!bDataTypeDetermined)
+	{
+		UE_LOG(LogDisplayClusterProjectionManual, Log, TEXT("No mandatory data specified for projection policy of viewport '%s'"), *GetViewportId());
+		return false;
+	}
+
+	return true;
+}
+
+void FDisplayClusterProjectionManualPolicy::HandleRemoveViewport()
+{
+	check(IsInGameThread());
+
+	UE_LOG(LogDisplayClusterProjectionManual, Log, TEXT("Removing viewport '%s'"), *GetViewportId());
+}
+
+bool FDisplayClusterProjectionManualPolicy::CalculateView(const uint32 ViewIdx, FVector& InOutViewLocation, FRotator& InOutViewRotation, const FVector& InViewOffset, const float InWorldToMeters, const float InNCP, const float InFCP)
+{
+	check(IsInGameThread());
+	check(ViewIdx < 2);
+
+	// Add local rotation specified in config
+	InOutViewRotation += ViewRotation;
+
+	// Store culling data
+	NCP = InNCP;
+	FCP = InFCP;
+
+	return true;
+}
+
+bool FDisplayClusterProjectionManualPolicy::GetProjectionMatrix(const uint32 ViewIdx, FMatrix& OutPrjMatrix)
+{
+	check(IsInGameThread());
+	check(ViewIdx < 2);
+
+	bool bResult = false;
+
+	switch (DataType)
+	{
+	case EManualDataType::Matrix:
+		OutPrjMatrix = ProjectionMatrix[ViewIdx];
+		bResult = true;
+		break;
+
+	case EManualDataType::FrustumAngles:
+		OutPrjMatrix = DisplayClusterHelpers::math::GetProjectionMatrixFromAngles(FrustumAngles[ViewIdx].Left, FrustumAngles[ViewIdx].Right, FrustumAngles[ViewIdx].Top, FrustumAngles[ViewIdx].Bottom, NCP, FCP);
+		bResult = true;
+		break;
+
+	default:
+		break;
+	}
+
+	return bResult;
+}
+
+bool FDisplayClusterProjectionManualPolicy::ExtractAngles(const FString& InAngles, FFrustumAngles& OutAngles)
+{
+	float Left;
+	if (!DisplayClusterHelpers::str::ExtractValue(InAngles, FString(DisplayClusterStrings::cfg::data::projection::manual::AngleL), Left))
+	{
+		if (!DisplayClusterHelpers::str::ExtractValue(InAngles, FString(DisplayClusterStrings::cfg::data::projection::manual::AngleLeft), Left))
+		{
+			return false;
+		}
+	}
+
+	float Right;
+	if (!DisplayClusterHelpers::str::ExtractValue(InAngles, FString(DisplayClusterStrings::cfg::data::projection::manual::AngleR), Right))
+	{
+		if (!DisplayClusterHelpers::str::ExtractValue(InAngles, FString(DisplayClusterStrings::cfg::data::projection::manual::AngleRight), Right))
+		{
+			return false;
+		}
+	}
+
+	float Top;
+	if (!DisplayClusterHelpers::str::ExtractValue(InAngles, FString(DisplayClusterStrings::cfg::data::projection::manual::AngleT), Top))
+	{
+		if (!DisplayClusterHelpers::str::ExtractValue(InAngles, FString(DisplayClusterStrings::cfg::data::projection::manual::AngleTop), Top))
+		{
+			return false;
+		}
+	}
+
+	float Bottom;
+	if (!DisplayClusterHelpers::str::ExtractValue(InAngles, FString(DisplayClusterStrings::cfg::data::projection::manual::AngleB), Bottom))
+	{
+		if (!DisplayClusterHelpers::str::ExtractValue(InAngles, FString(DisplayClusterStrings::cfg::data::projection::manual::AngleBottom), Bottom))
+		{
+			return false;
+		}
+	}
+
+	OutAngles.Left   = Left;
+	OutAngles.Right  = Right;
+	OutAngles.Top    = Top;
+	OutAngles.Bottom = Bottom;
+
+	return true;
+}
diff --git a/Engine/Plugins/Runtime/nDisplay/Source/DisplayClusterProjection/Private/Policy/Manual/DisplayClusterProjectionManualPolicy.h b/Engine/Plugins/Runtime/nDisplay/Source/DisplayClusterProjection/Private/Policy/Manual/DisplayClusterProjectionManualPolicy.h
new file mode 100644
--- /dev/null
+++ b/Engine/Plugins/Runtime/nDisplay/Source/DisplayClusterProjection/Private/Policy/Manual/DisplayClusterProjectionManualPolicy.h
@@ -0,0 +1,66 @@
+// Copyright 1998-2019 Epic Games, Inc. All Rights Reserved.
+
+#pragma once
+
+#include "Policy/DisplayClusterProjectionPolicyBase.h"
+
+
+/**
+ * Implements math behind the native (manual) quad based projections
+ */
+class FDisplayClusterProjectionManualPolicy
+	: public FDisplayClusterProjectionPolicyBase
+{
+public:
+	FDisplayClusterProjectionManualPolicy(const FString& ViewportId);
+	virtual ~FDisplayClusterProjectionManualPolicy();
+
+public:
+	enum class EManualDataType
+	{
+		Matrix,
+		FrustumAngles
+	};
+
+	EManualDataType GetDataType() const
+	{ return DataType; }
+
+public:
+	//////////////////////////////////////////////////////////////////////////////////////////////
+	// IDisplayClusterProjectionPolicy
+	//////////////////////////////////////////////////////////////////////////////////////////////
+	virtual void StartScene(UWorld* World) override;
+	virtual void EndScene() override;
+	virtual bool HandleAddViewport(const FIntPoint& ViewportSize, const uint32 ViewsAmount) override;
+	virtual void HandleRemoveViewport() override;
+
+	virtual bool CalculateView(const uint32 ViewIdx, FVector& InOutViewLocation, FRotator& InOutViewRotation, const FVector& InViewOffset, const float InWorldToMeters, const float InNCP, const float InFCP) override;
+	virtual bool GetProjectionMatrix(const uint32 ViewIdx, FMatrix& OutPrjMatrix) override;
+
+	virtual bool IsWarpBlendSupported() override
+	{ return false; }
+
+protected:
+	struct FFrustumAngles
+	{
+		float Left   = 0.f;
+		float Right  = 0.f;
+		float Top    = 0.f;
+		float Bottom = 0.f;
+	};
+
+	virtual bool ExtractAngles(const FString& InAngles, FFrustumAngles& OutAngles);
+
+private:
+	// Current data type (matrix, frustum angle, ...)
+	EManualDataType DataType;
+	// View rotation
+	FRotator ViewRotation = FRotator::ZeroRotator;
+	// Projection matrix
+	FMatrix  ProjectionMatrix[2] = { FMatrix::Identity };
+	// Frustum angles
+	FFrustumAngles FrustumAngles[2];
+	// Near/far clip planes
+	float NCP;
+	float FCP;
+};
diff --git a/Engine/Plugins/Runtime/nDisplay/Source/DisplayClusterProjection/Private/Policy/Manual/DisplayClusterProjectionManualPolicyFactory.cpp b/Engine/Plugins/Runtime/nDisplay/Source/DisplayClusterProjection/Private/Policy/Manual/DisplayClusterProjectionManualPolicyFactory.cpp
new file mode 100644
--- /dev/null
+++ b/Engine/Plugins/Runtime/nDisplay/Source/DisplayClusterProjection/Private/Policy/Manual/DisplayClusterProjectionManualPolicyFactory.cpp
@@ -0,0 +1,25 @@
+// Copyright 1998-2019 Epic Games, Inc. All Rights Reserved.
+
+#include "Policy/Manual/DisplayClusterProjectionManualPolicyFactory.h"
+#include "Policy/Manual/DisplayClusterProjectionManualPolicy.h"
+
+#include "DisplayClusterProjectionLog.h"
+
+
+FDisplayClusterProjectionManualPolicyFactory::FDisplayClusterProjectionManualPolicyFactory()
+{
+}
+
+FDisplayClusterProjectionManualPolicyFactory ::~FDisplayClusterProjectionManualPolicyFactory()
+{
+}
+
+
+//////////////////////////////////////////////////////////////////////////////////////////////
+// IDisplayClusterProjectionPolicyFactory
+//////////////////////////////////////////////////////////////////////////////////////////////
+TSharedPtr<IDisplayClusterProjectionPolicy> FDisplayClusterProjectionManualPolicyFactory::Create(const FString& PolicyType, const FString& RHIName, const FString& ViewportId)
+{
+	UE_LOG(LogDisplayClusterProjectionManual, Log, TEXT("Instantiating projection policy <%s>..."), *PolicyType);
+	return MakeShareable(new FDisplayClusterProjectionManualPolicy(ViewportId));
+}
diff --git a/Engine/Plugins/Runtime/nDisplay/Source/DisplayClusterProjection/Private/Policy/Manual/DisplayClusterProjectionManualPolicyFactory.h b/Engine/Plugins/Runtime/nDisplay/Source/DisplayClusterProjection/Private/Policy/Manual/DisplayClusterProjectionManualPolicyFactory.h
new file mode 100644
--- /dev/null
+++ b/Engine/Plugins/Runtime/nDisplay/Source/DisplayClusterProjection/Private/Policy/Manual/DisplayClusterProjectionManualPolicyFactory.h
@@ -0,0 +1,23 @@
+// Copyright 1998-2019 Epic Games, Inc. All Rights Reserved.
+
+#pragma once
+
+#include "Render/Projection/IDisplayClusterProjectionPolicyFactory.h"
+
+
+/**
+ * Implements projection policy factory for the 'manual' policy
+ */
+class FDisplayClusterProjectionManualPolicyFactory
+	: public IDisplayClusterProjectionPolicyFactory
+{
+public:
+	FDisplayClusterProjectionManualPolicyFactory();
+	virtual ~FDisplayClusterProjectionManualPolicyFactory();
+
+public:
+	//////////////////////////////////////////////////////////////////////////////////////////////
+	// IDisplayClusterProjectionPolicyFactory
+	//////////////////////////////////////////////////////////////////////////////////////////////
+	virtual TSharedPtr<IDisplayClusterProjectionPolicy> Create(const FString& PolicyType, const FString& RHIName, const FString& ViewportId) override;
+};
diff --git a/Engine/Plugins/Runtime/nDisplay/Source/DisplayClusterProjection/Private/Policy/Simple/DisplayClusterProjectionSimplePolicy.cpp b/Engine/Plugins/Runtime/nDisplay/Source/DisplayClusterProjection/Private/Policy/Simple/DisplayClusterProjectionSimplePolicy.cpp
--- a/Engine/Plugins/Runtime/nDisplay/Source/DisplayClusterProjection/Private/Policy/Simple/DisplayClusterProjectionSimplePolicy.cpp
+++ b/Engine/Plugins/Runtime/nDisplay/Source/DisplayClusterProjection/Private/Policy/Simple/DisplayClusterProjectionSimplePolicy.cpp
@@ -11,7 +11,7 @@
 #include "Config/IDisplayClusterConfigManager.h"
 #include "Game/IDisplayClusterGameManager.h"
 
-#include "DisplayClusterPawn.h"
+#include "DisplayClusterRootComponent.h"
 #include "DisplayClusterScreenComponent.h"
 
 
@@ -170,19 +170,7 @@ bool FDisplayClusterProjectionSimplePolicy::GetProjectionMatrix(const uint32 Vie
 	const float b = FVector::DotProduct(vu, va) * ndifd; // distance to bottom screen edge
 	const float t = FVector::DotProduct(vu, vc) * ndifd; // distance to top screen edge
 
-	// Normal LHS
-	const FMatrix pm = DisplayClusterHelpers::math::GetSafeProjectionMatrix(l, r, t, b, n, f);
-
-	// Invert Z-axis (UE4 uses Z-inverted LHS)
-	const FMatrix flipZ = FMatrix(
-		FPlane(1, 0, 0, 0),
-		FPlane(0, 1, 0, 0),
-		FPlane(0, 0, -1, 0),
-		FPlane(0, 0, 1, 1));
-
-	const FMatrix result(pm * flipZ);
-
-	OutPrjMatrix = result;
+	const FMatrix pm = DisplayClusterHelpers::math::GetProjectionMatrixFromOffsets(l, r, t, b, n, f);
 
 	return true;
 }
@@ -203,9 +191,9 @@ void FDisplayClusterProjectionSimplePolicy::InitializeMeshData()
 		return;
 	}
 
-	// Get our pawn
-	APawn* Pawn = GameMgr->GetRoot();
-	if (!Pawn)
+	// Get our VR root
+	UDisplayClusterRootComponent* Root = GameMgr->GetRoot();
+	if (!Root)
 	{
 		UE_LOG(LogDisplayClusterProjectionSimple, Error, TEXT("Couldn't get a VR root object"));
 		return;
@@ -215,7 +203,7 @@ void FDisplayClusterProjectionSimplePolicy::InitializeMeshData()
 	USceneComponent* ParentComp = nullptr;
 	if (CfgScreen.ParentId.IsEmpty())
 	{
-		ParentComp = Pawn->GetRootComponent();
+		ParentComp = Root;
 	}
 	else
 	{
@@ -229,7 +217,7 @@ void FDisplayClusterProjectionSimplePolicy::InitializeMeshData()
 	}
 
 	// Finally, create the component
-	ScreenComp = NewObject<UDisplayClusterScreenComponent>(Pawn, FName(*CfgScreen.Id), RF_Transient);
+	ScreenComp = NewObject<UDisplayClusterScreenComponent>(Root->GetOwner(), FName(*CfgScreen.Id), RF_Transient);
 	check(ScreenComp);
 
 	// Initialize it
diff --git a/Engine/Plugins/Runtime/nDisplay/Source/DisplayClusterProjection/Private/Policy/Simple/DisplayClusterProjectionSimplePolicyFactory.h b/Engine/Plugins/Runtime/nDisplay/Source/DisplayClusterProjection/Private/Policy/Simple/DisplayClusterProjectionSimplePolicyFactory.h
--- a/Engine/Plugins/Runtime/nDisplay/Source/DisplayClusterProjection/Private/Policy/Simple/DisplayClusterProjectionSimplePolicyFactory.h
+++ b/Engine/Plugins/Runtime/nDisplay/Source/DisplayClusterProjection/Private/Policy/Simple/DisplayClusterProjectionSimplePolicyFactory.h
@@ -3,7 +3,7 @@
 #pragma once
 
 #include "Render/Projection/IDisplayClusterProjectionPolicyFactory.h"
-#include "DisplayClusterProjectionStrings.h"
+
 
 /**
  * Implements projection policy factory for the 'simple' policy
diff --git a/Engine/Plugins/Runtime/nDisplay/Source/DisplayClusterUtils/DisplayClusterCommonHelpers.h b/Engine/Plugins/Runtime/nDisplay/Source/DisplayClusterUtils/DisplayClusterCommonHelpers.h
--- a/Engine/Plugins/Runtime/nDisplay/Source/DisplayClusterUtils/DisplayClusterCommonHelpers.h
+++ b/Engine/Plugins/Runtime/nDisplay/Source/DisplayClusterUtils/DisplayClusterCommonHelpers.h
@@ -24,11 +24,11 @@ namespace DisplayClusterHelpers
 	//////////////////////////////////////////////////////////////////////////////////////////////
 	namespace str
 	{
-		static inline auto BoolToStr(bool bVal, bool bAsWord = true)
+		static inline FString BoolToStr(bool bVal, bool bAsWord = true)
 		{
 			return (bVal ? 
-				(bAsWord ? TEXT("true")  : TEXT("1")) :
-				(bAsWord ? TEXT("false") : TEXT("0")));
+				(bAsWord ? FString(TEXT("true"))  : FString(TEXT("1"))) :
+				(bAsWord ? FString(TEXT("false")) : FString(TEXT("0"))));
 		}
 
 		static void TrimStringValue(FString& InLine, bool bTrimQuotes = true)
@@ -86,7 +86,7 @@ namespace DisplayClusterHelpers
 				ResultStr += FString::Printf(TEXT("%s%s"), *FDisplayClusterTypesConverter::ToString(InData[i]), *InSeparator);
 			}
 
-			if (InSeparator.Len() > 0)
+			if (InSeparator.Len() > 0 && InData.Num() > 0)
 			{
 				ResultStr.RemoveAt(ResultStr.Len() - InSeparator.Len());
 			}
@@ -464,34 +464,34 @@ namespace DisplayClusterHelpers
 	//////////////////////////////////////////////////////////////////////////////////////////////
 	namespace math
 	{
-		static FMatrix GetSafeProjectionMatrix(float l, float r, float t, float b, float n, float f)
+		static FMatrix GetProjectionMatrixFromOffsets(float l, float r, float t, float b, float n, float f)
 		{
-			const float LeftFOVRad = atan(l / n);
-			const float RightFOVRad = atan(r / n);
-			const float TopFOVRad = atan(t / n);
-			const float BottomFOVRad = atan(b / n);
+			const float LeftFOVRad   = FMath::Atan(l / n);
+			const float RightFOVRad  = FMath::Atan(r / n);
+			const float TopFOVRad    = FMath::Atan(t / n);
+			const float BottomFOVRad = FMath::Atan(b / n);
 
 			{
 				const float MinFOVRad = FMath::DegreesToRadians(1.0f);   // min FOV = 1 degree
 				const float MaxFOVRad = FMath::DegreesToRadians(179.0f); // max FOV = 179 degree
 
 				// clamp FOV values:
-				const float FOVRadH = (RightFOVRad - LeftFOVRad);
+				const float FOVRadH  = (RightFOVRad - LeftFOVRad);
 				float NewRightFOVRad = (FOVRadH < MinFOVRad) ? (LeftFOVRad + MinFOVRad) : RightFOVRad;
-				NewRightFOVRad = (FOVRadH > MaxFOVRad) ? (LeftFOVRad + MaxFOVRad) : NewRightFOVRad;
+				NewRightFOVRad       = (FOVRadH > MaxFOVRad) ? (LeftFOVRad + MaxFOVRad) : NewRightFOVRad;
 
 				const float FOVRadV = (TopFOVRad - BottomFOVRad);
-				float NewTopFOVRad = (FOVRadV < MinFOVRad) ? (BottomFOVRad + MinFOVRad) : TopFOVRad;
-				NewTopFOVRad = (FOVRadV > MaxFOVRad) ? (BottomFOVRad + MaxFOVRad) : NewTopFOVRad;
+				float NewTopFOVRad  = (FOVRadV < MinFOVRad) ? (BottomFOVRad + MinFOVRad) : TopFOVRad;
+				NewTopFOVRad        = (FOVRadV > MaxFOVRad) ? (BottomFOVRad + MaxFOVRad) : NewTopFOVRad;
 
 				if (RightFOVRad != NewRightFOVRad)
 				{
-					r = float(n*tan(NewRightFOVRad));
+					r = float(n * FMath::Tan(NewRightFOVRad));
 					//! add LOG warning
 				}
 				if (TopFOVRad != NewTopFOVRad)
 				{
-					t = float(n*tan(NewTopFOVRad));
+					t = float(n * FMath::Tan(NewTopFOVRad));
 					//! add LOG warning
 				}
 			}
@@ -503,18 +503,37 @@ namespace DisplayClusterHelpers
 
 			// Support unlimited far plane (f==n)
 			const float mc = (f == n) ? (1.0f - Z_PRECISION) : (f / (f - n));
-			const float md = (f == n) ? -n * (1.0f - Z_PRECISION) : (-(f * n) / (f - n));
+			const float md = (f == n) ? (-n * (1.0f - Z_PRECISION)) : (-(f * n) / (f - n));
 
 			const float me = 1.f;
 
 			// Normal LHS
-			FMatrix ProjectionMatrix = FMatrix(
+			const FMatrix ProjectionMatrix = FMatrix(
 				FPlane(mx,  0,  0,  0),
 				FPlane(0,  my,  0,  0),
 				FPlane(ma, mb, mc, me),
 				FPlane(0,  0,  md,  0));
 
-			return ProjectionMatrix;
+			// Invert Z-axis (UE4 uses Z-inverted LHS)
+			static const FMatrix flipZ = FMatrix(
+				FPlane(1, 0, 0, 0),
+				FPlane(0, 1, 0, 0),
+				FPlane(0, 0, -1, 0),
+				FPlane(0, 0, 1, 1));
+
+			const FMatrix ResultMatrix(ProjectionMatrix * flipZ);
+
+			return ResultMatrix;
+		}
+
+		static FMatrix GetProjectionMatrixFromAngles(float LeftAngle, float RightAngle, float TopAngle, float BottomAngle, float ZNear, float ZFar)
+		{
+			const float t = ZNear * FMath::Tan(FMath::DegreesToRadians(TopAngle));
+			const float b = ZNear * FMath::Tan(FMath::DegreesToRadians(BottomAngle));
+			const float l = ZNear * FMath::Tan(FMath::DegreesToRadians(LeftAngle));
+			const float r = ZNear * FMath::Tan(FMath::DegreesToRadians(RightAngle));
+
+			return DisplayClusterHelpers::math::GetProjectionMatrixFromOffsets(l, r, t, b, ZNear, ZFar);
 		}
 	}
 };
diff --git a/Engine/Plugins/Runtime/nDisplay/Source/DisplayClusterUtils/DisplayClusterCommonStrings.h b/Engine/Plugins/Runtime/nDisplay/Source/DisplayClusterUtils/DisplayClusterCommonStrings.h
--- a/Engine/Plugins/Runtime/nDisplay/Source/DisplayClusterUtils/DisplayClusterCommonStrings.h
+++ b/Engine/Plugins/Runtime/nDisplay/Source/DisplayClusterUtils/DisplayClusterCommonStrings.h
@@ -169,8 +169,9 @@ namespace DisplayClusterStrings
 			// General settings tokens
 			namespace general
 			{
-				static constexpr auto Header         = TEXT("[general]");
-				static constexpr auto SwapSyncPolicy = TEXT("swap_sync_policy");
+				static constexpr auto Header                = TEXT("[general]");
+				static constexpr auto SwapSyncPolicy        = TEXT("swap_sync_policy");
+				static constexpr auto UnrealInputSyncPolicy = TEXT("ue4_input_sync_policy");
 			}
 
 			// Stereo tokens
diff --git a/Engine/Plugins/Runtime/nDisplay/Source/MPCDI/Private/MPCDIWarpTexture.cpp b/Engine/Plugins/Runtime/nDisplay/Source/MPCDI/Private/MPCDIWarpTexture.cpp
--- a/Engine/Plugins/Runtime/nDisplay/Source/MPCDI/Private/MPCDIWarpTexture.cpp
+++ b/Engine/Plugins/Runtime/nDisplay/Source/MPCDI/Private/MPCDIWarpTexture.cpp
@@ -80,31 +80,17 @@ namespace // Helpers
 
 	FMatrix GetProjectionMatrixAssymetric(const IMPCDI::FFrustum::FAngles& ProjectionAngles, float zNear, float zFar)
 	{
-		float l = ProjectionAngles.Left;
-		float r = ProjectionAngles.Right;
-		float t = ProjectionAngles.Top;
-		float b = ProjectionAngles.Bottom;
+		const float l = ProjectionAngles.Left;
+		const float r = ProjectionAngles.Right;
+		const float t = ProjectionAngles.Top;
+		const float b = ProjectionAngles.Bottom;
 
-		static const FMatrix FlipZAxisToUE4 = FMatrix(
-			FPlane(1, 0, 0, 0),
-			FPlane(0, 1, 0, 0),
-			FPlane(0, 0, -1, 0),
-			FPlane(0, 0, 1, 1));
-
-		// Normal LHS
-		FMatrix ProjectionMatrix = DisplayClusterHelpers::math::GetSafeProjectionMatrix(l, r, t, b, zNear, zFar);
-		return ProjectionMatrix * FlipZAxisToUE4;
+		return DisplayClusterHelpers::math::GetProjectionMatrixFromOffsets(l, r, t, b, zNear, zFar);
 	}
 
 	FMatrix GetProjectionMatrixAssymetricFromFrustum(float LeftAngle, float RightAngle, float TopAngle, float BottomAngle, float ZNear, float ZFar)
 	{
-		IMPCDI::FFrustum::FAngles ProjectionAngles(
-			float(ZNear*tan(FMath::DegreesToRadians(TopAngle))),
-			float(ZNear*tan(FMath::DegreesToRadians(BottomAngle))),
-			float(ZNear*tan(FMath::DegreesToRadians(LeftAngle))),
-			float(ZNear*tan(FMath::DegreesToRadians(RightAngle))));
-
-		return GetProjectionMatrixAssymetric(ProjectionAngles, ZNear, ZFar);
+		return DisplayClusterHelpers::math::GetProjectionMatrixFromAngles(LeftAngle, RightAngle, TopAngle, BottomAngle, ZNear, ZFar);
 	}
 };
 
diff --git a/Engine/Plugins/Runtime/nDisplay/Source/PicpProjection/Private/Policy/MPCDI/PicpProjectionMPCDIPolicy.cpp b/Engine/Plugins/Runtime/nDisplay/Source/PicpProjection/Private/Policy/MPCDI/PicpProjectionMPCDIPolicy.cpp
--- a/Engine/Plugins/Runtime/nDisplay/Source/PicpProjection/Private/Policy/MPCDI/PicpProjectionMPCDIPolicy.cpp
+++ b/Engine/Plugins/Runtime/nDisplay/Source/PicpProjection/Private/Policy/MPCDI/PicpProjectionMPCDIPolicy.cpp
@@ -393,33 +393,24 @@ void FPicpProjectionOverlayViewportData::Empty()
 	ViewportUnder.Empty();
 
 	for (auto& It : Cameras)
-		{ It.Empty(); }
+	{
+		It.Empty();
+	}
+	
 	Cameras.Empty();
 }
 
 
 namespace // Helpers
 {
-	FMatrix GetProjectionMatrixAssymetric(float l, float r, float t, float b, float zNear, float zFar)
-	{
-		static const FMatrix FlipZAxisToUE4 = FMatrix(
-			FPlane(1, 0, 0, 0),
-			FPlane(0, 1, 0, 0),
-			FPlane(0, 0, -1, 0),
-			FPlane(0, 0, 1, 1));
-
-		FMatrix ProjectionMatrix = DisplayClusterHelpers::math::GetSafeProjectionMatrix(l, r, t, b, zNear, zFar);
-		return ProjectionMatrix * FlipZAxisToUE4;
-	}
-
 	FMatrix GetProjectionMatrix(float Fov, float ZNear, float ZFar)
 	{
-		float l = float(ZNear*tan(FMath::DegreesToRadians(-Fov / 2)));
-		float r = float(ZNear*tan(FMath::DegreesToRadians(Fov / 2)));
-		float b = float(ZNear*tan(FMath::DegreesToRadians(-Fov / 2)));
-		float t = float(ZNear*tan(FMath::DegreesToRadians(Fov / 2)));
+		const float r = Fov / 2;
+		const float l = -r;
+		const float t = Fov / 2;
+		const float b = -t;
 
-		return GetProjectionMatrixAssymetric(l, r, t, b, ZNear, ZFar);
+		return DisplayClusterHelpers::math::GetProjectionMatrixFromAngles(l, r, t, b, ZNear, ZFar);
 	}
 };
 
diff --git a/Engine/Plugins/Runtime/nDisplay/Source/PicpProjection/Private/Policy/PicpProjectionPolicyBase.cpp b/Engine/Plugins/Runtime/nDisplay/Source/PicpProjection/Private/Policy/PicpProjectionPolicyBase.cpp
--- a/Engine/Plugins/Runtime/nDisplay/Source/PicpProjection/Private/Policy/PicpProjectionPolicyBase.cpp
+++ b/Engine/Plugins/Runtime/nDisplay/Source/PicpProjection/Private/Policy/PicpProjectionPolicyBase.cpp
@@ -8,7 +8,8 @@
 #include "Config/IDisplayClusterConfigManager.h"
 #include "Game/IDisplayClusterGameManager.h"
 
-#include "DisplayClusterPawn.h"
+#include "DisplayClusterRootComponent.h"
+#include "DisplayClusterSceneComponent.h"
 
 
 FPicpProjectionPolicyBase::FPicpProjectionPolicyBase(const FString& InViewportId)
@@ -41,7 +42,7 @@ void FPicpProjectionPolicyBase::InitializeOriginComponent(const FString& OriginC
 	if(PolicyOriginComp == nullptr)
 	{
 		UE_LOG(LogPicpProjectionMPCDI, Log, TEXT("No custom origin set or component '%s' not found for viewport '%s'. VR root will be used."), *OriginCompId, *PolicyViewportId);
-		PolicyOriginComp = GameMgr->GetRoot()->GetRootComponent();
+		PolicyOriginComp = GameMgr->GetRoot();
 	}
 
 	if (!PolicyOriginComp)
diff --git a/Engine/Source/Runtime/Engine/Classes/GameFramework/PlayerInput.h b/Engine/Source/Runtime/Engine/Classes/GameFramework/PlayerInput.h
--- a/Engine/Source/Runtime/Engine/Classes/GameFramework/PlayerInput.h
+++ b/Engine/Source/Runtime/Engine/Classes/GameFramework/PlayerInput.h
@@ -12,6 +12,7 @@
 #include "UObject/ObjectMacros.h"
 #include "UObject/Object.h"
 #include "InputCoreTypes.h"
+#include "Components/InputComponent.h"
 #include "UObject/ScriptMacros.h"
 #include "Framework/Commands/InputChord.h"
 #include "GestureRecognizer.h"
@@ -524,7 +525,7 @@ public:
 	void Tick(float DeltaTime);
 
 	/** Process the frame's input events given the current input component stack. */
-	void ProcessInputStack(const TArray<UInputComponent*>& InputComponentStack, const float DeltaTime, const bool bGamePaused);
+	virtual void ProcessInputStack(const TArray<UInputComponent*>& InputComponentStack, const float DeltaTime, const bool bGamePaused);
 
 	/** Rather than processing input, consume it and discard without doing anything useful with it.  Like calling ProcessInputStack() and ignoring all results. */
 	void DiscardPlayerInput();
@@ -623,7 +624,75 @@ public:
 	static const TArray<FInputAxisKeyMapping>& GetEngineDefinedAxisMappings() { return EngineDefinedAxisMappings; }
 
 protected:
+	/** Runtime struct that gathers up the different kinds of delegates that might be issued */
+	struct FDelegateDispatchDetails
+	{
+		uint32 EventIndex;
+		uint32 FoundIndex;
+
+		FInputActionUnifiedDelegate ActionDelegate;
+		const FInputActionBinding* SourceAction;
+		FInputChord Chord;
+		TEnumAsByte<EInputEvent> KeyEvent;
+
+		FInputTouchUnifiedDelegate TouchDelegate;
+		FVector TouchLocation;
+		ETouchIndex::Type FingerIndex;
+
+		FInputGestureUnifiedDelegate GestureDelegate;
+		float GestureValue;
+
+		FDelegateDispatchDetails(const uint32 InEventIndex, const uint32 InFoundIndex, const FInputChord& InChord, FInputActionUnifiedDelegate InDelegate, const EInputEvent InKeyEvent, const FInputActionBinding* InSourceAction = NULL)
+			: EventIndex(InEventIndex)
+			, FoundIndex(InFoundIndex)
+			, ActionDelegate(MoveTemp(InDelegate))
+			, SourceAction(InSourceAction)
+			, Chord(InChord)
+			, KeyEvent(InKeyEvent)
+		{}
+
+		FDelegateDispatchDetails(const uint32 InEventIndex, const uint32 InFoundIndex, FInputTouchUnifiedDelegate InDelegate, const FVector InLocation, const ETouchIndex::Type InFingerIndex)
+			: EventIndex(InEventIndex)
+			, FoundIndex(InFoundIndex)
+			, TouchDelegate(MoveTemp(InDelegate))
+			, TouchLocation(InLocation)
+			, FingerIndex(InFingerIndex)
+		{}
+
+		FDelegateDispatchDetails(const uint32 InEventIndex, const uint32 InFoundIndex, FInputGestureUnifiedDelegate InDelegate, const float InValue)
+			: EventIndex(InEventIndex)
+			, FoundIndex(InFoundIndex)
+			, GestureDelegate(MoveTemp(InDelegate))
+			, GestureValue(InValue)
+		{}
+	};
+
+	// We collect axis contributions by delegate, so we can sum up 
+	// contributions from multiple bindings.
+	struct FAxisDelegateDetails
+	{
+		FInputAxisUnifiedDelegate Delegate;
+		float Value;
 
+		FAxisDelegateDetails(FInputAxisUnifiedDelegate InDelegate, const float InValue)
+			: Delegate(MoveTemp(InDelegate))
+			, Value(InValue)
+		{
+		}
+	};
+	struct FVectorAxisDelegateDetails
+	{
+		FInputVectorAxisUnifiedDelegate Delegate;
+		FVector Value;
+
+		FVectorAxisDelegateDetails(FInputVectorAxisUnifiedDelegate InDelegate, const FVector InValue)
+			: Delegate(MoveTemp(InDelegate))
+			, Value(InValue)
+		{
+		}
+	};
+
+protected:
 	/**
 	* Given raw keystate value of a vector axis, returns the "massaged" value. Override for any custom behavior,
 	* such as input changes dependent on a particular game state.
@@ -636,6 +705,9 @@ protected:
 	*/
 	virtual float MassageAxisInput(FKey Key, float RawValue);
 
+	virtual void ProcessDelegatesFilter(const TArray<UInputComponent*>& InputComponentStack, TArray<FAxisDelegateDetails>& AxisDelegates, TArray<FVectorAxisDelegateDetails>& VectorAxisDelegates, TArray<FDelegateDispatchDetails>& NonAxisDelegates)
+	{ }
+
 private:
 
 	/** Process non-axes keystates */
diff --git a/Engine/Source/Runtime/Engine/Private/UserInterface/PlayerInput.cpp b/Engine/Source/Runtime/Engine/Private/UserInterface/PlayerInput.cpp
--- a/Engine/Source/Runtime/Engine/Private/UserInterface/PlayerInput.cpp
+++ b/Engine/Source/Runtime/Engine/Private/UserInterface/PlayerInput.cpp
@@ -6,7 +6,6 @@
 
 #include "GameFramework/PlayerInput.h"
 #include "Misc/CommandLine.h"
-#include "Components/InputComponent.h"
 #include "Misc/App.h"
 #include "UObject/UObjectHash.h"
 #include "UObject/UObjectIterator.h"
@@ -33,48 +32,6 @@ const TArray<FInputAxisKeyMapping> UPlayerInput::NoAxisMappings;
 TArray<FInputActionKeyMapping> UPlayerInput::EngineDefinedActionMappings;
 TArray<FInputAxisKeyMapping> UPlayerInput::EngineDefinedAxisMappings;
 
-/** Runtime struct that gathers up the different kinds of delegates that might be issued */
-struct FDelegateDispatchDetails
-{
-	uint32 EventIndex;
-	uint32 FoundIndex;
-
-	FInputActionUnifiedDelegate ActionDelegate;
-	const FInputActionBinding* SourceAction;
-	FInputChord Chord;
-	TEnumAsByte<EInputEvent> KeyEvent;
-
-	FInputTouchUnifiedDelegate TouchDelegate;
-	FVector TouchLocation;
-	ETouchIndex::Type FingerIndex;
-
-	FInputGestureUnifiedDelegate GestureDelegate;
-	float GestureValue;
-
-	FDelegateDispatchDetails(const uint32 InEventIndex, const uint32 InFoundIndex, const FInputChord& InChord, FInputActionUnifiedDelegate InDelegate, const EInputEvent InKeyEvent, const FInputActionBinding* InSourceAction = NULL)
-		: EventIndex(InEventIndex)
-		, FoundIndex(InFoundIndex)
-		, ActionDelegate(MoveTemp(InDelegate))
-		, SourceAction(InSourceAction)
-		, Chord(InChord)
-		, KeyEvent(InKeyEvent)
-	{}
-
-	FDelegateDispatchDetails(const uint32 InEventIndex, const uint32 InFoundIndex, FInputTouchUnifiedDelegate InDelegate, const FVector InLocation, const ETouchIndex::Type InFingerIndex)
-		: EventIndex(InEventIndex)
-		, FoundIndex(InFoundIndex)
-		, TouchDelegate(MoveTemp(InDelegate))
-		, TouchLocation(InLocation)
-		, FingerIndex(InFingerIndex)
-	{}
-
-	FDelegateDispatchDetails(const uint32 InEventIndex, const uint32 InFoundIndex, FInputGestureUnifiedDelegate InDelegate, const float InValue)
-		: EventIndex(InEventIndex)
-		, FoundIndex(InFoundIndex)
-		, GestureDelegate(MoveTemp(InDelegate))
-		, GestureValue(InValue)
-	{}
-};
 
 UPlayerInput::UPlayerInput()
 	: Super()
@@ -987,31 +944,6 @@ void UPlayerInput::ProcessInputStack(const TArray<UInputComponent*>& InputCompon
 {
 	ConditionalBuildKeyMappings();
 
-	// We collect axis contributions by delegate, so we can sum up 
-	// contributions from multiple bindings.
-	struct FAxisDelegateDetails
-	{
-		FInputAxisUnifiedDelegate Delegate;
-		float Value;
-
-		FAxisDelegateDetails(FInputAxisUnifiedDelegate InDelegate, const float InValue)
-			: Delegate(MoveTemp(InDelegate))
-			, Value(InValue)
-		{
-		}
-	};
-	struct FVectorAxisDelegateDetails
-	{
-		FInputVectorAxisUnifiedDelegate Delegate;
-		FVector Value;
-
-		FVectorAxisDelegateDetails(FInputVectorAxisUnifiedDelegate InDelegate, const FVector InValue)
-			: Delegate(MoveTemp(InDelegate))
-			, Value(InValue)
-		{
-		}
-	};
-
 	static TArray<FAxisDelegateDetails> AxisDelegates;
 	static TArray<FVectorAxisDelegateDetails> VectorAxisDelegates;
 	static TArray<FDelegateDispatchDetails> NonAxisDelegates;
@@ -1318,6 +1250,9 @@ void UPlayerInput::ProcessInputStack(const TArray<UInputComponent*>& InputCompon
 
 	// Dispatch the delegates in the order they occurred
 	NonAxisDelegates.Sort(FDelegateDispatchDetailsSorter());
+
+	ProcessDelegatesFilter(InputComponentStack, AxisDelegates, VectorAxisDelegates, NonAxisDelegates);
+
 	for (const FDelegateDispatchDetails& Details : NonAxisDelegates)
 	{
 		if (Details.ActionDelegate.IsBound())
diff --git a/Engine/Source/Runtime/Renderer/Private/PostProcess/SceneRenderTargets.cpp b/Engine/Source/Runtime/Renderer/Private/PostProcess/SceneRenderTargets.cpp
--- a/Engine/Source/Runtime/Renderer/Private/PostProcess/SceneRenderTargets.cpp
+++ b/Engine/Source/Runtime/Renderer/Private/PostProcess/SceneRenderTargets.cpp
@@ -389,7 +389,7 @@ FIntPoint FSceneRenderTargets::ComputeDesiredSize(const FSceneViewFamily& ViewFa
 
 		bIsSceneCapture |= View->bIsSceneCapture;
 		bIsReflectionCapture |= View->bIsReflectionCapture;
-		bIsVRScene |= View->StereoPass != EStereoscopicPass::eSSP_FULL;
+		bIsVRScene |= (View->StereoPass != EStereoscopicPass::eSSP_FULL && GEngine->XRSystem.IsValid());
 	}
 
 	{
