Fix CIS warning introduced with CL 9761530.

diff --git a/Engine/Plugins/Runtime/nDisplay/Source/DisplayCluster/Private/Render/DisplayClusterRenderManager.cpp b/Engine/Plugins/Runtime/nDisplay/Source/DisplayCluster/Private/Render/DisplayClusterRenderManager.cpp
--- a/Engine/Plugins/Runtime/nDisplay/Source/DisplayCluster/Private/Render/DisplayClusterRenderManager.cpp
+++ b/Engine/Plugins/Runtime/nDisplay/Source/DisplayCluster/Private/Render/DisplayClusterRenderManager.cpp
@@ -96,7 +96,7 @@ bool FDisplayClusterRenderManager::StartSession(const FString& configPath, const
 		return true;
 	}
 
-	UDisplayClusterViewportClient* const GameViewport = Cast<UDisplayClusterViewportClient>(GEngine->GameViewport);
+	UDisplayClusterViewportClient* const GameViewport = GEngine ? Cast<UDisplayClusterViewportClient>(GEngine->GameViewport) : nullptr;
 	if (!GameViewport)
 	{
 		UE_LOG(LogDisplayClusterRender, Warning, TEXT("DisplayClusterViewportClient is not set as default GameViewport class"));
