Fixed black screen for simple projection policy
Fixed GameViewportClient class check logic (perform check after game viewport was created)
Fixed issue with DisplayClusterRootComponent hierarchy construction

diff --git a/Engine/Plugins/Runtime/nDisplay/Source/DisplayCluster/Private/Game/Classes/Scene/DisplayClusterSceneComponent.cpp b/Engine/Plugins/Runtime/nDisplay/Source/DisplayCluster/Private/Game/Classes/Scene/DisplayClusterSceneComponent.cpp
--- a/Engine/Plugins/Runtime/nDisplay/Source/DisplayCluster/Private/Game/Classes/Scene/DisplayClusterSceneComponent.cpp
+++ b/Engine/Plugins/Runtime/nDisplay/Source/DisplayCluster/Private/Game/Classes/Scene/DisplayClusterSceneComponent.cpp
@@ -2,8 +2,9 @@
 
 #include "DisplayClusterSceneComponent.h"
 
+#include "DisplayClusterRootComponent.h"
+
 #include "Config/DisplayClusterConfigTypes.h"
-#include "Game/IPDisplayClusterGameManager.h"
 #include "Input/IPDisplayClusterInputManager.h"
 
 #include "DisplayClusterGlobals.h"
@@ -49,7 +50,7 @@ void UDisplayClusterSceneComponent::TickComponent( float DeltaTime, ELevelTick T
 				// Update transform
 				this->SetRelativeLocationAndRotation(loc, rot);
 				// Force child transforms update
-				UpdateChildTransforms(/*true*/);
+				UpdateChildTransforms(EUpdateTransformFlags::PropagateFromParent);
 			}
 		}
 	}
@@ -70,11 +71,11 @@ bool UDisplayClusterSceneComponent::ApplySettings()
 	// Take place in hierarchy
 	if (!GetParentId().IsEmpty())
 	{
-		const IPDisplayClusterGameManager* const GameMgr = GDisplayCluster->GetPrivateGameMgr();
-		if (GameMgr)
+		UDisplayClusterRootComponent* const RootComp = Cast<UDisplayClusterRootComponent>(GetAttachParent());
+		if (RootComp)
 		{
 			UE_LOG(LogDisplayClusterGame, Log, TEXT("Attaching %s to %s"), *GetId(), *GetParentId());
-			UDisplayClusterSceneComponent* const pComp = GameMgr->GetNodeById(GetParentId());
+			UDisplayClusterSceneComponent* const pComp = RootComp->GetNodeById(GetParentId());
 			AttachToComponent(pComp, FAttachmentTransformRules(EAttachmentRule::KeepRelative, false));
 		}
 	}
diff --git a/Engine/Plugins/Runtime/nDisplay/Source/DisplayCluster/Private/Render/DisplayClusterRenderManager.cpp b/Engine/Plugins/Runtime/nDisplay/Source/DisplayCluster/Private/Render/DisplayClusterRenderManager.cpp
--- a/Engine/Plugins/Runtime/nDisplay/Source/DisplayCluster/Private/Render/DisplayClusterRenderManager.cpp
+++ b/Engine/Plugins/Runtime/nDisplay/Source/DisplayCluster/Private/Render/DisplayClusterRenderManager.cpp
@@ -96,11 +96,8 @@ bool FDisplayClusterRenderManager::StartSession(const FString& configPath, const
 		return true;
 	}
 
-	UDisplayClusterViewportClient* const GameViewport = GEngine ? Cast<UDisplayClusterViewportClient>(GEngine->GameViewport) : nullptr;
-	if (!GameViewport)
-	{
-		UE_LOG(LogDisplayClusterRender, Warning, TEXT("DisplayClusterViewportClient is not set as default GameViewport class"));
-	}
+	// Set callback on viewport created. We want to make sure the DisplayClusterViewportClient is used.
+	UGameViewportClient::OnViewportCreated().AddRaw(this, &FDisplayClusterRenderManager::OnViewportCreatedHandler_CheckViewportClass);
 
 	// Create synchronization object
 	UE_LOG(LogDisplayClusterRender, Log, TEXT("Instantiating synchronization policy object..."));
@@ -645,7 +642,7 @@ TSharedPtr<IDisplayClusterRenderDevice, ESPMode::ThreadSafe> FDisplayClusterRend
 		else
 		{
 			UE_LOG(LogDisplayClusterRender, Log, TEXT("A native present handler will be instantiated when viewport is available"));
-			UGameViewportClient::OnViewportCreated().AddRaw(this, &FDisplayClusterRenderManager::OnViewportCreatedHandler);
+			UGameViewportClient::OnViewportCreated().AddRaw(this, &FDisplayClusterRenderManager::OnViewportCreatedHandler_SetCustomPresent);
 		}
 	}
 	else if (CurrentOperationMode == EDisplayClusterOperationMode::Editor)
@@ -730,7 +727,7 @@ void FDisplayClusterRenderManager::ResizeWindow(int32 WinX, int32 WinY, int32 Re
 	window->ReshapeWindow(FVector2D(WinX, WinY), FVector2D(ResX, ResY));
 }
 
-void FDisplayClusterRenderManager::OnViewportCreatedHandler()
+void FDisplayClusterRenderManager::OnViewportCreatedHandler_SetCustomPresent()
 {
 	if (GEngine && GEngine->GameViewport)
 	{
@@ -741,6 +738,18 @@ void FDisplayClusterRenderManager::OnViewportCreatedHandler()
 	}
 }
 
+void FDisplayClusterRenderManager::OnViewportCreatedHandler_CheckViewportClass()
+{
+	if (GEngine && GEngine->GameViewport)
+	{
+		UDisplayClusterViewportClient* const GameViewport = Cast<UDisplayClusterViewportClient>(GEngine->GameViewport);
+		if (!GameViewport)
+		{
+			UE_LOG(LogDisplayClusterRender, Warning, TEXT("DisplayClusterViewportClient is not set as default GameViewport class"));
+		}
+	}
+}
+
 void FDisplayClusterRenderManager::OnBeginDrawHandler()
 {
 	static bool initialized = false;
diff --git a/Engine/Plugins/Runtime/nDisplay/Source/DisplayCluster/Private/Render/DisplayClusterRenderManager.h b/Engine/Plugins/Runtime/nDisplay/Source/DisplayCluster/Private/Render/DisplayClusterRenderManager.h
--- a/Engine/Plugins/Runtime/nDisplay/Source/DisplayCluster/Private/Render/DisplayClusterRenderManager.h
+++ b/Engine/Plugins/Runtime/nDisplay/Source/DisplayCluster/Private/Render/DisplayClusterRenderManager.h
@@ -87,7 +87,8 @@ private:
 	// FDisplayClusterRenderManager
 	//////////////////////////////////////////////////////////////////////////////////////////////
 	void ResizeWindow(int32 WinX, int32 WinY, int32 ResX, int32 ResY);
-	void OnViewportCreatedHandler();
+	void OnViewportCreatedHandler_SetCustomPresent();
+	void OnViewportCreatedHandler_CheckViewportClass();
 	void OnBeginDrawHandler();
 
 private:
diff --git a/Engine/Plugins/Runtime/nDisplay/Source/DisplayClusterProjection/Private/Policy/EasyBlend/DX11/DisplayClusterProjectionEasyBlendViewAdapterDX11.cpp b/Engine/Plugins/Runtime/nDisplay/Source/DisplayClusterProjection/Private/Policy/EasyBlend/DX11/DisplayClusterProjectionEasyBlendViewAdapterDX11.cpp
--- a/Engine/Plugins/Runtime/nDisplay/Source/DisplayClusterProjection/Private/Policy/EasyBlend/DX11/DisplayClusterProjectionEasyBlendViewAdapterDX11.cpp
+++ b/Engine/Plugins/Runtime/nDisplay/Source/DisplayClusterProjection/Private/Policy/EasyBlend/DX11/DisplayClusterProjectionEasyBlendViewAdapterDX11.cpp
@@ -137,9 +137,7 @@ bool FDisplayClusterProjectionEasyBlendViewAdapterDX11::GetProjectionMatrix(cons
 	const float b = Views[ViewIdx].EasyBlendMeshData->Frustum.BottomAngle;
 	const float t = Views[ViewIdx].EasyBlendMeshData->Frustum.TopAngle;
 
-	const FMatrix pm = DisplayClusterHelpers::math::GetProjectionMatrixFromAngles(l, r, t, b, n, f);
-
-	OutPrjMatrix = pm;
+	OutPrjMatrix = DisplayClusterHelpers::math::GetProjectionMatrixFromAngles(l, r, t, b, n, f);
 
 	return true;
 }
diff --git a/Engine/Plugins/Runtime/nDisplay/Source/DisplayClusterProjection/Private/Policy/Simple/DisplayClusterProjectionSimplePolicy.cpp b/Engine/Plugins/Runtime/nDisplay/Source/DisplayClusterProjection/Private/Policy/Simple/DisplayClusterProjectionSimplePolicy.cpp
--- a/Engine/Plugins/Runtime/nDisplay/Source/DisplayClusterProjection/Private/Policy/Simple/DisplayClusterProjectionSimplePolicy.cpp
+++ b/Engine/Plugins/Runtime/nDisplay/Source/DisplayClusterProjection/Private/Policy/Simple/DisplayClusterProjectionSimplePolicy.cpp
@@ -170,7 +170,7 @@ bool FDisplayClusterProjectionSimplePolicy::GetProjectionMatrix(const uint32 Vie
 	const float b = FVector::DotProduct(vu, va) * ndifd; // distance to bottom screen edge
 	const float t = FVector::DotProduct(vu, vc) * ndifd; // distance to top screen edge
 
-	const FMatrix pm = DisplayClusterHelpers::math::GetProjectionMatrixFromOffsets(l, r, t, b, n, f);
+	OutPrjMatrix = DisplayClusterHelpers::math::GetProjectionMatrixFromOffsets(l, r, t, b, n, f);
 
 	return true;
 }
