Fixed cluster sync objects

Fixed crash with non-nDisplay stereo device

diff --git a/Engine/Plugins/Runtime/nDisplay/Source/DisplayCluster/Private/Cluster/DisplayClusterClusterManager.cpp b/Engine/Plugins/Runtime/nDisplay/Source/DisplayCluster/Private/Cluster/DisplayClusterClusterManager.cpp
--- a/Engine/Plugins/Runtime/nDisplay/Source/DisplayCluster/Private/Cluster/DisplayClusterClusterManager.cpp
+++ b/Engine/Plugins/Runtime/nDisplay/Source/DisplayCluster/Private/Cluster/DisplayClusterClusterManager.cpp
@@ -268,19 +268,22 @@ bool FDisplayClusterClusterManager::IsCluster() const
 	return Controller ? Controller->IsCluster() : false;
 }
 
-void FDisplayClusterClusterManager::RegisterSyncObject(IDisplayClusterClusterSyncObject* pSyncObj, EDisplayClusterSyncGroup SyncGroup)
+void FDisplayClusterClusterManager::RegisterSyncObject(IDisplayClusterClusterSyncObject* SyncObj, EDisplayClusterSyncGroup SyncGroup)
 {
 	DISPLAY_CLUSTER_FUNC_TRACE(LogDisplayClusterCluster);
 
 	{
 		FScopeLock lock(&ObjectsToSyncCritSec);
-		ObjectsToSync[SyncGroup].Add(pSyncObj);
-	}
 
-	UE_LOG(LogDisplayClusterCluster, Log, TEXT("Registered sync object: %s"), *pSyncObj->GetSyncId());
+		if (SyncObj)
+		{
+			ObjectsToSync[SyncGroup].Add(SyncObj);
+			UE_LOG(LogDisplayClusterCluster, Log, TEXT("Registered sync object: %s"), *SyncObj->GetSyncId());
+		}
+	}
 }
 
-void FDisplayClusterClusterManager::UnregisterSyncObject(IDisplayClusterClusterSyncObject* pSyncObj)
+void FDisplayClusterClusterManager::UnregisterSyncObject(IDisplayClusterClusterSyncObject* SyncObj)
 {
 	DISPLAY_CLUSTER_FUNC_TRACE(LogDisplayClusterCluster);
 
@@ -289,11 +292,11 @@ void FDisplayClusterClusterManager::UnregisterSyncObject(IDisplayClusterClusterS
 
 		for (auto& GroupPair : ObjectsToSync)
 		{
-			GroupPair.Value.Remove(pSyncObj);
+			GroupPair.Value.Remove(SyncObj);
 		}
 	}
 
-	UE_LOG(LogDisplayClusterCluster, Log, TEXT("Unregistered sync object: %s"), *pSyncObj->GetSyncId());
+	UE_LOG(LogDisplayClusterCluster, Log, TEXT("Unregistered sync object: %s"), *SyncObj->GetSyncId());
 }
 
 void FDisplayClusterClusterManager::AddClusterEventListener(TScriptInterface<IDisplayClusterClusterEventListener> Listener)
@@ -395,13 +398,16 @@ void FDisplayClusterClusterManager::ExportSyncData(FDisplayClusterMessage::DataT
 
 		SyncData.Empty(SyncData.Num() | 0x7);
 
-		for (IDisplayClusterClusterSyncObject* Obj : ObjectsToSync[SyncGroup])
+		for (IDisplayClusterClusterSyncObject* SyncObj : ObjectsToSync[SyncGroup])
 		{
-			if (Obj && Obj->IsDirty())
+			if (SyncObj && SyncObj->IsActive())
 			{
-				UE_LOG(LogDisplayClusterCluster, Verbose, TEXT("Adding object to sync: %s"), *Obj->GetSyncId());
-				SyncData.Add(Obj->GetSyncId(), Obj->SerializeToString());
-				Obj->ClearDirty();
+				if (SyncObj->IsDirty())
+				{
+					UE_LOG(LogDisplayClusterCluster, Verbose, TEXT("Adding object to sync: %s"), *SyncObj->GetSyncId());
+					SyncData.Add(SyncObj->GetSyncId(), SyncObj->SerializeToString());
+					SyncObj->ClearDirty();
+				}
 			}
 		}
 	}
@@ -418,19 +424,22 @@ void FDisplayClusterClusterManager::ImportSyncData(const FDisplayClusterMessage:
 			UE_LOG(LogDisplayClusterCluster, Verbose, TEXT("sync-data: %s=%s"), *it->Key, *it->Value);
 		}
 
-		for (auto obj : ObjectsToSync[SyncGroup])
+		for (auto SyncObj : ObjectsToSync[SyncGroup])
 		{
-			const FString syncId = obj->GetSyncId();
-			if (!SyncData.Contains(syncId))
+			if (SyncObj && SyncObj->IsActive())
 			{
-				UE_LOG(LogDisplayClusterCluster, Verbose, TEXT("%s has nothing to update"), *syncId);
-				continue;
-			}
+				const FString SyncId = SyncObj->GetSyncId();
+				if (!SyncData.Contains(SyncId))
+				{
+					UE_LOG(LogDisplayClusterCluster, Verbose, TEXT("%s has nothing to update"), *SyncId);
+					continue;
+				}
 
-			UE_LOG(LogDisplayClusterCluster, Verbose, TEXT("Found %s in sync data. Applying..."), *syncId);
-			if (!obj->DeserializeFromString(SyncData[syncId]))
-			{
-				UE_LOG(LogDisplayClusterCluster, Error, TEXT("Couldn't apply sync data for sync object %s"), *syncId);
+				UE_LOG(LogDisplayClusterCluster, Verbose, TEXT("Found %s in sync data. Applying..."), *SyncId);
+				if (!SyncObj->DeserializeFromString(SyncData[SyncId]))
+				{
+					UE_LOG(LogDisplayClusterCluster, Error, TEXT("Couldn't apply sync data for sync object %s"), *SyncId);
+				}
 			}
 		}
 	}
diff --git a/Engine/Plugins/Runtime/nDisplay/Source/DisplayCluster/Private/Game/Classes/Basics/DisplayClusterViewportClient.cpp b/Engine/Plugins/Runtime/nDisplay/Source/DisplayCluster/Private/Game/Classes/Basics/DisplayClusterViewportClient.cpp
--- a/Engine/Plugins/Runtime/nDisplay/Source/DisplayCluster/Private/Game/Classes/Basics/DisplayClusterViewportClient.cpp
+++ b/Engine/Plugins/Runtime/nDisplay/Source/DisplayCluster/Private/Game/Classes/Basics/DisplayClusterViewportClient.cpp
@@ -35,6 +35,9 @@
 
 #include "Render\Device\IDisplayClusterRenderDevice.h"
 
+#include "DisplayClusterEnums.h"
+#include "DisplayClusterGlobals.h"
+
 
 //DECLARE_CYCLE_STAT(TEXT("UI Drawing Time"), STAT_UIDrawingTime, STATGROUP_UI);
 
@@ -145,7 +148,7 @@ void UDisplayClusterViewportClient::Draw(FViewport* InViewport, FCanvas* SceneCa
 		if (GEngine->StereoRenderingDevice.IsValid())
 		{
 			NumFamilies = GEngine->StereoRenderingDevice->GetDesiredNumberOfViews(bStereoRendering);
-			nDisplay = true;
+			nDisplay = (GDisplayCluster->GetOperationMode() == EDisplayClusterOperationMode::Cluster);
 		}
 	}
 
diff --git a/Engine/Plugins/Runtime/nDisplay/Source/DisplayCluster/Private/Game/Classes/Scene/DisplayClusterSceneComponentSync.cpp b/Engine/Plugins/Runtime/nDisplay/Source/DisplayCluster/Private/Game/Classes/Scene/DisplayClusterSceneComponentSync.cpp
--- a/Engine/Plugins/Runtime/nDisplay/Source/DisplayCluster/Private/Game/Classes/Scene/DisplayClusterSceneComponentSync.cpp
+++ b/Engine/Plugins/Runtime/nDisplay/Source/DisplayCluster/Private/Game/Classes/Scene/DisplayClusterSceneComponentSync.cpp
@@ -54,7 +54,7 @@ void UDisplayClusterSceneComponentSync::TickComponent( float DeltaTime, ELevelTi
 	// ...
 }
 
-void UDisplayClusterSceneComponentSync::DestroyComponent(bool bPromoteChildren)
+void UDisplayClusterSceneComponentSync::OnComponentDestroyed(bool bDestroyingHierarchy)
 {
 	if (GDisplayCluster->IsModuleInitialized())
 	{
@@ -68,13 +68,18 @@ void UDisplayClusterSceneComponentSync::DestroyComponent(bool bPromoteChildren)
 		}
 	}
 
-	Super::DestroyComponent(bPromoteChildren);
+	Super::OnComponentDestroyed(bDestroyingHierarchy);
 }
 
 
 //////////////////////////////////////////////////////////////////////////////////////////////
 // IDisplayClusterClusterSyncObject
 //////////////////////////////////////////////////////////////////////////////////////////////
+bool UDisplayClusterSceneComponentSync::IsActive() const
+{
+	return !this->IsPendingKill();
+}
+
 FString UDisplayClusterSceneComponentSync::GetSyncId() const
 {
 	return FString::Printf(TEXT("S_%s"), *GetOwner()->GetName());
diff --git a/Engine/Plugins/Runtime/nDisplay/Source/DisplayCluster/Private/Game/Classes/Scene/DisplayClusterSceneComponentSyncParent.cpp b/Engine/Plugins/Runtime/nDisplay/Source/DisplayCluster/Private/Game/Classes/Scene/DisplayClusterSceneComponentSyncParent.cpp
--- a/Engine/Plugins/Runtime/nDisplay/Source/DisplayCluster/Private/Game/Classes/Scene/DisplayClusterSceneComponentSyncParent.cpp
+++ b/Engine/Plugins/Runtime/nDisplay/Source/DisplayCluster/Private/Game/Classes/Scene/DisplayClusterSceneComponentSyncParent.cpp
@@ -28,12 +28,6 @@ void UDisplayClusterSceneComponentSyncParent::TickComponent( float DeltaTime, EL
 	// ...
 }
 
-void UDisplayClusterSceneComponentSyncParent::DestroyComponent(bool bPromoteChildren)
-{
-	Super::DestroyComponent(bPromoteChildren);
-}
-
-
 
 //////////////////////////////////////////////////////////////////////////////////////////////
 // IDisplayClusterClusterSyncObject
diff --git a/Engine/Plugins/Runtime/nDisplay/Source/DisplayCluster/Private/Game/Classes/Scene/DisplayClusterSceneComponentSyncThis.cpp b/Engine/Plugins/Runtime/nDisplay/Source/DisplayCluster/Private/Game/Classes/Scene/DisplayClusterSceneComponentSyncThis.cpp
--- a/Engine/Plugins/Runtime/nDisplay/Source/DisplayCluster/Private/Game/Classes/Scene/DisplayClusterSceneComponentSyncThis.cpp
+++ b/Engine/Plugins/Runtime/nDisplay/Source/DisplayCluster/Private/Game/Classes/Scene/DisplayClusterSceneComponentSyncThis.cpp
@@ -26,11 +26,6 @@ void UDisplayClusterSceneComponentSyncThis::TickComponent( float DeltaTime, ELev
 	// ...
 }
 
-void UDisplayClusterSceneComponentSyncThis::DestroyComponent(bool bPromoteChildren)
-{
-	Super::DestroyComponent(bPromoteChildren);
-}
-
 
 //////////////////////////////////////////////////////////////////////////////////////////////
 // IDisplayClusterClusterSyncObject
diff --git a/Engine/Plugins/Runtime/nDisplay/Source/DisplayCluster/Public/Cluster/IDisplayClusterClusterSyncObject.h b/Engine/Plugins/Runtime/nDisplay/Source/DisplayCluster/Public/Cluster/IDisplayClusterClusterSyncObject.h
--- a/Engine/Plugins/Runtime/nDisplay/Source/DisplayCluster/Public/Cluster/IDisplayClusterClusterSyncObject.h
+++ b/Engine/Plugins/Runtime/nDisplay/Source/DisplayCluster/Public/Cluster/IDisplayClusterClusterSyncObject.h
@@ -15,6 +15,8 @@ public:
 	virtual ~IDisplayClusterClusterSyncObject()
 	{ }
 
+	// Need to sync this object?
+	virtual bool IsActive() const = 0;
 	// Unique ID of synced object
 	virtual FString GetSyncId() const = 0;
 	// Check if object has changed since last ClearDirty call
diff --git a/Engine/Plugins/Runtime/nDisplay/Source/DisplayCluster/Public/DisplayClusterSceneComponentSync.h b/Engine/Plugins/Runtime/nDisplay/Source/DisplayCluster/Public/DisplayClusterSceneComponentSync.h
--- a/Engine/Plugins/Runtime/nDisplay/Source/DisplayCluster/Public/DisplayClusterSceneComponentSync.h
+++ b/Engine/Plugins/Runtime/nDisplay/Source/DisplayCluster/Public/DisplayClusterSceneComponentSync.h
@@ -30,6 +30,7 @@ public:
 	//////////////////////////////////////////////////////////////////////////////////////////////
 	// IDisplayClusterClusterSyncObject
 	//////////////////////////////////////////////////////////////////////////////////////////////
+	virtual bool IsActive() const override;
 	virtual FString GetSyncId() const override;
 	
 	virtual bool IsDirty() const override
@@ -48,7 +49,7 @@ public:
 public:
 	virtual void BeginPlay() override;
 	virtual void TickComponent(float DeltaTime, ELevelTick TickType, FActorComponentTickFunction* ThisTickFunction) override;
-	virtual void DestroyComponent(bool bPromoteChildren) override;
+	virtual void OnComponentDestroyed(bool bDestroyingHierarchy) override;
 
 protected:
 	virtual FTransform GetSyncTransform() const
diff --git a/Engine/Plugins/Runtime/nDisplay/Source/DisplayCluster/Public/DisplayClusterSceneComponentSyncParent.h b/Engine/Plugins/Runtime/nDisplay/Source/DisplayCluster/Public/DisplayClusterSceneComponentSyncParent.h
--- a/Engine/Plugins/Runtime/nDisplay/Source/DisplayCluster/Public/DisplayClusterSceneComponentSyncParent.h
+++ b/Engine/Plugins/Runtime/nDisplay/Source/DisplayCluster/Public/DisplayClusterSceneComponentSyncParent.h
@@ -21,7 +21,6 @@ public:
 public:
 	virtual void BeginPlay() override;
 	virtual void TickComponent( float DeltaTime, ELevelTick TickType, FActorComponentTickFunction* ThisTickFunction ) override;
-	virtual void DestroyComponent(bool bPromoteChildren) override;
 
 public:
 	//////////////////////////////////////////////////////////////////////////////////////////////
diff --git a/Engine/Plugins/Runtime/nDisplay/Source/DisplayCluster/Public/DisplayClusterSceneComponentSyncThis.h b/Engine/Plugins/Runtime/nDisplay/Source/DisplayCluster/Public/DisplayClusterSceneComponentSyncThis.h
--- a/Engine/Plugins/Runtime/nDisplay/Source/DisplayCluster/Public/DisplayClusterSceneComponentSyncThis.h
+++ b/Engine/Plugins/Runtime/nDisplay/Source/DisplayCluster/Public/DisplayClusterSceneComponentSyncThis.h
@@ -21,7 +21,6 @@ public:
 public:
 	virtual void BeginPlay() override;
 	virtual void TickComponent( float DeltaTime, ELevelTick TickType, FActorComponentTickFunction* ThisTickFunction ) override;
-	virtual void DestroyComponent(bool bPromoteChildren) override;
 
 public:
 	//////////////////////////////////////////////////////////////////////////////////////////////
