Removed 'final' specifier from public API

diff --git a/Engine/Plugins/Runtime/nDisplay/Source/DisplayCluster/Public/Config/DisplayClusterConfigTypes.h b/Engine/Plugins/Runtime/nDisplay/Source/DisplayCluster/Public/Config/DisplayClusterConfigTypes.h
--- a/Engine/Plugins/Runtime/nDisplay/Source/DisplayCluster/Public/Config/DisplayClusterConfigTypes.h
+++ b/Engine/Plugins/Runtime/nDisplay/Source/DisplayCluster/Public/Config/DisplayClusterConfigTypes.h
@@ -19,7 +19,7 @@ struct DISPLAYCLUSTER_API FDisplayClusterConfigBase : public IDisplayClusterStri
 	{ return FString("[]"); }
 
 	// Currently no need to serialize the data
-	virtual FString SerializeToString() const override final
+	virtual FString SerializeToString() const override
 	{ return FString(); }
 
 	// Deserialization from config file
diff --git a/Engine/Plugins/Runtime/nDisplay/Source/DisplayCluster/Public/DisplayClusterSceneComponentSync.h b/Engine/Plugins/Runtime/nDisplay/Source/DisplayCluster/Public/DisplayClusterSceneComponentSync.h
--- a/Engine/Plugins/Runtime/nDisplay/Source/DisplayCluster/Public/DisplayClusterSceneComponentSync.h
+++ b/Engine/Plugins/Runtime/nDisplay/Source/DisplayCluster/Public/DisplayClusterSceneComponentSync.h
@@ -43,8 +43,8 @@ public:
 	//////////////////////////////////////////////////////////////////////////////////////////////
 	// IDisplayClusterStringSerializable
 	//////////////////////////////////////////////////////////////////////////////////////////////
-	virtual FString SerializeToString() const override final;
-	virtual bool    DeserializeFromString(const FString& data) override final;
+	virtual FString SerializeToString() const override;
+	virtual bool    DeserializeFromString(const FString& data) override;
 
 public:
 	virtual void BeginPlay() override;
