Fixed crash during OpenLevel

diff --git a/Engine/Plugins/Runtime/nDisplay/Source/DisplayCluster/Private/Game/Classes/Scene/DisplayClusterSceneComponentSync.cpp b/Engine/Plugins/Runtime/nDisplay/Source/DisplayCluster/Private/Game/Classes/Scene/DisplayClusterSceneComponentSync.cpp
--- a/Engine/Plugins/Runtime/nDisplay/Source/DisplayCluster/Private/Game/Classes/Scene/DisplayClusterSceneComponentSync.cpp
+++ b/Engine/Plugins/Runtime/nDisplay/Source/DisplayCluster/Private/Game/Classes/Scene/DisplayClusterSceneComponentSync.cpp
@@ -27,7 +27,7 @@ void UDisplayClusterSceneComponentSync::BeginPlay()
 	}
 
 	// Generate unique sync id
-	SyncId = GetSyncId();
+	SyncId = GenerateSyncId();
 
 	GameMgr = GDisplayCluster->GetPrivateGameMgr();
 	if (GameMgr && GameMgr->IsDisplayClusterActive())
@@ -80,15 +80,15 @@ bool UDisplayClusterSceneComponentSync::IsActive() const
 	return !this->IsPendingKill();
 }
 
-FString UDisplayClusterSceneComponentSync::GetSyncId() const
-{
-	return FString::Printf(TEXT("S_%s"), *GetOwner()->GetName());
-}
-
 
 //////////////////////////////////////////////////////////////////////////////////////////////
 // IDisplayClusterStringSerializable
 //////////////////////////////////////////////////////////////////////////////////////////////
+FString UDisplayClusterSceneComponentSync::GenerateSyncId()
+{
+	return FString::Printf(TEXT("S_%s"), *GetOwner()->GetName());
+}
+
 FString UDisplayClusterSceneComponentSync::SerializeToString() const
 {
 	return GetSyncTransform().ToString();
diff --git a/Engine/Plugins/Runtime/nDisplay/Source/DisplayCluster/Private/Game/Classes/Scene/DisplayClusterSceneComponentSyncParent.cpp b/Engine/Plugins/Runtime/nDisplay/Source/DisplayCluster/Private/Game/Classes/Scene/DisplayClusterSceneComponentSyncParent.cpp
--- a/Engine/Plugins/Runtime/nDisplay/Source/DisplayCluster/Private/Game/Classes/Scene/DisplayClusterSceneComponentSyncParent.cpp
+++ b/Engine/Plugins/Runtime/nDisplay/Source/DisplayCluster/Private/Game/Classes/Scene/DisplayClusterSceneComponentSyncParent.cpp
@@ -32,12 +32,6 @@ void UDisplayClusterSceneComponentSyncParent::TickComponent( float DeltaTime, EL
 //////////////////////////////////////////////////////////////////////////////////////////////
 // IDisplayClusterClusterSyncObject
 //////////////////////////////////////////////////////////////////////////////////////////////
-FString UDisplayClusterSceneComponentSyncParent::GetSyncId() const
-{
-	return FString::Printf(TEXT("SP_%s.%s"), *GetOwner()->GetName(), *GetAttachParent()->GetName());
-}
-
-
 bool UDisplayClusterSceneComponentSyncParent::IsDirty() const
 {
 	USceneComponent* const pParent = GetAttachParent();
@@ -66,6 +60,11 @@ void UDisplayClusterSceneComponentSyncParent::ClearDirty()
 //////////////////////////////////////////////////////////////////////////////////////////////
 // UDisplayClusterSceneComponentSync
 //////////////////////////////////////////////////////////////////////////////////////////////
+FString UDisplayClusterSceneComponentSyncParent::GenerateSyncId()
+{
+	return FString::Printf(TEXT("SP_%s.%s"), *GetOwner()->GetName(), *GetAttachParent()->GetName());
+}
+
 FTransform UDisplayClusterSceneComponentSyncParent::GetSyncTransform() const
 {
 	return GetAttachParent()->GetRelativeTransform();
diff --git a/Engine/Plugins/Runtime/nDisplay/Source/DisplayCluster/Private/Game/Classes/Scene/DisplayClusterSceneComponentSyncThis.cpp b/Engine/Plugins/Runtime/nDisplay/Source/DisplayCluster/Private/Game/Classes/Scene/DisplayClusterSceneComponentSyncThis.cpp
--- a/Engine/Plugins/Runtime/nDisplay/Source/DisplayCluster/Private/Game/Classes/Scene/DisplayClusterSceneComponentSyncThis.cpp
+++ b/Engine/Plugins/Runtime/nDisplay/Source/DisplayCluster/Private/Game/Classes/Scene/DisplayClusterSceneComponentSyncThis.cpp
@@ -30,11 +30,6 @@ void UDisplayClusterSceneComponentSyncThis::TickComponent( float DeltaTime, ELev
 //////////////////////////////////////////////////////////////////////////////////////////////
 // IDisplayClusterClusterSyncObject
 //////////////////////////////////////////////////////////////////////////////////////////////
-FString UDisplayClusterSceneComponentSyncThis::GetSyncId() const
-{
-	return FString::Printf(TEXT("ST_%s"), *GetOwner()->GetName());
-}
-
 bool UDisplayClusterSceneComponentSyncThis::IsDirty() const
 {
 	return (LastSyncLoc != GetRelativeLocation() || LastSyncRot != GetRelativeRotation() || LastSyncScale != GetRelativeScale3D());
@@ -51,6 +46,11 @@ void UDisplayClusterSceneComponentSyncThis::ClearDirty()
 //////////////////////////////////////////////////////////////////////////////////////////////
 // UDisplayClusterSceneComponentSync
 //////////////////////////////////////////////////////////////////////////////////////////////
+FString UDisplayClusterSceneComponentSyncThis::GenerateSyncId()
+{
+	return FString::Printf(TEXT("ST_%s"), *GetOwner()->GetName());
+}
+
 FTransform UDisplayClusterSceneComponentSyncThis::GetSyncTransform() const
 {
 	return GetRelativeTransform();
diff --git a/Engine/Plugins/Runtime/nDisplay/Source/DisplayCluster/Private/Render/Device/DisplayClusterDeviceBase.cpp b/Engine/Plugins/Runtime/nDisplay/Source/DisplayCluster/Private/Render/Device/DisplayClusterDeviceBase.cpp
--- a/Engine/Plugins/Runtime/nDisplay/Source/DisplayCluster/Private/Render/Device/DisplayClusterDeviceBase.cpp
+++ b/Engine/Plugins/Runtime/nDisplay/Source/DisplayCluster/Private/Render/Device/DisplayClusterDeviceBase.cpp
@@ -149,16 +149,23 @@ bool FDisplayClusterDeviceBase::Initialize()
 	return true;
 }
 
-void FDisplayClusterDeviceBase::InitializeWorldContent(UWorld* InWorld)
+void FDisplayClusterDeviceBase::StartScene(UWorld* InWorld)
 {
 	DISPLAY_CLUSTER_FUNC_TRACE(LogDisplayClusterRender);
 
+	bIsSceneOpen = true;
+
 	for (FDisplayClusterRenderViewport& Viewport : RenderViewports)
 	{
 		Viewport.GetProjectionPolicy()->StartScene(InWorld);
 	}
 }
 
+void FDisplayClusterDeviceBase::EndScene()
+{
+	bIsSceneOpen = false;
+}
+
 void FDisplayClusterDeviceBase::SetViewportCamera(const FString& InCameraId /* = FString() */, const FString& InViewportId /* = FString() */)
 {
 	DISPLAY_CLUSTER_FUNC_TRACE(LogDisplayClusterRender);
@@ -345,6 +352,11 @@ void FDisplayClusterDeviceBase::CalculateStereoViewOffset(const enum EStereoscop
 	UE_LOG(LogDisplayClusterRender, VeryVerbose, TEXT("OLD ViewLoc: %s, ViewRot: %s"), *ViewLocation.ToString(), *ViewRotation.ToString());
 	UE_LOG(LogDisplayClusterRender, VeryVerbose, TEXT("WorldToMeters: %f"), WorldToMeters);
 
+	if (!bIsSceneOpen)
+	{
+		return;
+	}
+
 	// Get current viewport
 	const int CurrentViewportIndex = DecodeViewportIndex(StereoPassType);
 	check(int32(CurrentViewportIndex) < RenderViewports.Num());
@@ -449,9 +461,12 @@ FMatrix FDisplayClusterDeviceBase::GetStereoProjectionMatrix(const enum EStereos
 	FDisplayClusterRenderViewContext& ViewContext = Viewport.GetContext(ViewIndex);
 
 	FMatrix PrjMatrix = FMatrix::Identity;
-	if (!Viewport.GetProjectionPolicy()->GetProjectionMatrix(ViewIndex, PrjMatrix))
+	if (bIsSceneOpen)
 	{
-		UE_LOG(LogDisplayClusterRender, Warning, TEXT("Got invalid projection matrix: Viewport %s(%d), ViewIdx: %d"), *Viewport.GetId(), CurrentViewportIndex, int(ViewIndex));
+		if (!Viewport.GetProjectionPolicy()->GetProjectionMatrix(ViewIndex, PrjMatrix))
+		{
+			UE_LOG(LogDisplayClusterRender, Warning, TEXT("Got invalid projection matrix: Viewport %s(%d), ViewIdx: %d"), *Viewport.GetId(), CurrentViewportIndex, int(ViewIndex));
+		}
 	}
 
 	return PrjMatrix;
diff --git a/Engine/Plugins/Runtime/nDisplay/Source/DisplayCluster/Private/Render/Device/DisplayClusterDeviceBase.h b/Engine/Plugins/Runtime/nDisplay/Source/DisplayCluster/Private/Render/Device/DisplayClusterDeviceBase.h
--- a/Engine/Plugins/Runtime/nDisplay/Source/DisplayCluster/Private/Render/Device/DisplayClusterDeviceBase.h
+++ b/Engine/Plugins/Runtime/nDisplay/Source/DisplayCluster/Private/Render/Device/DisplayClusterDeviceBase.h
@@ -37,7 +37,8 @@ public:
 	// IDisplayClusterStereoDevice
 	//////////////////////////////////////////////////////////////////////////////////////////////
 	virtual bool Initialize() override;
-	virtual void InitializeWorldContent(UWorld* InWorld) override;
+	virtual void StartScene(UWorld* InWorld) override;
+	virtual void EndScene() override;
 	virtual void SetViewportCamera(const FString& InCameraId = FString(), const FString& InViewportId = FString()) override;
 	virtual void SetStartPostProcessingSettings(const FString& ViewportID, const FPostProcessSettings& StartPostProcessingSettings) override;
 	virtual void SetOverridePostProcessingSettings(const FString& ViewportID, const FPostProcessSettings& OverridePostProcessingSettings, float BlendWeight = 1.0f) override;
@@ -157,6 +158,8 @@ protected:
 	// UE4 main viewport
 	FViewport* MainViewport = nullptr;
 
+	bool bIsSceneOpen = false;
+
 	// Per-eye regions
 	FIntRect EyeRegions[2];
 
diff --git a/Engine/Plugins/Runtime/nDisplay/Source/DisplayCluster/Private/Render/DisplayClusterRenderManager.cpp b/Engine/Plugins/Runtime/nDisplay/Source/DisplayCluster/Private/Render/DisplayClusterRenderManager.cpp
--- a/Engine/Plugins/Runtime/nDisplay/Source/DisplayCluster/Private/Render/DisplayClusterRenderManager.cpp
+++ b/Engine/Plugins/Runtime/nDisplay/Source/DisplayCluster/Private/Render/DisplayClusterRenderManager.cpp
@@ -140,7 +140,7 @@ bool FDisplayClusterRenderManager::StartScene(UWorld* InWorld)
 
 	if (RenderDevicePtr)
 	{
-		RenderDevicePtr->InitializeWorldContent(InWorld);
+		RenderDevicePtr->StartScene(InWorld);
 	}
 
 	return true;
@@ -149,6 +149,11 @@ bool FDisplayClusterRenderManager::StartScene(UWorld* InWorld)
 void FDisplayClusterRenderManager::EndScene()
 {
 	DISPLAY_CLUSTER_FUNC_TRACE(LogDisplayClusterRender);
+
+	if (RenderDevicePtr)
+	{
+		RenderDevicePtr->EndScene();
+	}
 }
 
 void FDisplayClusterRenderManager::PreTick(float DeltaSeconds)
diff --git a/Engine/Plugins/Runtime/nDisplay/Source/DisplayCluster/Public/DisplayClusterSceneComponent.h b/Engine/Plugins/Runtime/nDisplay/Source/DisplayCluster/Public/DisplayClusterSceneComponent.h
--- a/Engine/Plugins/Runtime/nDisplay/Source/DisplayCluster/Public/DisplayClusterSceneComponent.h
+++ b/Engine/Plugins/Runtime/nDisplay/Source/DisplayCluster/Public/DisplayClusterSceneComponent.h
@@ -40,4 +40,5 @@ public:
 
 private:
 	FDisplayClusterConfigSceneNode Config;
+	FString SyncId;
 };
diff --git a/Engine/Plugins/Runtime/nDisplay/Source/DisplayCluster/Public/DisplayClusterSceneComponentSync.h b/Engine/Plugins/Runtime/nDisplay/Source/DisplayCluster/Public/DisplayClusterSceneComponentSync.h
--- a/Engine/Plugins/Runtime/nDisplay/Source/DisplayCluster/Public/DisplayClusterSceneComponentSync.h
+++ b/Engine/Plugins/Runtime/nDisplay/Source/DisplayCluster/Public/DisplayClusterSceneComponentSync.h
@@ -31,7 +31,9 @@ public:
 	// IDisplayClusterClusterSyncObject
 	//////////////////////////////////////////////////////////////////////////////////////////////
 	virtual bool IsActive() const override;
-	virtual FString GetSyncId() const override;
+	
+	virtual FString GetSyncId() const override
+	{ return SyncId; }
 	
 	virtual bool IsDirty() const override
 	{ return true; }
@@ -52,6 +54,8 @@ public:
 	virtual void OnComponentDestroyed(bool bDestroyingHierarchy) override;
 
 protected:
+	virtual FString GenerateSyncId();
+
 	virtual FTransform GetSyncTransform() const
 	{ return FTransform(); }
 
diff --git a/Engine/Plugins/Runtime/nDisplay/Source/DisplayCluster/Public/DisplayClusterSceneComponentSyncParent.h b/Engine/Plugins/Runtime/nDisplay/Source/DisplayCluster/Public/DisplayClusterSceneComponentSyncParent.h
--- a/Engine/Plugins/Runtime/nDisplay/Source/DisplayCluster/Public/DisplayClusterSceneComponentSyncParent.h
+++ b/Engine/Plugins/Runtime/nDisplay/Source/DisplayCluster/Public/DisplayClusterSceneComponentSyncParent.h
@@ -26,7 +26,6 @@ public:
 	//////////////////////////////////////////////////////////////////////////////////////////////
 	// IDisplayClusterClusterSyncObject
 	//////////////////////////////////////////////////////////////////////////////////////////////
-	virtual FString GetSyncId() const override;
 	virtual bool IsDirty() const override;
 	virtual void ClearDirty() override;
 
@@ -34,6 +33,7 @@ protected:
 	//////////////////////////////////////////////////////////////////////////////////////////////
 	// UDisplayClusterSceneComponentSync
 	//////////////////////////////////////////////////////////////////////////////////////////////
+	virtual FString GenerateSyncId() override;
 	virtual FTransform GetSyncTransform() const override;
 	virtual void SetSyncTransform(const FTransform& t) override;
 };
diff --git a/Engine/Plugins/Runtime/nDisplay/Source/DisplayCluster/Public/DisplayClusterSceneComponentSyncThis.h b/Engine/Plugins/Runtime/nDisplay/Source/DisplayCluster/Public/DisplayClusterSceneComponentSyncThis.h
--- a/Engine/Plugins/Runtime/nDisplay/Source/DisplayCluster/Public/DisplayClusterSceneComponentSyncThis.h
+++ b/Engine/Plugins/Runtime/nDisplay/Source/DisplayCluster/Public/DisplayClusterSceneComponentSyncThis.h
@@ -26,7 +26,6 @@ public:
 	//////////////////////////////////////////////////////////////////////////////////////////////
 	// IDisplayClusterClusterSyncObject
 	//////////////////////////////////////////////////////////////////////////////////////////////
-	virtual FString GetSyncId() const override;
 	virtual bool IsDirty() const override;
 	virtual void ClearDirty() override;
 
@@ -34,6 +33,7 @@ protected:
 	//////////////////////////////////////////////////////////////////////////////////////////////
 	// UDisplayClusterSceneComponentSync
 	//////////////////////////////////////////////////////////////////////////////////////////////
+	virtual FString GenerateSyncId() override;
 	virtual FTransform GetSyncTransform() const override;
 	virtual void SetSyncTransform(const FTransform& t) override;
 };
diff --git a/Engine/Plugins/Runtime/nDisplay/Source/DisplayCluster/Public/Render/Device/IDisplayClusterRenderDevice.h b/Engine/Plugins/Runtime/nDisplay/Source/DisplayCluster/Public/Render/Device/IDisplayClusterRenderDevice.h
--- a/Engine/Plugins/Runtime/nDisplay/Source/DisplayCluster/Public/Render/Device/IDisplayClusterRenderDevice.h
+++ b/Engine/Plugins/Runtime/nDisplay/Source/DisplayCluster/Public/Render/Device/IDisplayClusterRenderDevice.h
@@ -33,7 +33,13 @@ public:
 	/**
 	* Called on a scene start to allow a rendering device to initialize any world related content
 	*/
-	virtual void InitializeWorldContent(UWorld* World)
+	virtual void StartScene(UWorld* World)
+	{ }
+
+	/**
+	* Called before unload current level
+	*/
+	virtual void EndScene()
 	{ }
 
 	/**
