Updated template scene + fixed local paths for nDisplay configs

diff --git a/Engine/Plugins/Runtime/nDisplay/Source/DisplayCluster/Private/Config/DisplayClusterConfigManager.cpp b/Engine/Plugins/Runtime/nDisplay/Source/DisplayCluster/Private/Config/DisplayClusterConfigManager.cpp
--- a/Engine/Plugins/Runtime/nDisplay/Source/DisplayCluster/Private/Config/DisplayClusterConfigManager.cpp
+++ b/Engine/Plugins/Runtime/nDisplay/Source/DisplayCluster/Private/Config/DisplayClusterConfigManager.cpp
@@ -16,6 +16,8 @@
 #include "DisplayClusterLog.h"
 #include "DisplayClusterStrings.h"
 
+#include "HAL/FileManager.h"
+
 
 FDisplayClusterConfigManager::FDisplayClusterConfigManager()
 {
@@ -422,27 +424,36 @@ bool FDisplayClusterConfigManager::LoadConfig(const FString& cfgPath)
 {
 	DISPLAY_CLUSTER_FUNC_TRACE(LogDisplayClusterConfig);
 
+	FString ConfigFile = cfgPath;
+	ConfigFile.TrimStartAndEndInline();
+
+	if (FPaths::IsRelative(ConfigFile))
+	{
+		const FString ProjectDir = IFileManager::Get().ConvertToAbsolutePathForExternalAppForRead(*FPaths::ProjectDir());
+		ConfigFile = FPaths::ConvertRelativePathToFull(ProjectDir, ConfigFile);
+	}
+
 	// Actually the data is reset on EndFrame. This one is a safety call.
 	ResetConfigData();
 
 #ifdef DISPLAY_CLUSTER_USE_DEBUG_STANDALONE_CONFIG
-	if (cfgPath.Compare(FString(DisplayClusterStrings::misc::DbgStubConfig), ESearchCase::IgnoreCase) != 0 &&
-		FPaths::FileExists(cfgPath) == false)
+	if (ConfigFile.Compare(FString(DisplayClusterStrings::misc::DbgStubConfig), ESearchCase::IgnoreCase) != 0 &&
+		FPaths::FileExists(ConfigFile) == false)
 	{
-		UE_LOG(LogDisplayClusterConfig, Error, TEXT("File not found: %s"), *cfgPath);
+		UE_LOG(LogDisplayClusterConfig, Error, TEXT("File not found: %s"), *ConfigFile);
 		return false;
 	}
 #else
-	if (FPaths::FileExists(cfgPath) == false)
+	if (FPaths::FileExists(ConfigFile) == false)
 	{
-		UE_LOG(LogDisplayClusterConfig, Error, TEXT("File not found: %s"), *cfgPath);
+		UE_LOG(LogDisplayClusterConfig, Error, TEXT("File not found: %s"), *ConfigFile);
 		return false;
 	}
 #endif
 
 	// Instantiate appropriate parser
 	TUniquePtr<FDisplayClusterConfigParser> parser;
-	switch (GetConfigFileType(cfgPath))
+	switch (GetConfigFileType(ConfigFile))
 	{
 	case EConfigFileType::Text:
 		parser.Reset(new FDisplayClusterConfigParserText(this));
@@ -464,7 +475,7 @@ bool FDisplayClusterConfigManager::LoadConfig(const FString& cfgPath)
 		return false;
 	}
 
-	return parser->ParseFile(cfgPath);
+	return parser->ParseFile(ConfigFile);
 }
 
 void FDisplayClusterConfigManager::ResetConfigData()
diff --git a/Engine/Plugins/Runtime/nDisplay/Source/DisplayClusterEditor/Private/DisplayClusterEditorEngine.cpp b/Engine/Plugins/Runtime/nDisplay/Source/DisplayClusterEditor/Private/DisplayClusterEditorEngine.cpp
--- a/Engine/Plugins/Runtime/nDisplay/Source/DisplayClusterEditor/Private/DisplayClusterEditorEngine.cpp
+++ b/Engine/Plugins/Runtime/nDisplay/Source/DisplayClusterEditor/Private/DisplayClusterEditorEngine.cpp
@@ -7,7 +7,6 @@
 
 #include "DisplayCluster/Private/IPDisplayCluster.h"
 
-#include "Editor.h"
 
 void UDisplayClusterEditorEngine::Init(IEngineLoop* InEngineLoop)
 {
@@ -45,31 +44,57 @@ void UDisplayClusterEditorEngine::PreExit()
 	Super::PreExit();
 }
 
-void UDisplayClusterEditorEngine::PlayInEditor(UWorld* InWorld, bool bInSimulateInEditor, FPlayInEditorOverrides Overrides)
+UDisplayClusterRootComponent* UDisplayClusterEditorEngine::FindDisplayClusterRootComponent(UWorld* InWorld)
 {
-	UE_LOG(LogDisplayClusterEditorEngine, VeryVerbose, TEXT("UDisplayClusterEditorEngine::PlayInEditor"));
-
-	// Find nDisplay root
-	UDisplayClusterRootComponent* RootComponent = nullptr;
-	for (AActor* Actor : InWorld->PersistentLevel->Actors)
+	for (AActor* const Actor : InWorld->PersistentLevel->Actors)
 	{
 		if (Actor && !Actor->IsPendingKill())
 		{
-			RootComponent = Actor->FindComponentByClass<UDisplayClusterRootComponent>();
+			UDisplayClusterRootComponent* RootComponent = Actor->FindComponentByClass<UDisplayClusterRootComponent>();
 			if (RootComponent && !RootComponent->IsPendingKill())
 			{
 				UE_LOG(LogDisplayClusterEditorEngine, Log, TEXT("Found root component - %s"), *RootComponent->GetName());
-				break;
-			}
-			else
-			{
-				RootComponent = nullptr;
+				return RootComponent;
 			}
 		}
 	}
 
+	return nullptr;
+}
+
+void UDisplayClusterEditorEngine::PlayInEditor(UWorld* InWorld, bool bInSimulateInEditor, FPlayInEditorOverrides Overrides)
+{
+	UE_LOG(LogDisplayClusterEditorEngine, VeryVerbose, TEXT("UDisplayClusterEditorEngine::PlayInEditor"));
+
 	if (DisplayClusterModule)
 	{
+		// Find nDisplay root
+		UDisplayClusterRootComponent* RootComponent = FindDisplayClusterRootComponent(InWorld);
+		if (!RootComponent)
+		{
+			//Also search inside streamed levels:
+			const TArray<ULevelStreaming*>& StreamingLevels = InWorld->GetStreamingLevels();
+			for (ULevelStreaming* SreamingLevel : StreamingLevels)
+			{
+				switch (SreamingLevel->GetCurrentState())
+				{
+					case ULevelStreaming::ECurrentState::LoadedVisible:
+					{
+						// Parse only in loaded sub-levels
+						const TSoftObjectPtr<UWorld>& SubWorldAsset = SreamingLevel->GetWorldAsset();
+						RootComponent = FindDisplayClusterRootComponent(SubWorldAsset.Get());
+						if (RootComponent)
+						{
+							break;
+						}
+					}
+
+					default:
+						break;
+				}
+			}
+		}
+
 		// If we found a root component, start DisplayCluster PIE session
 		if (RootComponent)
 		{
diff --git a/Engine/Plugins/Runtime/nDisplay/Source/DisplayClusterEditor/Private/DisplayClusterEditorEngine.h b/Engine/Plugins/Runtime/nDisplay/Source/DisplayClusterEditor/Private/DisplayClusterEditorEngine.h
--- a/Engine/Plugins/Runtime/nDisplay/Source/DisplayClusterEditor/Private/DisplayClusterEditorEngine.h
+++ b/Engine/Plugins/Runtime/nDisplay/Source/DisplayClusterEditor/Private/DisplayClusterEditorEngine.h
@@ -7,6 +7,7 @@
 #include "DisplayClusterEditorEngine.generated.h"
 
 class IPDisplayCluster;
+class UDisplayClusterRootComponent;
 
 
 /**
@@ -24,6 +25,9 @@ public:
 	virtual void PlayInEditor(UWorld* InWorld, bool bInSimulateInEditor, FPlayInEditorOverrides Overrides = FPlayInEditorOverrides()) override;
 	virtual void Tick(float DeltaSeconds, bool bIdleMode) override;
 
+private:
+	UDisplayClusterRootComponent* FindDisplayClusterRootComponent(UWorld* InWorld);
+
 private:
 	IPDisplayCluster* DisplayClusterModule = nullptr;
 
