Fixed build error after 10083993

diff --git a/Engine/Plugins/Runtime/nDisplay/Source/DisplayClusterEditor/Private/DisplayClusterEditorEngine.cpp b/Engine/Plugins/Runtime/nDisplay/Source/DisplayClusterEditor/Private/DisplayClusterEditorEngine.cpp
--- a/Engine/Plugins/Runtime/nDisplay/Source/DisplayClusterEditor/Private/DisplayClusterEditorEngine.cpp
+++ b/Engine/Plugins/Runtime/nDisplay/Source/DisplayClusterEditor/Private/DisplayClusterEditorEngine.cpp
@@ -7,6 +7,8 @@
 
 #include "DisplayCluster/Private/IPDisplayCluster.h"
 
+#include "Editor.h"
+
 
 void UDisplayClusterEditorEngine::Init(IEngineLoop* InEngineLoop)
 {
