Added an #include to fix the build

diff --git a/Engine/Plugins/Runtime/nDisplay/Source/DisplayClusterEditor/Private/DisplayClusterEditorEngine.cpp b/Engine/Plugins/Runtime/nDisplay/Source/DisplayClusterEditor/Private/DisplayClusterEditorEngine.cpp
--- a/Engine/Plugins/Runtime/nDisplay/Source/DisplayClusterEditor/Private/DisplayClusterEditorEngine.cpp
+++ b/Engine/Plugins/Runtime/nDisplay/Source/DisplayClusterEditor/Private/DisplayClusterEditorEngine.cpp
@@ -7,6 +7,8 @@
 
 #include "DisplayCluster/Private/IPDisplayCluster.h"
 
+#include "Engine/LevelStreaming.h"
+
 #include "Editor.h"
 
 
@@ -76,14 +78,14 @@ void UDisplayClusterEditorEngine::PlayInEditor(UWorld* InWorld, bool bInSimulate
 		{
 			//Also search inside streamed levels:
 			const TArray<ULevelStreaming*>& StreamingLevels = InWorld->GetStreamingLevels();
-			for (ULevelStreaming* SreamingLevel : StreamingLevels)
+			for (ULevelStreaming* StreamingLevel : StreamingLevels)
 			{
-				switch (SreamingLevel->GetCurrentState())
+				switch (StreamingLevel->GetCurrentState())
 				{
 					case ULevelStreaming::ECurrentState::LoadedVisible:
 					{
 						// Parse only in loaded sub-levels
-						const TSoftObjectPtr<UWorld>& SubWorldAsset = SreamingLevel->GetWorldAsset();
+						const TSoftObjectPtr<UWorld>& SubWorldAsset = StreamingLevel->GetWorldAsset();
 						RootComponent = FindDisplayClusterRootComponent(SubWorldAsset.Get());
 						if (RootComponent)
 						{
