Fix attenuation functions now taking a listener index.
#rnx

diff --git a/Engine/Plugins/Runtime/nDisplay/Source/DisplayCluster/Private/Game/Classes/Basics/DisplayClusterViewportClient.cpp b/Engine/Plugins/Runtime/nDisplay/Source/DisplayCluster/Private/Game/Classes/Basics/DisplayClusterViewportClient.cpp
--- a/Engine/Plugins/Runtime/nDisplay/Source/DisplayCluster/Private/Game/Classes/Basics/DisplayClusterViewportClient.cpp
+++ b/Engine/Plugins/Runtime/nDisplay/Source/DisplayCluster/Private/Game/Classes/Basics/DisplayClusterViewportClient.cpp
@@ -356,11 +356,11 @@ void UDisplayClusterViewportClient::Draw(FViewport* InViewport, FCanvas* SceneCa
 									FVector OverrideAttenuation;
 									if (PlayerController->GetAudioListenerAttenuationOverridePosition(OverrideAttenuation))
 									{
-										AudioDevice->SetListenerAttenuationOverride(OverrideAttenuation);
+										AudioDevice->SetListenerAttenuationOverride(ViewportIndex, OverrideAttenuation);
 									}
 									else
 									{
-										AudioDevice->ClearListenerAttenuationOverride();
+										AudioDevice->ClearListenerAttenuationOverride(ViewportIndex);
 									}
 								}
 							}
