- Removed parameter to access Cluster Start/End frame delegate

diff --git a/Engine/Plugins/Runtime/nDisplay/Source/DisplayCluster/Private/DisplayClusterModule.h b/Engine/Plugins/Runtime/nDisplay/Source/DisplayCluster/Private/DisplayClusterModule.h
--- a/Engine/Plugins/Runtime/nDisplay/Source/DisplayCluster/Private/DisplayClusterModule.h
+++ b/Engine/Plugins/Runtime/nDisplay/Source/DisplayCluster/Private/DisplayClusterModule.h
@@ -70,10 +70,10 @@ public:
 	virtual FDisplayClusterEndSessionEvent& OnDisplayClusterEndSession() override
 	{ return DisplayClusterEndSessionEvent; }
 
-	virtual FDisplayClusterStartFrameEvent& OnDisplayClusterStartFrame(uint64 FrameNum) override
+	virtual FDisplayClusterStartFrameEvent& OnDisplayClusterStartFrame() override
 	{ return DisplayClusterStartFrameEvent; }
 
-	virtual FDisplayClusterEndFrameEvent& OnDisplayClusterEndFrame(uint64 FrameNum) override
+	virtual FDisplayClusterEndFrameEvent& OnDisplayClusterEndFrame() override
 	{ return DisplayClusterEndFrameEvent; }
 
 	virtual FDisplayClusterPreTickEvent& OnDisplayClusterPreTick() override
diff --git a/Engine/Plugins/Runtime/nDisplay/Source/DisplayCluster/Public/IDisplayCluster.h b/Engine/Plugins/Runtime/nDisplay/Source/DisplayCluster/Public/IDisplayCluster.h
--- a/Engine/Plugins/Runtime/nDisplay/Source/DisplayCluster/Public/IDisplayCluster.h
+++ b/Engine/Plugins/Runtime/nDisplay/Source/DisplayCluster/Public/IDisplayCluster.h
@@ -110,11 +110,11 @@ public:
 
 	/** Called on DisplayCluster StartFrame **/
 	DECLARE_EVENT_OneParam(IDisplayCluster, FDisplayClusterStartFrameEvent, uint64);
-	virtual FDisplayClusterStartFrameEvent& OnDisplayClusterStartFrame(uint64 FrameNum) = 0;
+	virtual FDisplayClusterStartFrameEvent& OnDisplayClusterStartFrame() = 0;
 
 	/** Called on DisplayCluster EndFrame **/
 	DECLARE_EVENT_OneParam(IDisplayCluster, FDisplayClusterEndFrameEvent, uint64);
-	virtual FDisplayClusterEndFrameEvent& OnDisplayClusterEndFrame(uint64 FrameNum) = 0;
+	virtual FDisplayClusterEndFrameEvent& OnDisplayClusterEndFrame() = 0;
 
 	/** Called on DisplayCluster PreTick **/
 	DECLARE_EVENT(IDisplayCluster, FDisplayClusterPreTickEvent);
