- Moving nDisplay engine related settings to Windows platform config
- Updated VirtualProduction templates with new values to avoid crashes on Linux and Mac

diff --git a/Engine/Plugins/Runtime/nDisplay/Source/DisplayClusterEditor/Private/DisplayClusterEditorSettings.cpp b/Engine/Plugins/Runtime/nDisplay/Source/DisplayClusterEditor/Private/DisplayClusterEditorSettings.cpp
--- a/Engine/Plugins/Runtime/nDisplay/Source/DisplayClusterEditor/Private/DisplayClusterEditorSettings.cpp
+++ b/Engine/Plugins/Runtime/nDisplay/Source/DisplayClusterEditor/Private/DisplayClusterEditorSettings.cpp
@@ -16,7 +16,9 @@ void UDisplayClusterEditorSettings::PostEditChangeProperty(FPropertyChangedEvent
 {
 	if (PropertyChangedEvent.Property != nullptr)
 	{
-		static const FString DefaultEnginePath = FString::Printf(TEXT("%sDefaultEngine.ini"), *FPaths::SourceConfigDir());
+		//Since nDisplay is Windows only, save configs that depends on nDisplay plugin assets in Windows specific config file
+		static const FString PlatformName = FPlatformProperties::PlatformName();
+		static const FString DefaultEnginePath = FString::Printf(TEXT("%s%s/%sEngine.ini"), *FPaths::SourceConfigDir(), *PlatformName, *PlatformName);
 		static const FString DefaultGamePath   = FString::Printf(TEXT("%sDefaultGame.ini"),   *FPaths::SourceConfigDir());
 
 		FName PropertyName(PropertyChangedEvent.Property->GetFName());
diff --git a/Templates/TP_ME_BlankBP/Config/DefaultEngine.ini b/Templates/TP_ME_BlankBP/Config/DefaultEngine.ini
--- a/Templates/TP_ME_BlankBP/Config/DefaultEngine.ini
+++ b/Templates/TP_ME_BlankBP/Config/DefaultEngine.ini
@@ -88,16 +88,9 @@ DefaultSaveSessionAs=
 ClientSettings=(DisplayName="",AvatarColor=(R=1.000000,G=1.000000,B=1.000000,A=1.000000),DesktopAvatarActorClass=/ConcertSyncClient/DesktopPresence.DesktopPresence_C,VRAvatarActorClass=/ConcertSyncClient/VRPresence.VRPresence_C,DiscoveryTimeoutSeconds=5,SessionTickFrequencySeconds=1,LatencyCompensationMs=0.000000,Tags=)
 EndpointSettings=(bEnableLogging=False,PurgeProcessedMessageDelaySeconds=30,RemoteEndpointTimeoutSeconds=60)
 
-; Setup for Ndisplay
-[/Script/Engine.Engine]
-GameEngine=/Script/DisplayCluster.DisplayClusterGameEngine
-UnrealEdEngine=/Script/DisplayClusterEditor.DisplayClusterEditorEngine
-GameViewportClientClassName=/Script/DisplayCluster.DisplayClusterViewportClient
-
 ; Setup for Ndisplay
 [/Script/DisplayClusterEditor.DisplayClusterEditorSettings]
 bEnabled=True
-
 [/Script/WindowsTargetPlatform.WindowsTargetSettings]
 ;Ndisplay Requirements
 DefaultGraphicsRHI=DefaultGraphicsRHI_DX11
diff --git a/Templates/TP_ME_BlankBP/Config/Windows/WindowsEngine.ini b/Templates/TP_ME_BlankBP/Config/Windows/WindowsEngine.ini
new file mode 100644
--- /dev/null
+++ b/Templates/TP_ME_BlankBP/Config/Windows/WindowsEngine.ini
@@ -0,0 +1,6 @@
+[/Script/Engine.Engine]
+GameEngine=/Script/DisplayCluster.DisplayClusterGameEngine
+UnrealEdEngine=/Script/DisplayClusterEditor.DisplayClusterEditorEngine
+GameViewportClientClassName=/Script/DisplayCluster.DisplayClusterViewportClient
+
+
diff --git a/Templates/TP_ME_VProdBP/Config/DefaultEngine.ini b/Templates/TP_ME_VProdBP/Config/DefaultEngine.ini
--- a/Templates/TP_ME_VProdBP/Config/DefaultEngine.ini
+++ b/Templates/TP_ME_VProdBP/Config/DefaultEngine.ini
@@ -96,12 +96,6 @@ ClientSettings=(DisplayName="",AvatarColor=(R=1.000000,G=1.000000,B=1.000000,A=1
 EndpointSettings=(bEnableLogging=False,PurgeProcessedMessageDelaySeconds=30,RemoteEndpointTimeoutSeconds=60)
 
 
-;Setup for Ndisplay Cluster
-[/Script/Engine.Engine]
-GameEngine=/Script/DisplayCluster.DisplayClusterGameEngine
-UnrealEdEngine=/Script/DisplayClusterEditor.DisplayClusterEditorEngine
-GameViewportClientClassName=/Script/DisplayCluster.DisplayClusterViewportClient
-
 ;Setup for Ndisplay Cluster
 [/Script/DisplayClusterEditor.DisplayClusterEditorSettings]
 bEnabled=True
diff --git a/Templates/TP_ME_VProdBP/Config/Windows/WindowsEngine.ini b/Templates/TP_ME_VProdBP/Config/Windows/WindowsEngine.ini
new file mode 100644
--- /dev/null
+++ b/Templates/TP_ME_VProdBP/Config/Windows/WindowsEngine.ini
@@ -0,0 +1,6 @@
+[/Script/Engine.Engine]
+GameEngine=/Script/DisplayCluster.DisplayClusterGameEngine
+UnrealEdEngine=/Script/DisplayClusterEditor.DisplayClusterEditorEngine
+GameViewportClientClassName=/Script/DisplayCluster.DisplayClusterViewportClient
+
+
diff --git a/Templates/TP_nDisplayBP/Config/DefaultEngine.ini b/Templates/TP_nDisplayBP/Config/DefaultEngine.ini
--- a/Templates/TP_nDisplayBP/Config/DefaultEngine.ini
+++ b/Templates/TP_nDisplayBP/Config/DefaultEngine.ini
@@ -8,11 +8,6 @@ AppliedDefaultGraphicsPerformance=Maximum
 EditorStartupMap=/Game/Maps/Main.Main
 GameDefaultMap=/Game/Maps/Main.Main
 
-[/Script/Engine.Engine]
-GameEngine=/Script/DisplayCluster.DisplayClusterGameEngine
-UnrealEdEngine=/Script/DisplayClusterEditor.DisplayClusterEditorEngine
-GameViewportClientClassName=/Script/DisplayCluster.DisplayClusterViewportClient
-
 [/Script/DisplayClusterEditor.DisplayClusterEditorSettings]
 bEnabled=True
 
diff --git a/Templates/TP_nDisplayBP/Config/Windows/WindowsEngine.ini b/Templates/TP_nDisplayBP/Config/Windows/WindowsEngine.ini
new file mode 100644
--- /dev/null
+++ b/Templates/TP_nDisplayBP/Config/Windows/WindowsEngine.ini
@@ -0,0 +1,6 @@
+[/Script/Engine.Engine]
+GameEngine=/Script/DisplayCluster.DisplayClusterGameEngine
+UnrealEdEngine=/Script/DisplayClusterEditor.DisplayClusterEditorEngine
+GameViewportClientClassName=/Script/DisplayCluster.DisplayClusterViewportClient
+
+
