Native input data replication + restored PlayerInput to its original state

diff --git a/Engine/Plugins/Runtime/nDisplay/Source/DisplayCluster/Private/Cluster/Controller/DisplayClusterClusterNodeCtrlSlave.h b/Engine/Plugins/Runtime/nDisplay/Source/DisplayCluster/Private/Cluster/Controller/DisplayClusterClusterNodeCtrlSlave.h
--- a/Engine/Plugins/Runtime/nDisplay/Source/DisplayCluster/Private/Cluster/Controller/DisplayClusterClusterNodeCtrlSlave.h
+++ b/Engine/Plugins/Runtime/nDisplay/Source/DisplayCluster/Private/Cluster/Controller/DisplayClusterClusterNodeCtrlSlave.h
@@ -41,7 +41,7 @@ public:
 	virtual void GetSyncData(FDisplayClusterMessage::DataType& SyncData, EDisplayClusterSyncGroup SyncGroup) override;
 	virtual void GetInputData(FDisplayClusterMessage::DataType& InputData) override;
 	virtual void GetEventsData(FDisplayClusterMessage::DataType& EventsData) override;
-	virtual void GetNativeInputData(FDisplayClusterMessage::DataType& EventsData) override;
+	virtual void GetNativeInputData(FDisplayClusterMessage::DataType& NativeInputData) override;
 
 public:
 	//////////////////////////////////////////////////////////////////////////////////////////////
diff --git a/Engine/Plugins/Runtime/nDisplay/Source/DisplayCluster/Private/Cluster/Controller/DisplayClusterNodeCtrlBase.h b/Engine/Plugins/Runtime/nDisplay/Source/DisplayCluster/Private/Cluster/Controller/DisplayClusterNodeCtrlBase.h
--- a/Engine/Plugins/Runtime/nDisplay/Source/DisplayCluster/Private/Cluster/Controller/DisplayClusterNodeCtrlBase.h
+++ b/Engine/Plugins/Runtime/nDisplay/Source/DisplayCluster/Private/Cluster/Controller/DisplayClusterNodeCtrlBase.h
@@ -82,7 +82,7 @@ public:
 	virtual void GetEventsData(FDisplayClusterMessage::DataType& EventsData) override
 	{ }
 
-	virtual void GetNativeInputData(FDisplayClusterMessage::DataType& EventsData) override
+	virtual void GetNativeInputData(FDisplayClusterMessage::DataType& NativeInputData) override
 	{ }
 
 public:
diff --git a/Engine/Plugins/Runtime/nDisplay/Source/DisplayCluster/Private/Game/Classes/Basics/DisplayClusterPlayerInput.cpp b/Engine/Plugins/Runtime/nDisplay/Source/DisplayCluster/Private/Game/Classes/Basics/DisplayClusterPlayerInput.cpp
--- a/Engine/Plugins/Runtime/nDisplay/Source/DisplayCluster/Private/Game/Classes/Basics/DisplayClusterPlayerInput.cpp
+++ b/Engine/Plugins/Runtime/nDisplay/Source/DisplayCluster/Private/Game/Classes/Basics/DisplayClusterPlayerInput.cpp
@@ -9,6 +9,8 @@
 #include "DisplayClusterLog.h"
 #include "DisplayClusterUtils/DisplayClusterTypesConverter.h"
 
+#include "Engine/World.h"
+
 
 UDisplayClusterPlayerInput::UDisplayClusterPlayerInput()
 	: Super()
@@ -20,373 +22,112 @@ void UDisplayClusterPlayerInput::ProcessInputStack(const TArray<UInputComponent*
 	UE_LOG(LogDisplayClusterGame, Verbose, TEXT("Processing input stack..."));
 
 	static IPDisplayClusterClusterManager* const ClusterMgr = GDisplayCluster->GetPrivateClusterMgr();
+
 	if (ClusterMgr && ClusterMgr->IsCluster())
 	{
+		TMap<FString, FString> KeyStates;
+
 		if (ClusterMgr->IsMaster())
 		{
-			// Build master auxiliary data
-			BuildDelegatesMap(InputComponentStack, true);
-			// Start procssing the stack
-			Super::ProcessInputStack(InputComponentStack, DeltaTime, bGamePaused);
+			// Export key states data
+			SerializeKeyStateMap(KeyStates);
+			ClusterMgr->ProvideNativeInputData(KeyStates);
 		}
 		else
 		{
-			// Build slave auxiliary data
-			BuildDelegatesMap(InputComponentStack, false);
-			// Get native input data from master
-			TMap<FString, FString> NativeInputData;
-			ClusterMgr->SyncNativeInput(NativeInputData);
-			// Process native input data got from master
-			ImportNativeInputData(NativeInputData, InputComponentStack);
+			// Import key states data
+			ClusterMgr->SyncNativeInput(KeyStates);
+			DeserializeKeyStateMap(KeyStates);
 		}
 	}
-}
-
-void UDisplayClusterPlayerInput::ProcessDelegatesFilter(const TArray<UInputComponent*>& InputComponentStack, TArray<FAxisDelegateDetails>& AxisDelegates, TArray<FVectorAxisDelegateDetails>& VectorAxisDelegates, TArray<FDelegateDispatchDetails>& NonAxisDelegates)
-{
-	UE_LOG(LogDisplayClusterGame, Verbose, TEXT("Processing input delegates..."));
-
-	static IPDisplayClusterClusterManager* const ClusterMgr = GDisplayCluster->GetPrivateClusterMgr();
 
-	// Serialize native input data
-	TMap<FString, FString> NativeInputData;
-	ExportNativeInputData(NativeInputData, InputComponentStack, AxisDelegates, VectorAxisDelegates, NonAxisDelegates);
-	// Provide data to slave nodes
-	ClusterMgr->ProvideNativeInputData(NativeInputData);
+	Super::ProcessInputStack(InputComponentStack, DeltaTime, bGamePaused);
 }
 
-void UDisplayClusterPlayerInput::BuildDelegatesMap(const TArray<UInputComponent*>& InputComponentStack, bool IsMaster)
+bool UDisplayClusterPlayerInput::SerializeKeyStateMap(TMap<FString, FString>& OutKeyStateMap)
 {
-	CleanDelegateMap(IsMaster);
-
-	// Master related logic
-	if (IsMaster)
+	TMap<FKey, FKeyState>& StateMap = GetKeyStateMap();
+	for (auto it = StateMap.CreateConstIterator(); it; ++it)
 	{
-		for (UInputComponent* IC : InputComponentStack)
-		{
-			// Cache axis delegates
-			for (const FInputAxisBinding& AxisBinding : IC->AxisBindings)
-			{
-				void* Addr1 = (void*)&AxisBinding.AxisDelegate.GetDelegate();
-				void* Addr2 = (void*)&AxisBinding.AxisDelegate.GetDynamicDelegate();
+		const FString KeyName = it->Key.ToString();
 
-				if (!AxisDelegateNames.Contains(Addr1))
-				{
-					AxisDelegateNames.Emplace(Addr1);
-				}
+		FString StrKeyState;
+		StrKeyState.Reserve(2048);
 
-				AxisDelegateNames[Addr1].Emplace(Addr2, AxisBinding.AxisName);
-			}
+		float SendValue = GetWorld()->GetRealTimeSeconds() - it->Value.LastUpDownTransitionTime;
+		check(SendValue > 0.f);
 
-			// Cache gesture delegates
-			for (const FInputGestureBinding& GestureBinding : IC->GestureBindings)
-			{
-				void* Addr1 = (void*)&GestureBinding.GestureDelegate.GetDelegate();
-				void* Addr2 = (void*)&GestureBinding.GestureDelegate.GetDynamicDelegate();
-
-				if (!GestureDelegateNames.Contains(Addr1))
-				{
-					GestureDelegateNames.Emplace(Addr1);
-				}
+		StrKeyState = FString::Printf(TEXT("%s;%s;%s;%s;%s;%s;%s;%s;"),
+			*FDisplayClusterTypesConverter::ToString(it->Value.RawValue),
+			*FDisplayClusterTypesConverter::ToString(it->Value.Value),
+			*FDisplayClusterTypesConverter::ToString(it->Value.LastUpDownTransitionTime),
+			*DisplayClusterHelpers::str::BoolToStr(it->Value.bDown ? true : false, false),
+			*DisplayClusterHelpers::str::BoolToStr(it->Value.bDownPrevious ? true : false, false),
+			*DisplayClusterHelpers::str::BoolToStr(it->Value.bConsumed ? true : false, false),
+			*FDisplayClusterTypesConverter::ToString(it->Value.SampleCountAccumulator),
+			*FDisplayClusterTypesConverter::ToString(it->Value.RawValueAccumulator));
 
-				GestureDelegateNames[Addr1].Emplace(Addr2, GestureBinding.GestureKey.GetFName());
-			}
-
-			// Cache vector axis delegates
-			for (const FInputVectorAxisBinding& VectorAxisBinding : IC->VectorAxisBindings)
-			{
-				void* Addr1 = (void*)&VectorAxisBinding.AxisDelegate.GetDelegate();
-				void* Addr2 = (void*)&VectorAxisBinding.AxisDelegate.GetDynamicDelegate();
-
-				if (!VectorAxisDelegateNames.Contains(Addr1))
-				{
-					VectorAxisDelegateNames.Emplace(Addr1);
-				}
-
-				VectorAxisDelegateNames[Addr1].Emplace(Addr2, VectorAxisBinding.AxisKey.GetFName());
-			}
-		}
-	}
-	// Slave related logic
-	else
-	{
-		for (UInputComponent* IC : InputComponentStack)
+		for (int i = 0; i < IE_MAX; ++i)
 		{
-			// Cache axis delegates
-			for (const FInputAxisBinding& AxisBinding : IC->AxisBindings)
-			{
-				AxisDelegateInstances.Emplace(AxisBinding.AxisName, &AxisBinding.AxisDelegate);
-			}
-
-			// Cache action delegates
-			const int32 ActionBindingsAmount = IC->GetNumActionBindings();
-			for (int i = 0; i < ActionBindingsAmount; ++i)
-			{
-				const FInputActionBinding& ActionBinding = IC->GetActionBinding(i);
-				const FName ActionName = ActionBinding.GetActionName();
-				if (!ActionDelegateInstances.Contains(ActionName))
-				{
-					ActionDelegateInstances.Emplace(ActionName);
-				}
-
-				const EInputEvent InputEvent = ActionBinding.KeyEvent.GetValue();
-				if (!ActionDelegateInstances[ActionName].Contains(InputEvent))
-				{
-					ActionDelegateInstances[ActionName].Emplace(InputEvent);
-				}
-
-				ActionDelegateInstances[ActionName][InputEvent] = &ActionBinding.ActionDelegate;
-			}
-
-			// Cache gesture delegates
-			for (const FInputGestureBinding& GestureBinding : IC->GestureBindings)
-			{
-				const FName GestureKey = GestureBinding.GestureKey.GetFName();
-				GestureDelegateInstances.Emplace(GestureKey, &GestureBinding.GestureDelegate);
-			}
-
-			// Cache vector axis delegates
-			for (const FInputVectorAxisBinding& VectorAxisBinding : IC->VectorAxisBindings)
-			{
-				const FName AxisKey = VectorAxisBinding.AxisKey.GetFName();
-				VectorAxisDelegateInstances.Emplace(AxisKey, &VectorAxisBinding.AxisDelegate);
-			}
+			StrKeyState += FString::Printf(TEXT("%s;"), *DisplayClusterHelpers::str::ArrayToStr(it->Value.EventCounts[i], FString(","), false));
 		}
-	}
 
-	// For any node type
-	for (UInputComponent* IC : InputComponentStack)
-	{
-		// Cache touch delegates
-		for (const FInputTouchBinding& TouchBinding : IC->TouchBindings)
+		for (int i = 0; i < IE_MAX; ++i)
 		{
-			const EInputEvent InputEvent = TouchBinding.KeyEvent.GetValue();
-			TouchDelegateInstances.Emplace(InputEvent, &TouchBinding.TouchDelegate);
+			StrKeyState += FString::Printf(TEXT("%s;"), *DisplayClusterHelpers::str::ArrayToStr(it->Value.EventAccumulator[i], FString(","), false));
 		}
-	}
-}
-
-void UDisplayClusterPlayerInput::CleanDelegateMap(bool IsMaster)
-{
-	if (IsMaster)
-	{
-		AxisDelegateNames.Reset();
-		GestureDelegateNames.Reset();
-		VectorAxisDelegateNames.Reset();
-	}
-	else
-	{
-		AxisDelegateInstances.Reset();
-		ActionDelegateInstances.Reset();
-		GestureDelegateInstances.Reset();
-		VectorAxisDelegateInstances.Reset();
-	}
-
-	TouchDelegateInstances.Reset();
-}
-
-void UDisplayClusterPlayerInput::ExportNativeInputData(TMap<FString, FString>& InputData, const TArray<UInputComponent*>& InputComponentStack, TArray<FAxisDelegateDetails>& AxisDelegates, TArray<FVectorAxisDelegateDetails>& VectorAxisDelegates, TArray<FDelegateDispatchDetails>& NonAxisDelegates)
-{
-	UE_LOG(LogDisplayClusterGame, Verbose, TEXT("Exporting native input data..."));
-
-	InputData.Reset();
-
-	InputData.Emplace(FString("action"), SerializeActionData(InputComponentStack, NonAxisDelegates));
-	InputData.Emplace(FString("axis"),   SerializeAxisData(InputComponentStack,   AxisDelegates));
-	InputData.Emplace(FString("vector"), SerializeVectorAxisData(InputComponentStack, VectorAxisDelegates));
-
-	UE_LOG(LogDisplayClusterGame, Verbose, TEXT("[EXP] Native output action: %s"), *InputData[FString("action")]);
-	UE_LOG(LogDisplayClusterGame, Verbose, TEXT("[EXP] Native output axis  : %s"), *InputData[FString("axis")]);
-	UE_LOG(LogDisplayClusterGame, Verbose, TEXT("[EXP] Native output vector: %s"), *InputData[FString("vector")]);
-}
-
-void UDisplayClusterPlayerInput::ImportNativeInputData(const TMap<FString, FString>& InputData, const TArray<UInputComponent*>& InputComponentStack)
-{
-	UE_LOG(LogDisplayClusterGame, Verbose, TEXT("Importing native input data..."));
-
-	check(InputData.Num() == 3);
-
-	UE_LOG(LogDisplayClusterGame, Verbose, TEXT("[IMP] Native output action: %s"), *InputData[FString("action")]);
-	UE_LOG(LogDisplayClusterGame, Verbose, TEXT("[IMP] Native output axis  : %s"), *InputData[FString("axis")]);
-	UE_LOG(LogDisplayClusterGame, Verbose, TEXT("[IMP] Native output vector: %s"), *InputData[FString("vector")]);
 
-	DeserializeAndProcessActionInput    (InputData[FString("action")]);
-	DeserializeAndProcessAxisInput      (InputData[FString("axis")]);
-	DeserializeAndProcessVectorAxisInput(InputData[FString("vector")]);
-}
-
-FString UDisplayClusterPlayerInput::SerializeAxisData(const TArray<UInputComponent*>& InputComponentStack, const TArray<FAxisDelegateDetails>& AxisDelegates)
-{
-	FString OutData;
-	OutData.Reserve(512);
-
-	for (const FAxisDelegateDetails& Details : AxisDelegates)
-	{
-		const void* const Addr1 = (void*)&Details.Delegate.GetDelegate();
-		const void* const Addr2 = (void*)&Details.Delegate.GetDynamicDelegate();
-
-		if (AxisDelegateNames.Contains(Addr1) && AxisDelegateNames[Addr1].Contains(Addr2))
-		{
-			OutData += FString::Printf(TEXT("%s=%f;"), *AxisDelegateNames[Addr1][Addr2].ToString(), Details.Value);
-		}
+		OutKeyStateMap.Emplace(KeyName, StrKeyState);
 	}
 
-	return OutData;
+	return true;
 }
 
-FString UDisplayClusterPlayerInput::SerializeVectorAxisData(const TArray<UInputComponent*>& InputComponentStack, const TArray<FVectorAxisDelegateDetails>& VectorAxisDelegates)
+bool UDisplayClusterPlayerInput::DeserializeKeyStateMap(const TMap<FString, FString>& InKeyStateMap)
 {
-	FString OutData;
-	OutData.Reserve(512);
+	TMap<FKey, FKeyState>& StateMap = GetKeyStateMap();
+	
+	// Reset local key state map
+	StateMap.Reset();
 
-	for (const FVectorAxisDelegateDetails& Details : VectorAxisDelegates)
+	int idx = 0;
+	for (auto it = InKeyStateMap.CreateConstIterator(); it; ++it)
 	{
-		void* Addr1 = (void*)&Details.Delegate.GetDelegate();
-		void* Addr2 = (void*)&Details.Delegate.GetDynamicDelegate();
+		UE_LOG(LogDisplayClusterGame, Log, TEXT("Input data [%d]: %s = %s"), idx, *it->Key, *it->Value);
+		++idx;
 
-		if (VectorAxisDelegateNames.Contains(Addr1) && VectorAxisDelegateNames[Addr1].Contains(Addr2))
-		{
-			OutData += FString::Printf(TEXT("%s=%s;"), *VectorAxisDelegateNames[Addr1][Addr2].ToString(), *FDisplayClusterTypesConverter::ToString(Details.Value));
-		}
-	}
+		FKey Key(*it->Key);
+		FKeyState KeyState;
 
-	return OutData;
-}
+		TArray<FString> Fields = DisplayClusterHelpers::str::StrToArray<FString>(it->Value, FString(";"), false);
 
-FString UDisplayClusterPlayerInput::SerializeActionData(const TArray<UInputComponent*>& InputComponentStack, const TArray<FDelegateDispatchDetails>& NonAxisDelegates)
-{
-	FString OutData;
-	OutData.Reserve(512);
+		KeyState.RawValue                 = FDisplayClusterTypesConverter::FromString<FVector>(Fields[0]);
+		KeyState.Value                    = FDisplayClusterTypesConverter::FromString<FVector>(Fields[1]);
+		KeyState.LastUpDownTransitionTime = FDisplayClusterTypesConverter::FromString<float>(Fields[2]);
+		KeyState.bDown                    = FDisplayClusterTypesConverter::FromString<bool>(Fields[3]) ? 1 : 0;
+		KeyState.bDownPrevious            = FDisplayClusterTypesConverter::FromString<bool>(Fields[4]) ? 1 : 0;
+		KeyState.bConsumed                = FDisplayClusterTypesConverter::FromString<bool>(Fields[5]) ? 1 : 0;
+		KeyState.SampleCountAccumulator   = FDisplayClusterTypesConverter::FromString<uint8>(Fields[6]);
+		KeyState.RawValueAccumulator      = FDisplayClusterTypesConverter::FromString<FVector>(Fields[7]);
 
-	for (const FDelegateDispatchDetails& Details : NonAxisDelegates)
-	{
-		if (Details.SourceAction)
+		for (int i = 0; i < IE_MAX; ++i)
 		{
-			const FString ActionName = Details.SourceAction->GetActionName().ToString();
-			if (Details.ActionDelegate.IsBound())
+			const FString& EventCountsStr = Fields[8 + i];
+			if (!EventCountsStr.IsEmpty())
 			{
-				OutData += FString::Printf(TEXT("key=%s|%s|%d;"), *ActionName, *Details.Chord.Key.GetFName().ToString(), (int)Details.SourceAction->KeyEvent.GetValue());
+				KeyState.EventCounts[i] = DisplayClusterHelpers::str::StrToArray<uint32>(EventCountsStr, FString(","));
 			}
-			else if (Details.TouchDelegate.IsBound())
-			{
-				OutData += FString::Printf(TEXT("touch=%d|%s|%d;"), (int)Details.KeyEvent.GetValue(), *FDisplayClusterTypesConverter::ToString(Details.TouchLocation), (int)Details.FingerIndex);
-			}
-			else if (Details.GestureDelegate.IsBound())
-			{
-				const void* const Addr1 = (void*)&Details.GestureDelegate.GetDelegate();
-				const void* const Addr2 = (void*)&Details.GestureDelegate.GetDynamicDelegate();
-
-				if (GestureDelegateNames.Contains(Addr1) && GestureDelegateNames[Addr1].Contains(Addr2))
-				{
-					OutData += FString::Printf(TEXT("gesture=%s|%f;"), *GestureDelegateNames[Addr1][Addr2].ToString(), Details.GestureValue);
-				}
-			}
-		}
-	}
-
-	return OutData;
-}
 
-void UDisplayClusterPlayerInput::DeserializeAndProcessAxisInput(const FString& Data)
-{
-	const TArray<FString> AxisRecords = DisplayClusterHelpers::str::StrToArray<FString>(Data, FString(";"));
-	for (const FString& AxisRecord : AxisRecords)
-	{
-		if (AxisRecord.Len() > 0)
-		{
-			const TArray<FString> AxisParams = DisplayClusterHelpers::str::StrToArray<FString>(AxisRecord, FString(TEXT("=")));
-			check(AxisParams.Num() == 2);
-
-			const FName AxisName = *AxisParams[0];
-			if (AxisDelegateInstances.Contains(AxisName))
+			const FString& EventAccumulatorStr = Fields[8 + IE_MAX + i];
+			if (!EventAccumulatorStr.IsEmpty())
 			{
-				const float AxisValue = FDisplayClusterTypesConverter::FromString<float>(AxisParams[1]);
-				UE_LOG(LogDisplayClusterGame, Verbose, TEXT("Native input - axis %s, value %f"), *AxisName.ToString(), AxisValue);
-				AxisDelegateInstances[AxisName]->Execute(AxisValue);
+				KeyState.EventAccumulator[i] = DisplayClusterHelpers::str::StrToArray<uint32>(Fields[8 + IE_MAX + i], FString(","));
 			}
 		}
-	}
-}
 
-void UDisplayClusterPlayerInput::DeserializeAndProcessVectorAxisInput(const FString& Data)
-{
-	const TArray<FString> VectorRecords = DisplayClusterHelpers::str::StrToArray<FString>(Data, FString(";"));
-	for (const FString& VectorRecord : VectorRecords)
-	{
-		if (VectorRecord.Len() > 0)
-		{
-			const TArray<FString> VectorParams = DisplayClusterHelpers::str::StrToArray<FString>(VectorRecord, FString(TEXT("=")));
-			check(VectorParams.Num() == 2);
-
-			FName VectorAxisName = *VectorParams[0];
-			if (VectorAxisDelegateInstances.Contains(VectorAxisName))
-			{
-				FVector VectorValue = FDisplayClusterTypesConverter::FromString<FVector>(VectorParams[1]);
-				UE_LOG(LogDisplayClusterGame, Verbose, TEXT("Native input - vector axis %s, value %f"), *VectorAxisName.ToString(), *VectorValue.ToString());
-				VectorAxisDelegateInstances[VectorAxisName]->Execute(VectorValue);
-			}
-		}
+		// Add incoming data to the local map
+		StateMap.Add(Key, KeyState);
 	}
-}
-
-void UDisplayClusterPlayerInput::DeserializeAndProcessActionInput(const FString& Data)
-{
-	const TArray<FString> ActionRecords = DisplayClusterHelpers::str::StrToArray<FString>(Data, FString(";"));
-	for (const FString& ActionRecord : ActionRecords)
-	{
-		if(ActionRecord.Len() > 0)
-		{
-			const TArray<FString> ActionParams = DisplayClusterHelpers::str::StrToArray<FString>(ActionRecord, FString(TEXT("=")));
-			check(ActionParams.Num() == 2);
-			const TArray<FString> ActionParamsEx = DisplayClusterHelpers::str::StrToArray<FString>(ActionParams[1], FString(TEXT("|")));
-
-			if (ActionParams[0].Equals(FString("key"), ESearchCase::IgnoreCase))
-			{
-				check(ActionParamsEx.Num() == 3);
-
-				const FName ActionName = *ActionParamsEx[0];
-				if (ActionDelegateInstances.Contains(ActionName))
-				{
-					const EInputEvent InputEvent = (EInputEvent)FDisplayClusterTypesConverter::FromString<int>(ActionParamsEx[2]);
-					if (ActionDelegateInstances[ActionName].Contains(InputEvent))
-					{
-						const FKey ActionKey(*ActionParamsEx[1]);
 
-						UE_LOG(LogDisplayClusterGame, Verbose, TEXT("Native input - key action %s, key %s, input event %d"), *ActionName.ToString(), *ActionKey.ToString(), (int)InputEvent);
-
-						ActionDelegateInstances[ActionName][InputEvent]->Execute(ActionKey);
-					}
-				}
-			}
-			else if (ActionParams[0].Equals(FString("touch"), ESearchCase::IgnoreCase))
-			{
-				check(ActionParamsEx.Num() == 3);
-
-				const EInputEvent InputEvent = (EInputEvent)FDisplayClusterTypesConverter::FromString<int>(ActionParamsEx[0]);
-				if (TouchDelegateInstances.Contains(InputEvent))
-				{
-					const FVector           TouchLocation = FDisplayClusterTypesConverter::FromString<FVector>(ActionParamsEx[1]);
-					const ETouchIndex::Type FingerIndex   = (ETouchIndex::Type)FDisplayClusterTypesConverter::FromString<int>(ActionParamsEx[2]);
-
-					UE_LOG(LogDisplayClusterGame, Verbose, TEXT("Native input - touch action, finger index %d, touch location %s, input event %d"), (int)FingerIndex, *TouchLocation.ToString(), (int)InputEvent);
-
-					TouchDelegateInstances[InputEvent]->Execute(FingerIndex, TouchLocation);
-				}
-			}
-			else if (ActionParams[0].Equals(FString("gesture"), ESearchCase::IgnoreCase))
-			{
-				check(ActionParamsEx.Num() == 2);
-
-				const FName GestureKey = *ActionParamsEx[0];
-				if (GestureDelegateInstances.Contains(GestureKey))
-				{
-					const float GestureValue = FDisplayClusterTypesConverter::FromString<float>(ActionParamsEx[1]);
-
-					UE_LOG(LogDisplayClusterGame, Verbose, TEXT("Native input - gesture action, key %s, value %f"), *GestureKey.ToString(), GestureValue);
-
-					GestureDelegateInstances[GestureKey]->Execute(GestureValue);
-				}
-			}
-		}
-	}
+	return true;
 }
diff --git a/Engine/Plugins/Runtime/nDisplay/Source/DisplayCluster/Private/Game/Classes/Scene/DisplayClusterRootComponent.cpp b/Engine/Plugins/Runtime/nDisplay/Source/DisplayCluster/Private/Game/Classes/Scene/DisplayClusterRootComponent.cpp
--- a/Engine/Plugins/Runtime/nDisplay/Source/DisplayCluster/Private/Game/Classes/Scene/DisplayClusterRootComponent.cpp
+++ b/Engine/Plugins/Runtime/nDisplay/Source/DisplayCluster/Private/Game/Classes/Scene/DisplayClusterRootComponent.cpp
@@ -48,6 +48,15 @@ void UDisplayClusterRootComponent::BeginPlay()
 	if (OperationMode == EDisplayClusterOperationMode::Cluster ||
 		OperationMode == EDisplayClusterOperationMode::Editor)
 	{
+		// Read native input synchronization settings
+		IPDisplayClusterConfigManager* const ConfigMgr = GDisplayCluster->GetPrivateConfigMgr();
+		if (ConfigMgr)
+		{
+			FDisplayClusterConfigGeneral CfgGeneral = ConfigMgr->GetConfigGeneral();
+			NativeInputSyncPolicy = CfgGeneral.NativeInputSyncPolicy;
+			UE_LOG(LogDisplayClusterGame, Log, TEXT("Native input sync policy: %d"), NativeInputSyncPolicy);
+		}
+
 		// Build nDisplay hierarchy
 		InitializeHierarchy();
 
@@ -112,21 +121,25 @@ bool UDisplayClusterRootComponent::InitializeHierarchy()
 {
 	DISPLAY_CLUSTER_FUNC_TRACE(LogDisplayClusterGame);
 
-	// Add replication components
-	const TSet<UActorComponent*> ActorComponents = GetOwner()->GetComponents();
-	int CompIdx = 0;
-	for (UActorComponent* Component : ActorComponents)
+	// Use sync components if no native input synchronization activated
+	if (NativeInputSyncPolicy == 0)
 	{
-		USceneComponent* const SceneComp = Cast<USceneComponent>(Component);
-		if (SceneComp)
+		// Add replication components
+		const TSet<UActorComponent*> ActorComponents = GetOwner()->GetComponents();
+		int CompIdx = 0;
+		for (UActorComponent* Component : ActorComponents)
 		{
-			const FString CompName = FString::Printf(TEXT("sync_comp_%d"), CompIdx++);
-
-			UDisplayClusterSceneComponentSyncParent* NewSyncComp = NewObject<UDisplayClusterSceneComponentSyncParent>(Component->GetOwner(), FName(*CompName));
-			if (NewSyncComp)
+			USceneComponent* const SceneComp = Cast<USceneComponent>(Component);
+			if (SceneComp)
 			{
-				NewSyncComp->AttachToComponent(SceneComp, FAttachmentTransformRules(EAttachmentRule::KeepRelative, false));
-				NewSyncComp->RegisterComponent();
+				const FString CompName = FString::Printf(TEXT("sync_comp_%d"), CompIdx++);
+
+				UDisplayClusterSceneComponentSyncParent* NewSyncComp = NewObject<UDisplayClusterSceneComponentSyncParent>(Component->GetOwner(), FName(*CompName));
+				if (NewSyncComp)
+				{
+					NewSyncComp->AttachToComponent(SceneComp, FAttachmentTransformRules(EAttachmentRule::KeepRelative, false));
+					NewSyncComp->RegisterComponent();
+				}
 			}
 		}
 	}
@@ -164,22 +177,17 @@ bool UDisplayClusterRootComponent::InitializeHierarchy()
 
 bool UDisplayClusterRootComponent::InitializeExtra()
 {
-	// Read native input synchronization settings
-	IPDisplayClusterConfigManager* const ConfigMgr = GDisplayCluster->GetPrivateConfigMgr();
-	if (ConfigMgr)
-	{
-		FDisplayClusterConfigGeneral CfgGeneral = ConfigMgr->GetConfigGeneral();
-		NativeInputSyncPolicy = CfgGeneral.NativeInputSyncPolicy;
-		UE_LOG(LogDisplayClusterGame, Log, TEXT("Native input sync policy: %d"), NativeInputSyncPolicy);
-	}
-
-	// Optionally activate native input synchronization
-	if (NativeInputSyncPolicy == 1)
+	EDisplayClusterOperationMode OpMode = GDisplayCluster->GetOperationMode();
+	if (OpMode == EDisplayClusterOperationMode::Cluster)
 	{
-		APlayerController* const PlayerController = GetWorld()->GetFirstPlayerController();
-		if (PlayerController)
+		// Optionally activate native input synchronization
+		if (NativeInputSyncPolicy == 1)
 		{
-			PlayerController->PlayerInput = NewObject<UDisplayClusterPlayerInput>(PlayerController);
+			APlayerController* const PlayerController = GetWorld()->GetFirstPlayerController();
+			if (PlayerController)
+			{
+				PlayerController->PlayerInput = NewObject<UDisplayClusterPlayerInput>(PlayerController);
+			}
 		}
 	}
 
diff --git a/Engine/Plugins/Runtime/nDisplay/Source/DisplayCluster/Public/DisplayClusterPlayerInput.h b/Engine/Plugins/Runtime/nDisplay/Source/DisplayCluster/Public/DisplayClusterPlayerInput.h
--- a/Engine/Plugins/Runtime/nDisplay/Source/DisplayCluster/Public/DisplayClusterPlayerInput.h
+++ b/Engine/Plugins/Runtime/nDisplay/Source/DisplayCluster/Public/DisplayClusterPlayerInput.h
@@ -21,34 +21,7 @@ public:
 	/** Process the frame's input events given the current input component stack. */
 	virtual void ProcessInputStack(const TArray<UInputComponent*>& InputComponentStack, const float DeltaTime, const bool bGamePaused) override;
 
-protected:
-	virtual void ProcessDelegatesFilter(const TArray<UInputComponent*>& InputComponentStack, TArray<FAxisDelegateDetails>& AxisDelegates, TArray<FVectorAxisDelegateDetails>& VectorAxisDelegates, TArray<FDelegateDispatchDetails>& NonAxisDelegates) override;
-
 private:
-	void ExportNativeInputData(TMap<FString, FString>& InputData, const TArray<UInputComponent*>& InputComponentStack, TArray<FAxisDelegateDetails>& AxisDelegates, TArray<FVectorAxisDelegateDetails>& VectorAxisDelegates, TArray<FDelegateDispatchDetails>& NonAxisDelegates);
-	void ImportNativeInputData(const TMap<FString, FString>& InputData, const TArray<UInputComponent*>& InputComponentStack);
-
-	void CleanDelegateMap(bool IsMaster);
-	void BuildDelegatesMap(const TArray<UInputComponent*>& InputComponentStack, bool IsMaster);
-
-	FString SerializeAxisData  (const TArray<UInputComponent*>& InputComponentStack, const TArray<FAxisDelegateDetails>& AxisDelegates);
-	FString SerializeVectorAxisData(const TArray<UInputComponent*>& InputComponentStack, const TArray<FVectorAxisDelegateDetails>& VectorAxisDelegates);
-	FString SerializeActionData(const TArray<UInputComponent*>& InputComponentStack, const TArray<FDelegateDispatchDetails>& NonAxisDelegates);
-
-	void DeserializeAndProcessAxisInput  (const FString& Data);
-	void DeserializeAndProcessVectorAxisInput(const FString& Data);
-	void DeserializeAndProcessActionInput(const FString& Data);
-
-private:
-	TMap<void*, TMap<void*, FName>> AxisDelegateNames;
-	TMap<FName, const FInputAxisUnifiedDelegate*> AxisDelegateInstances;
-
-	TMap<FName, TMap<EInputEvent, const FInputActionUnifiedDelegate*>> ActionDelegateInstances;
-	TMap<EInputEvent, const FInputTouchUnifiedDelegate*> TouchDelegateInstances;
-
-	TMap<void*, TMap<void*, FName>> GestureDelegateNames;
-	TMap<FName, const FInputGestureUnifiedDelegate*> GestureDelegateInstances;
-
-	TMap<void*, TMap<void*, FName>> VectorAxisDelegateNames;
-	TMap<FName, const FInputVectorAxisUnifiedDelegate*> VectorAxisDelegateInstances;
+	bool SerializeKeyStateMap(TMap<FString, FString>& OutKeyStateMap);
+	bool DeserializeKeyStateMap(const TMap<FString, FString>& InKeyStateMap);
 };
diff --git a/Engine/Plugins/Runtime/nDisplay/Source/DisplayClusterUtils/DisplayClusterCommonHelpers.h b/Engine/Plugins/Runtime/nDisplay/Source/DisplayClusterUtils/DisplayClusterCommonHelpers.h
--- a/Engine/Plugins/Runtime/nDisplay/Source/DisplayClusterUtils/DisplayClusterCommonHelpers.h
+++ b/Engine/Plugins/Runtime/nDisplay/Source/DisplayClusterUtils/DisplayClusterCommonHelpers.h
@@ -47,21 +47,29 @@ namespace DisplayClusterHelpers
 		// Parses string items separated by specified separator into array
 		// Example: item1, item2,item3  ,  item4 => {item1, item2, item3, item4}
 		template<typename T>
-		static TArray<T> StrToArray(const FString& InData, const FString& InSeparator)
+		static TArray<T> StrToArray(const FString& InData, const FString& InSeparator, bool bCullEmpty = true)
 		{
 			TArray<T> OutData;
+			TArray<FString> TempData;
+			
+			InData.ParseIntoArray(TempData, *InSeparator, bCullEmpty);
 
-			FString TmpR = InData;
-			FString TmpL;
-
-			while (TmpR.Split(InSeparator, &TmpL, &TmpR, ESearchCase::IgnoreCase, ESearchDir::FromStart))
+			for (FString& Item : TempData)
 			{
-				TrimStringValue(TmpL, false);
-				OutData.Add(FDisplayClusterTypesConverter::template FromString<T>(TmpL));
+				TrimStringValue(Item, false);
 			}
 
-			TrimStringValue(TmpR, false);
-			OutData.Add(FDisplayClusterTypesConverter::template FromString<T>(TmpR));
+			for (const FString& Item : TempData)
+			{
+				if (!bCullEmpty && Item.IsEmpty())
+				{
+					OutData.Add(T());
+				}
+				else
+				{
+					OutData.Add(FDisplayClusterTypesConverter::template FromString<T>(Item));
+				}
+			}
 
 			return OutData;
 		}
diff --git a/Engine/Source/Runtime/Engine/Classes/GameFramework/PlayerInput.h b/Engine/Source/Runtime/Engine/Classes/GameFramework/PlayerInput.h
--- a/Engine/Source/Runtime/Engine/Classes/GameFramework/PlayerInput.h
+++ b/Engine/Source/Runtime/Engine/Classes/GameFramework/PlayerInput.h
@@ -12,7 +12,6 @@
 #include "UObject/ObjectMacros.h"
 #include "UObject/Object.h"
 #include "InputCoreTypes.h"
-#include "Components/InputComponent.h"
 #include "UObject/ScriptMacros.h"
 #include "Framework/Commands/InputChord.h"
 #include "GestureRecognizer.h"
@@ -624,75 +623,9 @@ public:
 	static const TArray<FInputAxisKeyMapping>& GetEngineDefinedAxisMappings() { return EngineDefinedAxisMappings; }
 
 protected:
-	/** Runtime struct that gathers up the different kinds of delegates that might be issued */
-	struct FDelegateDispatchDetails
-	{
-		uint32 EventIndex;
-		uint32 FoundIndex;
-
-		FInputActionUnifiedDelegate ActionDelegate;
-		const FInputActionBinding* SourceAction;
-		FInputChord Chord;
-		TEnumAsByte<EInputEvent> KeyEvent;
-
-		FInputTouchUnifiedDelegate TouchDelegate;
-		FVector TouchLocation;
-		ETouchIndex::Type FingerIndex;
-
-		FInputGestureUnifiedDelegate GestureDelegate;
-		float GestureValue;
-
-		FDelegateDispatchDetails(const uint32 InEventIndex, const uint32 InFoundIndex, const FInputChord& InChord, FInputActionUnifiedDelegate InDelegate, const EInputEvent InKeyEvent, const FInputActionBinding* InSourceAction = NULL)
-			: EventIndex(InEventIndex)
-			, FoundIndex(InFoundIndex)
-			, ActionDelegate(MoveTemp(InDelegate))
-			, SourceAction(InSourceAction)
-			, Chord(InChord)
-			, KeyEvent(InKeyEvent)
-		{}
-
-		FDelegateDispatchDetails(const uint32 InEventIndex, const uint32 InFoundIndex, FInputTouchUnifiedDelegate InDelegate, const FVector InLocation, const ETouchIndex::Type InFingerIndex)
-			: EventIndex(InEventIndex)
-			, FoundIndex(InFoundIndex)
-			, TouchDelegate(MoveTemp(InDelegate))
-			, TouchLocation(InLocation)
-			, FingerIndex(InFingerIndex)
-		{}
-
-		FDelegateDispatchDetails(const uint32 InEventIndex, const uint32 InFoundIndex, FInputGestureUnifiedDelegate InDelegate, const float InValue)
-			: EventIndex(InEventIndex)
-			, FoundIndex(InFoundIndex)
-			, GestureDelegate(MoveTemp(InDelegate))
-			, GestureValue(InValue)
-		{}
-	};
-
-	// We collect axis contributions by delegate, so we can sum up 
-	// contributions from multiple bindings.
-	struct FAxisDelegateDetails
-	{
-		FInputAxisUnifiedDelegate Delegate;
-		float Value;
+	TMap<FKey, FKeyState>& GetKeyStateMap()
+	{ return KeyStateMap; }
 
-		FAxisDelegateDetails(FInputAxisUnifiedDelegate InDelegate, const float InValue)
-			: Delegate(MoveTemp(InDelegate))
-			, Value(InValue)
-		{
-		}
-	};
-	struct FVectorAxisDelegateDetails
-	{
-		FInputVectorAxisUnifiedDelegate Delegate;
-		FVector Value;
-
-		FVectorAxisDelegateDetails(FInputVectorAxisUnifiedDelegate InDelegate, const FVector InValue)
-			: Delegate(MoveTemp(InDelegate))
-			, Value(InValue)
-		{
-		}
-	};
-
-protected:
 	/**
 	* Given raw keystate value of a vector axis, returns the "massaged" value. Override for any custom behavior,
 	* such as input changes dependent on a particular game state.
@@ -705,9 +638,6 @@ protected:
 	*/
 	virtual float MassageAxisInput(FKey Key, float RawValue);
 
-	virtual void ProcessDelegatesFilter(const TArray<UInputComponent*>& InputComponentStack, TArray<FAxisDelegateDetails>& AxisDelegates, TArray<FVectorAxisDelegateDetails>& VectorAxisDelegates, TArray<FDelegateDispatchDetails>& NonAxisDelegates)
-	{ }
-
 private:
 
 	/** Process non-axes keystates */
diff --git a/Engine/Source/Runtime/Engine/Private/UserInterface/PlayerInput.cpp b/Engine/Source/Runtime/Engine/Private/UserInterface/PlayerInput.cpp
--- a/Engine/Source/Runtime/Engine/Private/UserInterface/PlayerInput.cpp
+++ b/Engine/Source/Runtime/Engine/Private/UserInterface/PlayerInput.cpp
@@ -6,6 +6,7 @@
 
 #include "GameFramework/PlayerInput.h"
 #include "Misc/CommandLine.h"
+#include "Components/InputComponent.h"
 #include "Misc/App.h"
 #include "UObject/UObjectHash.h"
 #include "UObject/UObjectIterator.h"
@@ -32,6 +33,48 @@ const TArray<FInputAxisKeyMapping> UPlayerInput::NoAxisMappings;
 TArray<FInputActionKeyMapping> UPlayerInput::EngineDefinedActionMappings;
 TArray<FInputAxisKeyMapping> UPlayerInput::EngineDefinedAxisMappings;
 
+/** Runtime struct that gathers up the different kinds of delegates that might be issued */
+struct FDelegateDispatchDetails
+{
+	uint32 EventIndex;
+	uint32 FoundIndex;
+
+	FInputActionUnifiedDelegate ActionDelegate;
+	const FInputActionBinding* SourceAction;
+	FInputChord Chord;
+	TEnumAsByte<EInputEvent> KeyEvent;
+
+	FInputTouchUnifiedDelegate TouchDelegate;
+	FVector TouchLocation;
+	ETouchIndex::Type FingerIndex;
+
+	FInputGestureUnifiedDelegate GestureDelegate;
+	float GestureValue;
+
+	FDelegateDispatchDetails(const uint32 InEventIndex, const uint32 InFoundIndex, const FInputChord& InChord, FInputActionUnifiedDelegate InDelegate, const EInputEvent InKeyEvent, const FInputActionBinding* InSourceAction = NULL)
+		: EventIndex(InEventIndex)
+		, FoundIndex(InFoundIndex)
+		, ActionDelegate(MoveTemp(InDelegate))
+		, SourceAction(InSourceAction)
+		, Chord(InChord)
+		, KeyEvent(InKeyEvent)
+	{}
+
+	FDelegateDispatchDetails(const uint32 InEventIndex, const uint32 InFoundIndex, FInputTouchUnifiedDelegate InDelegate, const FVector InLocation, const ETouchIndex::Type InFingerIndex)
+		: EventIndex(InEventIndex)
+		, FoundIndex(InFoundIndex)
+		, TouchDelegate(MoveTemp(InDelegate))
+		, TouchLocation(InLocation)
+		, FingerIndex(InFingerIndex)
+	{}
+
+	FDelegateDispatchDetails(const uint32 InEventIndex, const uint32 InFoundIndex, FInputGestureUnifiedDelegate InDelegate, const float InValue)
+		: EventIndex(InEventIndex)
+		, FoundIndex(InFoundIndex)
+		, GestureDelegate(MoveTemp(InDelegate))
+		, GestureValue(InValue)
+	{}
+};
 
 UPlayerInput::UPlayerInput()
 	: Super()
@@ -944,6 +987,31 @@ void UPlayerInput::ProcessInputStack(const TArray<UInputComponent*>& InputCompon
 {
 	ConditionalBuildKeyMappings();
 
+	// We collect axis contributions by delegate, so we can sum up 
+	// contributions from multiple bindings.
+	struct FAxisDelegateDetails
+	{
+		FInputAxisUnifiedDelegate Delegate;
+		float Value;
+
+		FAxisDelegateDetails(FInputAxisUnifiedDelegate InDelegate, const float InValue)
+			: Delegate(MoveTemp(InDelegate))
+			, Value(InValue)
+		{
+		}
+	};
+	struct FVectorAxisDelegateDetails
+	{
+		FInputVectorAxisUnifiedDelegate Delegate;
+		FVector Value;
+
+		FVectorAxisDelegateDetails(FInputVectorAxisUnifiedDelegate InDelegate, const FVector InValue)
+			: Delegate(MoveTemp(InDelegate))
+			, Value(InValue)
+		{
+		}
+	};
+
 	static TArray<FAxisDelegateDetails> AxisDelegates;
 	static TArray<FVectorAxisDelegateDetails> VectorAxisDelegates;
 	static TArray<FDelegateDispatchDetails> NonAxisDelegates;
@@ -1250,9 +1318,6 @@ void UPlayerInput::ProcessInputStack(const TArray<UInputComponent*>& InputCompon
 
 	// Dispatch the delegates in the order they occurred
 	NonAxisDelegates.Sort(FDelegateDispatchDetailsSorter());
-
-	ProcessDelegatesFilter(InputComponentStack, AxisDelegates, VectorAxisDelegates, NonAxisDelegates);
-
 	for (const FDelegateDispatchDetails& Details : NonAxisDelegates)
 	{
 		if (Details.ActionDelegate.IsBound())
