Added DisplayClusterRootActor

diff --git a/Engine/Plugins/Runtime/nDisplay/Source/DisplayCluster/Private/Blueprints/DisplayClusterBlueprintAPIImpl.cpp b/Engine/Plugins/Runtime/nDisplay/Source/DisplayCluster/Private/Blueprints/DisplayClusterBlueprintAPIImpl.cpp
--- a/Engine/Plugins/Runtime/nDisplay/Source/DisplayCluster/Private/Blueprints/DisplayClusterBlueprintAPIImpl.cpp
+++ b/Engine/Plugins/Runtime/nDisplay/Source/DisplayCluster/Private/Blueprints/DisplayClusterBlueprintAPIImpl.cpp
@@ -169,14 +169,27 @@ void UDisplayClusterBlueprintAPIImpl::GetLocalViewports(bool IsRTT, TArray<FStri
 // Game API
 //////////////////////////////////////////////////////////////////////////////////////////////
 // Root
-UDisplayClusterRootComponent* UDisplayClusterBlueprintAPIImpl::GetRoot()
+ADisplayClusterRootActor* UDisplayClusterBlueprintAPIImpl::GetRootActor()
 {
 	DISPLAY_CLUSTER_FUNC_TRACE(LogDisplayClusterBlueprint);
 
 	IDisplayClusterGameManager* const Manager = IDisplayCluster::Get().GetGameMgr();
 	if (Manager)
 	{
-		return Manager->GetRoot();
+		return Manager->GetRootActor();
+	}
+
+	return nullptr;
+}
+
+UDisplayClusterRootComponent* UDisplayClusterBlueprintAPIImpl::GetRootComponent()
+{
+	DISPLAY_CLUSTER_FUNC_TRACE(LogDisplayClusterBlueprint);
+
+	IDisplayClusterGameManager* const Manager = IDisplayCluster::Get().GetGameMgr();
+	if (Manager)
+	{
+		return Manager->GetRootComponent();
 	}
 
 	return nullptr;
diff --git a/Engine/Plugins/Runtime/nDisplay/Source/DisplayCluster/Private/Blueprints/DisplayClusterBlueprintAPIImpl.h b/Engine/Plugins/Runtime/nDisplay/Source/DisplayCluster/Private/Blueprints/DisplayClusterBlueprintAPIImpl.h
--- a/Engine/Plugins/Runtime/nDisplay/Source/DisplayCluster/Private/Blueprints/DisplayClusterBlueprintAPIImpl.h
+++ b/Engine/Plugins/Runtime/nDisplay/Source/DisplayCluster/Private/Blueprints/DisplayClusterBlueprintAPIImpl.h
@@ -79,8 +79,11 @@ public:
 	// Game API
 	//////////////////////////////////////////////////////////////////////////////////////////////
 	// Root
-	UFUNCTION(BlueprintCallable, meta = (DisplayName = "Get root"), Category = "DisplayCluster|Game")
-	virtual UDisplayClusterRootComponent* GetRoot() override;
+	UFUNCTION(BlueprintCallable, meta = (DisplayName = "Get root actor"), Category = "DisplayCluster|Game")
+	virtual ADisplayClusterRootActor* GetRootActor() override;
+
+	UFUNCTION(BlueprintCallable, meta = (DisplayName = "Get root component"), Category = "DisplayCluster|Game")
+	virtual UDisplayClusterRootComponent* GetRootComponent() override;
 
 	// Screens
 	UFUNCTION(BlueprintCallable, meta = (DisplayName = "Get screen by ID"), Category = "DisplayCluster|Game")
diff --git a/Engine/Plugins/Runtime/nDisplay/Source/DisplayCluster/Private/Game/Classes/Scene/DisplayClusterRootActor.cpp b/Engine/Plugins/Runtime/nDisplay/Source/DisplayCluster/Private/Game/Classes/Scene/DisplayClusterRootActor.cpp
new file mode 100644
--- /dev/null
+++ b/Engine/Plugins/Runtime/nDisplay/Source/DisplayCluster/Private/Game/Classes/Scene/DisplayClusterRootActor.cpp
@@ -0,0 +1,119 @@
+// Copyright 1998-2019 Epic Games, Inc. All Rights Reserved.
+
+#include "DisplayClusterRootActor.h"
+
+#include "Components/SceneComponent.h"
+#include "GameFramework/PlayerController.h"
+#include "Engine/World.h"
+
+#include "Config/IPDisplayClusterConfigManager.h"
+#include "Config/DisplayClusterConfigTypes.h"
+#include "Misc/DisplayClusterAppExit.h"
+
+#include "DisplayClusterRootComponent.h"
+#include "DisplayClusterSceneComponentSyncParent.h"
+#include "DisplayClusterPlayerInput.h"
+
+#include "DisplayClusterGlobals.h"
+#include "DisplayClusterLog.h"
+
+
+ADisplayClusterRootActor::ADisplayClusterRootActor(const FObjectInitializer& ObjectInitializer)
+	: Super(ObjectInitializer)
+	, bExitOnEsc(true)
+	, bShowProjectionScreens(false)
+	, ProjectionScreensMaterial(nullptr)
+{
+	DISPLAY_CLUSTER_FUNC_TRACE(LogDisplayClusterGame);
+
+	// Root component
+	RootComponent = CreateDefaultSubobject<USceneComponent>(TEXT("RootComponent"));
+
+	// DisplayCluster sync
+	DisplayClusterRootComponent = CreateDefaultSubobject<UDisplayClusterRootComponent>(TEXT("DisplayClusterRoot"));
+	DisplayClusterRootComponent->AttachToComponent(RootComponent, FAttachmentTransformRules(EAttachmentRule::KeepRelative, false));
+
+	PrimaryActorTick.bCanEverTick = true;
+	bFindCameraComponentWhenViewTarget = false;
+	bReplicates = false;
+	SpawnCollisionHandlingMethod = ESpawnActorCollisionHandlingMethod::AdjustIfPossibleButAlwaysSpawn;
+}
+
+void ADisplayClusterRootActor::BeginPlay()
+{
+	DISPLAY_CLUSTER_FUNC_TRACE(LogDisplayClusterGame);
+
+	int32 NativeInputSyncPolicy = 0;
+
+	// Store current operation mode
+	OperationMode = GDisplayCluster->GetOperationMode();
+
+	if (OperationMode == EDisplayClusterOperationMode::Cluster ||
+		OperationMode == EDisplayClusterOperationMode::Editor)
+	{
+		// Read native input synchronization settings
+		IPDisplayClusterConfigManager* const ConfigMgr = GDisplayCluster->GetPrivateConfigMgr();
+		if (ConfigMgr)
+		{
+			FDisplayClusterConfigGeneral CfgGeneral = ConfigMgr->GetConfigGeneral();
+			NativeInputSyncPolicy = CfgGeneral.NativeInputSyncPolicy;
+			UE_LOG(LogDisplayClusterGame, Log, TEXT("Native input sync policy: %d"), NativeInputSyncPolicy);
+		}
+	}
+
+	if (OperationMode == EDisplayClusterOperationMode::Cluster)
+	{
+		// Optionally activate native input synchronization
+		if (NativeInputSyncPolicy == 1)
+		{
+			APlayerController* const PlayerController = GetWorld()->GetFirstPlayerController();
+			if (PlayerController)
+			{
+				PlayerController->PlayerInput = NewObject<UDisplayClusterPlayerInput>(PlayerController);
+			}
+		}
+	}
+
+	Super::BeginPlay();
+}
+
+void ADisplayClusterRootActor::BeginDestroy()
+{
+	DISPLAY_CLUSTER_FUNC_TRACE(LogDisplayClusterGame);
+
+	Super::BeginDestroy();
+}
+
+void ADisplayClusterRootActor::Tick(float DeltaSeconds)
+{
+	DISPLAY_CLUSTER_FUNC_TRACE(LogDisplayClusterGame);
+
+	if (OperationMode == EDisplayClusterOperationMode::Cluster ||
+		OperationMode == EDisplayClusterOperationMode::Editor)
+	{
+		UWorld* const CurWorld = GetWorld();
+		if (CurWorld)
+		{
+			APlayerController* const CurPlayerController = CurWorld->GetFirstPlayerController();
+			if (CurPlayerController)
+			{
+				APlayerCameraManager* const CurPlayerCameraManager = CurPlayerController->PlayerCameraManager;
+				if (CurPlayerCameraManager)
+				{
+					SetActorLocation(CurPlayerCameraManager->GetCameraLocation());
+					SetActorRotation(CurPlayerCameraManager->GetCameraRotation());
+				}
+
+				if (bExitOnEsc)
+				{
+					if (CurPlayerController->WasInputKeyJustPressed(EKeys::Escape))
+					{
+						FDisplayClusterAppExit::ExitApplication(FDisplayClusterAppExit::ExitType::NormalSoft, FString("Exit on ESC requested"));
+					}
+				}
+			}
+		}
+	}
+
+	Super::Tick(DeltaSeconds);
+}
diff --git a/Engine/Plugins/Runtime/nDisplay/Source/DisplayCluster/Private/Game/Classes/Scene/DisplayClusterRootComponent.cpp b/Engine/Plugins/Runtime/nDisplay/Source/DisplayCluster/Private/Game/Classes/Scene/DisplayClusterRootComponent.cpp
--- a/Engine/Plugins/Runtime/nDisplay/Source/DisplayCluster/Private/Game/Classes/Scene/DisplayClusterRootComponent.cpp
+++ b/Engine/Plugins/Runtime/nDisplay/Source/DisplayCluster/Private/Game/Classes/Scene/DisplayClusterRootComponent.cpp
@@ -5,37 +5,28 @@
 #include "Config/DisplayClusterConfigTypes.h"
 #include "Cluster/IPDisplayClusterClusterManager.h"
 #include "Config/IPDisplayClusterConfigManager.h"
-#include "Game/IPDisplayClusterGameManager.h"
-
 #include "Config/DisplayClusterConfigTypes.h"
+#include "Game/IPDisplayClusterGameManager.h"
 
 #include "Camera/CameraComponent.h"
-#include "GameFramework/PlayerController.h"
 
 #include "DisplayClusterCameraComponent.h"
 #include "DisplayClusterSceneComponent.h"
-#include "DisplayClusterSceneComponentSyncParent.h"
 #include "DisplayClusterScreenComponent.h"
-#include "DisplayClusterPlayerInput.h"
 
 #include "DisplayClusterGlobals.h"
 #include "DisplayClusterHelpers.h"
 #include "DisplayClusterLog.h"
 #include "DisplayClusterStrings.h"
 
-#include "Misc/DisplayClusterAppExit.h"
-
 
 UDisplayClusterRootComponent::UDisplayClusterRootComponent(const FObjectInitializer& ObjectInitializer)
 	: USceneComponent(ObjectInitializer)
-	, bExitOnEsc(true)
-	, bShowProjectionScreens(false)
-	, ProjectionScreensMaterial(nullptr)
 {
 	DISPLAY_CLUSTER_FUNC_TRACE(LogDisplayClusterGame);
 
 	// This component settings
-	PrimaryComponentTick.bCanEverTick = true;
+	PrimaryComponentTick.bCanEverTick = false;
 }
 
 void UDisplayClusterRootComponent::BeginPlay()
@@ -43,28 +34,13 @@ void UDisplayClusterRootComponent::BeginPlay()
 	DISPLAY_CLUSTER_FUNC_TRACE(LogDisplayClusterGame);
 
 	// Store current operation mode on BeginPlay (not in the constructor)
-	OperationMode = GDisplayCluster->GetOperationMode();
+	const EDisplayClusterOperationMode OperationMode = GDisplayCluster->GetOperationMode();
 
 	if (OperationMode == EDisplayClusterOperationMode::Cluster ||
 		OperationMode == EDisplayClusterOperationMode::Editor)
 	{
-		// Read native input synchronization settings
-		IPDisplayClusterConfigManager* const ConfigMgr = GDisplayCluster->GetPrivateConfigMgr();
-		if (ConfigMgr)
-		{
-			FDisplayClusterConfigGeneral CfgGeneral = ConfigMgr->GetConfigGeneral();
-			NativeInputSyncPolicy = CfgGeneral.NativeInputSyncPolicy;
-			UE_LOG(LogDisplayClusterGame, Log, TEXT("Native input sync policy: %d"), NativeInputSyncPolicy);
-		}
-
 		// Build nDisplay hierarchy
 		InitializeHierarchy();
-
-		// Extra initialization
-		InitializeExtra();
-
-		// Set this root as currently active
-		GDisplayCluster->GetPrivateGameMgr()->SetRoot(this);
 	}
 
 	Super::BeginPlay();
@@ -74,76 +50,23 @@ void UDisplayClusterRootComponent::BeginDestroy()
 {
 	DISPLAY_CLUSTER_FUNC_TRACE(LogDisplayClusterGame);
 
-	// Clean containers. We store only pointers so there is no need to do any additional
-	// operations. All components will be destroyed by the engine.
-	ScreenComponents.Reset();
-	CameraComponents.Reset();
-	SceneNodeComponents.Reset();
-
-	Super::BeginDestroy();
-}
-
-void UDisplayClusterRootComponent::TickComponent( float DeltaTime, ELevelTick TickType, FActorComponentTickFunction* ThisTickFunction )
-{
-	DISPLAY_CLUSTER_FUNC_TRACE(LogDisplayClusterGame);
-
-	Super::TickComponent(DeltaTime, TickType, ThisTickFunction);
-
-	if (OperationMode == EDisplayClusterOperationMode::Cluster ||
-		OperationMode == EDisplayClusterOperationMode::Editor)
 	{
-		UWorld* const CurWorld = GetWorld();
-		if (CurWorld)
-		{
-			APlayerController* const CurPlayerController = CurWorld->GetFirstPlayerController();
-			if (CurPlayerController)
-			{
-				APlayerCameraManager* const CurPlayerCameraManager = CurPlayerController->PlayerCameraManager;
-				if (CurPlayerCameraManager)
-				{
-					this->SetWorldLocation(CurPlayerCameraManager->GetCameraLocation());
-					this->SetWorldRotation(CurPlayerCameraManager->GetCameraRotation());
-				}
-
-				if (bExitOnEsc)
-				{
-					if (CurPlayerController->WasInputKeyJustPressed(EKeys::Escape))
-					{
-						FDisplayClusterAppExit::ExitApplication(FDisplayClusterAppExit::ExitType::NormalSoft, FString("Exit on ESC requested"));
-					}
-				}
-			}
-		}
+		FScopeLock lock(&InternalsSyncScope);
+
+		// Clean containers. We store only pointers so there is no need to do any additional
+		// operations. All components will be destroyed by the engine.
+		ScreenComponents.Reset();
+		CameraComponents.Reset();
+		SceneNodeComponents.Reset();
 	}
+
+	Super::BeginDestroy();
 }
 
 bool UDisplayClusterRootComponent::InitializeHierarchy()
 {
 	DISPLAY_CLUSTER_FUNC_TRACE(LogDisplayClusterGame);
 
-	// Use sync components if no native input synchronization activated
-	if (NativeInputSyncPolicy == 0)
-	{
-		// Add replication components
-		const TSet<UActorComponent*> ActorComponents = GetOwner()->GetComponents();
-		int CompIdx = 0;
-		for (UActorComponent* Component : ActorComponents)
-		{
-			USceneComponent* const SceneComp = Cast<USceneComponent>(Component);
-			if (SceneComp)
-			{
-				const FString CompName = FString::Printf(TEXT("sync_comp_%d"), CompIdx++);
-
-				UDisplayClusterSceneComponentSyncParent* NewSyncComp = NewObject<UDisplayClusterSceneComponentSyncParent>(Component->GetOwner(), FName(*CompName));
-				if (NewSyncComp)
-				{
-					NewSyncComp->AttachToComponent(SceneComp, FAttachmentTransformRules(EAttachmentRule::KeepRelative, false));
-					NewSyncComp->RegisterComponent();
-				}
-			}
-		}
-	}
-
 	// Now create all child components
 	if (!(CreateCameras() && CreateScreens() && CreateNodes()))
 	{
@@ -175,25 +98,6 @@ bool UDisplayClusterRootComponent::InitializeHierarchy()
 	return true;
 }
 
-bool UDisplayClusterRootComponent::InitializeExtra()
-{
-	EDisplayClusterOperationMode OpMode = GDisplayCluster->GetOperationMode();
-	if (OpMode == EDisplayClusterOperationMode::Cluster)
-	{
-		// Optionally activate native input synchronization
-		if (NativeInputSyncPolicy == 1)
-		{
-			APlayerController* const PlayerController = GetWorld()->GetFirstPlayerController();
-			if (PlayerController)
-			{
-				PlayerController->PlayerInput = NewObject<UDisplayClusterPlayerInput>(PlayerController);
-			}
-		}
-	}
-
-	return true;
-}
-
 bool UDisplayClusterRootComponent::CreateScreens()
 {
 	// Create screens
@@ -316,7 +220,6 @@ int32 UDisplayClusterRootComponent::GetCamerasAmount() const
 UDisplayClusterCameraComponent* UDisplayClusterRootComponent::GetDefaultCamera() const
 {
 	FScopeLock lock(&InternalsSyncScope);
-
 	return DefaultCameraComponent;
 }
 
diff --git a/Engine/Plugins/Runtime/nDisplay/Source/DisplayCluster/Private/Game/Classes/Scene/DisplayClusterScreenComponent.cpp b/Engine/Plugins/Runtime/nDisplay/Source/DisplayCluster/Private/Game/Classes/Scene/DisplayClusterScreenComponent.cpp
--- a/Engine/Plugins/Runtime/nDisplay/Source/DisplayCluster/Private/Game/Classes/Scene/DisplayClusterScreenComponent.cpp
+++ b/Engine/Plugins/Runtime/nDisplay/Source/DisplayCluster/Private/Game/Classes/Scene/DisplayClusterScreenComponent.cpp
@@ -8,7 +8,7 @@
 #include "Materials/Material.h"
 #include "UObject/ConstructorHelpers.h"
 
-#include "DisplayClusterRootComponent.h"
+#include "DisplayClusterRootActor.h"
 
 #include "Game/IPDisplayClusterGameManager.h"
 #include "DisplayClusterGlobals.h"
@@ -47,15 +47,15 @@ void UDisplayClusterScreenComponent::TickComponent( float DeltaTime, ELevelTick
 	static const IPDisplayClusterGameManager* const GameMgr = GDisplayCluster->GetPrivateGameMgr();
 	if (GameMgr)
 	{
-		UDisplayClusterRootComponent* const RootComp = GameMgr->GetRoot();
-		if (RootComp)
+		ADisplayClusterRootActor* const RootActor = GameMgr->GetRootActor();
+		if (RootActor)
 		{
-			if (RootComp->GetProjectionScreenMaterial())
+			if (RootActor->GetProjectionScreenMaterial())
 			{
-				ScreenGeometryComponent->SetMaterial(0, RootComp->GetProjectionScreenMaterial());
+				ScreenGeometryComponent->SetMaterial(0, RootActor->GetProjectionScreenMaterial());
 			}
 
-			const bool bScreenVisible = RootComp->GetShowProjectionScreens();
+			const bool bScreenVisible = RootActor->GetShowProjectionScreens();
 			ScreenGeometryComponent->SetVisibility(bScreenVisible);
 		}
 	}
diff --git a/Engine/Plugins/Runtime/nDisplay/Source/DisplayCluster/Private/Game/DisplayClusterGameManager.cpp b/Engine/Plugins/Runtime/nDisplay/Source/DisplayCluster/Private/Game/DisplayClusterGameManager.cpp
--- a/Engine/Plugins/Runtime/nDisplay/Source/DisplayCluster/Private/Game/DisplayClusterGameManager.cpp
+++ b/Engine/Plugins/Runtime/nDisplay/Source/DisplayCluster/Private/Game/DisplayClusterGameManager.cpp
@@ -2,15 +2,13 @@
 
 #include "Game/DisplayClusterGameManager.h"
 
-#include "Config/IPDisplayClusterConfigManager.h"
-
-#include "Misc/DisplayClusterHelpers.h"
-
 #include "Kismet/GameplayStatics.h"
 #include "Misc/CommandLine.h"
+#include "Misc/DisplayClusterHelpers.h"
 #include "Config/DisplayClusterConfigTypes.h"
 #include "Config/IPDisplayClusterConfigManager.h"
 
+#include "DisplayClusterRootActor.h"
 #include "Camera/CameraComponent.h"
 #include "Components/SceneComponent.h"
 #include "DisplayClusterCameraComponent.h"
@@ -23,6 +21,7 @@
 #include "DisplayClusterStrings.h"
 
 #include "GameFramework/Actor.h"
+#include "Engine/LevelStreaming.h"
 
 
 FDisplayClusterGameManager::FDisplayClusterGameManager()
@@ -75,27 +74,39 @@ bool FDisplayClusterGameManager::StartScene(UWorld* InWorld)
 	check(InWorld);
 	CurrentWorld = InWorld;
 
-	// Look for existing VR root components
-	for (AActor* Actor : InWorld->PersistentLevel->Actors)
+	// Find nDisplay root actor
+	DisplayClusterRootActor = FindDisplayClusterRootActor(InWorld);
+	if (!DisplayClusterRootActor)
 	{
-		if (Actor && !Actor->IsPendingKill())
+		// Also search inside streamed levels
+		const TArray<ULevelStreaming*>& StreamingLevels = InWorld->GetStreamingLevels();
+		for (const ULevelStreaming* StreamingLevel : StreamingLevels)
 		{
-			CurrentRoot = Actor->FindComponentByClass<UDisplayClusterRootComponent>();
-			if (CurrentRoot)
+			switch (StreamingLevel->GetCurrentState())
+			{
+			case ULevelStreaming::ECurrentState::LoadedVisible:
 			{
-				UE_LOG(LogDisplayClusterGame, Log, TEXT("Found root component - %s"), *CurrentRoot->GetName());
+				// Look for the actor in those sub-levels that have been loaded already
+				const TSoftObjectPtr<UWorld>& SubWorldAsset = StreamingLevel->GetWorldAsset();
+				DisplayClusterRootActor = FindDisplayClusterRootActor(SubWorldAsset.Get());
+				if (DisplayClusterRootActor)
+				{
+					break;
+				}
+			}
+
+			default:
 				break;
 			}
 		}
 	}
 
-	APlayerController* const PlayerController = GetWorld()->GetFirstPlayerController();
-	if (!CurrentRoot && PlayerController)
+	// In cluster mode we spawn root actor if no actors found
+	if (GDisplayCluster->GetOperationMode() == EDisplayClusterOperationMode::Cluster)
 	{
-		APawn* CurPawn = PlayerController->GetPawn();
-		if (CurPawn)
+		if (!DisplayClusterRootActor)
 		{
-			SpawnRootComponent(CurPawn);
+			DisplayClusterRootActor = Cast<ADisplayClusterRootActor>(CurrentWorld->SpawnActor(ADisplayClusterRootActor::StaticClass()));
 		}
 	}
 
@@ -107,161 +118,159 @@ void FDisplayClusterGameManager::EndScene()
 	DISPLAY_CLUSTER_FUNC_TRACE(LogDisplayClusterGame);
 	FScopeLock lock(&InternalsSyncScope);
 
-	CurrentRoot = nullptr;
+	DisplayClusterRootActor = nullptr;
 }
 
 
 //////////////////////////////////////////////////////////////////////////////////////////////
 // IDisplayClusterGameManager
 //////////////////////////////////////////////////////////////////////////////////////////////
-UDisplayClusterRootComponent* FDisplayClusterGameManager::GetRoot() const
+ADisplayClusterRootActor* FDisplayClusterGameManager::GetRootActor() const
 {
 	FScopeLock lock(&InternalsSyncScope);
-	return CurrentRoot;
+	return DisplayClusterRootActor;
 }
 
-void FDisplayClusterGameManager::SetRoot(UDisplayClusterRootComponent* InRoot)
+UDisplayClusterRootComponent* FDisplayClusterGameManager::GetRootComponent() const
 {
 	FScopeLock lock(&InternalsSyncScope);
-	CurrentRoot = InRoot;
+	return DisplayClusterRootActor->GetDisplayClusterRootComponent();
 }
 
 TArray<UDisplayClusterScreenComponent*> FDisplayClusterGameManager::GetAllScreens() const
 {
 	FScopeLock lock(&InternalsSyncScope);
 
-	if (!CurrentRoot)
+	if (!DisplayClusterRootActor)
 	{
 		return TArray<UDisplayClusterScreenComponent*>();
 	}
 
-	return CurrentRoot->GetAllScreens();
+	return DisplayClusterRootActor->GetDisplayClusterRootComponent()->GetAllScreens();
 }
 
 UDisplayClusterScreenComponent* FDisplayClusterGameManager::GetScreenById(const FString& ScreenId) const
 {
 	FScopeLock lock(&InternalsSyncScope);
 
-	if (!CurrentRoot)
+	if (!DisplayClusterRootActor)
 	{
 		return nullptr;
 	}
 
-	return CurrentRoot->GetScreenById(ScreenId);
+	return DisplayClusterRootActor->GetDisplayClusterRootComponent()->GetScreenById(ScreenId);
 }
 
 int32 FDisplayClusterGameManager::GetScreensAmount() const
 {
 	FScopeLock lock(&InternalsSyncScope);
 
-	if (!CurrentRoot)
+	if (!DisplayClusterRootActor)
 	{
 		return 0;
 	}
 
-	return CurrentRoot->GetScreensAmount();
+	return DisplayClusterRootActor->GetDisplayClusterRootComponent()->GetScreensAmount();
 }
 
 UDisplayClusterCameraComponent* FDisplayClusterGameManager::GetCameraById(const FString& CameraId) const
 {
 	FScopeLock lock(&InternalsSyncScope);
 
-	if (!CurrentRoot)
+	if (!DisplayClusterRootActor)
 	{
 		return nullptr;
 	}
 
-	return CurrentRoot->GetCameraById(CameraId);
+	return DisplayClusterRootActor->GetDisplayClusterRootComponent()->GetCameraById(CameraId);
 }
 
 TArray<UDisplayClusterCameraComponent*> FDisplayClusterGameManager::GetAllCameras() const
 {
 	FScopeLock lock(&InternalsSyncScope);
 
-	if (!CurrentRoot)
+	if (!DisplayClusterRootActor)
 	{
 		return TArray<UDisplayClusterCameraComponent*>();
 	}
 
-	return CurrentRoot->GetAllCameras();
+	return DisplayClusterRootActor->GetDisplayClusterRootComponent()->GetAllCameras();
 }
 
 int32 FDisplayClusterGameManager::GetCamerasAmount() const
 {
 	FScopeLock lock(&InternalsSyncScope);
 
-	if (!CurrentRoot)
+	if (!DisplayClusterRootActor)
 	{
 		return 0;
 	}
 
-	return CurrentRoot->GetCamerasAmount();
+	return DisplayClusterRootActor->GetDisplayClusterRootComponent()->GetCamerasAmount();
 }
 
 UDisplayClusterCameraComponent* FDisplayClusterGameManager::GetDefaultCamera() const
 {
 	FScopeLock lock(&InternalsSyncScope);
 
-	if (!CurrentRoot)
+	if (!DisplayClusterRootActor)
 	{
 		return nullptr;
 	}
 
-	return CurrentRoot->GetDefaultCamera();
+	return DisplayClusterRootActor->GetDisplayClusterRootComponent()->GetDefaultCamera();
 }
 
 void FDisplayClusterGameManager::SetDefaultCamera(const FString& id)
 {
 	DISPLAY_CLUSTER_FUNC_TRACE(LogDisplayClusterGame);
 
-	if(!CurrentRoot)
+	if(!DisplayClusterRootActor)
 	{
 		return;
 	}
 
-	CurrentRoot->SetDefaultCamera(id);
+	DisplayClusterRootActor->GetDisplayClusterRootComponent()->SetDefaultCamera(id);
 }
 
 UDisplayClusterSceneComponent* FDisplayClusterGameManager::GetNodeById(const FString& NodeId) const
 {
 	FScopeLock lock(&InternalsSyncScope);
 
-	if (!CurrentRoot)
+	if (!DisplayClusterRootActor)
 	{
 		return nullptr;
 	}
 
-	return CurrentRoot->GetNodeById(NodeId);
+	return DisplayClusterRootActor->GetDisplayClusterRootComponent()->GetNodeById(NodeId);
 }
 
 TArray<UDisplayClusterSceneComponent*> FDisplayClusterGameManager::GetAllNodes() const
 {
 	FScopeLock lock(&InternalsSyncScope);
 
-	if (!CurrentRoot)
+	if (!DisplayClusterRootActor)
 	{
 		return TArray<UDisplayClusterSceneComponent*>();
 	}
 
-	return CurrentRoot->GetAllNodes();
+	return DisplayClusterRootActor->GetDisplayClusterRootComponent()->GetAllNodes();
 }
 
-bool FDisplayClusterGameManager::SpawnRootComponent(AActor* Actor)
+ADisplayClusterRootActor* FDisplayClusterGameManager::FindDisplayClusterRootActor(UWorld* InWorld)
 {
-	check(Actor);
-
-	FScopeLock lock(&InternalsSyncScope);
-
-	// Create new root component
-	UDisplayClusterRootComponent* NewRoot = NewObject<UDisplayClusterRootComponent>(Actor, FName(TEXT("nDisplay_Root")));
-	if (!NewRoot)
+	for (AActor* const Actor : InWorld->PersistentLevel->Actors)
 	{
-		return false;
+		if (Actor && !Actor->IsPendingKill())
+		{
+			ADisplayClusterRootActor* RootActor = Cast<ADisplayClusterRootActor>(Actor);
+			if (RootActor)
+			{
+				UE_LOG(LogDisplayClusterGame, Log, TEXT("Found root actor - %s"), *RootActor->GetName());
+				return RootActor;
+			}
+		}
 	}
 
-	// Set up
-	NewRoot->AttachToComponent(Actor->GetRootComponent(), FAttachmentTransformRules(EAttachmentRule::KeepRelative, false));
-	NewRoot->RegisterComponent();
-
-	return true;
+	return nullptr;
 }
diff --git a/Engine/Plugins/Runtime/nDisplay/Source/DisplayCluster/Private/Game/DisplayClusterGameManager.h b/Engine/Plugins/Runtime/nDisplay/Source/DisplayCluster/Private/Game/DisplayClusterGameManager.h
--- a/Engine/Plugins/Runtime/nDisplay/Source/DisplayCluster/Private/Game/DisplayClusterGameManager.h
+++ b/Engine/Plugins/Runtime/nDisplay/Source/DisplayCluster/Private/Game/DisplayClusterGameManager.h
@@ -8,6 +8,7 @@
 
 
 class AActor;
+class ADisplayClusterRootActor;
 class USceneComponent;
 class UDisplayClusterCameraComponent;
 class UDisplayClusterRootComponent;
@@ -41,8 +42,8 @@ public:
 	//////////////////////////////////////////////////////////////////////////////////////////////
 	// IDisplayClusterGameManager
 	//////////////////////////////////////////////////////////////////////////////////////////////
-	virtual UDisplayClusterRootComponent*           GetRoot() const override;
-	virtual void                                    SetRoot(UDisplayClusterRootComponent* InRoot) override;
+	virtual ADisplayClusterRootActor*               GetRootActor() const override;
+	virtual UDisplayClusterRootComponent*           GetRootComponent() const override;
 
 	virtual TArray<UDisplayClusterScreenComponent*> GetAllScreens() const override;
 	virtual UDisplayClusterScreenComponent*         GetScreenById(const FString& id) const override;
@@ -67,12 +68,12 @@ public:
 	virtual bool IsDisplayClusterActive() const override
 	{ return CurrentOperationMode != EDisplayClusterOperationMode::Disabled; }
 
-protected:
-	bool SpawnRootComponent(AActor* Actor);
+private:
+	ADisplayClusterRootActor* FindDisplayClusterRootActor(UWorld* InWorld);
 
 private:
 	// Active DisplayCluster root
-	UDisplayClusterRootComponent* CurrentRoot = nullptr;
+	ADisplayClusterRootActor* DisplayClusterRootActor = nullptr;
 
 	EDisplayClusterOperationMode CurrentOperationMode;
 	FString ConfigPath;
diff --git a/Engine/Plugins/Runtime/nDisplay/Source/DisplayCluster/Public/Blueprints/IDisplayClusterBlueprintAPI.h b/Engine/Plugins/Runtime/nDisplay/Source/DisplayCluster/Public/Blueprints/IDisplayClusterBlueprintAPI.h
--- a/Engine/Plugins/Runtime/nDisplay/Source/DisplayCluster/Public/Blueprints/IDisplayClusterBlueprintAPI.h
+++ b/Engine/Plugins/Runtime/nDisplay/Source/DisplayCluster/Public/Blueprints/IDisplayClusterBlueprintAPI.h
@@ -95,10 +95,13 @@ public:
 	//////////////////////////////////////////////////////////////////////////////////////////////
 	// Game API
 	//////////////////////////////////////////////////////////////////////////////////////////////
-	// Root
-	/** Returns Cluster Pawn. */
-	UFUNCTION(BlueprintCallable, meta = (DisplayName = "Get root"), Category = "DisplayCluster|Game")
-	virtual UDisplayClusterRootComponent* GetRoot() = 0;
+	/** Returns DisplayCluster root actor. */
+	UFUNCTION(BlueprintCallable, meta = (DisplayName = "Get root actor"), Category = "DisplayCluster|Game")
+	virtual ADisplayClusterRootActor* GetRootActor() = 0;
+
+	/** Returns DisplayCluster root component. */
+	UFUNCTION(BlueprintCallable, meta = (DisplayName = "Get root component"), Category = "DisplayCluster|Game")
+	virtual UDisplayClusterRootComponent* GetRootComponent() = 0;
 
 	// Screens
 	/** Returns screen reference by id name. */
diff --git a/Engine/Plugins/Runtime/nDisplay/Source/DisplayCluster/Public/DisplayClusterRootActor.h b/Engine/Plugins/Runtime/nDisplay/Source/DisplayCluster/Public/DisplayClusterRootActor.h
new file mode 100644
--- /dev/null
+++ b/Engine/Plugins/Runtime/nDisplay/Source/DisplayCluster/Public/DisplayClusterRootActor.h
@@ -0,0 +1,82 @@
+// Copyright 1998-2019 Epic Games, Inc. All Rights Reserved.
+
+#pragma once
+
+#include "CoreMinimal.h"
+#include "GameFramework/Actor.h"
+#include "Camera/PlayerCameraManager.h"
+
+#include "DisplayClusterEnums.h"
+
+#include "DisplayClusterRootActor.generated.h"
+
+class UDisplayClusterRootComponent;
+class UMaterial;
+
+
+/**
+ * VR root. This contains nDisplay VR hierarchy in the game.
+ */
+UCLASS()
+class DISPLAYCLUSTER_API ADisplayClusterRootActor
+	: public AActor
+{
+	GENERATED_BODY()
+
+public:
+	ADisplayClusterRootActor(const FObjectInitializer& ObjectInitializer);
+
+	inline UDisplayClusterRootComponent* GetDisplayClusterRootComponent() const
+	{ return DisplayClusterRootComponent; }
+
+public:
+	//////////////////////////////////////////////////////////////////////////////////////////////
+	// AActor
+	//////////////////////////////////////////////////////////////////////////////////////////////
+	virtual void BeginPlay() override;
+	virtual void BeginDestroy() override;
+	virtual void Tick(float DeltaSeconds) override;
+
+	
+public:
+	bool GetShowProjectionScreens() const
+	{ return bShowProjectionScreens; }
+
+	UMaterial* GetProjectionScreenMaterial() const
+	{ return ProjectionScreensMaterial; }
+
+protected:
+	UPROPERTY(EditAnywhere, Category = "DisplayCluster", meta = (DisplayName = "Exit when ESC pressed"))
+	bool bExitOnEsc;
+
+	UPROPERTY(EditAnywhere, Category = "DisplayCluster", meta = (DisplayName = "Show projection screens"))
+	bool bShowProjectionScreens;
+
+	UPROPERTY(EditAnywhere, Category = "DisplayCluster", meta = (DisplayName = "Projection screens material"))
+	UMaterial* ProjectionScreensMaterial;
+
+#if WITH_EDITORONLY_DATA
+public:
+	FString GetEditorConfigPath() const
+	{ return EditorConfigPath; }
+
+	FString GetEditorNodeId() const
+	{ return EditorNodeId; }
+
+protected:
+	UPROPERTY(EditAnywhere, Category = "DisplayCluster (Editor only)", meta = (DisplayName = "Config file"))
+	FString EditorConfigPath;
+
+	UPROPERTY(EditAnywhere, Category = "DisplayCluster (Editor only)", meta = (DisplayName = "Node ID"))
+	FString EditorNodeId;
+#endif
+
+protected:
+	/** Camera component */
+	UPROPERTY(VisibleAnywhere, BlueprintReadOnly, Category = "DisplayCluster")
+	UDisplayClusterRootComponent* DisplayClusterRootComponent;
+
+private:
+	// Current operation mode
+	EDisplayClusterOperationMode OperationMode = EDisplayClusterOperationMode::Disabled;
+};
diff --git a/Engine/Plugins/Runtime/nDisplay/Source/DisplayCluster/Public/DisplayClusterRootComponent.h b/Engine/Plugins/Runtime/nDisplay/Source/DisplayCluster/Public/DisplayClusterRootComponent.h
--- a/Engine/Plugins/Runtime/nDisplay/Source/DisplayCluster/Public/DisplayClusterRootComponent.h
+++ b/Engine/Plugins/Runtime/nDisplay/Source/DisplayCluster/Public/DisplayClusterRootComponent.h
@@ -3,9 +3,6 @@
 #pragma once
 
 #include "Components/SceneComponent.h"
-#include "Camera/PlayerCameraManager.h"
-
-#include "DisplayClusterEnums.h"
 
 #include "DisplayClusterRootComponent.generated.h"
 
@@ -14,13 +11,12 @@ class UCameraComponent;
 class UDisplayClusterCameraComponent;
 class UDisplayClusterSceneComponent;
 class UDisplayClusterScreenComponent;
-class UMaterial;
 
 
 /**
  * DisplayCluster root component (nDisplay root)
  */
-UCLASS( ClassGroup=(Custom), meta=(BlueprintSpawnableComponent) )
+UCLASS( ClassGroup=(Custom) )
 class DISPLAYCLUSTER_API UDisplayClusterRootComponent
 	: public USceneComponent
 {
@@ -34,7 +30,6 @@ public:
 public:
 	virtual void BeginPlay() override;
 	virtual void BeginDestroy() override;
-	virtual void TickComponent(float DeltaTime, ELevelTick TickType, FActorComponentTickFunction* ThisTickFunction) override;
 
 public:
 	TArray<UDisplayClusterScreenComponent*> GetAllScreens() const;
@@ -61,43 +56,9 @@ protected:
 	// Default camera for this root
 	UDisplayClusterCameraComponent* DefaultCameraComponent = nullptr;
 
-public:
-	bool GetShowProjectionScreens() const
-	{ return bShowProjectionScreens; }
-
-	UMaterial* GetProjectionScreenMaterial() const
-	{ return ProjectionScreensMaterial; }
-
-protected:
-	UPROPERTY(EditAnywhere, Category = "DisplayCluster", meta = (DisplayName = "Exit when ESC pressed"))
-	bool bExitOnEsc;
-
-	UPROPERTY(EditAnywhere, Category = "DisplayCluster", meta = (DisplayName = "Show projection screens"))
-	bool bShowProjectionScreens;
-
-	UPROPERTY(EditAnywhere, Category = "DisplayCluster", meta = (DisplayName = "Projection screens material"))
-	UMaterial* ProjectionScreensMaterial;
-
-#if WITH_EDITORONLY_DATA
-public:
-	FString GetEditorConfigPath() const
-	{ return EditorConfigPath; }
-
-	FString GetEditorNodeId() const
-	{ return EditorNodeId; }
-
-protected:
-	UPROPERTY(EditAnywhere, Category = "nDisplay (Editor only)", meta = (DisplayName = "Config file"))
-	FString EditorConfigPath;
-
-	UPROPERTY(EditAnywhere, Category = "nDisplay (Editor only)", meta = (DisplayName = "Node ID"))
-	FString EditorNodeId;
-#endif
-
 protected:
 	// Creates all hierarchy objects declared in a config file
 	virtual bool InitializeHierarchy();
-	virtual bool InitializeExtra();
 	virtual bool CreateScreens();
 	virtual bool CreateNodes();
 	virtual bool CreateCameras();
@@ -113,9 +74,4 @@ private:
 
 private:
 	mutable FCriticalSection InternalsSyncScope;
-
-	// Current operation mode
-	EDisplayClusterOperationMode OperationMode = EDisplayClusterOperationMode::Disabled;
-
-	int32 NativeInputSyncPolicy = 0;
 };
diff --git a/Engine/Plugins/Runtime/nDisplay/Source/DisplayCluster/Public/Game/IDisplayClusterGameManager.h b/Engine/Plugins/Runtime/nDisplay/Source/DisplayCluster/Public/Game/IDisplayClusterGameManager.h
--- a/Engine/Plugins/Runtime/nDisplay/Source/DisplayCluster/Public/Game/IDisplayClusterGameManager.h
+++ b/Engine/Plugins/Runtime/nDisplay/Source/DisplayCluster/Public/Game/IDisplayClusterGameManager.h
@@ -2,6 +2,7 @@
 
 #pragma once
 
+class ADisplayClusterRootActor;
 class UDisplayClusterCameraComponent;
 class UDisplayClusterRootComponent;
 class UDisplayClusterSceneComponent;
@@ -19,8 +20,8 @@ public:
 	{ }
 
 public:
-	virtual UDisplayClusterRootComponent*           GetRoot() const = 0;
-	virtual void                                    SetRoot(UDisplayClusterRootComponent* InRoot) = 0;
+	virtual ADisplayClusterRootActor*               GetRootActor() const = 0;
+	virtual UDisplayClusterRootComponent*           GetRootComponent() const = 0;
 
 	virtual TArray<UDisplayClusterScreenComponent*> GetAllScreens() const = 0;
 	virtual UDisplayClusterScreenComponent*         GetScreenById(const FString& id) const = 0;
diff --git a/Engine/Plugins/Runtime/nDisplay/Source/DisplayClusterEditor/Private/DisplayClusterEditorEngine.cpp b/Engine/Plugins/Runtime/nDisplay/Source/DisplayClusterEditor/Private/DisplayClusterEditorEngine.cpp
--- a/Engine/Plugins/Runtime/nDisplay/Source/DisplayClusterEditor/Private/DisplayClusterEditorEngine.cpp
+++ b/Engine/Plugins/Runtime/nDisplay/Source/DisplayClusterEditor/Private/DisplayClusterEditorEngine.cpp
@@ -3,7 +3,7 @@
 #include "DisplayClusterEditorEngine.h"
 #include "DisplayClusterEditorLog.h"
 
-#include "DisplayClusterRootComponent.h"
+#include "DisplayClusterRootActor.h"
 
 #include "DisplayCluster/Private/IPDisplayCluster.h"
 
@@ -48,17 +48,17 @@ void UDisplayClusterEditorEngine::PreExit()
 	Super::PreExit();
 }
 
-UDisplayClusterRootComponent* UDisplayClusterEditorEngine::FindDisplayClusterRootComponent(UWorld* InWorld)
+ADisplayClusterRootActor* UDisplayClusterEditorEngine::FindDisplayClusterRootActor(UWorld* InWorld)
 {
 	for (AActor* const Actor : InWorld->PersistentLevel->Actors)
 	{
 		if (Actor && !Actor->IsPendingKill())
 		{
-			UDisplayClusterRootComponent* RootComponent = Actor->FindComponentByClass<UDisplayClusterRootComponent>();
-			if (RootComponent && !RootComponent->IsPendingKill())
+			ADisplayClusterRootActor* RootActor = Cast<ADisplayClusterRootActor>(Actor);
+			if (RootActor)
 			{
-				UE_LOG(LogDisplayClusterEditorEngine, Log, TEXT("Found root component - %s"), *RootComponent->GetName());
-				return RootComponent;
+				UE_LOG(LogDisplayClusterEditorEngine, Log, TEXT("Found root actor - %s"), *RootActor->GetName());
+				return RootActor;
 			}
 		}
 	}
@@ -72,11 +72,11 @@ void UDisplayClusterEditorEngine::PlayInEditor(UWorld* InWorld, bool bInSimulate
 
 	if (DisplayClusterModule)
 	{
-		// Find nDisplay root
-		UDisplayClusterRootComponent* RootComponent = FindDisplayClusterRootComponent(InWorld);
-		if (!RootComponent)
+		// Find nDisplay root actor
+		ADisplayClusterRootActor* RootActor = FindDisplayClusterRootActor(InWorld);
+		if (!RootActor)
 		{
-			//Also search inside streamed levels:
+			// Also search inside streamed levels
 			const TArray<ULevelStreaming*>& StreamingLevels = InWorld->GetStreamingLevels();
 			for (ULevelStreaming* StreamingLevel : StreamingLevels)
 			{
@@ -84,10 +84,10 @@ void UDisplayClusterEditorEngine::PlayInEditor(UWorld* InWorld, bool bInSimulate
 				{
 					case ULevelStreaming::ECurrentState::LoadedVisible:
 					{
-						// Parse only in loaded sub-levels
+						// Look for the actor in those sub-levels that have been loaded already
 						const TSoftObjectPtr<UWorld>& SubWorldAsset = StreamingLevel->GetWorldAsset();
-						RootComponent = FindDisplayClusterRootComponent(SubWorldAsset.Get());
-						if (RootComponent)
+						RootActor = FindDisplayClusterRootActor(SubWorldAsset.Get());
+						if (RootActor)
 						{
 							break;
 						}
@@ -99,12 +99,12 @@ void UDisplayClusterEditorEngine::PlayInEditor(UWorld* InWorld, bool bInSimulate
 			}
 		}
 
-		// If we found a root component, start DisplayCluster PIE session
-		if (RootComponent)
+		// If we found a root actor, start DisplayCluster PIE session
+		if (RootActor)
 		{
 			bIsNDisplayPIE = true;
 
-			if (!DisplayClusterModule->StartSession(RootComponent->GetEditorConfigPath(), RootComponent->GetEditorNodeId()))
+			if (!DisplayClusterModule->StartSession(RootActor->GetEditorConfigPath(), RootActor->GetEditorNodeId()))
 			{
 				UE_LOG(LogDisplayClusterEditorEngine, Error, TEXT("Couldn't start DisplayCluster session"));
 
diff --git a/Engine/Plugins/Runtime/nDisplay/Source/DisplayClusterEditor/Private/DisplayClusterEditorEngine.h b/Engine/Plugins/Runtime/nDisplay/Source/DisplayClusterEditor/Private/DisplayClusterEditorEngine.h
--- a/Engine/Plugins/Runtime/nDisplay/Source/DisplayClusterEditor/Private/DisplayClusterEditorEngine.h
+++ b/Engine/Plugins/Runtime/nDisplay/Source/DisplayClusterEditor/Private/DisplayClusterEditorEngine.h
@@ -7,7 +7,7 @@
 #include "DisplayClusterEditorEngine.generated.h"
 
 class IPDisplayCluster;
-class UDisplayClusterRootComponent;
+class ADisplayClusterRootActor;
 
 
 /**
@@ -26,7 +26,7 @@ public:
 	virtual void Tick(float DeltaSeconds, bool bIdleMode) override;
 
 private:
-	UDisplayClusterRootComponent* FindDisplayClusterRootComponent(UWorld* InWorld);
+	ADisplayClusterRootActor* FindDisplayClusterRootActor(UWorld* InWorld);
 
 private:
 	IPDisplayCluster* DisplayClusterModule = nullptr;
diff --git a/Engine/Plugins/Runtime/nDisplay/Source/DisplayClusterProjection/Private/Policy/DisplayClusterProjectionPolicyBase.cpp b/Engine/Plugins/Runtime/nDisplay/Source/DisplayClusterProjection/Private/Policy/DisplayClusterProjectionPolicyBase.cpp
--- a/Engine/Plugins/Runtime/nDisplay/Source/DisplayClusterProjection/Private/Policy/DisplayClusterProjectionPolicyBase.cpp
+++ b/Engine/Plugins/Runtime/nDisplay/Source/DisplayClusterProjection/Private/Policy/DisplayClusterProjectionPolicyBase.cpp
@@ -42,7 +42,7 @@ void FDisplayClusterProjectionPolicyBase::InitializeOriginComponent(const FStrin
 	if(PolicyOriginComp == nullptr)
 	{
 		UE_LOG(LogDisplayClusterProjection, Log, TEXT("No custom origin set or component '%s' not found for viewport '%s'. VR root will be used."), *OriginCompId, *PolicyViewportId);
-		PolicyOriginComp = GameMgr->GetRoot();
+		PolicyOriginComp = GameMgr->GetRootComponent();
 	}
 
 	if (!PolicyOriginComp)
diff --git a/Engine/Plugins/Runtime/nDisplay/Source/DisplayClusterProjection/Private/Policy/Simple/DisplayClusterProjectionSimplePolicy.cpp b/Engine/Plugins/Runtime/nDisplay/Source/DisplayClusterProjection/Private/Policy/Simple/DisplayClusterProjectionSimplePolicy.cpp
--- a/Engine/Plugins/Runtime/nDisplay/Source/DisplayClusterProjection/Private/Policy/Simple/DisplayClusterProjectionSimplePolicy.cpp
+++ b/Engine/Plugins/Runtime/nDisplay/Source/DisplayClusterProjection/Private/Policy/Simple/DisplayClusterProjectionSimplePolicy.cpp
@@ -192,7 +192,7 @@ void FDisplayClusterProjectionSimplePolicy::InitializeMeshData()
 	}
 
 	// Get our VR root
-	UDisplayClusterRootComponent* Root = GameMgr->GetRoot();
+	UDisplayClusterRootComponent* Root = GameMgr->GetRootComponent();
 	if (!Root)
 	{
 		UE_LOG(LogDisplayClusterProjectionSimple, Error, TEXT("Couldn't get a VR root object"));
diff --git a/Engine/Plugins/Runtime/nDisplay/Source/PicpProjection/Private/Policy/PicpProjectionPolicyBase.cpp b/Engine/Plugins/Runtime/nDisplay/Source/PicpProjection/Private/Policy/PicpProjectionPolicyBase.cpp
--- a/Engine/Plugins/Runtime/nDisplay/Source/PicpProjection/Private/Policy/PicpProjectionPolicyBase.cpp
+++ b/Engine/Plugins/Runtime/nDisplay/Source/PicpProjection/Private/Policy/PicpProjectionPolicyBase.cpp
@@ -42,7 +42,7 @@ void FPicpProjectionPolicyBase::InitializeOriginComponent(const FString& OriginC
 	if(PolicyOriginComp == nullptr)
 	{
 		UE_LOG(LogPicpProjectionMPCDI, Log, TEXT("No custom origin set or component '%s' not found for viewport '%s'. VR root will be used."), *OriginCompId, *PolicyViewportId);
-		PolicyOriginComp = GameMgr->GetRoot();
+		PolicyOriginComp = GameMgr->GetRootComponent();
 	}
 
 	if (!PolicyOriginComp)
