- Facilitated configuration of MessageBus Message Interceptor
- Master will now send its settings to the cluster so all nodes use the same setup

diff --git a/Engine/Plugins/Runtime/LiveLinkOvernDisplay/Source/LiveLinkOverNDisplay/Private/LiveLinkOverNDisplayModule.cpp b/Engine/Plugins/Runtime/LiveLinkOvernDisplay/Source/LiveLinkOverNDisplay/Private/LiveLinkOverNDisplayModule.cpp
--- a/Engine/Plugins/Runtime/LiveLinkOvernDisplay/Source/LiveLinkOverNDisplay/Private/LiveLinkOverNDisplayModule.cpp
+++ b/Engine/Plugins/Runtime/LiveLinkOvernDisplay/Source/LiveLinkOverNDisplay/Private/LiveLinkOverNDisplayModule.cpp
@@ -33,7 +33,7 @@ FLiveLinkOverNDisplayModule::FLiveLinkOverNDisplayModule()
 void FLiveLinkOverNDisplayModule::StartupModule()
 {
 	// Register Engine callbacks
-	OnEngineLoopInitCompleteDelegate = FCoreDelegates::OnFEngineLoopInitComplete.AddRaw(this, &FLiveLinkOverNDisplayModule::OnEngineLoopInitComplete);
+	FCoreDelegates::OnFEngineLoopInitComplete.AddRaw(this, &FLiveLinkOverNDisplayModule::OnEngineLoopInitComplete);
 
 #if WITH_EDITOR
 	if (ISettingsModule* SettingsModule = FModuleManager::GetModulePtr<ISettingsModule>("Settings"))
@@ -67,14 +67,11 @@ FNDisplayLiveLinkSubjectReplicator& FLiveLinkOverNDisplayModule::GetSubjectRepli
 
 void FLiveLinkOverNDisplayModule::OnEngineLoopInitComplete()
 {
-	if (UDisplayClusterGameEngine* DisplayClusterEngine = Cast<UDisplayClusterGameEngine>(GEngine))
+	if (IDisplayCluster::IsAvailable() && IDisplayCluster::Get().GetOperationMode() == EDisplayClusterOperationMode::Cluster) 
 	{
-		if (EDisplayClusterOperationMode::Cluster == DisplayClusterEngine->GetOperationMode())
+		if (GetDefault<ULiveLinkOverNDisplaySettings>()->IsLiveLinkOverNDisplayEnabled())
 		{
-			if (GetDefault<ULiveLinkOverNDisplaySettings>()->IsLiveLinkOverNDisplayEnabled())
-			{
-				LiveLinkReplicator->Activate();
-			}
+			LiveLinkReplicator->Activate();
 		}
 	}
 }
diff --git a/Engine/Plugins/Runtime/LiveLinkOvernDisplay/Source/LiveLinkOverNDisplay/Private/LiveLinkOverNDisplayModule.h b/Engine/Plugins/Runtime/LiveLinkOvernDisplay/Source/LiveLinkOverNDisplay/Private/LiveLinkOverNDisplayModule.h
--- a/Engine/Plugins/Runtime/LiveLinkOvernDisplay/Source/LiveLinkOverNDisplay/Private/LiveLinkOverNDisplayModule.h
+++ b/Engine/Plugins/Runtime/LiveLinkOvernDisplay/Source/LiveLinkOverNDisplay/Private/LiveLinkOverNDisplayModule.h
@@ -11,7 +11,7 @@ class FNDisplayLiveLinkSubjectReplicator;
 
 
 /**
- * Entry point for LiveLinkOverNDisplay functionnality. 
+ * Entry point for LiveLinkOverNDisplay functionality. 
  */
 class FLiveLinkOverNDisplayModule : public ILiveLinkOverNDisplayModule, public TSharedFromThis<FLiveLinkOverNDisplayModule>
 {
@@ -32,9 +32,6 @@ private:
 
 private:
 
-	/** Cluster event listener delegate */
-	FDelegateHandle OnEngineLoopInitCompleteDelegate;
-
 	/** Replicator SyncObject used to transfer data across all cluster machines */
 	TUniquePtr<FNDisplayLiveLinkSubjectReplicator> LiveLinkReplicator;
 };
diff --git a/Engine/Plugins/Runtime/nDisplay/Source/DisplayClusterMessageInterceptor/Private/DisplayClusterMessageInterceptionModule.cpp b/Engine/Plugins/Runtime/nDisplay/Source/DisplayClusterMessageInterceptor/Private/DisplayClusterMessageInterceptionModule.cpp
--- a/Engine/Plugins/Runtime/nDisplay/Source/DisplayClusterMessageInterceptor/Private/DisplayClusterMessageInterceptionModule.cpp
+++ b/Engine/Plugins/Runtime/nDisplay/Source/DisplayClusterMessageInterceptor/Private/DisplayClusterMessageInterceptionModule.cpp
@@ -3,13 +3,17 @@
 #include "Modules/ModuleInterface.h"
 #include "Modules/ModuleManager.h"
 
-#include "HAL/IConsoleManager.h"
-#include "IMessagingModule.h"
-#include "IMessageBus.h"
-#include "IDisplayCluster.h"
+#include "Cluster/DisplayClusterClusterEvent.h"
 #include "Cluster/IDisplayClusterClusterManager.h"
+#include "Cluster/IDisplayClusterClusterSyncObject.h"
+#include "DisplayClusterGameEngine.h"
 #include "DisplayClusterMessageInterceptor.h"
 #include "DisplayClusterMessageInterceptionSettings.h"
+#include "Engine/Engine.h"
+#include "HAL/IConsoleManager.h"
+#include "IDisplayCluster.h"
+#include "IMessageBus.h"
+#include "IMessagingModule.h"
 
 #if WITH_EDITOR
 	#include "ISettingsModule.h"
@@ -19,6 +23,12 @@
 #define LOCTEXT_NAMESPACE "DisplayClusterInterception"
 
 
+namespace DisplayClusterInterceptionModuleUtils
+{
+	static const FString MessageInterceptionSetupEventCategory = TEXT("nDCISetup");
+	static const FString MessageInterceptionSetupEventParameterSettings = TEXT("Settings");
+}
+
 /**
  * Display Cluster Message Interceptor module
  * Intercept a specified set of message bus messages that are received across all the display nodes
@@ -29,21 +39,8 @@ class FDisplayClusterMessageInterceptionModule : public IModuleInterface
 public:
 	virtual void StartupModule() override
 	{
-		// Create the message interceptor
-		Interceptor = MakeShared<FDisplayClusterMessageInterceptor, ESPMode::ThreadSafe>();
-
-		// Register cluster event listeners
-		IDisplayClusterClusterManager* ClusterManager = IDisplayCluster::Get().GetClusterMgr();
-		if (ClusterManager && !ListenerDelegate.IsBound())
-		{
-			ListenerDelegate = FOnClusterEventListener::CreateRaw(this, &FDisplayClusterMessageInterceptionModule::HandleClusterEvent);
-			ClusterManager->AddClusterEventListener(ListenerDelegate);
-		}
-
-		// Register cluster session events
-		IDisplayCluster::Get().OnDisplayClusterStartSession().AddRaw(this, &FDisplayClusterMessageInterceptionModule::StartInterception);
-		IDisplayCluster::Get().OnDisplayClusterEndSession().AddRaw(this, &FDisplayClusterMessageInterceptionModule::StopInterception);
-		IDisplayCluster::Get().OnDisplayClusterPreTick().AddRaw(this, &FDisplayClusterMessageInterceptionModule::HandleClusterPreTick);
+		// Register for Cluster StartSession callback so everything is setup before launching interception
+		IDisplayCluster::Get().OnDisplayClusterStartSession().AddRaw(this, &FDisplayClusterMessageInterceptionModule::OnDisplayClusterStartSession);
 
 		// Setup console command to start/stop interception
 		StartMessageSyncCommand = MakeUnique<FAutoConsoleCommand>(
@@ -100,16 +97,79 @@ public:
 	}
 
 private:
+	void OnDisplayClusterStartSession()
+	{
+		if (IDisplayCluster::IsAvailable() && IDisplayCluster::Get().GetOperationMode() == EDisplayClusterOperationMode::Cluster)
+		{
+			// Create the message interceptor only when we're in cluster mode
+			Interceptor = MakeShared<FDisplayClusterMessageInterceptor, ESPMode::ThreadSafe>();
+
+			// Register cluster event listeners
+			IDisplayClusterClusterManager* ClusterManager = IDisplayCluster::Get().GetClusterMgr();
+			if (ClusterManager && !ListenerDelegate.IsBound())
+			{
+				ListenerDelegate = FOnClusterEventListener::CreateRaw(this, &FDisplayClusterMessageInterceptionModule::HandleClusterEvent);
+				ClusterManager->AddClusterEventListener(ListenerDelegate);
+
+				//Master will send out its interceptor settings to the cluster so everyone uses the same things
+				if (ClusterManager->IsMaster())
+				{
+					FString ExportedSettings;
+					const UDisplayClusterMessageInterceptionSettings* CurrentSettings = GetDefault<UDisplayClusterMessageInterceptionSettings>();
+					FMessageInterceptionSettings::StaticStruct()->ExportText(ExportedSettings, &CurrentSettings->InterceptionSettings, nullptr, nullptr, PPF_None, nullptr);
+					
+					FDisplayClusterClusterEvent SettingsEvent;
+					SettingsEvent.Category = DisplayClusterInterceptionModuleUtils::MessageInterceptionSetupEventCategory;
+					SettingsEvent.Name = ClusterManager->GetNodeId();
+					SettingsEvent.Parameters.FindOrAdd(DisplayClusterInterceptionModuleUtils::MessageInterceptionSetupEventParameterSettings) = MoveTemp(ExportedSettings);
+
+					const bool bMasterOnly = true; 
+					ClusterManager->EmitClusterEvent(SettingsEvent, bMasterOnly);
+				}
+			}
+
+			//Start with interception enabled
+			bStartInterceptionRequested = true;
+
+			// Register cluster session events
+			IDisplayCluster::Get().OnDisplayClusterEndSession().AddRaw(this, &FDisplayClusterMessageInterceptionModule::StopInterception);
+			IDisplayCluster::Get().OnDisplayClusterPreTick().AddRaw(this, &FDisplayClusterMessageInterceptionModule::HandleClusterPreTick);
+		}
+	}
+
 	void HandleClusterEvent(const FDisplayClusterClusterEvent& InEvent)
 	{
-		if (Interceptor)
+		if (InEvent.Category == DisplayClusterInterceptionModuleUtils::MessageInterceptionSetupEventCategory)
+		{
+			HandleMessageInterceptorSetupEvent(InEvent);
+		}
+		else
 		{
-			Interceptor->HandleClusterEvent(InEvent);
+			//All events except our settings synchronization one are passed to the interceptor
+			if (Interceptor)
+			{
+				Interceptor->HandleClusterEvent(InEvent);
+			}
 		}
+
 	}
 
 	void HandleClusterPreTick()
 	{
+		// StartInterception will be handled once settings synchronization is completed to ensure the same behavior across the cluster
+		if (bStartInterceptionRequested)
+		{
+			IDisplayClusterClusterManager* ClusterManager = IDisplayCluster::Get().GetClusterMgr();
+			if (Interceptor && ClusterManager)
+			{
+				if (bSettingsSynchronizationDone)
+				{
+					bStartInterceptionRequested = false;
+					Interceptor->Start(IMessagingModule::Get().GetDefaultBus().ToSharedRef());
+				}
+			}
+		}
+
 		if (Interceptor)
 		{
 			Interceptor->SyncMessages();
@@ -118,12 +178,7 @@ private:
 
 	void StartInterception()
 	{
-		IDisplayClusterClusterManager* ClusterManager = IDisplayCluster::Get().GetClusterMgr();
-		if (Interceptor && ClusterManager)
-		{
-			Interceptor->Setup(ClusterManager, IMessagingModule::Get().GetDefaultBus().ToSharedRef());
-			Interceptor->Start();
-		}		
+		bStartInterceptionRequested = true;
 	}
 
 	void StopInterception()
@@ -131,20 +186,47 @@ private:
 		if (Interceptor)
 		{
 			Interceptor->Stop();
-			// remove internal reference to message bus
-			Interceptor->Setup(nullptr, nullptr);
 		}
 	}
 
+	void HandleMessageInterceptorSetupEvent(const FDisplayClusterClusterEvent& InEvent)
+	{
+		if (Interceptor)
+		{
+			const FString& ExportedSettings = InEvent.Parameters.FindChecked(DisplayClusterInterceptionModuleUtils::MessageInterceptionSetupEventParameterSettings);
+			FMessageInterceptionSettings::StaticStruct()->ImportText(*ExportedSettings, &SynchronizedSettings, nullptr, EPropertyPortFlags::PPF_None, GLog, FMessageInterceptionSettings::StaticStruct()->GetName());
+
+			IDisplayClusterClusterManager* ClusterManager = IDisplayCluster::Get().GetClusterMgr();
+			if (ClusterManager)
+			{
+				Interceptor->Setup(ClusterManager, SynchronizedSettings);
+				bSettingsSynchronizationDone = true;
+
+				UE_LOG(LogDisplayClusterInterception, Display, TEXT("Node '%s' received synchronization settings event"), *ClusterManager->GetNodeId());
+			}
+		}
+	}
+
+private:
+
 	/** MessageBus interceptor */
 	TSharedPtr<FDisplayClusterMessageInterceptor, ESPMode::ThreadSafe> Interceptor;
 
+	/** Settings to be used synchronized around the cluster */
+	FMessageInterceptionSettings SynchronizedSettings;
+
 	/** Cluster event listener delegate */
 	FOnClusterEventListener ListenerDelegate;
 
-	// Console commands handle.
+	/** Console commands handle. */
 	TUniquePtr<FAutoConsoleCommand> StartMessageSyncCommand;
 	TUniquePtr<FAutoConsoleCommand> StopMessageSyncCommand;
+
+	/** Request to start interception. Cached to be done once synchronization is done. */
+	bool bStartInterceptionRequested;
+
+	/** Flag to check if synchronization was done */
+	bool bSettingsSynchronizationDone = false;
 };
 
 #undef LOCTEXT_NAMESPACE
diff --git a/Engine/Plugins/Runtime/nDisplay/Source/DisplayClusterMessageInterceptor/Private/DisplayClusterMessageInterceptor.cpp b/Engine/Plugins/Runtime/nDisplay/Source/DisplayClusterMessageInterceptor/Private/DisplayClusterMessageInterceptor.cpp
--- a/Engine/Plugins/Runtime/nDisplay/Source/DisplayClusterMessageInterceptor/Private/DisplayClusterMessageInterceptor.cpp
+++ b/Engine/Plugins/Runtime/nDisplay/Source/DisplayClusterMessageInterceptor/Private/DisplayClusterMessageInterceptor.cpp
@@ -9,7 +9,13 @@
 #include "Cluster/IDisplayClusterClusterManager.h"
 
 
-DEFINE_LOG_CATEGORY_STATIC(LogDisplayClusterInterception, Log, All);
+DEFINE_LOG_CATEGORY(LogDisplayClusterInterception);
+
+
+namespace DisplayClusterMessageInterceptorUtils
+{
+	static const FInterceptedMessageDescriptor MultiUserInterception(TArray<FName>({ TEXT("ConcertSession_CustomEvent") }), TEXT("ConcertMessageId"));
+}
 
 FDisplayClusterMessageInterceptor::FDisplayClusterMessageInterceptor()
 	: bIsIntercepting(false)
@@ -18,10 +24,16 @@ FDisplayClusterMessageInterceptor::FDisplayClusterMessageInterceptor()
 	, ClusterManager(nullptr)
 {}
 
-void FDisplayClusterMessageInterceptor::Setup(IDisplayClusterClusterManager* InClusterManager, TSharedPtr<IMessageBus, ESPMode::ThreadSafe> InBus)
+void FDisplayClusterMessageInterceptor::Setup(IDisplayClusterClusterManager* InClusterManager, const FMessageInterceptionSettings& InInterceptionSettings)
 {
 	ClusterManager = InClusterManager;
-	InterceptedBus = InBus;
+	InterceptionSettings = InInterceptionSettings;
+
+	//Setup desired messages to be intercepted
+	if (InterceptionSettings.bInterceptMultiUserMessages)
+	{
+		InterceptedMessages.Emplace(DisplayClusterMessageInterceptorUtils::MultiUserInterception);
+	}
 }
 
 void FDisplayClusterMessageInterceptor::Purge()
@@ -44,20 +56,25 @@ void FDisplayClusterMessageInterceptor::Purge()
 	}
 }
 
-void FDisplayClusterMessageInterceptor::Start()
+void FDisplayClusterMessageInterceptor::Start(TSharedPtr<IMessageBus, ESPMode::ThreadSafe> InBus)
 {
-	const UDisplayClusterMessageInterceptionSettings* InterceptionSettings = GetDefault<UDisplayClusterMessageInterceptionSettings>();
-	if (!bIsIntercepting && InterceptionSettings->bIsEnabled && InterceptedBus)
+	InterceptedBus = InBus;
+
+	if (!bIsIntercepting && InterceptionSettings.bIsEnabled && InterceptedBus)
 	{
-		InterceptedAnnotation = InterceptionSettings->Annotation;
-		FString DisplayString = InterceptedAnnotation.ToString();
-		UE_LOG(LogDisplayClusterInterception, Display, TEXT("Starting interception of bus messages with annotation: %s"), *DisplayString);
-		for (const FName& MessageType : InterceptionSettings->MessageTypes)
+		UE_LOG(LogDisplayClusterInterception, Display, TEXT("Starting interception of bus messages"));
+
+		for (const FInterceptedMessageDescriptor& Descriptor : InterceptedMessages)
 		{
-			InterceptedBus->Intercept(AsShared(), MessageType);
-			DisplayString = MessageType.ToString();
-			UE_LOG(LogDisplayClusterInterception, Display, TEXT("Intercepted message type: %s"), *DisplayString);
+			const FString AnnotationString = Descriptor.Annotation.ToString();
+			for (const FName& MessageType : Descriptor.MessageTypes)
+			{
+				InterceptedBus->Intercept(AsShared(), MessageType);
+				UE_LOG(LogDisplayClusterInterception, Display, TEXT("Intercepting message type: '%s' with annotation '%s'"), *MessageType.ToString(), *AnnotationString);
+			}
+
 		}
+
 		bIsIntercepting = true;
 	}
 }
@@ -71,6 +88,8 @@ void FDisplayClusterMessageInterceptor::Stop()
 		Purge();
 		UE_LOG(LogDisplayClusterInterception, Display, TEXT("Stopping interception of bus messages."));
 	}
+
+	InterceptedBus.Reset();
 }
 
 void FDisplayClusterMessageInterceptor::SyncMessages()
@@ -88,31 +107,33 @@ void FDisplayClusterMessageInterceptor::SyncMessages()
 		for (const FString& MessageId : MessageIds)
 		{
 			SyncMessagesEvent.Type = MessageId;					// the actually message id we received
-			ClusterManager->EmitClusterEvent(SyncMessagesEvent, false);
+			const bool bMasterOnly = false; //All nodes are broadcasting events to synchronize them across cluster
+			ClusterManager->EmitClusterEvent(SyncMessagesEvent, bMasterOnly);
 			UE_LOG(LogDisplayClusterInterception, VeryVerbose, TEXT("Emitting cluster event for message %s on frame %d"), *MessageId, GFrameCounter);
 		}
 	}
 
 	// remove out of date messages that are not marked reliable
-	const FTimespan MessageTimeoutSpan = FTimespan(0, 0, 1);
-	FDateTime UtcNow = FDateTime::UtcNow();
+	const FTimespan MessageTimeoutSpan = FTimespan::FromSeconds(InterceptionSettings.TimeoutSeconds);
+	const FDateTime UtcNow = FDateTime::UtcNow();
 
 	TArray<TSharedPtr<IMessageContext, ESPMode::ThreadSafe>> ContextToForward;
 	{
 		FScopeLock Lock(&ContextQueueCS);
 		for (auto It = ContextMap.CreateIterator(); It; ++It)
 		{
-			if (It.Value().ContextPtr->GetTimeSent() + MessageTimeoutSpan <= UtcNow)
+			if (It.Value().ContextPtr->GetTimeForwarded() + MessageTimeoutSpan <= UtcNow)
 			{
 				if (!EnumHasAnyFlags(It.Value().ContextPtr->GetFlags(), EMessageFlags::Reliable))
 				{
-					UE_LOG(LogDisplayClusterInterception, VeryVerbose, TEXT("discarding unreliable message %s left intercepted for more than 1s"), *It.Key());
+					UE_LOG(LogDisplayClusterInterception, VeryVerbose, TEXT("Discarding unreliable message '%s' left intercepted for more than %0.1f seconds"), *It.Key(), InterceptionSettings.TimeoutSeconds);
 					It.RemoveCurrent();
 				}
 				else
 				{
-					UE_LOG(LogDisplayClusterInterception, Warning, TEXT("Forcing dispatching of reliable message %s left intercepted for more than 1s"), *It.Key());
-					// Force treatment if not synced after a second and the message is reliable
+					UE_LOG(LogDisplayClusterInterception, Warning, TEXT("Forcing dispatching of reliable message '%s' left intercepted for more than %0.1f seconds"), *It.Key(), InterceptionSettings.TimeoutSeconds);
+
+					// Force treatment if not synced after a certain amount of time and the message is reliable
 					ContextToForward.Add(It.Value().ContextPtr);
 				}
 			}
@@ -163,7 +184,6 @@ const FGuid& FDisplayClusterMessageInterceptor::GetInterceptorId() const
 	return InterceptorId;
 }
 
-
 bool FDisplayClusterMessageInterceptor::InterceptMessage(const TSharedRef<IMessageContext, ESPMode::ThreadSafe>& Context)
 {
 	// we do not intercept forwarded message, they will be either coming off of the network or being forwarded by ourselves
@@ -172,15 +192,24 @@ bool FDisplayClusterMessageInterceptor::InterceptMessage(const TSharedRef<IMessa
 		return false;
 	}
 
-	const FString* MessageId = Context->GetAnnotations().Find(InterceptedAnnotation);
-	if (Context->GetForwarder() != Address && MessageId)
+	if (Context->GetForwarder() != Address)
 	{
-		UE_LOG(LogDisplayClusterInterception, VeryVerbose, TEXT("Intercepting message %s"), **MessageId);
+		const FName MessageType = Context->GetMessageType();
+		const FInterceptedMessageDescriptor* Descriptor = InterceptedMessages.FindByPredicate([MessageType](const FInterceptedMessageDescriptor& Other) { return Other.MessageTypes.Contains(MessageType); });
+		if (Descriptor)
+		{
+			const FString* MessageId = Context->GetAnnotations().Find(Descriptor->Annotation);
+			if (MessageId)
+			{
+				UE_LOG(LogDisplayClusterInterception, VeryVerbose, TEXT("Intercepted message '%s'"), **MessageId);
 
-		FScopeLock Lock(&ContextQueueCS);
-		ContextMap.Add(*MessageId, FContextSync(Context));
-		return true;
+				FScopeLock Lock(&ContextQueueCS);
+				ContextMap.Add(*MessageId, FContextSync(Context));
+				return true;
+			}
+		}
 	}
+
 	return false;
 }
 
diff --git a/Engine/Plugins/Runtime/nDisplay/Source/DisplayClusterMessageInterceptor/Private/DisplayClusterMessageInterceptor.h b/Engine/Plugins/Runtime/nDisplay/Source/DisplayClusterMessageInterceptor/Private/DisplayClusterMessageInterceptor.h
--- a/Engine/Plugins/Runtime/nDisplay/Source/DisplayClusterMessageInterceptor/Private/DisplayClusterMessageInterceptor.h
+++ b/Engine/Plugins/Runtime/nDisplay/Source/DisplayClusterMessageInterceptor/Private/DisplayClusterMessageInterceptor.h
@@ -11,6 +11,21 @@ class IMessageBus;
 class IDisplayClusterClusterManager;
 struct FDisplayClusterClusterEvent;
 
+
+DECLARE_LOG_CATEGORY_EXTERN(LogDisplayClusterInterception, Log, All);
+
+
+struct FInterceptedMessageDescriptor
+{
+	FInterceptedMessageDescriptor(TArray<FName>&& InMessageTypes, FName InAnnotation)
+		: MessageTypes(MoveTemp(InMessageTypes))
+		, Annotation(InAnnotation)
+	{}
+
+	TArray<FName> MessageTypes;
+	FName Annotation;
+};
+
 /**
  * Display Cluster Message Interceptor 
  * Intercept message bus messages to process them in sync across the display cluster
@@ -22,9 +37,9 @@ class FDisplayClusterMessageInterceptor : public TSharedFromThis<FDisplayCluster
 public:
 	FDisplayClusterMessageInterceptor();
 
-	void Setup(IDisplayClusterClusterManager* InClusterManager, TSharedPtr<IMessageBus, ESPMode::ThreadSafe> InBus);
+	void Setup(IDisplayClusterClusterManager* InClusterManager, const FMessageInterceptionSettings& InInterceptionSettings);
 	
-	void Start();
+	void Start(TSharedPtr<IMessageBus, ESPMode::ThreadSafe> InBus);
 	void Stop();
 	
 	void SyncMessages();
@@ -51,9 +66,6 @@ private:
 	/** Sender address needed by message bus. */
 	FMessageAddress Address;
 
-	/** Message with this annotation will be intercepted. */
-	FName InterceptedAnnotation;
-
 	/** The bus this interceptor is intercepting messages on. */
 	TSharedPtr<IMessageBus, ESPMode::ThreadSafe> InterceptedBus;
 
@@ -63,6 +75,12 @@ private:
 	/** Critical section for the intercepted messages. */
 	FCriticalSection ContextQueueCS;
 
+	/** Settings to be used for interception */
+	FMessageInterceptionSettings InterceptionSettings;
+
+	/** Description of messages we intercept */
+	TArray<FInterceptedMessageDescriptor> InterceptedMessages;
+
 	/** Holds intercepted messages and their current sync state across the cluster. */
 	struct FContextSync
 	{
diff --git a/Engine/Plugins/Runtime/nDisplay/Source/DisplayClusterMessageInterceptor/Public/DisplayClusterMessageInterceptionSettings.h b/Engine/Plugins/Runtime/nDisplay/Source/DisplayClusterMessageInterceptor/Public/DisplayClusterMessageInterceptionSettings.h
--- a/Engine/Plugins/Runtime/nDisplay/Source/DisplayClusterMessageInterceptor/Public/DisplayClusterMessageInterceptionSettings.h
+++ b/Engine/Plugins/Runtime/nDisplay/Source/DisplayClusterMessageInterceptor/Public/DisplayClusterMessageInterceptionSettings.h
@@ -5,21 +5,43 @@
 #include "UObject/ObjectMacros.h"
 #include "DisplayClusterMessageInterceptionSettings.generated.h"
 
-UCLASS(config=Engine)
-class DISPLAYCLUSTERMESSAGEINTERCEPTION_API UDisplayClusterMessageInterceptionSettings : public UObject
+
+USTRUCT()
+struct FMessageInterceptionSettings
 {
 	GENERATED_BODY()
 
-public:
+	/** 
+	 * By default, interception is enabled and multi user interception is active
+	 * Old config properties are not redirectable so by default, new system will be enabled
+	 * fixing any discrepencies across the system with settings synchronization.
+	 */
+
 	/** Indicates if message interception is enabled. */
 	UPROPERTY(config, EditAnywhere, Category = "Interception Settings")
-	bool bIsEnabled = false;
+	bool bIsEnabled = true;
 
-	/** List of message types that should be taken into consideration for interception. */
-	UPROPERTY(config, EditAnywhere, Category = "Interception Settings")
-	TArray<FName> MessageTypes;
+	/** Indicates whether messages from multi user are intercepted. */
+	UPROPERTY(config, EditAnywhere, Category = "Interception Settings", meta=(EditCondition = "bIsEnabled"))
+	bool bInterceptMultiUserMessages = true;
+	
+	/** Maximum seconds to keep intercepted message to be synchronized across cluster. */
+	UPROPERTY(config, EditAnywhere, AdvancedDisplay, Category = "Interception Settings")
+	float TimeoutSeconds = 1.0f;
+};
 
-	/** The annotations to look for identifying messages across the cluster. */
-	UPROPERTY(config, EditAnywhere, Category = "Interception Settings")
-	FName Annotation;
+
+UCLASS(config=Engine)
+class DISPLAYCLUSTERMESSAGEINTERCEPTION_API UDisplayClusterMessageInterceptionSettings : public UObject
+{
+	GENERATED_BODY()
+
+public:
+
+	/**
+	 * Settings related to interception. 
+	 * @note Settings from master will be synchronized across the cluster
+	 */
+	UPROPERTY(config, EditAnywhere, Category = "Interception Settings", meta = (ShowOnlyInnerProperties))
+	FMessageInterceptionSettings InterceptionSettings;
 };
\ No newline at end of file
