Deprecation fixes in UDisplayClusterEditorEngine

#rnx

diff --git a/Engine/Plugins/Runtime/nDisplay/Source/DisplayClusterEditor/Private/DisplayClusterEditorEngine.cpp b/Engine/Plugins/Runtime/nDisplay/Source/DisplayClusterEditor/Private/DisplayClusterEditorEngine.cpp
--- a/Engine/Plugins/Runtime/nDisplay/Source/DisplayClusterEditor/Private/DisplayClusterEditorEngine.cpp
+++ b/Engine/Plugins/Runtime/nDisplay/Source/DisplayClusterEditor/Private/DisplayClusterEditorEngine.cpp
@@ -39,9 +39,9 @@ void UDisplayClusterEditorEngine::PreExit()
 	Super::PreExit();
 }
 
-void UDisplayClusterEditorEngine::PlayInEditor(UWorld* InWorld, bool bInSimulateInEditor, FPlayInEditorOverrides Overrides)
+void UDisplayClusterEditorEngine::StartPlayInEditorSession(FRequestPlaySessionParams& InRequestParams)
 {
-	UE_LOG(LogDisplayClusterEditorEngine, VeryVerbose, TEXT("UDisplayClusterEditorEngine::PlayInEditor"));
+	UE_LOG(LogDisplayClusterEditorEngine, VeryVerbose, TEXT("UDisplayClusterEditorEngine::StartPlayInEditorSession"));
 
 	Super::PlayInEditor(InWorld, bInSimulateInEditor, Overrides);
 }
diff --git a/Engine/Plugins/Runtime/nDisplay/Source/DisplayClusterEditor/Private/DisplayClusterEditorEngine.h b/Engine/Plugins/Runtime/nDisplay/Source/DisplayClusterEditor/Private/DisplayClusterEditorEngine.h
--- a/Engine/Plugins/Runtime/nDisplay/Source/DisplayClusterEditor/Private/DisplayClusterEditorEngine.h
+++ b/Engine/Plugins/Runtime/nDisplay/Source/DisplayClusterEditor/Private/DisplayClusterEditorEngine.h
@@ -21,7 +21,7 @@ class UDisplayClusterEditorEngine
 public:
 	virtual void Init(IEngineLoop* InEngineLoop) override;
 	virtual void PreExit() override;
-	virtual void PlayInEditor(UWorld* InWorld, bool bInSimulateInEditor, FPlayInEditorOverrides Overrides = FPlayInEditorOverrides()) override;
+	virtual void StartPlayInEditorSession(FRequestPlaySessionParams& InRequestParams) override;
 
 private:
 	
