CIS fixes.

#rnx

diff --git a/Engine/Plugins/Runtime/nDisplay/Source/DisplayClusterEditor/Private/DisplayClusterEditorEngine.cpp b/Engine/Plugins/Runtime/nDisplay/Source/DisplayClusterEditor/Private/DisplayClusterEditorEngine.cpp
--- a/Engine/Plugins/Runtime/nDisplay/Source/DisplayClusterEditor/Private/DisplayClusterEditorEngine.cpp
+++ b/Engine/Plugins/Runtime/nDisplay/Source/DisplayClusterEditor/Private/DisplayClusterEditorEngine.cpp
@@ -43,5 +43,5 @@ void UDisplayClusterEditorEngine::StartPlayInEditorSession(FRequestPlaySessionPa
 {
 	UE_LOG(LogDisplayClusterEditorEngine, VeryVerbose, TEXT("UDisplayClusterEditorEngine::StartPlayInEditorSession"));
 
-	Super::PlayInEditor(InWorld, bInSimulateInEditor, Overrides);
+	Super::StartPlayInEditorSession(InRequestParams);
 }
