CIS: Shadowed Member variable fix.

#rnx

diff --git a/Engine/Plugins/Runtime/nDisplay/Source/DisplayClusterEditor/Private/DisplayClusterEditorEngine.cpp b/Engine/Plugins/Runtime/nDisplay/Source/DisplayClusterEditor/Private/DisplayClusterEditorEngine.cpp
--- a/Engine/Plugins/Runtime/nDisplay/Source/DisplayClusterEditor/Private/DisplayClusterEditorEngine.cpp
+++ b/Engine/Plugins/Runtime/nDisplay/Source/DisplayClusterEditor/Private/DisplayClusterEditorEngine.cpp
@@ -49,11 +49,11 @@ void UDisplayClusterEditorEngine::StartPlayInEditorSession(FRequestPlaySessionPa
 {
 	UE_LOG(LogDisplayClusterEditorEngine, VeryVerbose, TEXT("UDisplayClusterEditorEngine::StartPlayInEditorSession"));
 
-	UWorld* EditorWorld = GetEditorWorldContext().World();
+	UWorld* EditorWorldPreDup = GetEditorWorldContext().World();
 
 	// Find nDisplay root
 	UDisplayClusterRootComponent* RootComponent = nullptr;
-	for (AActor* Actor : EditorWorld->PersistentLevel->Actors)
+	for (AActor* Actor : EditorWorldPreDup->PersistentLevel->Actors)
 	{
 		if (Actor && !Actor->IsPendingKill())
 		{
@@ -86,7 +86,7 @@ void UDisplayClusterEditorEngine::StartPlayInEditorSession(FRequestPlaySessionPa
 				return;
 			}
 
-			DisplayClusterModule->StartScene(EditorWorld);
+			DisplayClusterModule->StartScene(EditorWorldPreDup);
 		}
 	}
 
