Aquire sync data on main thread only

diff --git a/Engine/Plugins/Runtime/nDisplay/Source/DisplayCluster/Private/Cluster/Controller/DisplayClusterClusterNodeCtrlMaster.cpp b/Engine/Plugins/Runtime/nDisplay/Source/DisplayCluster/Private/Cluster/Controller/DisplayClusterClusterNodeCtrlMaster.cpp
--- a/Engine/Plugins/Runtime/nDisplay/Source/DisplayCluster/Private/Cluster/Controller/DisplayClusterClusterNodeCtrlMaster.cpp
+++ b/Engine/Plugins/Runtime/nDisplay/Source/DisplayCluster/Private/Cluster/Controller/DisplayClusterClusterNodeCtrlMaster.cpp
@@ -25,17 +25,31 @@
 FDisplayClusterClusterNodeCtrlMaster::FDisplayClusterClusterNodeCtrlMaster(const FString& ctrlName, const FString& nodeName) :
 	FDisplayClusterClusterNodeCtrlSlave(ctrlName, nodeName)
 {
-	IsDataCached_GetSyncData.Emplace(EDisplayClusterSyncGroup::PreTick,  false);
-	IsDataCached_GetSyncData.Emplace(EDisplayClusterSyncGroup::Tick,     false);
-	IsDataCached_GetSyncData.Emplace(EDisplayClusterSyncGroup::PostTick, false);
-
-	CachedSyncData_GetSyncData.Emplace(EDisplayClusterSyncGroup::PreTick);
-	CachedSyncData_GetSyncData.Emplace(EDisplayClusterSyncGroup::Tick);
-	CachedSyncData_GetSyncData.Emplace(EDisplayClusterSyncGroup::PostTick);
+	CachedSyncData.Emplace(EDisplayClusterSyncGroup::PreTick);
+	CachedSyncData.Emplace(EDisplayClusterSyncGroup::Tick);
+	CachedSyncData.Emplace(EDisplayClusterSyncGroup::PostTick);
+
+	CachedSyncDataEvents.Emplace(EDisplayClusterSyncGroup::PreTick, FPlatformProcess::CreateSynchEvent(true));
+	CachedSyncDataEvents.Emplace(EDisplayClusterSyncGroup::Tick, FPlatformProcess::CreateSynchEvent(true));
+	CachedSyncDataEvents.Emplace(EDisplayClusterSyncGroup::PostTick, FPlatformProcess::CreateSynchEvent(true));
+
+	CachedInputDataEvent  = FPlatformProcess::CreateSynchEvent(true);
+	CachedDeltaTimeEvent  = FPlatformProcess::CreateSynchEvent(true);
+	CachedEventsDataEvent = FPlatformProcess::CreateSynchEvent(true);
+	CachedTimeCodeFrameRateEvent = FPlatformProcess::CreateSynchEvent(true);
 }
 
 FDisplayClusterClusterNodeCtrlMaster::~FDisplayClusterClusterNodeCtrlMaster()
 {
+	delete CachedInputDataEvent;
+	delete CachedDeltaTimeEvent;
+	delete CachedEventsDataEvent;
+	delete CachedTimeCodeFrameRateEvent;
+
+	for (auto& it : CachedSyncDataEvents)
+	{
+		delete it.Value;
+	}
 }
 
 
@@ -44,57 +58,51 @@ FDisplayClusterClusterNodeCtrlMaster::~FDisplayClusterClusterNodeCtrlMaster()
 //////////////////////////////////////////////////////////////////////////////////////////////
 void FDisplayClusterClusterNodeCtrlMaster::GetDeltaTime(float& DeltaSeconds)
 {
-	FScopeLock lock(&InternalsSyncScope);
-
-	// Cache data so it will be the same for all requests within current frame
-	if (!bIsDataCached_GetDeltaTime)
+	if (IsInGameThread())
 	{
-		bIsDataCached_GetDeltaTime = true;
-
-		CachedDeltaTime_GetDeltaTime = FApp::GetDeltaTime();
-
-		UE_LOG(LogDisplayClusterCluster, Verbose, TEXT("GetDeltaTime cached values: DeltaSeconds %f"), CachedDeltaTime_GetDeltaTime);
+		// Cache data so it will be the same for all requests within current frame
+		CachedDeltaTime = FApp::GetDeltaTime();
+		UE_LOG(LogDisplayClusterCluster, Verbose, TEXT("GetDeltaTime cached values: DeltaSeconds %f"), CachedDeltaTime);
+		CachedDeltaTimeEvent->Trigger();
 	}
 
+	// Wait until data is available
+	CachedDeltaTimeEvent->Wait();
+
 	// Return cached value
-	DeltaSeconds = CachedDeltaTime_GetDeltaTime;
+	DeltaSeconds = CachedDeltaTime;
 }
 
 void FDisplayClusterClusterNodeCtrlMaster::GetTimecode(FTimecode& Timecode, FFrameRate& FrameRate)
 {
-	FScopeLock lock(&InternalsSyncScope);
-
-	// Cache data so it will be the same for all requests within current frame
-	if (!bIsDataCached_GetTimecode)
+	if (IsInGameThread())
 	{
-		bIsDataCached_GetTimecode = true;
-
 		// This values are updated in UEngine::UpdateTimeAndHandleMaxTickRate (via UpdateTimecode).
-		CachedTimecode_GetTimecode  = FApp::GetTimecode();
-		CachedFramerate_GetTimecode = FApp::GetTimecodeFrameRate();
+		CachedTimecode  = FApp::GetTimecode();
+		CachedFramerate = FApp::GetTimecodeFrameRate();
 
-		UE_LOG(LogDisplayClusterCluster, Verbose, TEXT("GetDeltaTime cached values: DeltaSeconds %f"), CachedDeltaTime_GetDeltaTime);
+		UE_LOG(LogDisplayClusterCluster, Verbose, TEXT("GetTimecode cached values: Timecode %s, FrameRate %s"), *CachedTimecode.ToString(), *FDisplayClusterTypesConverter::ToString(CachedFramerate));
+
+		CachedTimeCodeFrameRateEvent->Trigger();
 	}
 
+	CachedTimeCodeFrameRateEvent->Wait();
+
 	// Return cached value
-	Timecode  = CachedTimecode_GetTimecode;
-	FrameRate = CachedFramerate_GetTimecode;
+	Timecode  = CachedTimecode;
+	FrameRate = CachedFramerate;
 }
 
 void FDisplayClusterClusterNodeCtrlMaster::GetSyncData(FDisplayClusterMessage::DataType& SyncData, EDisplayClusterSyncGroup SyncGroup)
 {
-	FScopeLock lock(&InternalsSyncScope);
-
 	static IPDisplayClusterClusterManager* const ClusterMgr = GDisplayCluster->GetPrivateClusterMgr();
 
-	// Cache data so it will be the same for all requests within current frame
-	if (!IsDataCached_GetSyncData[SyncGroup])
+	if (IsInGameThread())
 	{
-		IsDataCached_GetSyncData[SyncGroup] = true;
-
+		// Cache data so it will be the same for all requests within current frame
 		FDisplayClusterMessage::DataType SyncDataToCache;
 		ClusterMgr->ExportSyncData(SyncDataToCache, SyncGroup);
-		CachedSyncData_GetSyncData.Emplace(SyncGroup, SyncDataToCache);
+		CachedSyncData.Emplace(SyncGroup, SyncDataToCache);
 
 		UE_LOG(LogDisplayClusterCluster, Verbose, TEXT("GetSyncData cached values amount: %d"), SyncDataToCache.Num());
 
@@ -103,68 +111,78 @@ void FDisplayClusterClusterNodeCtrlMaster::GetSyncData(FDisplayClusterMessage::D
 		{
 			UE_LOG(LogDisplayClusterCluster, Verbose, TEXT("GetSyncData cached value %d: %s - %s"), i++, *it->Key, *it->Value);
 		}
+
+		UE_LOG(LogDisplayClusterCluster, Verbose, TEXT("SyncData has %d records"), CachedSyncData[SyncGroup].Num());
+
+		// Notify data is available
+		CachedSyncDataEvents[SyncGroup]->Trigger();
 	}
 
-	UE_LOG(LogDisplayClusterCluster, Verbose, TEXT("SyncData has %d records"), CachedSyncData_GetSyncData[SyncGroup].Num());
+	// Wait until data is available
+	CachedSyncDataEvents[SyncGroup]->Wait();
 
 	// Return cached value
-	SyncData = CachedSyncData_GetSyncData[SyncGroup];
+	SyncData = CachedSyncData[SyncGroup];
 }
 
 void FDisplayClusterClusterNodeCtrlMaster::GetInputData(FDisplayClusterMessage::DataType& InputData)
 {
-	FScopeLock lock(&InternalsSyncScope);
-
 	static IPDisplayClusterInputManager* const InputMgr = GDisplayCluster->GetPrivateInputMgr();
 
-	// Cache data so it will be the same for all requests within current frame
-	if (!bIsDataCached_GetInputData)
+	if (IsInGameThread())
 	{
-		bIsDataCached_GetInputData = true;
-
-		InputMgr->ExportInputData(CachedInputData_GetInputData);
+		// Cache data so it will be the same for all requests within current frame
+		InputMgr->ExportInputData(CachedInputData);
 
-		UE_LOG(LogDisplayClusterCluster, Verbose, TEXT("GetInputData cached values amount: %d"), CachedInputData_GetInputData.Num());
+		UE_LOG(LogDisplayClusterCluster, Verbose, TEXT("GetInputData cached values amount: %d"), CachedInputData.Num());
 
 		int i = 0;
-		for (auto it = CachedInputData_GetInputData.CreateConstIterator(); it; ++it)
+		for (auto it = CachedInputData.CreateConstIterator(); it; ++it)
 		{
 			UE_LOG(LogDisplayClusterCluster, Verbose, TEXT("GetInputData cached value %d: %s - %s"), i++, *it->Key, *it->Value);
 		}
+
+		UE_LOG(LogDisplayClusterCluster, Verbose, TEXT("InputData has %d records"), CachedInputData.Num());
+
+		// Notify data is available
+		CachedInputDataEvent->Trigger();
 	}
 
-	UE_LOG(LogDisplayClusterCluster, Verbose, TEXT("InputData has %d records"), CachedInputData_GetInputData.Num());
+	// Wait until data is available
+	CachedInputDataEvent->Wait();
 
 	// Return cached value
-	InputData = CachedInputData_GetInputData;
+	InputData = CachedInputData;
 }
 
 void FDisplayClusterClusterNodeCtrlMaster::GetEventsData(FDisplayClusterMessage::DataType& EventsData)
 {
-	FScopeLock lock(&InternalsSyncScope);
-
 	static IPDisplayClusterClusterManager* const ClusterMgr = GDisplayCluster->GetPrivateClusterMgr();
 
-	// Cache data so it will be the same for all requests within current frame
-	if (!bIsDataCached_GetEventsData)
+	if (IsInGameThread())
 	{
-		bIsDataCached_GetEventsData = true;
+		// Cache data so it will be the same for all requests within current frame
+		ClusterMgr->ExportEventsData(CachedEventsData);
 
-		ClusterMgr->ExportEventsData(CachedEventsData_GetEventsData);
-
-		UE_LOG(LogDisplayClusterCluster, Verbose, TEXT("GetEventsData cached values amount: %d"), CachedEventsData_GetEventsData.Num());
+		UE_LOG(LogDisplayClusterCluster, Verbose, TEXT("GetEventsData cached values amount: %d"), CachedEventsData.Num());
 
 		int i = 0;
-		for (auto it = CachedEventsData_GetEventsData.CreateConstIterator(); it; ++it)
+		for (auto it = CachedEventsData.CreateConstIterator(); it; ++it)
 		{
 			UE_LOG(LogDisplayClusterCluster, Verbose, TEXT("GetEventsData cached value %d: %s - %s"), i++, *it->Key, *it->Value);
 		}
+
+		UE_LOG(LogDisplayClusterCluster, Verbose, TEXT("EventsData has %d records"), CachedEventsData.Num());
+
+		// Notify data is available
+		CachedEventsDataEvent->Trigger();
 	}
 
-	UE_LOG(LogDisplayClusterCluster, Verbose, TEXT("EventsData has %d records"), CachedEventsData_GetEventsData.Num());
+	// Wait until data is available
+	CachedEventsDataEvent->Wait();
 
 	// Return cached value
-	EventsData = CachedEventsData_GetEventsData;
+	EventsData = CachedEventsData;
 }
 
 void FDisplayClusterClusterNodeCtrlMaster::GetNativeInputData(FDisplayClusterMessage::DataType& NativeInputData)
@@ -181,22 +199,22 @@ void FDisplayClusterClusterNodeCtrlMaster::ClearCache()
 {
 	FScopeLock lock(&InternalsSyncScope);
 
-	// Reset all cache flags
-	bIsDataCached_GetDeltaTime  = false;
-	bIsDataCached_GetTimecode   = false;
-	bIsDataCached_GetInputData  = false;
-	bIsDataCached_GetEventsData = false;
+	// Reset all cache events
+	CachedDeltaTimeEvent->Reset();
+	CachedTimeCodeFrameRateEvent->Reset();
+	CachedEventsDataEvent->Reset();
+	CachedInputDataEvent->Reset();
 
 	// Reset cache containers
-	CachedInputData_GetInputData.Reset();
-	CachedEventsData_GetEventsData.Reset();
+	CachedInputData.Reset();
+	CachedEventsData.Reset();
 
-	for (auto& it : IsDataCached_GetSyncData)
+	for (auto& it : CachedSyncDataEvents)
 	{
-		it.Value = false;
+		it.Value->Reset();
 	}
 
-	for (auto& it : CachedSyncData_GetSyncData)
+	for (auto& it : CachedSyncData)
 	{
 		it.Value.Reset();
 	}
diff --git a/Engine/Plugins/Runtime/nDisplay/Source/DisplayCluster/Private/Cluster/Controller/DisplayClusterClusterNodeCtrlMaster.h b/Engine/Plugins/Runtime/nDisplay/Source/DisplayCluster/Private/Cluster/Controller/DisplayClusterClusterNodeCtrlMaster.h
--- a/Engine/Plugins/Runtime/nDisplay/Source/DisplayCluster/Private/Cluster/Controller/DisplayClusterClusterNodeCtrlMaster.h
+++ b/Engine/Plugins/Runtime/nDisplay/Source/DisplayCluster/Private/Cluster/Controller/DisplayClusterClusterNodeCtrlMaster.h
@@ -62,25 +62,25 @@ private:
 
 private:
 	// GetDeltaTime internals
-	bool  bIsDataCached_GetDeltaTime = false;
-	float CachedDeltaTime_GetDeltaTime = 0.f;
+	FEvent* CachedDeltaTimeEvent = nullptr;
+	float   CachedDeltaTime = 0.f;
 
 	// GetTimecode internals
-	bool       bIsDataCached_GetTimecode = false;
-	FTimecode  CachedTimecode_GetTimecode;
-	FFrameRate CachedFramerate_GetTimecode;
+	FEvent* CachedTimeCodeFrameRateEvent = nullptr;
+	FTimecode  CachedTimecode;
+	FFrameRate CachedFramerate;
 
 	// GetSyncData internals
-	TMap<EDisplayClusterSyncGroup, bool> IsDataCached_GetSyncData;
-	TMap<EDisplayClusterSyncGroup, FDisplayClusterMessage::DataType> CachedSyncData_GetSyncData;
+	TMap<EDisplayClusterSyncGroup, FEvent*> CachedSyncDataEvents;
+	TMap<EDisplayClusterSyncGroup, FDisplayClusterMessage::DataType> CachedSyncData;
 
 	// GetInputData internals
-	bool bIsDataCached_GetInputData = false;
-	FDisplayClusterMessage::DataType CachedInputData_GetInputData;
+	FEvent* CachedInputDataEvent = nullptr;
+	FDisplayClusterMessage::DataType CachedInputData;
 
 	// GetEventsData internals
-	bool bIsDataCached_GetEventsData = false;
-	FDisplayClusterMessage::DataType CachedEventsData_GetEventsData;
+	FEvent* CachedEventsDataEvent = nullptr;
+	FDisplayClusterMessage::DataType CachedEventsData;
 
 private:
 	mutable FCriticalSection InternalsSyncScope;
diff --git a/Engine/Plugins/Runtime/nDisplay/Source/DisplayCluster/Private/Cluster/DisplayClusterClusterManager.cpp b/Engine/Plugins/Runtime/nDisplay/Source/DisplayCluster/Private/Cluster/DisplayClusterClusterManager.cpp
--- a/Engine/Plugins/Runtime/nDisplay/Source/DisplayCluster/Private/Cluster/DisplayClusterClusterManager.cpp
+++ b/Engine/Plugins/Runtime/nDisplay/Source/DisplayCluster/Private/Cluster/DisplayClusterClusterManager.cpp
@@ -207,6 +207,16 @@ void FDisplayClusterClusterManager::PreTick(float DeltaSeconds)
 {
 	DISPLAY_CLUSTER_FUNC_TRACE(LogDisplayClusterCluster);
 
+	// Move cluster events from the primary pool to the output pool. These will be synchronized on the current frame.
+	{
+		FScopeLock lock(&ClusterEventsCritSec);
+
+		// Clear the output pool since we have all data cached already
+		ClusterEventsPoolOut.Empty(ClusterEventsPoolOut.Num() | 0x07);
+		ClusterEventsPoolOut = MoveTemp(ClusterEventsPoolMain);
+		ClusterEventsPoolMain.Empty(ClusterEventsPoolOut.Num() | 0x07);
+	}
+
 	// Update input state in the cluster
 	SyncInput();
 
@@ -215,15 +225,6 @@ void FDisplayClusterClusterManager::PreTick(float DeltaSeconds)
 
 	// Sync cluster events
 	SyncEvents();
-
-	// Move cluster events from the primary pool to the output pool. These will be synchronized on the current frame.
-	{
-		FScopeLock lock(& ClusterEventsCritSec);
-
-		ClusterEventsPoolOut = MoveTemp(ClusterEventsPoolMain);
-		ClusterEventsPoolMain.Empty(ClusterEventsPoolOut.Num() | 0x07);
-		ClusterEventsCacheOut.Empty(ClusterEventsPoolOut.Num() | 0x07);
-	}
 }
 
 void FDisplayClusterClusterManager::Tick(float DeltaSeconds)
@@ -463,16 +464,11 @@ void FDisplayClusterClusterManager::ExportEventsData(FDisplayClusterMessage::Dat
 					for (const auto& NamedEvent : TypeMap.Value)
 					{
 						UE_LOG(LogDisplayClusterCluster, Verbose, TEXT("Adding event to sync: %s::%s"), *NamedEvent.Value.Name, *NamedEvent.Value.Type);
-						ClusterEventsCacheOut.Add(FString::Printf(TEXT("EVENT_%d"), ObjID++), NamedEvent.Value.SerializeToString());
+						EventsData.Add(FString::Printf(TEXT("EVENT_%d"), ObjID++), NamedEvent.Value.SerializeToString());
 					}
 				}
 			}
-
-			// Clear the output pool since we have all data cached already
-			ClusterEventsPoolOut.Empty(ClusterEventsPoolOut.Num() | 0x07);
 		}
-
-		EventsData = ClusterEventsCacheOut;
 	}
 }
 
@@ -542,28 +538,15 @@ void FDisplayClusterClusterManager::SyncEvents()
 {
 	DISPLAY_CLUSTER_FUNC_TRACE(LogDisplayClusterCluster);
 
-	// No need to do the sync for master
-	if (IsSlave())
-	{
-		UE_LOG(LogDisplayClusterCluster, Verbose, TEXT("Downloading synchronization data (events)..."));
-		TMap<FString, FString> EventsData;
+	TMap<FString, FString> EventsData;
 
-		{
-			FScopeLock lock(&ClusterEventsCritSec);
-			Controller->GetEventsData(EventsData);
-		}
+	// Get events data from a provider
+	UE_LOG(LogDisplayClusterCluster, Verbose, TEXT("Downloading synchronization data (events)..."));
+	Controller->GetEventsData(EventsData);
+	UE_LOG(LogDisplayClusterCluster, Verbose, TEXT("Downloading finished. Available %d records (events)."), EventsData.Num());
 
-		UE_LOG(LogDisplayClusterCluster, Verbose, TEXT("Downloading finished. Available %d records (events)."), EventsData.Num());
-
-		// Perform data load (objects state update)
-		GDisplayCluster->GetPrivateClusterMgr()->ImportEventsData(EventsData);
-	}
-	else
-	{
-		TMap<FString, FString> EventsData;
-		ExportEventsData(EventsData);
-		ImportEventsData(EventsData);
-	}
+	// Import and process them
+	ImportEventsData(EventsData);
 }
 
 void FDisplayClusterClusterManager::ProvideNativeInputData(const TMap<FString, FString>& NativeInputData)
diff --git a/Engine/Plugins/Runtime/nDisplay/Source/DisplayCluster/Private/Cluster/DisplayClusterClusterManager.h b/Engine/Plugins/Runtime/nDisplay/Source/DisplayCluster/Private/Cluster/DisplayClusterClusterManager.h
--- a/Engine/Plugins/Runtime/nDisplay/Source/DisplayCluster/Private/Cluster/DisplayClusterClusterManager.h
+++ b/Engine/Plugins/Runtime/nDisplay/Source/DisplayCluster/Private/Cluster/DisplayClusterClusterManager.h
@@ -120,7 +120,6 @@ private:
 	// Sync events - data
 	FClusterEventsContainer                      ClusterEventsPoolMain;
 	mutable FClusterEventsContainer              ClusterEventsPoolOut;
-	mutable FDisplayClusterMessage::DataType     ClusterEventsCacheOut;
 	mutable FCriticalSection                     ClusterEventsCritSec;
 	FOnClusterEvent                              OnClusterEvent;
 	TArray<TScriptInterface<IDisplayClusterClusterEventListener>> ClusterEventListeners;
