Fixed 32K limit in networking

diff --git a/Engine/Plugins/Runtime/nDisplay/Source/DisplayCluster/Private/DisplayClusterConstants.h b/Engine/Plugins/Runtime/nDisplay/Source/DisplayCluster/Private/DisplayClusterConstants.h
--- a/Engine/Plugins/Runtime/nDisplay/Source/DisplayCluster/Private/DisplayClusterConstants.h
+++ b/Engine/Plugins/Runtime/nDisplay/Source/DisplayCluster/Private/DisplayClusterConstants.h
@@ -10,8 +10,7 @@ namespace DisplayClusterConstants
 {
 	namespace net
 	{
-		static constexpr int32  SocketBufferSize            = INT16_MAX; // bytes
-		static constexpr int32  MessageBufferSize           = INT16_MAX; // bytes
+		static constexpr int32  MessageBufferSize           = 4 * 1024 * 1024; // bytes
 	};
 
 	namespace misc
diff --git a/Engine/Plugins/Runtime/nDisplay/Source/DisplayCluster/Private/Network/DisplayClusterClient.cpp b/Engine/Plugins/Runtime/nDisplay/Source/DisplayCluster/Private/Network/DisplayClusterClient.cpp
--- a/Engine/Plugins/Runtime/nDisplay/Source/DisplayCluster/Private/Network/DisplayClusterClient.cpp
+++ b/Engine/Plugins/Runtime/nDisplay/Source/DisplayCluster/Private/Network/DisplayClusterClient.cpp
@@ -67,9 +67,9 @@ void FDisplayClusterClient::Disconnect()
 	}
 }
 
-FSocket* FDisplayClusterClient::CreateSocket(const FString& InName, const int32 BuffSize)
+FSocket* FDisplayClusterClient::CreateSocket(const FString& InName)
 {
-	FSocket* pSock = FTcpSocketBuilder(*InName).AsBlocking().WithReceiveBufferSize(BuffSize).WithSendBufferSize(BuffSize);
+	FSocket* pSock = FTcpSocketBuilder(*InName).AsBlocking();
 	check(pSock);
 	return pSock;
 }
diff --git a/Engine/Plugins/Runtime/nDisplay/Source/DisplayCluster/Private/Network/DisplayClusterClient.h b/Engine/Plugins/Runtime/nDisplay/Source/DisplayCluster/Private/Network/DisplayClusterClient.h
--- a/Engine/Plugins/Runtime/nDisplay/Source/DisplayCluster/Private/Network/DisplayClusterClient.h
+++ b/Engine/Plugins/Runtime/nDisplay/Source/DisplayCluster/Private/Network/DisplayClusterClient.h
@@ -39,7 +39,7 @@ public:
 
 protected:
 	// Creates client socket
-	FSocket* CreateSocket(const FString& InName, const int32 BuffSize = DisplayClusterConstants::net::SocketBufferSize);
+	FSocket* CreateSocket(const FString& InName);
 
 private:
 	// Client name
diff --git a/Engine/Plugins/Runtime/nDisplay/Source/DisplayCluster/Private/Network/DisplayClusterServer.cpp b/Engine/Plugins/Runtime/nDisplay/Source/DisplayCluster/Private/Network/DisplayClusterServer.cpp
--- a/Engine/Plugins/Runtime/nDisplay/Source/DisplayCluster/Private/Network/DisplayClusterServer.cpp
+++ b/Engine/Plugins/Runtime/nDisplay/Source/DisplayCluster/Private/Network/DisplayClusterServer.cpp
@@ -80,11 +80,6 @@ bool FDisplayClusterServer::ConnectionHandler(FSocket* InSock, const FIPv4Endpoi
 		InSock->SetLinger(false, 0);
 		InSock->SetNonBlocking(false);
 
-		const int32 RequestedBufferSize = static_cast<int32>(DisplayClusterConstants::net::SocketBufferSize);
-		int32 FinalBufferSize;
-		InSock->SetReceiveBufferSize(RequestedBufferSize, FinalBufferSize);
-		InSock->SetSendBufferSize(RequestedBufferSize, FinalBufferSize);
-
 		TSharedPtr<FDisplayClusterSessionBase> Session = CreateSession(InSock, InEP);
 		Session->StartSession();
 		Sessions.Add(Session);
diff --git a/Engine/Plugins/Runtime/nDisplay/Source/DisplayCluster/Private/Network/DisplayClusterSocketOps.cpp b/Engine/Plugins/Runtime/nDisplay/Source/DisplayCluster/Private/Network/DisplayClusterSocketOps.cpp
--- a/Engine/Plugins/Runtime/nDisplay/Source/DisplayCluster/Private/Network/DisplayClusterSocketOps.cpp
+++ b/Engine/Plugins/Runtime/nDisplay/Source/DisplayCluster/Private/Network/DisplayClusterSocketOps.cpp
@@ -59,9 +59,15 @@ bool FDisplayClusterSocketOps::SendMsg(const TSharedPtr<FDisplayClusterMessage>&
 	// Initialize the message header
 	FDisplayClusterMessageHeader MessageHeader;
 	const int32 MessageLength = DataBuffer.Num();
-	MessageHeader.Length = static_cast<uint16>(MessageLength & 0xFFFF) - sizeof(FDisplayClusterMessageHeader);
+	MessageHeader.Length = static_cast<uint32>(MessageLength & 0xFFFFFFFF) - sizeof(FDisplayClusterMessageHeader);
 	UE_LOG(LogDisplayClusterNetworkMsg, Verbose, TEXT("Outgoing message body length %d"), MessageHeader.Length);
 
+	if (MessageHeader.Length > INT32_MAX)
+	{
+		UE_LOG(LogDisplayClusterNetworkMsg, Error, TEXT("Message length %u exceeds maximum limit"), MessageHeader.Length);
+		return false;
+	}
+
 	// Fill packet header with message data length
 	FMemory::Memcpy(DataBuffer.GetData(), &MessageHeader, sizeof(FDisplayClusterMessageHeader));
 
@@ -146,9 +152,15 @@ bool FDisplayClusterSocketOps::SendJson(const TSharedPtr<FJsonObject>& Message)
 
 	// Initialize message header
 	FDisplayClusterMessageHeader MessageHeader;
-	MessageHeader.Length = static_cast<uint16>(MessageLength & 0xFFFF);
+	MessageHeader.Length = static_cast<uint32>(MessageLength & 0xFFFFFFFF);
 	UE_LOG(LogDisplayClusterNetworkMsg, Verbose, TEXT("Outgoing json body length %d"), MessageHeader.Length);
 
+	if (MessageHeader.Length > INT32_MAX)
+	{
+		UE_LOG(LogDisplayClusterNetworkMsg, Error, TEXT("Message length %u exceeds maximum limit"), MessageHeader.Length);
+		return false;
+	}
+
 	// Fill packet header with message data length
 	FMemory::Memcpy(DataBuffer.GetData(), &MessageHeader, sizeof(FDisplayClusterMessageHeader));
 
diff --git a/Engine/Plugins/Runtime/nDisplay/Source/DisplayCluster/Private/Network/DisplayClusterSocketOps.h b/Engine/Plugins/Runtime/nDisplay/Source/DisplayCluster/Private/Network/DisplayClusterSocketOps.h
--- a/Engine/Plugins/Runtime/nDisplay/Source/DisplayCluster/Private/Network/DisplayClusterSocketOps.h
+++ b/Engine/Plugins/Runtime/nDisplay/Source/DisplayCluster/Private/Network/DisplayClusterSocketOps.h
@@ -47,7 +47,7 @@ protected:
 private:
 	struct FDisplayClusterMessageHeader
 	{
-		uint16 Length;
+		uint32 Length;
 
 		FString ToString()
 		{
diff --git a/Engine/Source/Programs/nDisplayLauncher/Cluster/Launcher/Launcher.cs b/Engine/Source/Programs/nDisplayLauncher/Cluster/Launcher/Launcher.cs
--- a/Engine/Source/Programs/nDisplayLauncher/Cluster/Launcher/Launcher.cs
+++ b/Engine/Source/Programs/nDisplayLauncher/Cluster/Launcher/Launcher.cs
@@ -597,7 +597,7 @@ namespace nDisplayLauncher.Cluster
 				NetworkStream networkStream = nodeClient.GetStream();
 
 				byte[] OutData = ASCIIEncoding.ASCII.GetBytes(cmd);
-				byte[] OutSize = BitConverter.GetBytes((short)OutData.Length);
+				byte[] OutSize = BitConverter.GetBytes((Int32)OutData.Length);
 
 				networkStream.Write(OutSize, 0, OutSize.Length);
 				networkStream.Write(OutData, 0, OutData.Length);
