Fixed cluster eventes in Editor (PIE)

diff --git a/Engine/Plugins/Runtime/nDisplay/Source/DisplayCluster/Private/Cluster/Controller/DisplayClusterClusterNodeCtrlBase.h b/Engine/Plugins/Runtime/nDisplay/Source/DisplayCluster/Private/Cluster/Controller/DisplayClusterClusterNodeCtrlBase.h
--- a/Engine/Plugins/Runtime/nDisplay/Source/DisplayCluster/Private/Cluster/Controller/DisplayClusterClusterNodeCtrlBase.h
+++ b/Engine/Plugins/Runtime/nDisplay/Source/DisplayCluster/Private/Cluster/Controller/DisplayClusterClusterNodeCtrlBase.h
@@ -22,4 +22,8 @@ public:
 	//////////////////////////////////////////////////////////////////////////////////////////////
 	virtual bool IsStandalone() const override final
 	{ return false; }
+
+protected:
+	virtual void OverrideMasterAddr(FString& Addr)
+	{ }
 };
diff --git a/Engine/Plugins/Runtime/nDisplay/Source/DisplayCluster/Private/Cluster/Controller/DisplayClusterClusterNodeCtrlEditor.cpp b/Engine/Plugins/Runtime/nDisplay/Source/DisplayCluster/Private/Cluster/Controller/DisplayClusterClusterNodeCtrlEditor.cpp
--- a/Engine/Plugins/Runtime/nDisplay/Source/DisplayCluster/Private/Cluster/Controller/DisplayClusterClusterNodeCtrlEditor.cpp
+++ b/Engine/Plugins/Runtime/nDisplay/Source/DisplayCluster/Private/Cluster/Controller/DisplayClusterClusterNodeCtrlEditor.cpp
@@ -14,7 +14,7 @@
 
 
 FDisplayClusterClusterNodeCtrlEditor::FDisplayClusterClusterNodeCtrlEditor(const FString& ctrlName, const FString& nodeName)
-	: FDisplayClusterClusterNodeCtrlBase(ctrlName, nodeName)
+	: FDisplayClusterClusterNodeCtrlMaster(ctrlName, nodeName)
 {
 }
 
@@ -23,87 +23,11 @@ FDisplayClusterClusterNodeCtrlEditor::~FDisplayClusterClusterNodeCtrlEditor()
 }
 
 
-//////////////////////////////////////////////////////////////////////////////////////////////
-// IPDisplayClusterClusterEventsProtocol
-//////////////////////////////////////////////////////////////////////////////////////////////
-void FDisplayClusterClusterNodeCtrlEditor::EmitClusterEvent(const FDisplayClusterClusterEvent& Event)
-{
-	UE_LOG(LogDisplayClusterCluster, Warning, TEXT("This should never be called!"));
-	ClusterEventsClient->EmitClusterEvent(Event);
-}
-
 //////////////////////////////////////////////////////////////////////////////////////////////
 // FDisplayClusterNodeCtrlBase
 //////////////////////////////////////////////////////////////////////////////////////////////
-bool FDisplayClusterClusterNodeCtrlEditor::InitializeServers()
-{
-	if (GDisplayCluster->GetOperationMode() == EDisplayClusterOperationMode::Disabled)
-	{
-		return false;
-	}
-
-	UE_LOG(LogDisplayClusterCluster, Log, TEXT("%s - initializing editor servers..."), *GetControllerName());
-
-	// Get config data
-	FDisplayClusterConfigClusterNode CfgMasterNode;
-	if (GDisplayCluster->GetPrivateConfigMgr()->GetMasterClusterNode(CfgMasterNode) == false)
-	{
-		UE_LOG(LogDisplayClusterCluster, Error, TEXT("No master node configuration data found"));
-		return false;
-	}
-
-	// Editor controller uses 127.0.0.1 only
-	CfgMasterNode.Addr = FString("127.0.0.1");
-
-	// Instantiate node servers
-	UE_LOG(LogDisplayClusterCluster, Log, TEXT("Servers: addr %s, port_ce %d"), *CfgMasterNode.Addr, CfgMasterNode.Port_CE);
-	ClusterEventsServer.Reset(new FDisplayClusterClusterEventsService(CfgMasterNode.Addr, CfgMasterNode.Port_CE));
-
-	return ClusterEventsServer.IsValid();
-}
-
-bool FDisplayClusterClusterNodeCtrlEditor::StartServers()
-{
-	UE_LOG(LogDisplayClusterCluster, Log, TEXT("%s - starting editor servers..."), *GetControllerName());
-	return StartServerWithLogs(ClusterEventsServer.Get());
-}
-
-void FDisplayClusterClusterNodeCtrlEditor::StopServers()
-{
-	ClusterEventsServer->Shutdown();
-}
-
-bool FDisplayClusterClusterNodeCtrlEditor::InitializeClients()
-{
-	UE_LOG(LogDisplayClusterCluster, Log, TEXT("%s - initializing editor clients..."), *GetControllerName());
-
-	// Instantiate local clients
-	ClusterEventsClient.Reset(new FDisplayClusterClusterEventsClient);
-
-	return ClusterEventsClient.IsValid();
-}
-
-bool FDisplayClusterClusterNodeCtrlEditor::StartClients()
+void FDisplayClusterClusterNodeCtrlEditor::OverrideMasterAddr(FString& Addr)
 {
-	UE_LOG(LogDisplayClusterCluster, Log, TEXT("%s - starting editor clients..."), *GetControllerName());
-
-	// Master config
-	FDisplayClusterConfigClusterNode CfgMasterNode;
-	if (GDisplayCluster->GetPrivateConfigMgr()->GetMasterClusterNode(CfgMasterNode) == false)
-	{
-		UE_LOG(LogDisplayClusterCluster, Error, TEXT("No master node configuration data found"));
-		return false;
-	}
-
 	// Editor controller uses 127.0.0.1 only
-	CfgMasterNode.Addr = FString("127.0.0.1");
-
-	const FDisplayClusterConfigNetwork CfgNetwork = GDisplayCluster->GetPrivateConfigMgr()->GetConfigNetwork();
-
-	return StartClientWithLogs(ClusterEventsClient.Get(), CfgMasterNode.Addr, CfgMasterNode.Port_CE, CfgNetwork.ClientConnectTriesAmount, CfgNetwork.ClientConnectRetryDelay);
-}
-
-void FDisplayClusterClusterNodeCtrlEditor::StopClients()
-{
-	ClusterEventsClient->Disconnect();
+	Addr = FString("127.0.0.1");
 }
diff --git a/Engine/Plugins/Runtime/nDisplay/Source/DisplayCluster/Private/Cluster/Controller/DisplayClusterClusterNodeCtrlEditor.h b/Engine/Plugins/Runtime/nDisplay/Source/DisplayCluster/Private/Cluster/Controller/DisplayClusterClusterNodeCtrlEditor.h
--- a/Engine/Plugins/Runtime/nDisplay/Source/DisplayCluster/Private/Cluster/Controller/DisplayClusterClusterNodeCtrlEditor.h
+++ b/Engine/Plugins/Runtime/nDisplay/Source/DisplayCluster/Private/Cluster/Controller/DisplayClusterClusterNodeCtrlEditor.h
@@ -2,7 +2,7 @@
 
 #pragma once
 
-#include "Cluster/Controller/DisplayClusterClusterNodeCtrlBase.h"
+#include "Cluster/Controller/DisplayClusterClusterNodeCtrlMaster.h"
 
 #include "Network/DisplayClusterMessage.h"
 
@@ -14,41 +14,15 @@ class FDisplayClusterClusterEventsService;
  * Editor node controller implementation.
  */
 class FDisplayClusterClusterNodeCtrlEditor
-	: public FDisplayClusterClusterNodeCtrlBase
+	: public FDisplayClusterClusterNodeCtrlMaster
 {
 public:
 	FDisplayClusterClusterNodeCtrlEditor(const FString& ctrlName, const FString& nodeName);
 	virtual ~FDisplayClusterClusterNodeCtrlEditor();
 
-public:
-	//////////////////////////////////////////////////////////////////////////////////////////////
-	// IPDisplayClusterNodeController
-	//////////////////////////////////////////////////////////////////////////////////////////////
-	virtual bool IsSlave() const override final
-	{ return false; }
-
-public:
-	//////////////////////////////////////////////////////////////////////////////////////////////
-	// IPDisplayClusterClusterEventsProtocol
-	//////////////////////////////////////////////////////////////////////////////////////////////
-	virtual void EmitClusterEvent(const FDisplayClusterClusterEvent& Event) override;
-
 protected:
 	//////////////////////////////////////////////////////////////////////////////////////////////
 	// FDisplayClusterNodeCtrlBase
 	//////////////////////////////////////////////////////////////////////////////////////////////
-	virtual bool InitializeServers() override;
-	virtual bool StartServers()      override;
-	virtual void StopServers()       override;
-
-	virtual bool InitializeClients() override;
-	virtual bool StartClients()      override;
-	virtual void StopClients()       override;
-
-private:
-	// Node clients
-	TUniquePtr<FDisplayClusterClusterEventsClient> ClusterEventsClient;
-
-	// Node servers
-	TUniquePtr<FDisplayClusterClusterEventsService> ClusterEventsServer;
+	virtual void OverrideMasterAddr(FString& Addr) override;
 };
diff --git a/Engine/Plugins/Runtime/nDisplay/Source/DisplayCluster/Private/Cluster/Controller/DisplayClusterClusterNodeCtrlMaster.cpp b/Engine/Plugins/Runtime/nDisplay/Source/DisplayCluster/Private/Cluster/Controller/DisplayClusterClusterNodeCtrlMaster.cpp
--- a/Engine/Plugins/Runtime/nDisplay/Source/DisplayCluster/Private/Cluster/Controller/DisplayClusterClusterNodeCtrlMaster.cpp
+++ b/Engine/Plugins/Runtime/nDisplay/Source/DisplayCluster/Private/Cluster/Controller/DisplayClusterClusterNodeCtrlMaster.cpp
@@ -248,6 +248,9 @@ bool FDisplayClusterClusterNodeCtrlMaster::InitializeServers()
 		return false;
 	}
 
+	// Allow children to override master's address
+	OverrideMasterAddr(masterCfg.Addr);
+
 	// Instantiate node servers
 	UE_LOG(LogDisplayClusterCluster, Log, TEXT("Servers: addr %s, port_cs %d, port_ss %d, port_ce %d"), *masterCfg.Addr, masterCfg.Port_CS, masterCfg.Port_SS, masterCfg.Port_CE);
 	ClusterSyncServer.Reset(new FDisplayClusterClusterSyncService(masterCfg.Addr, masterCfg.Port_CS));
diff --git a/Engine/Plugins/Runtime/nDisplay/Source/DisplayCluster/Private/Cluster/Controller/DisplayClusterClusterNodeCtrlSlave.cpp b/Engine/Plugins/Runtime/nDisplay/Source/DisplayCluster/Private/Cluster/Controller/DisplayClusterClusterNodeCtrlSlave.cpp
--- a/Engine/Plugins/Runtime/nDisplay/Source/DisplayCluster/Private/Cluster/Controller/DisplayClusterClusterNodeCtrlSlave.cpp
+++ b/Engine/Plugins/Runtime/nDisplay/Source/DisplayCluster/Private/Cluster/Controller/DisplayClusterClusterNodeCtrlSlave.cpp
@@ -165,6 +165,9 @@ bool FDisplayClusterClusterNodeCtrlSlave::StartClients()
 		return false;
 	}
 
+	// Allow children to override master's address
+	OverrideMasterAddr(MasterCfg.Addr);
+
 	const FDisplayClusterConfigNetwork CfgNetwork = GDisplayCluster->GetPrivateConfigMgr()->GetConfigNetwork();
 
 	return StartClientWithLogs(ClusterSyncClient.Get(),   MasterCfg.Addr, MasterCfg.Port_CS, CfgNetwork.ClientConnectTriesAmount, CfgNetwork.ClientConnectRetryDelay)
