#nDisplay Move DisplayClusterUtils into a module. Including source outside a module isn't supported very well and it was not necessary

diff --git a/Engine/Plugins/Runtime/nDisplay/Source/DisplayCluster/Public/DisplayClusterUtils/DisplayClusterCommonHelpers.h b/Engine/Plugins/Runtime/nDisplay/Source/DisplayCluster/Public/DisplayClusterUtils/DisplayClusterCommonHelpers.h
new file mode 100644
--- /dev/null
+++ b/Engine/Plugins/Runtime/nDisplay/Source/DisplayCluster/Public/DisplayClusterUtils/DisplayClusterCommonHelpers.h
@@ -0,0 +1,548 @@
+// Copyright Epic Games, Inc. All Rights Reserved.
+
+#pragma once
+
+#include "CoreMinimal.h"
+#include "EngineUtils.h"
+
+#include "IDisplayCluster.h"
+#include "Cluster/IDisplayClusterClusterManager.h"
+#include "Config/IDisplayClusterConfigManager.h"
+#include "Config/DisplayClusterConfigTypes.h"
+#include "Game/IDisplayClusterGameManager.h"
+#include "DisplayClusterTypesConverter.h"
+
+#include "DisplayClusterUtils/DisplayClusterTypesConverter.h"
+
+class AActor;
+class UDisplayClusterCameraComponent;
+
+
+namespace DisplayClusterHelpers
+{
+	//////////////////////////////////////////////////////////////////////////////////////////////
+	// Common String helpers
+	//////////////////////////////////////////////////////////////////////////////////////////////
+	namespace str
+	{
+		static inline FString BoolToStr(bool bVal, bool bAsWord = true)
+		{
+			return (bVal ? 
+				(bAsWord ? FString(TEXT("true"))  : FString(TEXT("1"))) :
+				(bAsWord ? FString(TEXT("false")) : FString(TEXT("0"))));
+		}
+
+		static void TrimStringValue(FString& InLine, bool bTrimQuotes = true)
+		{
+			InLine.TrimStartAndEndInline();
+			InLine.RemoveFromStart(DisplayClusterStrings::strKeyValSeparator);
+
+			if (bTrimQuotes)
+			{
+				InLine = InLine.TrimQuotes();
+			}
+
+			InLine.TrimStartAndEndInline();
+		}
+
+		// Parses string items separated by specified separator into array
+		// Example: item1, item2,item3  ,  item4 => {item1, item2, item3, item4}
+		template<typename T>
+		static TArray<T> StrToArray(const FString& InData, const FString& InSeparator, bool bCullEmpty = true)
+		{
+			TArray<T> OutData;
+			TArray<FString> TempData;
+			
+			InData.ParseIntoArray(TempData, *InSeparator, bCullEmpty);
+
+			for (FString& Item : TempData)
+			{
+				TrimStringValue(Item, false);
+			}
+
+			for (const FString& Item : TempData)
+			{
+				if (!bCullEmpty && Item.IsEmpty())
+				{
+					OutData.Add(T());
+				}
+				else
+				{
+					OutData.Add(FDisplayClusterTypesConverter::template FromString<T>(Item));
+				}
+			}
+
+			return OutData;
+		}
+
+		// Exports array data to a string
+		// Example: {item1, item2, item3, item4} => "item1,item2,item3,item4"
+		template<typename T>
+		static FString ArrayToStr(const TArray<T>& InData, const FString& InSeparator = DisplayClusterStrings::strArrayValSeparator, bool bAddQuotes = true)
+		{
+			static const auto Quotes = TEXT("\"");
+
+			FString ResultStr;
+			ResultStr.Reserve(255);
+
+			if (bAddQuotes)
+			{
+				ResultStr = Quotes;
+			}
+
+			for (int i = 0; i < InData.Num(); ++i)
+			{
+				ResultStr += FString::Printf(TEXT("%s%s"), *FDisplayClusterTypesConverter::ToString(InData[i]), *InSeparator);
+			}
+
+			if (InSeparator.Len() > 0 && InData.Num() > 0)
+			{
+				ResultStr.RemoveAt(ResultStr.Len() - InSeparator.Len());
+			}
+
+			if (bAddQuotes)
+			{
+				ResultStr += Quotes;
+			}
+
+			return ResultStr;
+		}
+
+		// Parses string of key-value pairs separated by specified separator into map
+		// Example: "key1=val1 key2=val2 key3=val3" => {{key1, val2}, {key2, val2}, {key3, val3}}
+		template<typename TKey, typename TVal>
+		void StrToMap(const FString& InData, const FString& InPairSeparator, const FString& InKeyValSeparator, TMap<TKey, TVal>& OutData)
+		{
+			TArray<FString> StrPairs = StrToArray<FString>(InData, InPairSeparator);
+
+			for (const FString& StrPair : StrPairs)
+			{
+				FString StrKey;
+				FString StrVal;
+
+				if (StrPair.Split(InKeyValSeparator, &StrKey, &StrVal, ESearchCase::IgnoreCase))
+				{
+					OutData.Emplace(FDisplayClusterTypesConverter::template FromString<TKey>(StrKey), FDisplayClusterTypesConverter::template FromString<TVal>(StrVal));
+				}
+			}
+		}
+
+		// Exports map data to a string
+		// Example: {{key1,val1},{key2,val2},{key3,val3}} => "key1=val1 key2=val2 key3=var3"
+		template<typename TKey, typename TVal>
+		FString MapToStr(const TMap<TKey, TVal>& InData, const FString& InPairSeparator = DisplayClusterStrings::strPairSeparator, const FString& InKeyValSeparator = DisplayClusterStrings::strKeyValSeparator, bool bAddQuoutes = true)
+		{
+			static const auto Quotes = TEXT("\"");
+			
+			FString ResultStr;
+			ResultStr.Reserve(255);
+
+			if (bAddQuoutes)
+			{
+				ResultStr = Quotes;
+			}
+
+			for (const auto& Pair : InData)
+			{
+				ResultStr = FString::Printf(TEXT("%s%s%s%s%s"), *ResultStr, *FDisplayClusterTypesConverter::ToString(Pair.Key), *InKeyValSeparator, *FDisplayClusterTypesConverter::ToString(Pair.Value), *InPairSeparator);
+			}
+
+			if (InPairSeparator.Len() > 0)
+			{
+				ResultStr.RemoveAt(ResultStr.Len() - InPairSeparator.Len());
+			}
+
+			if (bAddQuoutes)
+			{
+				ResultStr += Quotes;
+			}
+
+			return ResultStr;
+		}
+
+		// Extracts value either from a command line string or any other line that matches the same format
+		// Example: extracting value of param2
+		// "param1=value1 param2=value2 param3=value3" => value2
+		template<typename T>
+		static bool ExtractValue(const FString& InLine, const FString& InParamName, T& OutValue, bool bInTrimQuotes = true)
+		{
+			FString TmpVal;
+
+			const FString FullParamName = InParamName + DisplayClusterStrings::strKeyValSeparator;
+			if (FParse::Value(*InLine, *FullParamName, TmpVal, false))
+			{
+				TrimStringValue(TmpVal, bInTrimQuotes);
+				OutValue = FDisplayClusterTypesConverter::template FromString<T>(TmpVal);
+				return true;
+			}
+
+			return false;
+		}
+
+		// Extracts array value either from a command line string or any other line that matches the same format
+		// Example: extracting array value of param2
+		// "param1=value1 param2="a,b,c,d" param3=value3" => {a,b,c,d}
+		template<typename T>
+		static bool ExtractArray(const FString& InLine, const FString& InParamName, const FString& InSeparator, TArray<T>& OutValue)
+		{
+			FString TmpVal;
+
+			if (ExtractValue(InLine, InParamName, TmpVal, false))
+			{
+				OutValue = StrToArray<T>(TmpVal, InSeparator);
+				return true;
+			}
+
+			return false;
+		}
+
+		// Extracts map value either from a command line string or any other line that matches the same format
+		// Example: extracting map value of param2
+		// "param1=value1 param2="a:1,b:7,c:22" param3=value3" => {{a,1},{b,7}{c,22}}
+		template<typename TKey, typename TVal>
+		static bool ExtractMap(const FString& InLine, const FString& InParamName, const FString& InPairSeparator, const FString& InKeyValSeparator, TMap<TKey, TVal>& OutData)
+		{
+			TArray<FString> TmpPairs;
+			if (!ExtractArray(InLine, InParamName, InPairSeparator, TmpPairs))
+			{
+				return false;
+			}
+
+			for (const FString& StrPair : TmpPairs)
+			{
+				StrToMap(StrPair, InPairSeparator, InKeyValSeparator, OutData);
+			}
+
+			return true;
+		}
+	};
+
+
+	//////////////////////////////////////////////////////////////////////////////////////////////
+	// Array helpers
+	//////////////////////////////////////////////////////////////////////////////////////////////	struct str
+	namespace arrays
+	{
+		// Max element in array
+		template<typename T>
+		T max(const T* data, int size)
+		{
+			T result = data[0];
+			for (int i = 1; i < size; i++)
+				if (result < data[i])
+					result = data[i];
+			return result;
+		}
+
+		// Max element's index in array
+		template<typename T>
+		size_t max_idx(const T* data, int size)
+		{
+			size_t idx = 0;
+			T result = data[0];
+			for (int i = 1; i < size; i++)
+				if (result < data[i])
+				{
+					result = data[i];
+					idx = i;
+				}
+			return idx;
+		}
+
+		// Min element in array
+		template<typename T>
+		T min(const T* data, int size)
+		{
+			T result = data[0];
+			for (int i = 1; i < size; i++)
+				if (result > data[i])
+					result = data[i];
+			return result;
+		}
+
+		// Min element's index in array
+		template<typename T>
+		size_t min_idx(const T* data, int size)
+		{
+			size_t idx = 0;
+			T result = data[0];
+			for (int i = 1; i < size; i++)
+				if (result > data[i])
+				{
+					result = data[i];
+					idx = i;
+				}
+			return idx;
+		}
+
+		// Helper for array size
+		template <typename T, size_t n>
+		constexpr size_t array_size(const T(&)[n])
+		{
+			return n;
+		}
+	}
+
+
+	//////////////////////////////////////////////////////////////////////////////////////////////
+	// Game helpers
+	//////////////////////////////////////////////////////////////////////////////////////////////
+	namespace game
+	{
+		template<typename T>
+		static void FindAllActors(UWorld* World, TArray<T*>& Out)
+		{
+			for (TActorIterator<AActor> It(World, T::StaticClass()); It; ++It)
+			{
+				T* Actor = Cast<T>(*It);
+				if (Actor && !Actor->IsPendingKill())
+				{
+					Out.Add(Actor);
+				}
+			}
+		}
+
+		static UDisplayClusterCameraComponent* GetCamera(const FString& CameraId)
+		{
+			static const IDisplayClusterGameManager* const GameMgr = IDisplayCluster::Get().GetGameMgr();
+			if (GameMgr)
+			{
+				return GameMgr->GetCameraById(CameraId);
+			}
+
+			return nullptr;
+		}
+	}
+
+	//////////////////////////////////////////////////////////////////////////////////////////////
+	// Config helpers
+	//////////////////////////////////////////////////////////////////////////////////////////////
+	namespace config
+	{
+		static bool GetLocalClusterNode(FDisplayClusterConfigClusterNode& LocalClusterNode)
+		{
+			IDisplayCluster& DisplayCluster = IDisplayCluster::Get();
+
+			if (DisplayCluster.GetOperationMode() == EDisplayClusterOperationMode::Disabled)
+			{
+				return false;
+			}
+
+			const IDisplayClusterClusterManager* const ClusterMgr = DisplayCluster.GetClusterMgr();
+			if (!ClusterMgr)
+			{
+				return false;
+			}
+
+			const FString LocalNodeId = ClusterMgr->GetNodeId();
+			const IDisplayClusterConfigManager* const ConfigMgr = DisplayCluster.GetConfigMgr();
+			if (!ConfigMgr)
+			{
+				return false;
+			}
+
+			return ConfigMgr->GetClusterNode(LocalNodeId, LocalClusterNode);
+		}
+
+		static bool GetLocalWindow(FDisplayClusterConfigWindow& LocalWindow)
+		{
+			FDisplayClusterConfigClusterNode LocalClusterNode;
+			if (!GetLocalClusterNode(LocalClusterNode))
+			{
+				return false;
+			}
+
+			static const IDisplayClusterConfigManager* const ConfigMgr = IDisplayCluster::Get().GetConfigMgr();
+			if (!ConfigMgr)
+			{
+				return false;
+			}
+
+			return ConfigMgr->GetWindow(LocalClusterNode.WindowId, LocalWindow);
+		}
+
+		static TArray<FDisplayClusterConfigViewport> GetLocalViewports()
+		{
+			TArray<FDisplayClusterConfigViewport> LocalViewports;
+
+			FDisplayClusterConfigWindow LocalWindow;
+			if (!DisplayClusterHelpers::config::GetLocalWindow(LocalWindow))
+			{
+				return LocalViewports;
+			}
+
+			static const IDisplayClusterConfigManager* const ConfigMgr = IDisplayCluster::Get().GetConfigMgr();
+			if (!ConfigMgr)
+			{
+				return LocalViewports;
+			}
+
+			LocalViewports = ConfigMgr->GetViewports().FilterByPredicate([&LocalWindow](const FDisplayClusterConfigViewport& ItemViewport)
+			{
+				return LocalWindow.ViewportIds.ContainsByPredicate([ItemViewport](const FString& ItemId)
+				{
+					return ItemViewport.Id.Compare(ItemId, ESearchCase::IgnoreCase) == 0;
+				});
+			});
+
+			return LocalViewports;
+		}
+
+		static TArray<FDisplayClusterConfigPostprocess> GetLocalPostprocess()
+		{
+			TArray<FDisplayClusterConfigPostprocess> LocalPostprocess;
+
+			FDisplayClusterConfigWindow LocalWindow;
+			if (!DisplayClusterHelpers::config::GetLocalWindow(LocalWindow))
+			{
+				return LocalPostprocess;
+			}
+
+			static const IDisplayClusterConfigManager* const ConfigMgr = IDisplayCluster::Get().GetConfigMgr();
+			if (!ConfigMgr)
+			{
+				return LocalPostprocess;
+			}
+
+			LocalPostprocess = ConfigMgr->GetPostprocess().FilterByPredicate([&LocalWindow](const FDisplayClusterConfigPostprocess& ItemPostprocess)
+			{
+				return LocalWindow.PostprocessIds.ContainsByPredicate([ItemPostprocess](const FString& ItemId)
+				{
+					return ItemPostprocess.Id.Compare(ItemId, ESearchCase::IgnoreCase) == 0;
+				});
+			});
+
+			return LocalPostprocess;
+		}
+
+		static TArray<FDisplayClusterConfigProjection> GetLocalProjections()
+		{
+			TArray<FDisplayClusterConfigProjection> LocalProjections;
+
+			static const IDisplayClusterConfigManager* const ConfigMgr = IDisplayCluster::Get().GetConfigMgr();
+			if (!ConfigMgr)
+			{
+				return LocalProjections;
+			}
+
+			TArray<FDisplayClusterConfigViewport> LocalViewports = GetLocalViewports();
+
+			LocalProjections = ConfigMgr->GetProjections().FilterByPredicate([&LocalViewports](const FDisplayClusterConfigProjection& ItemProjection)
+			{
+				return LocalViewports.ContainsByPredicate([ItemProjection](const FDisplayClusterConfigViewport& ItemViewport)
+				{
+					return ItemViewport.ProjectionId.Compare(ItemProjection.Id, ESearchCase::IgnoreCase) == 0;
+				});
+			});
+
+			return LocalProjections;
+		}
+
+		static bool GetViewportProjection(const FString& ViewportId, FDisplayClusterConfigProjection& ViewportProjection)
+		{
+			static const IDisplayClusterConfigManager* const ConfigMgr = IDisplayCluster::Get().GetConfigMgr();
+			if (ConfigMgr)
+			{
+				FDisplayClusterConfigViewport CfgViewport;
+				if (ConfigMgr->GetViewport(ViewportId, CfgViewport))
+				{
+					if (ConfigMgr->GetProjection(CfgViewport.ProjectionId, ViewportProjection))
+					{
+						return true;
+					}
+				}
+			}
+
+			return false;
+		}
+
+		static FString GetFullPath(const FString& LocalPath)
+		{
+			static const IDisplayClusterConfigManager* const ConfigMgr = IDisplayCluster::Get().GetConfigMgr();
+			if (ConfigMgr)
+			{
+				return ConfigMgr->GetFullPathToFile(LocalPath);
+			}
+
+			return LocalPath;
+		}
+	}
+
+	//////////////////////////////////////////////////////////////////////////////////////////////
+	// Math helpers
+	//////////////////////////////////////////////////////////////////////////////////////////////
+	namespace math
+	{
+		static FMatrix GetProjectionMatrixFromOffsets(float l, float r, float t, float b, float n, float f)
+		{
+			const float LeftFOVRad   = FMath::Atan(l / n);
+			const float RightFOVRad  = FMath::Atan(r / n);
+			const float TopFOVRad    = FMath::Atan(t / n);
+			const float BottomFOVRad = FMath::Atan(b / n);
+
+			{
+				const float MinFOVRad = FMath::DegreesToRadians(1.0f);   // min FOV = 1 degree
+				const float MaxFOVRad = FMath::DegreesToRadians(179.0f); // max FOV = 179 degree
+
+				// clamp FOV values:
+				const float FOVRadH  = (RightFOVRad - LeftFOVRad);
+				float NewRightFOVRad = (FOVRadH < MinFOVRad) ? (LeftFOVRad + MinFOVRad) : RightFOVRad;
+				NewRightFOVRad       = (FOVRadH > MaxFOVRad) ? (LeftFOVRad + MaxFOVRad) : NewRightFOVRad;
+
+				const float FOVRadV = (TopFOVRad - BottomFOVRad);
+				float NewTopFOVRad  = (FOVRadV < MinFOVRad) ? (BottomFOVRad + MinFOVRad) : TopFOVRad;
+				NewTopFOVRad        = (FOVRadV > MaxFOVRad) ? (BottomFOVRad + MaxFOVRad) : NewTopFOVRad;
+
+				if (RightFOVRad != NewRightFOVRad)
+				{
+					r = float(n * FMath::Tan(NewRightFOVRad));
+					//! add LOG warning
+				}
+				if (TopFOVRad != NewTopFOVRad)
+				{
+					t = float(n * FMath::Tan(NewTopFOVRad));
+					//! add LOG warning
+				}
+			}
+
+			const float mx = 2.f * n / (r - l);
+			const float my = 2.f * n / (t - b);
+			const float ma = -(r + l) / (r - l);
+			const float mb = -(t + b) / (t - b);
+
+			// Support unlimited far plane (f==n)
+			const float mc = (f == n) ? (1.0f - Z_PRECISION) : (f / (f - n));
+			const float md = (f == n) ? (-n * (1.0f - Z_PRECISION)) : (-(f * n) / (f - n));
+
+			const float me = 1.f;
+
+			// Normal LHS
+			const FMatrix ProjectionMatrix = FMatrix(
+				FPlane(mx,  0,  0,  0),
+				FPlane(0,  my,  0,  0),
+				FPlane(ma, mb, mc, me),
+				FPlane(0,  0,  md,  0));
+
+			// Invert Z-axis (UE4 uses Z-inverted LHS)
+			static const FMatrix flipZ = FMatrix(
+				FPlane(1, 0, 0, 0),
+				FPlane(0, 1, 0, 0),
+				FPlane(0, 0, -1, 0),
+				FPlane(0, 0, 1, 1));
+
+			const FMatrix ResultMatrix(ProjectionMatrix * flipZ);
+
+			return ResultMatrix;
+		}
+
+		static FMatrix GetProjectionMatrixFromAngles(float LeftAngle, float RightAngle, float TopAngle, float BottomAngle, float ZNear, float ZFar)
+		{
+			const float t = ZNear * FMath::Tan(FMath::DegreesToRadians(TopAngle));
+			const float b = ZNear * FMath::Tan(FMath::DegreesToRadians(BottomAngle));
+			const float l = ZNear * FMath::Tan(FMath::DegreesToRadians(LeftAngle));
+			const float r = ZNear * FMath::Tan(FMath::DegreesToRadians(RightAngle));
+
+			return DisplayClusterHelpers::math::GetProjectionMatrixFromOffsets(l, r, t, b, ZNear, ZFar);
+		}
+	}
+};
diff --git a/Engine/Plugins/Runtime/nDisplay/Source/DisplayCluster/Public/DisplayClusterUtils/DisplayClusterCommonStrings.h b/Engine/Plugins/Runtime/nDisplay/Source/DisplayCluster/Public/DisplayClusterUtils/DisplayClusterCommonStrings.h
new file mode 100644
--- /dev/null
+++ b/Engine/Plugins/Runtime/nDisplay/Source/DisplayCluster/Public/DisplayClusterUtils/DisplayClusterCommonStrings.h
@@ -0,0 +1,235 @@
+// Copyright Epic Games, Inc. All Rights Reserved.
+
+#pragma once
+
+#include "CoreMinimal.h"
+
+
+namespace DisplayClusterStrings
+{
+	// Common strings
+	static constexpr auto strPairSeparator     = TEXT(" ");
+	static constexpr auto strKeyValSeparator   = TEXT("=");
+	static constexpr auto strArrayValSeparator = TEXT(",");
+
+	namespace cfg
+	{
+		// Config special constants
+		namespace spec
+		{
+			static constexpr auto Comment          = TEXT("#");
+			static constexpr auto KeyValSeparator  = TEXT("=");
+			static constexpr auto ValTrue          = TEXT("true");
+			static constexpr auto ValFalse         = TEXT("false");
+			static constexpr auto MappingDelimiter = TEXT(",");
+		}
+
+		// Config file extensions
+		namespace file
+		{
+			static constexpr auto FileExtCfg1 = TEXT("cfg");
+			static constexpr auto FileExtCfg2 = TEXT("conf");
+			static constexpr auto FileExtCfg3 = TEXT("config");
+			static constexpr auto FileExtTxt  = TEXT("txt");
+			static constexpr auto FileExtXml  = TEXT("xml");
+		}
+
+		// Config data tokens
+		namespace data
+		{
+			static constexpr auto Id       = TEXT("id");
+			static constexpr auto ParentId = TEXT("parent");
+			static constexpr auto Loc      = TEXT("loc");
+			static constexpr auto Rot      = TEXT("rot");
+
+			// Config info
+			namespace info
+			{
+				static constexpr auto Header  = TEXT("[info]");
+				static constexpr auto Version = TEXT("version");
+			}
+
+			// Cluster tokens
+			namespace cluster
+			{
+				static constexpr auto Header = TEXT("[cluster_node]");
+				static constexpr auto Addr   = TEXT("addr");
+				static constexpr auto Window = TEXT("window");
+				static constexpr auto PortCS = TEXT("port_cs");
+				static constexpr auto PortSS = TEXT("port_ss");
+				static constexpr auto PortCE = TEXT("port_ce");
+				static constexpr auto Master = TEXT("master");
+				static constexpr auto Sound  = TEXT("sound");
+				// + Id
+			}
+
+			// Window tokens
+			namespace window
+			{
+				static constexpr auto Header       = TEXT("[window]");
+				static constexpr auto Viewports    = TEXT("viewports");
+				static constexpr auto Postprocess  = TEXT("postprocess");
+				static constexpr auto Fullscreen   = TEXT("fullscreen");
+				static constexpr auto WinX         = TEXT("WinX");
+				static constexpr auto WinY         = TEXT("WinY");
+				static constexpr auto ResX         = TEXT("ResX");
+				static constexpr auto ResY         = TEXT("ResY");
+				// + Id
+			}
+
+			// Screen tokens
+			namespace screen
+			{
+				static constexpr auto Header = TEXT("[screen]");
+				static constexpr auto Size   = TEXT("size");
+				// + Id, Parent, Loc, Rot
+			}
+
+			// Viewport tokens
+			namespace viewport
+			{
+				static constexpr auto Header      = TEXT("[viewport]");
+				static constexpr auto Projection  = TEXT("projection");
+				static constexpr auto Camera      = TEXT("camera");
+				static constexpr auto PosX        = TEXT("x");
+				static constexpr auto PosY        = TEXT("y");
+				static constexpr auto Width       = TEXT("width");
+				static constexpr auto Height      = TEXT("height");
+				static constexpr auto RTT         = TEXT("rtt");
+				static constexpr auto BufferRatio = TEXT("buffer_ratio");
+				// + Id
+			}
+
+			// Posprocess tokens
+			namespace postprocess
+			{
+				static constexpr auto Header        = TEXT("[postprocess]");
+				static constexpr auto PostprocessId = TEXT("type");
+				// + Id
+			}
+
+			// Camera tokens
+			namespace camera
+			{
+				static constexpr auto Header      = TEXT("[camera]");
+				static constexpr auto EyeDist     = TEXT("eye_dist");
+				static constexpr auto EyeSwap     = TEXT("eye_swap");
+				static constexpr auto ForceOffset = TEXT("force_offset");
+
+				// + Id, Loc, Rot, Parent
+			}
+
+			// Scene node (transforms)
+			namespace scene
+			{
+				static constexpr auto Header    = TEXT("[scene_node]");
+				static constexpr auto TrackerId = TEXT("tracker_id");
+				static constexpr auto TrackerCh = TEXT("tracker_ch");
+				// + Id, Loc, Rot, Parent
+			}
+
+			// Input tokens
+			namespace input
+			{
+				static constexpr auto Header  = TEXT("[input]");
+				static constexpr auto Type    = TEXT("type");
+				static constexpr auto Address = TEXT("addr");
+				static constexpr auto Remap   = TEXT("remap");
+
+				static constexpr auto Right   = TEXT("right");
+				static constexpr auto Front   = TEXT("front");
+				static constexpr auto Up      = TEXT("up");
+
+				static constexpr auto MapX    = TEXT("x");
+				static constexpr auto MapNX   = TEXT("-x");
+				static constexpr auto MapY    = TEXT("y");
+				static constexpr auto MapNY   = TEXT("-y");
+				static constexpr auto MapZ    = TEXT("z");
+				static constexpr auto MapNZ   = TEXT("-z");
+
+				static constexpr auto DeviceTracker  = TEXT("tracker");
+				static constexpr auto DeviceAnalog   = TEXT("analog");
+				static constexpr auto DeviceButtons  = TEXT("buttons");
+				static constexpr auto DeviceKeyboard = TEXT("keyboard");
+
+				// + Id
+			}
+
+			// Input Setup tokens
+			namespace inputsetup
+			{
+				static constexpr auto Header  = TEXT("[input_setup]");
+				static constexpr auto Channel = TEXT("ch");
+				static constexpr auto Key     = TEXT("key");
+				static constexpr auto Bind    = TEXT("bind");
+
+				// + Id
+			}
+
+			// General settings tokens
+			namespace general
+			{
+				static constexpr auto Header                = TEXT("[general]");
+				static constexpr auto SwapSyncPolicy        = TEXT("swap_sync_policy");
+				static constexpr auto UnrealInputSyncPolicy = TEXT("ue4_input_sync_policy");
+			}
+
+			// Stereo tokens
+			namespace stereo
+			{
+				static constexpr auto Header = TEXT("[stereo]");
+			}
+
+			// Render tokens
+			namespace render
+			{
+				static constexpr auto Header = TEXT("[render]");
+			}
+
+			// Network tokens
+			namespace network
+			{
+				static constexpr auto Header = TEXT("[network]");
+				static constexpr auto ClientConnectTriesAmount    = TEXT("cln_conn_tries_amount");
+				static constexpr auto ClientConnectRetryDelay     = TEXT("cln_conn_retry_delay");
+				static constexpr auto BarrierGameStartWaitTimeout = TEXT("game_start_timeout");
+				static constexpr auto BarrierWaitTimeout          = TEXT("barrier_wait_timeout");
+			}
+
+			// Debug tokens
+			namespace debug
+			{
+				static constexpr auto Header    = TEXT("[debug]");
+				static constexpr auto LagSim    = TEXT("lag_simulation");
+				static constexpr auto LagTime   = TEXT("lag_max_time");
+				static constexpr auto DrawStats = TEXT("draw_stats");
+			}
+
+			// Custom arguments
+			namespace custom
+			{
+				static constexpr auto Header = TEXT("[custom]");
+			}
+
+			// Projection data
+			namespace projection
+			{
+				static constexpr auto Header = TEXT("[projection]");
+				static constexpr auto Type   = TEXT("type");
+
+				// + Id
+			}
+		}
+	};
+
+	namespace rhi
+	{
+		static constexpr auto D3D11  = TEXT("D3D11");
+		static constexpr auto D3D12  = TEXT("D3D12");
+	}
+
+	namespace platform
+	{
+		static constexpr auto Windows = TEXT("Windows");
+	}
+};
diff --git a/Engine/Plugins/Runtime/nDisplay/Source/DisplayCluster/Public/DisplayClusterUtils/DisplayClusterTypesConverter.h b/Engine/Plugins/Runtime/nDisplay/Source/DisplayCluster/Public/DisplayClusterUtils/DisplayClusterTypesConverter.h
new file mode 100644
--- /dev/null
+++ b/Engine/Plugins/Runtime/nDisplay/Source/DisplayCluster/Public/DisplayClusterUtils/DisplayClusterTypesConverter.h
@@ -0,0 +1,197 @@
+// Copyright Epic Games, Inc. All Rights Reserved.
+
+#pragma once
+
+#include "CoreMinimal.h"
+#include "Misc/Timecode.h"
+#include "Misc/FrameRate.h"
+#include "DisplayClusterEnums.h"
+#include "DisplayClusterUtils/DisplayClusterCommonStrings.h"
+#include "DisplayClusterUtils/DisplayClusterCommonHelpers.h"
+
+
+/**
+ * Auxiliary class with different type conversion functions
+ */
+class FDisplayClusterTypesConverter
+{
+public:
+	//////////////////////////////////////////////////////////////////////////////////////////////
+	// TYPE --> STRING
+	//////////////////////////////////////////////////////////////////////////////////////////////
+	template <typename ConvertFrom>
+	static FString ToString(const ConvertFrom& from);
+
+	//////////////////////////////////////////////////////////////////////////////////////////////
+	// STRING --> TYPE
+	//////////////////////////////////////////////////////////////////////////////////////////////
+	template <typename ConvertTo>
+	static ConvertTo FromString(const FString& from);
+};
+
+template <> inline FString FDisplayClusterTypesConverter::ToString<>(const FString& from)    { return from; }
+template <> inline FString FDisplayClusterTypesConverter::ToString<>(const bool& from)       { return (from ? DisplayClusterStrings::cfg::spec::ValTrue : DisplayClusterStrings::cfg::spec::ValFalse); }
+template <> inline FString FDisplayClusterTypesConverter::ToString<> (const int8& from)      { return FString::FromInt(from); }
+template <> inline FString FDisplayClusterTypesConverter::ToString<> (const uint8& from)     { return ToString(static_cast<int8>(from)); }
+template <> inline FString FDisplayClusterTypesConverter::ToString<> (const int32& from)     { return FString::FromInt(from); }
+template <> inline FString FDisplayClusterTypesConverter::ToString<> (const uint32& from)    { return ToString(static_cast<int32>(from)); }
+template <> inline FString FDisplayClusterTypesConverter::ToString<>(const float& from)      { return FString::SanitizeFloat(from); }
+template <> inline FString FDisplayClusterTypesConverter::ToString<>(const double& from)     { return FString::Printf(TEXT("%lf"), from); }
+template <> inline FString FDisplayClusterTypesConverter::ToString<>(const FVector& from)    { return from.ToString(); }
+template <> inline FString FDisplayClusterTypesConverter::ToString<>(const FVector2D& from)  { return from.ToString(); }
+template <> inline FString FDisplayClusterTypesConverter::ToString<>(const FRotator& from)   { return from.ToString(); }
+template <> inline FString FDisplayClusterTypesConverter::ToString<>(const FMatrix& from)    { return from.ToString(); }
+
+// We can't just use FTimecode ToString as that loses information.
+template <> inline FString FDisplayClusterTypesConverter::ToString<>(const FTimecode& from)  { return FString::Printf(TEXT("%d;%d;%d;%d;%d"), from.bDropFrameFormat ? 1 : 0, from.Hours, from.Minutes, from.Seconds, from.Frames); }
+template <> inline FString FDisplayClusterTypesConverter::ToString<>(const FFrameRate& from) { return FString::Printf(TEXT("%d;%d"), from.Numerator, from.Denominator); }
+
+template <> inline FString FDisplayClusterTypesConverter::ToString<>(const EDisplayClusterOperationMode& from)
+{
+	switch (from)
+	{
+	case EDisplayClusterOperationMode::Cluster:
+		return FString("cluster");
+	case EDisplayClusterOperationMode::Standalone:
+		return FString("standalone");
+	case EDisplayClusterOperationMode::Editor:
+		return FString("editor");
+	case EDisplayClusterOperationMode::Disabled:
+		return FString("disabled");
+	default:
+		return FString("unknown");
+	}
+}
+
+template <> inline FString FDisplayClusterTypesConverter::ToString<>(const EDisplayClusterSyncGroup& from)
+{
+	switch (from)
+	{
+	case EDisplayClusterSyncGroup::PreTick:
+		return FString("PreTick");
+	case EDisplayClusterSyncGroup::Tick:
+		return FString("Tick");
+	case EDisplayClusterSyncGroup::PostTick:
+		return FString("PostTick");
+	}
+
+	return FString("Tick");
+}
+
+
+template <> inline FString   FDisplayClusterTypesConverter::FromString<> (const FString& from) { return from; }
+template <> inline bool      FDisplayClusterTypesConverter::FromString<> (const FString& from) { return (from == FString("1") || (from.Compare(DisplayClusterStrings::cfg::spec::ValTrue, ESearchCase::IgnoreCase) == 0)); }
+template <> inline int8      FDisplayClusterTypesConverter::FromString<> (const FString& from) { return FCString::Atoi(*from); }
+template <> inline uint8     FDisplayClusterTypesConverter::FromString<> (const FString& from) { return static_cast<uint8>(FromString<int8>(from)); }
+template <> inline int32     FDisplayClusterTypesConverter::FromString<> (const FString& from) { return FCString::Atoi(*from); }
+template <> inline uint32    FDisplayClusterTypesConverter::FromString<> (const FString& from) { return static_cast<uint32>(FromString<int32>(from)); }
+template <> inline float     FDisplayClusterTypesConverter::FromString<> (const FString& from) { return FCString::Atof(*from); }
+template <> inline double    FDisplayClusterTypesConverter::FromString<> (const FString& from) { return FCString::Atod(*from); }
+template <> inline FVector   FDisplayClusterTypesConverter::FromString<> (const FString& from) { FVector vec;  vec.InitFromString(from); return vec; }
+template <> inline FVector2D FDisplayClusterTypesConverter::FromString<> (const FString& from) { FVector2D vec;  vec.InitFromString(from); return vec; }
+template <> inline FRotator  FDisplayClusterTypesConverter::FromString<> (const FString& from) { FRotator rot; rot.InitFromString(from); return rot; }
+
+template <> inline FMatrix   FDisplayClusterTypesConverter::FromString<>(const FString& from)
+{
+	FMatrix ResultMatrix = FMatrix::Identity;
+	FPlane  Planes[4];
+
+	int32 IdxStart = 0;
+	int32 IdxEnd = 0;
+	for (int PlaneNum = 0; PlaneNum < 4; ++PlaneNum)
+	{
+		IdxStart = from.Find(FString(TEXT("[")), ESearchCase::IgnoreCase, ESearchDir::FromStart, IdxEnd);
+		IdxEnd   = from.Find(FString(TEXT("]")), ESearchCase::IgnoreCase, ESearchDir::FromStart, IdxStart);
+		if (IdxStart == INDEX_NONE || IdxEnd == INDEX_NONE || (IdxEnd - IdxStart) < 8)
+		{
+			return ResultMatrix;
+		}
+
+		FString StrPlane = from.Mid(IdxStart + 1, IdxEnd - IdxStart - 1);
+
+		int32 StrLen = 0;
+		int32 NewStrLen = -1;
+		while (NewStrLen < StrLen)
+		{
+			StrLen = StrPlane.Len();
+			StrPlane.ReplaceInline(TEXT("  "), TEXT(" "));
+			NewStrLen = StrPlane.Len();
+		}
+
+		TArray<FString> StrPlaneValues;
+		StrPlane.ParseIntoArray(StrPlaneValues, TEXT(" "));
+		if (StrPlaneValues.Num() != 4)
+		{
+			return ResultMatrix;
+		}
+
+		float PlaneValues[4] = { 0.f };
+		for (int i = 0; i < 4; ++i)
+		{
+			PlaneValues[i] = FromString<float>(StrPlaneValues[i]);
+		}
+
+		Planes[PlaneNum] = FPlane(PlaneValues[0], PlaneValues[1], PlaneValues[2], PlaneValues[3]);
+	}
+
+	ResultMatrix = FMatrix(Planes[0], Planes[1], Planes[2], Planes[3]);
+
+	return ResultMatrix;
+}
+
+template <> inline FTimecode FDisplayClusterTypesConverter::FromString<> (const FString& from)
+{
+	FTimecode timecode;
+
+	TArray<FString> parts;
+	parts.Reserve(5);
+	const int32 found = from.ParseIntoArray(parts, TEXT(";"));
+
+	// We are expecting 5 "parts" - DropFrame, Hours, Minutes, Seconds, Frames.
+	if (found == 5)
+	{
+		timecode.bDropFrameFormat = FromString<bool>(parts[0]);
+		timecode.Hours = FromString<int32>(parts[1]);
+		timecode.Minutes = FromString<int32>(parts[2]);
+		timecode.Seconds = FromString<int32>(parts[3]);
+		timecode.Frames = FromString<int32>(parts[4]);
+	}
+
+	return timecode;
+}
+
+template <> inline FFrameRate FDisplayClusterTypesConverter::FromString<> (const FString& from)
+{
+	FFrameRate frameRate;
+
+	TArray<FString> parts;
+	parts.Reserve(2);
+	const int32 found = from.ParseIntoArray(parts, TEXT(";"));
+
+	// We are expecting 2 "parts" - Numerator, Denominator.
+	if (found == 2)
+	{
+		frameRate.Numerator   = FromString<int32>(parts[0]);
+		frameRate.Denominator = FromString<int32>(parts[1]);
+	}
+
+	return frameRate;
+}
+
+template <> inline EDisplayClusterSyncGroup FDisplayClusterTypesConverter::FromString<>(const FString& from)
+{
+	if (from.Equals(TEXT("PreTick")))
+	{
+		return EDisplayClusterSyncGroup::PreTick;
+	}
+	else if (from.Equals(TEXT("Tick")))
+	{
+		return EDisplayClusterSyncGroup::Tick;
+	}
+	else if (from.Equals(TEXT("PostTick")))
+	{
+		return EDisplayClusterSyncGroup::PostTick;
+	}
+
+	return EDisplayClusterSyncGroup::Tick;
+}
diff --git a/Engine/Plugins/Runtime/nDisplay/Source/DisplayClusterUtils/DisplayClusterCommonHelpers.h b/Engine/Plugins/Runtime/nDisplay/Source/DisplayClusterUtils/DisplayClusterCommonHelpers.h
deleted file mode 100644
diff --git a/Engine/Plugins/Runtime/nDisplay/Source/DisplayClusterUtils/DisplayClusterCommonStrings.h b/Engine/Plugins/Runtime/nDisplay/Source/DisplayClusterUtils/DisplayClusterCommonStrings.h
deleted file mode 100644
diff --git a/Engine/Plugins/Runtime/nDisplay/Source/DisplayClusterUtils/DisplayClusterTypesConverter.h b/Engine/Plugins/Runtime/nDisplay/Source/DisplayClusterUtils/DisplayClusterTypesConverter.h
deleted file mode 100644
