Set up custom present on PreTick

diff --git a/Engine/Plugins/Runtime/nDisplay/Source/DisplayCluster/Private/Render/Device/DisplayClusterDeviceBase.cpp b/Engine/Plugins/Runtime/nDisplay/Source/DisplayCluster/Private/Render/Device/DisplayClusterDeviceBase.cpp
--- a/Engine/Plugins/Runtime/nDisplay/Source/DisplayCluster/Private/Render/Device/DisplayClusterDeviceBase.cpp
+++ b/Engine/Plugins/Runtime/nDisplay/Source/DisplayCluster/Private/Render/Device/DisplayClusterDeviceBase.cpp
@@ -334,6 +334,25 @@ bool FDisplayClusterDeviceBase::EnableStereo(bool stereo /*= true*/)
 void FDisplayClusterDeviceBase::InitCanvasFromView(class FSceneView* InView, class UCanvas* Canvas)
 {
 	DISPLAY_CLUSTER_FUNC_TRACE(LogDisplayClusterRender);
+
+	if (!bIsCustomPresentSet)
+	{
+		// Set up our new present handler
+		if (MainViewport)
+		{
+			// Current sync policy
+			TSharedPtr<IDisplayClusterRenderSyncPolicy> SyncPolicy = GDisplayCluster->GetRenderMgr()->GetCurrentSynchronizationPolicy();
+			check(SyncPolicy.IsValid());
+
+			// Create present handler
+			FDisplayClusterPresentationBase* const CustomPresentHandler = CreatePresentationObject(MainViewport, SyncPolicy);
+			check(CustomPresentHandler);
+
+			MainViewport->GetViewportRHI()->SetCustomPresent(CustomPresentHandler);
+		}
+
+		bIsCustomPresentSet = true;
+	}
 }
 
 void FDisplayClusterDeviceBase::AdjustViewRect(enum EStereoscopicPass StereoPassType, int32& X, int32& Y, uint32& SizeX, uint32& SizeY) const
@@ -584,15 +603,6 @@ void FDisplayClusterDeviceBase::UpdateViewport(bool bUseSeparateRenderTarget, co
 	{
 		// UE viewport
 		MainViewport = (FViewport*)&Viewport;
-		// Current sync policy
-		TSharedPtr<IDisplayClusterRenderSyncPolicy> SyncPolicy = GDisplayCluster->GetRenderMgr()->GetCurrentSynchronizationPolicy();
-
-		// Create present handler
-		FDisplayClusterPresentationBase* const CustomPresentHandler = CreatePresentationObject(MainViewport, SyncPolicy);
-		check(CustomPresentHandler);
-
-		// Set up our new present handler
-		Viewport.GetViewportRHI()->SetCustomPresent(CustomPresentHandler);
 	}
 }
 
diff --git a/Engine/Plugins/Runtime/nDisplay/Source/DisplayCluster/Private/Render/Device/DisplayClusterDeviceBase.h b/Engine/Plugins/Runtime/nDisplay/Source/DisplayCluster/Private/Render/Device/DisplayClusterDeviceBase.h
--- a/Engine/Plugins/Runtime/nDisplay/Source/DisplayCluster/Private/Render/Device/DisplayClusterDeviceBase.h
+++ b/Engine/Plugins/Runtime/nDisplay/Source/DisplayCluster/Private/Render/Device/DisplayClusterDeviceBase.h
@@ -159,6 +159,7 @@ protected:
 	FViewport* MainViewport = nullptr;
 
 	bool bIsSceneOpen = false;
+	bool bIsCustomPresentSet = false;
 
 	// Per-eye regions
 	FIntRect EyeRegions[2];
