Fix circular include causing missing symbol errors
#rnx

diff --git a/Engine/Plugins/Runtime/nDisplay/Source/DisplayCluster/Public/DisplayClusterUtils/DisplayClusterTypesConverter.h b/Engine/Plugins/Runtime/nDisplay/Source/DisplayCluster/Public/DisplayClusterUtils/DisplayClusterTypesConverter.h
--- a/Engine/Plugins/Runtime/nDisplay/Source/DisplayCluster/Public/DisplayClusterUtils/DisplayClusterTypesConverter.h
+++ b/Engine/Plugins/Runtime/nDisplay/Source/DisplayCluster/Public/DisplayClusterUtils/DisplayClusterTypesConverter.h
@@ -7,7 +7,6 @@
 #include "Misc/FrameRate.h"
 #include "DisplayClusterEnums.h"
 #include "DisplayClusterUtils/DisplayClusterCommonStrings.h"
-#include "DisplayClusterUtils/DisplayClusterCommonHelpers.h"
 
 
 /**
