Remove IsInGameThread check for GetViewportRect

diff --git a/Engine/Plugins/Runtime/nDisplay/Source/DisplayCluster/Private/Render/DisplayClusterRenderManager.cpp b/Engine/Plugins/Runtime/nDisplay/Source/DisplayCluster/Private/Render/DisplayClusterRenderManager.cpp
--- a/Engine/Plugins/Runtime/nDisplay/Source/DisplayCluster/Private/Render/DisplayClusterRenderManager.cpp
+++ b/Engine/Plugins/Runtime/nDisplay/Source/DisplayCluster/Private/Render/DisplayClusterRenderManager.cpp
@@ -476,7 +476,6 @@ void FDisplayClusterRenderManager::SetViewportCamera(const FString& InCameraId /
 bool FDisplayClusterRenderManager::GetViewportRect(const FString& InViewportID, FIntRect& Rect)
 {
 	DISPLAY_CLUSTER_FUNC_TRACE(LogDisplayClusterRender);
-	check(IsInGameThread());
 
 	if (!RenderDevicePtr)
 	{
