Fixed custom origin for PICP and PICP-MPCDI projection policies

diff --git a/Engine/Plugins/Runtime/nDisplay/Source/DisplayClusterProjection/Private/Policy/MPCDI/DisplayClusterProjectionMPCDIPolicy.cpp b/Engine/Plugins/Runtime/nDisplay/Source/DisplayClusterProjection/Private/Policy/MPCDI/DisplayClusterProjectionMPCDIPolicy.cpp
--- a/Engine/Plugins/Runtime/nDisplay/Source/DisplayClusterProjection/Private/Policy/MPCDI/DisplayClusterProjectionMPCDIPolicy.cpp
+++ b/Engine/Plugins/Runtime/nDisplay/Source/DisplayClusterProjection/Private/Policy/MPCDI/DisplayClusterProjectionMPCDIPolicy.cpp
@@ -77,6 +77,9 @@ bool FDisplayClusterProjectionMPCDIPolicy::HandleAddViewport(const FIntPoint& In
 		return false;
 	}
 
+	// Support custom origin node
+	OriginCompId = CfgData.OriginType;
+
 	UE_LOG(LogDisplayClusterProjectionMPCDI, Log, TEXT("MPCDI policy has been initialized [%s:%s in %s]"), *CfgData.BufferId, *CfgData.RegionId, *CfgData.MPCDIFileName);
 
 	// Finally, initialize internal views data container
diff --git a/Engine/Plugins/Runtime/nDisplay/Source/PicpProjection/Private/Policy/MPCDI/PicpProjectionMPCDIPolicy.cpp b/Engine/Plugins/Runtime/nDisplay/Source/PicpProjection/Private/Policy/MPCDI/PicpProjectionMPCDIPolicy.cpp
--- a/Engine/Plugins/Runtime/nDisplay/Source/PicpProjection/Private/Policy/MPCDI/PicpProjectionMPCDIPolicy.cpp
+++ b/Engine/Plugins/Runtime/nDisplay/Source/PicpProjection/Private/Policy/MPCDI/PicpProjectionMPCDIPolicy.cpp
@@ -92,7 +92,10 @@ bool FPicpProjectionMPCDIPolicy::HandleAddViewport(const FIntPoint& InViewportSi
 		UE_LOG(LogPicpProjectionMPCDI, Warning, TEXT("Couldn't load PICP MPCDI config: %s"), *CfgData.ConfigLineStr);
 		return false;
 	}
-	
+
+	// Support custom origin node
+	OriginCompId = CfgData.OriginType;
+
 	UE_LOG(LogPicpProjectionMPCDI, Log, TEXT("PICP MPCDI policy has been initialized [%s]"), *CfgData.ConfigLineStr);
 
 	// Finally, initialize internal views data container
