copy CL 11090020 and 11090793 from Dev-Editor to fix Tools-Staging build failure
#rnx

diff --git a/Engine/Plugins/Runtime/nDisplay/Source/DisplayCluster/Private/Config/DisplayClusterConfigManager.cpp b/Engine/Plugins/Runtime/nDisplay/Source/DisplayCluster/Private/Config/DisplayClusterConfigManager.cpp
--- a/Engine/Plugins/Runtime/nDisplay/Source/DisplayCluster/Private/Config/DisplayClusterConfigManager.cpp
+++ b/Engine/Plugins/Runtime/nDisplay/Source/DisplayCluster/Private/Config/DisplayClusterConfigManager.cpp
@@ -539,6 +539,29 @@ FString FDisplayClusterConfigManager::GetFullPathToFile(const FString& FileName)
 		}
 
 		//@Handle error, file not found
+		UE_LOG(LogDisplayClusterConfig, Warning, TEXT("File '%s' not found. In case of relative path do not forget to put './' at the beginning"), *FileName);
+	}
+
+	return FileName;
+}
+
+FString FDisplayClusterConfigManager::GetFullPathToNewFile(const FString& FileName) const
+{
+	TArray<FString> OrderedBaseDirs;
+
+	//Add ordered search base dirs
+	OrderedBaseDirs.Add(FPaths::GetPath(ConfigPath));
+	OrderedBaseDirs.Add(FPaths::RootDir());
+
+	// Process base dirs in order:
+	for (auto It : OrderedBaseDirs)
+	{
+		FString FullPath = FPaths::ConvertRelativePathToFull(It, FileName);
+		
+		if (FPaths::DirectoryExists(FPaths::GetPath(FullPath)))
+		{
+			return FullPath;
+		}
 	}
 
 	return FileName;
diff --git a/Engine/Plugins/Runtime/nDisplay/Source/DisplayCluster/Private/Config/DisplayClusterConfigManager.h b/Engine/Plugins/Runtime/nDisplay/Source/DisplayCluster/Private/Config/DisplayClusterConfigManager.h
--- a/Engine/Plugins/Runtime/nDisplay/Source/DisplayCluster/Private/Config/DisplayClusterConfigManager.h
+++ b/Engine/Plugins/Runtime/nDisplay/Source/DisplayCluster/Private/Config/DisplayClusterConfigManager.h
@@ -96,6 +96,7 @@ public:
 	{ return CfgCustom; }
 
 	virtual FString GetFullPathToFile(const FString& FileName) const override;
+	virtual FString GetFullPathToNewFile(const FString& FileName) const override;
 
 public:
 	//////////////////////////////////////////////////////////////////////////////////////////////
diff --git a/Engine/Plugins/Runtime/nDisplay/Source/DisplayCluster/Public/Config/IDisplayClusterConfigManager.h b/Engine/Plugins/Runtime/nDisplay/Source/DisplayCluster/Public/Config/IDisplayClusterConfigManager.h
--- a/Engine/Plugins/Runtime/nDisplay/Source/DisplayCluster/Public/Config/IDisplayClusterConfigManager.h
+++ b/Engine/Plugins/Runtime/nDisplay/Source/DisplayCluster/Public/Config/IDisplayClusterConfigManager.h
@@ -62,4 +62,5 @@ public:
 	virtual FDisplayClusterConfigCustom  GetConfigCustom()  const = 0;
 
 	virtual FString GetFullPathToFile(const FString& FileName) const = 0;
+	virtual FString GetFullPathToNewFile(const FString& FileName) const = 0;
 };
diff --git a/Engine/Plugins/Runtime/nDisplay/Source/DisplayClusterUtils/DisplayClusterCommonHelpers.h b/Engine/Plugins/Runtime/nDisplay/Source/DisplayClusterUtils/DisplayClusterCommonHelpers.h
--- a/Engine/Plugins/Runtime/nDisplay/Source/DisplayClusterUtils/DisplayClusterCommonHelpers.h
+++ b/Engine/Plugins/Runtime/nDisplay/Source/DisplayClusterUtils/DisplayClusterCommonHelpers.h
@@ -11,7 +11,7 @@
 #include "Config/DisplayClusterConfigTypes.h"
 #include "Game/IDisplayClusterGameManager.h"
 
-#include "DisplayClusterUtils/DisplayClusterTypesConverter.h"
+#include "DisplayClusterTypesConverter.h"
 
 class AActor;
 class UDisplayClusterCameraComponent;
@@ -465,6 +465,17 @@ namespace DisplayClusterHelpers
 
 			return LocalPath;
 		}
+
+		static FString GetNewFileFullPath(const FString& LocalPath)
+		{
+			static const IDisplayClusterConfigManager* const ConfigMgr = IDisplayCluster::Get().GetConfigMgr();
+			if (ConfigMgr)
+			{
+				return ConfigMgr->GetFullPathToNewFile(LocalPath);
+			}
+
+			return LocalPath;
+		}
 	}
 
 	//////////////////////////////////////////////////////////////////////////////////////////////
