copy Dev-Editor CL 11092476 to Tools-Staging to resolve CIS error
#rnx

diff --git a/Engine/Plugins/Runtime/nDisplay/Source/DisplayCluster/Private/Blueprints/DisplayClusterBlueprintAPIImpl.cpp b/Engine/Plugins/Runtime/nDisplay/Source/DisplayCluster/Private/Blueprints/DisplayClusterBlueprintAPIImpl.cpp
--- a/Engine/Plugins/Runtime/nDisplay/Source/DisplayCluster/Private/Blueprints/DisplayClusterBlueprintAPIImpl.cpp
+++ b/Engine/Plugins/Runtime/nDisplay/Source/DisplayCluster/Private/Blueprints/DisplayClusterBlueprintAPIImpl.cpp
@@ -15,9 +15,9 @@
 #include "DisplayClusterRootComponent.h"
 
 #include "Config/DisplayClusterConfigTypes.h"
-#include "Misc/DisplayClusterHelpers.h"
 
 #include "DisplayClusterGlobals.h"
+#include "DisplayClusterHelpers.h"
 
 
 //////////////////////////////////////////////////////////////////////////////////////////////
diff --git a/Engine/Plugins/Runtime/nDisplay/Source/DisplayCluster/Private/Cluster/DisplayClusterClusterManager.cpp b/Engine/Plugins/Runtime/nDisplay/Source/DisplayCluster/Private/Cluster/DisplayClusterClusterManager.cpp
--- a/Engine/Plugins/Runtime/nDisplay/Source/DisplayCluster/Private/Cluster/DisplayClusterClusterManager.cpp
+++ b/Engine/Plugins/Runtime/nDisplay/Source/DisplayCluster/Private/Cluster/DisplayClusterClusterManager.cpp
@@ -14,7 +14,6 @@
 #include "Dom/JsonObject.h"
 
 #include "Misc/DisplayClusterAppExit.h"
-#include "Misc/DisplayClusterHelpers.h"
 
 #include "DisplayClusterUtils/DisplayClusterTypesConverter.h"
 
@@ -24,6 +23,7 @@
 
 #include "DisplayClusterBuildConfig.h"
 #include "DisplayClusterGlobals.h"
+#include "DisplayClusterHelpers.h"
 #include "DisplayClusterLog.h"
 #include "DisplayClusterStrings.h"
 
@@ -586,7 +586,7 @@ FDisplayClusterClusterManager::TController FDisplayClusterClusterManager::Create
 {
 	DISPLAY_CLUSTER_FUNC_TRACE(LogDisplayClusterCluster);
 
-	UE_LOG(LogDisplayClusterCluster, Log, TEXT("Current operation mode: %s"), *FDisplayClusterTypesConverter::ToString(CurrentOperationMode));
+	UE_LOG(LogDisplayClusterCluster, Log, TEXT("Current operation mode: %s"), *FDisplayClusterTypesConverter::template ToString(CurrentOperationMode));
 
 	// Instantiate appropriate controller depending on operation mode and cluster role
 	FDisplayClusterNodeCtrlBase* pController = nullptr;
diff --git a/Engine/Plugins/Runtime/nDisplay/Source/DisplayCluster/Private/Config/DisplayClusterConfigTypes.cpp b/Engine/Plugins/Runtime/nDisplay/Source/DisplayCluster/Private/Config/DisplayClusterConfigTypes.cpp
--- a/Engine/Plugins/Runtime/nDisplay/Source/DisplayCluster/Private/Config/DisplayClusterConfigTypes.cpp
+++ b/Engine/Plugins/Runtime/nDisplay/Source/DisplayCluster/Private/Config/DisplayClusterConfigTypes.cpp
@@ -2,9 +2,9 @@
 
 #include "Config/DisplayClusterConfigTypes.h"
 
+#include "DisplayClusterHelpers.h"
 #include "DisplayClusterLog.h"
 #include "DisplayClusterStrings.h"
-#include "Misc/DisplayClusterHelpers.h"
 
 #include "DisplayClusterUtils/DisplayClusterTypesConverter.h"
 
diff --git a/Engine/Plugins/Runtime/nDisplay/Source/DisplayCluster/Private/DisplayClusterHelpers.h b/Engine/Plugins/Runtime/nDisplay/Source/DisplayCluster/Private/DisplayClusterHelpers.h
--- a/Engine/Plugins/Runtime/nDisplay/Source/DisplayCluster/Private/DisplayClusterHelpers.h
+++ b/Engine/Plugins/Runtime/nDisplay/Source/DisplayCluster/Private/DisplayClusterHelpers.h
@@ -3,10 +3,28 @@
 #pragma once
 
 #include "CoreMinimal.h"
-#include "DisplayClusterBuildConfig.h"
-#include "DisplayClusterUtils\DisplayClusterCommonHelpers.h"
+
+#include "Interfaces/IPv4/IPv4Address.h"
+#include "Interfaces/IPv4/IPv4Endpoint.h"
+
+#include "DisplayClusterUtils/DisplayClusterCommonHelpers.h"
 
 
 namespace DisplayClusterHelpers
 {
-};
\ No newline at end of file
+	//////////////////////////////////////////////////////////////////////////////////////////////
+	// Network helpers
+	//////////////////////////////////////////////////////////////////////////////////////////////
+	namespace net
+	{
+		static bool GenIPv4Endpoint(const FString& addr, const int32 port, FIPv4Endpoint& ep)
+		{
+			FIPv4Address ipAddr;
+			if (!FIPv4Address::Parse(addr, ipAddr))
+				return false;
+
+			ep = FIPv4Endpoint(ipAddr, port);
+			return true;
+		}
+	};
+};
diff --git a/Engine/Plugins/Runtime/nDisplay/Source/DisplayCluster/Private/DisplayClusterModule.cpp b/Engine/Plugins/Runtime/nDisplay/Source/DisplayCluster/Private/DisplayClusterModule.cpp
--- a/Engine/Plugins/Runtime/nDisplay/Source/DisplayCluster/Private/DisplayClusterModule.cpp
+++ b/Engine/Plugins/Runtime/nDisplay/Source/DisplayCluster/Private/DisplayClusterModule.cpp
@@ -82,7 +82,7 @@ bool FDisplayClusterModule::Init(EDisplayClusterOperationMode OperationMode)
 
 	CurrentOperationMode = OperationMode;
 
-	UE_LOG(LogDisplayClusterModule, Log, TEXT("Initializing subsystems to %s operation mode"), *FDisplayClusterTypesConverter::ToString(CurrentOperationMode));
+	UE_LOG(LogDisplayClusterModule, Log, TEXT("Initializing subsystems to %s operation mode"), *FDisplayClusterTypesConverter::template ToString(CurrentOperationMode));
 
 	bool result = true;
 	auto it = Managers.CreateIterator();
diff --git a/Engine/Plugins/Runtime/nDisplay/Source/DisplayCluster/Private/Game/Classes/Basics/DisplayClusterGameEngine.cpp b/Engine/Plugins/Runtime/nDisplay/Source/DisplayCluster/Private/Game/Classes/Basics/DisplayClusterGameEngine.cpp
--- a/Engine/Plugins/Runtime/nDisplay/Source/DisplayCluster/Private/Game/Classes/Basics/DisplayClusterGameEngine.cpp
+++ b/Engine/Plugins/Runtime/nDisplay/Source/DisplayCluster/Private/Game/Classes/Basics/DisplayClusterGameEngine.cpp
@@ -12,11 +12,10 @@
 #include "Misc/DisplayClusterAppExit.h"
 #include "Misc/Parse.h"
 
-#include "Misc/DisplayClusterHelpers.h"
-
 #include "DisplayClusterBuildConfig.h"
 #include "DisplayClusterEnums.h"
 #include "DisplayClusterGlobals.h"
+#include "DisplayClusterHelpers.h"
 #include "DisplayClusterLog.h"
 #include "DisplayClusterStrings.h"
 
@@ -111,7 +110,7 @@ EDisplayClusterOperationMode UDisplayClusterGameEngine::DetectOperationMode()
 		OpMode = EDisplayClusterOperationMode::Standalone;
 	}
 
-	UE_LOG(LogDisplayClusterEngine, Log, TEXT("Detected operation mode: %s"), *FDisplayClusterTypesConverter::ToString(OpMode));
+	UE_LOG(LogDisplayClusterEngine, Log, TEXT("Detected operation mode: %s"), *FDisplayClusterTypesConverter::template ToString(OpMode));
 
 	return OpMode;
 }
@@ -130,7 +129,7 @@ bool UDisplayClusterGameEngine::InitializeInternals()
 	FDisplayClusterConfigClusterNode LocalClusterNode;
 	if (DisplayClusterHelpers::config::GetLocalClusterNode(LocalClusterNode))
 	{
-		UE_LOG(LogDisplayClusterEngine, Log, TEXT("Configuring sound enabled: %s"), *FDisplayClusterTypesConverter::ToString(LocalClusterNode.SoundEnabled));
+		UE_LOG(LogDisplayClusterEngine, Log, TEXT("Configuring sound enabled: %s"), *FDisplayClusterTypesConverter::template ToString(LocalClusterNode.SoundEnabled));
 		bUseSound = LocalClusterNode.SoundEnabled;
 	}
 
diff --git a/Engine/Plugins/Runtime/nDisplay/Source/DisplayCluster/Private/Game/Classes/Basics/DisplayClusterPlayerInput.cpp b/Engine/Plugins/Runtime/nDisplay/Source/DisplayCluster/Private/Game/Classes/Basics/DisplayClusterPlayerInput.cpp
--- a/Engine/Plugins/Runtime/nDisplay/Source/DisplayCluster/Private/Game/Classes/Basics/DisplayClusterPlayerInput.cpp
+++ b/Engine/Plugins/Runtime/nDisplay/Source/DisplayCluster/Private/Game/Classes/Basics/DisplayClusterPlayerInput.cpp
@@ -5,9 +5,9 @@
 
 #include "Cluster/IPDisplayClusterClusterManager.h"
 #include "DisplayClusterGlobals.h"
-#include "DisplayClusterHelpers.h"
 #include "DisplayClusterLog.h"
 #include "DisplayClusterUtils/DisplayClusterTypesConverter.h"
+#include "DisplayClusterHelpers.h"
 
 #include "Engine/World.h"
 
@@ -58,14 +58,14 @@ bool UDisplayClusterPlayerInput::SerializeKeyStateMap(TMap<FString, FString>& Ou
 		check(SendValue > 0.f);
 
 		StrKeyState = FString::Printf(TEXT("%s;%s;%s;%s;%s;%s;%s;%s;"),
-			*FDisplayClusterTypesConverter::ToString(it->Value.RawValue),
-			*FDisplayClusterTypesConverter::ToString(it->Value.Value),
-			*FDisplayClusterTypesConverter::ToString(it->Value.LastUpDownTransitionTime),
+			*FDisplayClusterTypesConverter::template ToString(it->Value.RawValue),
+			*FDisplayClusterTypesConverter::template ToString(it->Value.Value),
+			*FDisplayClusterTypesConverter::template ToString(it->Value.LastUpDownTransitionTime),
 			*DisplayClusterHelpers::str::BoolToStr(it->Value.bDown ? true : false, false),
 			*DisplayClusterHelpers::str::BoolToStr(it->Value.bDownPrevious ? true : false, false),
 			*DisplayClusterHelpers::str::BoolToStr(it->Value.bConsumed ? true : false, false),
-			*FDisplayClusterTypesConverter::ToString(it->Value.SampleCountAccumulator),
-			*FDisplayClusterTypesConverter::ToString(it->Value.RawValueAccumulator));
+			*FDisplayClusterTypesConverter::template ToString(it->Value.SampleCountAccumulator),
+			*FDisplayClusterTypesConverter::template ToString(it->Value.RawValueAccumulator));
 
 		for (int i = 0; i < IE_MAX; ++i)
 		{
@@ -101,14 +101,14 @@ bool UDisplayClusterPlayerInput::DeserializeKeyStateMap(const TMap<FString, FStr
 
 		TArray<FString> Fields = DisplayClusterHelpers::str::StrToArray<FString>(it->Value, FString(";"), false);
 
-		KeyState.RawValue                 = FDisplayClusterTypesConverter::FromString<FVector>(Fields[0]);
-		KeyState.Value                    = FDisplayClusterTypesConverter::FromString<FVector>(Fields[1]);
-		KeyState.LastUpDownTransitionTime = FDisplayClusterTypesConverter::FromString<float>(Fields[2]);
-		KeyState.bDown                    = FDisplayClusterTypesConverter::FromString<bool>(Fields[3]) ? 1 : 0;
-		KeyState.bDownPrevious            = FDisplayClusterTypesConverter::FromString<bool>(Fields[4]) ? 1 : 0;
-		KeyState.bConsumed                = FDisplayClusterTypesConverter::FromString<bool>(Fields[5]) ? 1 : 0;
-		KeyState.SampleCountAccumulator   = FDisplayClusterTypesConverter::FromString<uint8>(Fields[6]);
-		KeyState.RawValueAccumulator      = FDisplayClusterTypesConverter::FromString<FVector>(Fields[7]);
+		KeyState.RawValue                 = FDisplayClusterTypesConverter::template FromString<FVector>(Fields[0]);
+		KeyState.Value                    = FDisplayClusterTypesConverter::template FromString<FVector>(Fields[1]);
+		KeyState.LastUpDownTransitionTime = FDisplayClusterTypesConverter::template FromString<float>(Fields[2]);
+		KeyState.bDown                    = FDisplayClusterTypesConverter::template FromString<bool>(Fields[3]) ? 1 : 0;
+		KeyState.bDownPrevious            = FDisplayClusterTypesConverter::template FromString<bool>(Fields[4]) ? 1 : 0;
+		KeyState.bConsumed                = FDisplayClusterTypesConverter::template FromString<bool>(Fields[5]) ? 1 : 0;
+		KeyState.SampleCountAccumulator   = FDisplayClusterTypesConverter::template FromString<uint8>(Fields[6]);
+		KeyState.RawValueAccumulator      = FDisplayClusterTypesConverter::template FromString<FVector>(Fields[7]);
 
 		for (int i = 0; i < IE_MAX; ++i)
 		{
diff --git a/Engine/Plugins/Runtime/nDisplay/Source/DisplayCluster/Private/Game/DisplayClusterGameManager.cpp b/Engine/Plugins/Runtime/nDisplay/Source/DisplayCluster/Private/Game/DisplayClusterGameManager.cpp
--- a/Engine/Plugins/Runtime/nDisplay/Source/DisplayCluster/Private/Game/DisplayClusterGameManager.cpp
+++ b/Engine/Plugins/Runtime/nDisplay/Source/DisplayCluster/Private/Game/DisplayClusterGameManager.cpp
@@ -4,7 +4,6 @@
 
 #include "Kismet/GameplayStatics.h"
 #include "Misc/CommandLine.h"
-#include "Misc/DisplayClusterHelpers.h"
 #include "Config/DisplayClusterConfigTypes.h"
 #include "Config/IPDisplayClusterConfigManager.h"
 
@@ -17,6 +16,7 @@
 #include "DisplayClusterScreenComponent.h"
 
 #include "DisplayClusterGlobals.h"
+#include "DisplayClusterHelpers.h"
 #include "DisplayClusterLog.h"
 #include "DisplayClusterStrings.h"
 
diff --git a/Engine/Plugins/Runtime/nDisplay/Source/DisplayCluster/Private/Input/Devices/VRPN/Analog/DisplayClusterVrpnAnalogInputDevice.cpp b/Engine/Plugins/Runtime/nDisplay/Source/DisplayCluster/Private/Input/Devices/VRPN/Analog/DisplayClusterVrpnAnalogInputDevice.cpp
--- a/Engine/Plugins/Runtime/nDisplay/Source/DisplayCluster/Private/Input/Devices/VRPN/Analog/DisplayClusterVrpnAnalogInputDevice.cpp
+++ b/Engine/Plugins/Runtime/nDisplay/Source/DisplayCluster/Private/Input/Devices/VRPN/Analog/DisplayClusterVrpnAnalogInputDevice.cpp
@@ -2,8 +2,7 @@
 
 #include "Input/Devices/VRPN/Analog/DisplayClusterVrpnAnalogInputDevice.h"
 
-#include "Misc/DisplayClusterHelpers.h"
-
+#include "DisplayClusterHelpers.h"
 #include "DisplayClusterLog.h"
 #include "DisplayClusterStrings.h"
 
diff --git a/Engine/Plugins/Runtime/nDisplay/Source/DisplayCluster/Private/Input/Devices/VRPN/Button/DisplayClusterVrpnButtonInputDevice.cpp b/Engine/Plugins/Runtime/nDisplay/Source/DisplayCluster/Private/Input/Devices/VRPN/Button/DisplayClusterVrpnButtonInputDevice.cpp
--- a/Engine/Plugins/Runtime/nDisplay/Source/DisplayCluster/Private/Input/Devices/VRPN/Button/DisplayClusterVrpnButtonInputDevice.cpp
+++ b/Engine/Plugins/Runtime/nDisplay/Source/DisplayCluster/Private/Input/Devices/VRPN/Button/DisplayClusterVrpnButtonInputDevice.cpp
@@ -2,8 +2,7 @@
 
 #include "Input/Devices/VRPN/Button/DisplayClusterVrpnButtonInputDevice.h"
 
-#include "Misc/DisplayClusterHelpers.h"
-
+#include "DisplayClusterHelpers.h"
 #include "DisplayClusterLog.h"
 #include "DisplayClusterStrings.h"
 
diff --git a/Engine/Plugins/Runtime/nDisplay/Source/DisplayCluster/Private/Input/Devices/VRPN/Keyboard/DisplayClusterVrpnKeyboardInputDevice.cpp b/Engine/Plugins/Runtime/nDisplay/Source/DisplayCluster/Private/Input/Devices/VRPN/Keyboard/DisplayClusterVrpnKeyboardInputDevice.cpp
--- a/Engine/Plugins/Runtime/nDisplay/Source/DisplayCluster/Private/Input/Devices/VRPN/Keyboard/DisplayClusterVrpnKeyboardInputDevice.cpp
+++ b/Engine/Plugins/Runtime/nDisplay/Source/DisplayCluster/Private/Input/Devices/VRPN/Keyboard/DisplayClusterVrpnKeyboardInputDevice.cpp
@@ -2,8 +2,7 @@
 
 #include "DisplayClusterVrpnKeyboardInputDevice.h"
 
-#include "Misc/DisplayClusterHelpers.h"
-
+#include "DisplayClusterHelpers.h"
 #include "DisplayClusterLog.h"
 #include "DisplayClusterStrings.h"
 
diff --git a/Engine/Plugins/Runtime/nDisplay/Source/DisplayCluster/Private/Input/Devices/VRPN/Tracker/DisplayClusterVrpnTrackerInputDevice.cpp b/Engine/Plugins/Runtime/nDisplay/Source/DisplayCluster/Private/Input/Devices/VRPN/Tracker/DisplayClusterVrpnTrackerInputDevice.cpp
--- a/Engine/Plugins/Runtime/nDisplay/Source/DisplayCluster/Private/Input/Devices/VRPN/Tracker/DisplayClusterVrpnTrackerInputDevice.cpp
+++ b/Engine/Plugins/Runtime/nDisplay/Source/DisplayCluster/Private/Input/Devices/VRPN/Tracker/DisplayClusterVrpnTrackerInputDevice.cpp
@@ -2,9 +2,8 @@
 
 #include "Input/Devices/VRPN/Tracker/DisplayClusterVrpnTrackerInputDevice.h"
 
-#include "Misc/DisplayClusterHelpers.h"
-
 #include "DisplayClusterBuildConfig.h"
+#include "DisplayClusterHelpers.h"
 #include "DisplayClusterLog.h"
 #include "DisplayClusterStrings.h"
 
diff --git a/Engine/Plugins/Runtime/nDisplay/Source/DisplayCluster/Private/Misc/DisplayClusterHelpers.h b/Engine/Plugins/Runtime/nDisplay/Source/DisplayCluster/Private/Misc/DisplayClusterHelpers.h
deleted file mode 100644
diff --git a/Engine/Plugins/Runtime/nDisplay/Source/DisplayCluster/Private/Network/DisplayClusterMessage.h b/Engine/Plugins/Runtime/nDisplay/Source/DisplayCluster/Private/Network/DisplayClusterMessage.h
--- a/Engine/Plugins/Runtime/nDisplay/Source/DisplayCluster/Private/Network/DisplayClusterMessage.h
+++ b/Engine/Plugins/Runtime/nDisplay/Source/DisplayCluster/Private/Network/DisplayClusterMessage.h
@@ -44,7 +44,7 @@ public:
 		if (Arguments.Contains(ArgName))
 		{
 			FString StrVal = Arguments[ArgName];
-			ArgVal = FDisplayClusterTypesConverter::FromString<ValType>(StrVal);
+			ArgVal = FDisplayClusterTypesConverter::template FromString<ValType>(StrVal);
 			return true;
 		}
 		return false;
@@ -54,7 +54,7 @@ public:
 	template <typename ValType>
 	void SetArg(const FString& ArgName, const ValType& ArgVal)
 	{
-		Arguments.Add(ArgName, FDisplayClusterTypesConverter::ToString<ValType>(ArgVal));
+		Arguments.Add(ArgName, FDisplayClusterTypesConverter::template ToString<ValType>(ArgVal));
 	}
 
 	// Get all arguments (be careful with the reference)
diff --git a/Engine/Plugins/Runtime/nDisplay/Source/DisplayCluster/Private/Network/DisplayClusterSocketOps.cpp b/Engine/Plugins/Runtime/nDisplay/Source/DisplayCluster/Private/Network/DisplayClusterSocketOps.cpp
--- a/Engine/Plugins/Runtime/nDisplay/Source/DisplayCluster/Private/Network/DisplayClusterSocketOps.cpp
+++ b/Engine/Plugins/Runtime/nDisplay/Source/DisplayCluster/Private/Network/DisplayClusterSocketOps.cpp
@@ -7,6 +7,8 @@
 
 #include "Misc/ScopeLock.h"
 
+#include "Network/DisplayClusterMessage.h"
+
 #include "Dom/JsonObject.h"
 #include "Serialization/JsonReader.h"
 #include "Serialization/JsonWriter.h"
diff --git a/Engine/Plugins/Runtime/nDisplay/Source/DisplayCluster/Private/Network/DisplayClusterSocketOps.h b/Engine/Plugins/Runtime/nDisplay/Source/DisplayCluster/Private/Network/DisplayClusterSocketOps.h
--- a/Engine/Plugins/Runtime/nDisplay/Source/DisplayCluster/Private/Network/DisplayClusterSocketOps.h
+++ b/Engine/Plugins/Runtime/nDisplay/Source/DisplayCluster/Private/Network/DisplayClusterSocketOps.h
@@ -5,8 +5,7 @@
 #include "CoreMinimal.h"
 #include "Sockets.h"
 
-#include "Network/DisplayClusterMessage.h"
-
+class FDisplayClusterMessage;
 class FJsonObject;
 
 
diff --git a/Engine/Plugins/Runtime/nDisplay/Source/DisplayCluster/Private/Network/DisplayClusterTcpListener.cpp b/Engine/Plugins/Runtime/nDisplay/Source/DisplayCluster/Private/Network/DisplayClusterTcpListener.cpp
--- a/Engine/Plugins/Runtime/nDisplay/Source/DisplayCluster/Private/Network/DisplayClusterTcpListener.cpp
+++ b/Engine/Plugins/Runtime/nDisplay/Source/DisplayCluster/Private/Network/DisplayClusterTcpListener.cpp
@@ -7,8 +7,8 @@
 #include "Common/TcpSocketBuilder.h"
 
 #include "Misc/DisplayClusterAppExit.h"
-#include "Misc/DisplayClusterHelpers.h"
 
+#include "DisplayClusterHelpers.h"
 #include "DisplayClusterLog.h"
 
 
diff --git a/Engine/Plugins/Runtime/nDisplay/Source/DisplayCluster/Private/Network/Session/DisplayClusterSessionBase.cpp b/Engine/Plugins/Runtime/nDisplay/Source/DisplayCluster/Private/Network/Session/DisplayClusterSessionBase.cpp
--- a/Engine/Plugins/Runtime/nDisplay/Source/DisplayCluster/Private/Network/Session/DisplayClusterSessionBase.cpp
+++ b/Engine/Plugins/Runtime/nDisplay/Source/DisplayCluster/Private/Network/Session/DisplayClusterSessionBase.cpp
@@ -2,7 +2,6 @@
 
 #include "Network/Session/DisplayClusterSessionBase.h"
 #include "Network/DisplayClusterServer.h"
-#include "Network/DisplayClusterMessage.h"
 
 #include "HAL/RunnableThread.h"
 
diff --git a/Engine/Plugins/Runtime/nDisplay/Source/DisplayCluster/Private/Render/Device/DisplayClusterDeviceBase.cpp b/Engine/Plugins/Runtime/nDisplay/Source/DisplayCluster/Private/Render/Device/DisplayClusterDeviceBase.cpp
--- a/Engine/Plugins/Runtime/nDisplay/Source/DisplayCluster/Private/Render/Device/DisplayClusterDeviceBase.cpp
+++ b/Engine/Plugins/Runtime/nDisplay/Source/DisplayCluster/Private/Render/Device/DisplayClusterDeviceBase.cpp
@@ -15,8 +15,6 @@
 #include "RHIStaticStates.h"
 #include "Slate/SceneViewport.h"
 
-#include "Misc/DisplayClusterHelpers.h"
-
 #include "Render/Device/DisplayClusterRenderViewport.h"
 #include "Render/PostProcess/IDisplayClusterPostProcess.h"
 #include "Render/Presentation/DisplayClusterPresentationBase.h"
@@ -25,6 +23,7 @@
 #include "Render/Synchronization/IDisplayClusterRenderSyncPolicy.h"
 
 #include "DisplayClusterGlobals.h"
+#include "DisplayClusterHelpers.h"
 #include "DisplayClusterLog.h"
 
 #include <utility>
diff --git a/Engine/Plugins/Runtime/nDisplay/Source/DisplayCluster/Private/Render/Device/DisplayClusterDeviceBase_PostProcess.cpp b/Engine/Plugins/Runtime/nDisplay/Source/DisplayCluster/Private/Render/Device/DisplayClusterDeviceBase_PostProcess.cpp
--- a/Engine/Plugins/Runtime/nDisplay/Source/DisplayCluster/Private/Render/Device/DisplayClusterDeviceBase_PostProcess.cpp
+++ b/Engine/Plugins/Runtime/nDisplay/Source/DisplayCluster/Private/Render/Device/DisplayClusterDeviceBase_PostProcess.cpp
@@ -3,7 +3,7 @@
 #include "Render/Device/DisplayClusterDeviceBase_PostProcess.h"
 #include "Render/Device/DisplayClusterRenderViewport.h"
 
-#include "Misc/DisplayClusterHelpers.h"
+#include "DisplayClusterHelpers.h"
 #include "DisplayClusterLog.h"
 
 #include "HAL/IConsoleManager.h"
diff --git a/Engine/Plugins/Runtime/nDisplay/Source/DisplayCluster/Private/Render/DisplayClusterRenderManager.cpp b/Engine/Plugins/Runtime/nDisplay/Source/DisplayCluster/Private/Render/DisplayClusterRenderManager.cpp
--- a/Engine/Plugins/Runtime/nDisplay/Source/DisplayCluster/Private/Render/DisplayClusterRenderManager.cpp
+++ b/Engine/Plugins/Runtime/nDisplay/Source/DisplayCluster/Private/Render/DisplayClusterRenderManager.cpp
@@ -33,8 +33,8 @@
 #include "Render/Synchronization/DisplayClusterRenderSyncPolicyNone.h"
 #include "Render/Synchronization/DisplayClusterRenderSyncPolicySoftwareGeneric.h"
 
-#include "Misc/DisplayClusterHelpers.h"
 #include "DisplayClusterUtils/DisplayClusterTypesConverter.h"
+#include "DisplayClusterHelpers.h"
 
 #include "UnrealClient.h"
 #include "Kismet/GameplayStatics.h"
@@ -691,7 +691,7 @@ TSharedPtr<IDisplayClusterRenderSyncPolicy> FDisplayClusterRenderManager::Create
 
 	// Create sync policy specified in a config file
 	FDisplayClusterConfigGeneral CfgGeneral = GDisplayCluster->GetPrivateConfigMgr()->GetConfigGeneral();
-	const FString SyncPolicyType = FDisplayClusterTypesConverter::ToString(CfgGeneral.SwapSyncPolicy);
+	const FString SyncPolicyType = FDisplayClusterTypesConverter::template ToString(CfgGeneral.SwapSyncPolicy);
 	const FString RHIName = GDynamicRHI->GetName();
 	TSharedPtr<IDisplayClusterRenderSyncPolicy> NewSyncPolicy;
 
diff --git a/Engine/Plugins/Runtime/nDisplay/Source/DisplayClusterUtils/DisplayClusterCommonHelpers.h b/Engine/Plugins/Runtime/nDisplay/Source/DisplayClusterUtils/DisplayClusterCommonHelpers.h
--- a/Engine/Plugins/Runtime/nDisplay/Source/DisplayClusterUtils/DisplayClusterCommonHelpers.h
+++ b/Engine/Plugins/Runtime/nDisplay/Source/DisplayClusterUtils/DisplayClusterCommonHelpers.h
@@ -11,7 +11,7 @@
 #include "Config/DisplayClusterConfigTypes.h"
 #include "Game/IDisplayClusterGameManager.h"
 
-#include "DisplayClusterTypesConverter.h"
+#include "DisplayClusterUtils/DisplayClusterTypesConverter.h"
 
 class AActor;
 class UDisplayClusterCameraComponent;
@@ -91,7 +91,7 @@ namespace DisplayClusterHelpers
 
 			for (int i = 0; i < InData.Num(); ++i)
 			{
-				ResultStr += FString::Printf(TEXT("%s%s"), *FDisplayClusterTypesConverter::ToString(InData[i]), *InSeparator);
+				ResultStr += FString::Printf(TEXT("%s%s"), *FDisplayClusterTypesConverter::template ToString(InData[i]), *InSeparator);
 			}
 
 			if (InSeparator.Len() > 0 && InData.Num() > 0)
@@ -143,7 +143,7 @@ namespace DisplayClusterHelpers
 
 			for (const auto& Pair : InData)
 			{
-				ResultStr = FString::Printf(TEXT("%s%s%s%s%s"), *ResultStr, *FDisplayClusterTypesConverter::ToString(Pair.Key), *InKeyValSeparator, *FDisplayClusterTypesConverter::ToString(Pair.Value), *InPairSeparator);
+				ResultStr = FString::Printf(TEXT("%s%s%s%s%s"), *ResultStr, *FDisplayClusterTypesConverter::template ToString(Pair.Key), *InKeyValSeparator, *FDisplayClusterTypesConverter::template ToString(Pair.Value), *InPairSeparator);
 			}
 
 			if (InPairSeparator.Len() > 0)
diff --git a/Engine/Plugins/Runtime/nDisplay/Source/DisplayClusterUtils/DisplayClusterTypesConverter.h b/Engine/Plugins/Runtime/nDisplay/Source/DisplayClusterUtils/DisplayClusterTypesConverter.h
--- a/Engine/Plugins/Runtime/nDisplay/Source/DisplayClusterUtils/DisplayClusterTypesConverter.h
+++ b/Engine/Plugins/Runtime/nDisplay/Source/DisplayClusterUtils/DisplayClusterTypesConverter.h
@@ -7,7 +7,6 @@
 #include "Misc/FrameRate.h"
 #include "DisplayClusterEnums.h"
 #include "DisplayClusterUtils/DisplayClusterCommonStrings.h"
-#include "DisplayClusterUtils/DisplayClusterCommonHelpers.h"
 
 
 /**
