Unshelved from pending changelist '10672112':

Unsubscribe from cluster events on 'EndPlay' instead of 'Destroyed'

diff --git a/Engine/Plugins/Runtime/nDisplay/Source/DisplayCluster/Private/Game/Classes/Scene/DisplayClusterSceneComponentSync.cpp b/Engine/Plugins/Runtime/nDisplay/Source/DisplayCluster/Private/Game/Classes/Scene/DisplayClusterSceneComponentSync.cpp
--- a/Engine/Plugins/Runtime/nDisplay/Source/DisplayCluster/Private/Game/Classes/Scene/DisplayClusterSceneComponentSync.cpp
+++ b/Engine/Plugins/Runtime/nDisplay/Source/DisplayCluster/Private/Game/Classes/Scene/DisplayClusterSceneComponentSync.cpp
@@ -46,15 +46,7 @@ void UDisplayClusterSceneComponentSync::BeginPlay()
 	}
 }
 
-
-void UDisplayClusterSceneComponentSync::TickComponent( float DeltaTime, ELevelTick TickType, FActorComponentTickFunction* ThisTickFunction )
-{
-	Super::TickComponent( DeltaTime, TickType, ThisTickFunction );
-
-	// ...
-}
-
-void UDisplayClusterSceneComponentSync::OnComponentDestroyed(bool bDestroyingHierarchy)
+void UDisplayClusterSceneComponentSync::EndPlay(const EEndPlayReason::Type EndPlayReason)
 {
 	if (GDisplayCluster->IsModuleInitialized())
 	{
@@ -68,7 +60,14 @@ void UDisplayClusterSceneComponentSync::OnComponentDestroyed(bool bDestroyingHie
 		}
 	}
 
-	Super::OnComponentDestroyed(bDestroyingHierarchy);
+	Super::EndPlay(EndPlayReason);
+}
+
+void UDisplayClusterSceneComponentSync::TickComponent( float DeltaTime, ELevelTick TickType, FActorComponentTickFunction* ThisTickFunction )
+{
+	Super::TickComponent( DeltaTime, TickType, ThisTickFunction );
+
+	// ...
 }
 
 
diff --git a/Engine/Plugins/Runtime/nDisplay/Source/DisplayCluster/Public/DisplayClusterSceneComponentSync.h b/Engine/Plugins/Runtime/nDisplay/Source/DisplayCluster/Public/DisplayClusterSceneComponentSync.h
--- a/Engine/Plugins/Runtime/nDisplay/Source/DisplayCluster/Public/DisplayClusterSceneComponentSync.h
+++ b/Engine/Plugins/Runtime/nDisplay/Source/DisplayCluster/Public/DisplayClusterSceneComponentSync.h
@@ -50,8 +50,8 @@ public:
 
 public:
 	virtual void BeginPlay() override;
+	virtual void EndPlay(const EEndPlayReason::Type EndPlayReason) override;
 	virtual void TickComponent(float DeltaTime, ELevelTick TickType, FActorComponentTickFunction* ThisTickFunction) override;
-	virtual void OnComponentDestroyed(bool bDestroyingHierarchy) override;
 
 protected:
 	virtual FString GenerateSyncId();
