Stability update: added some nullptr checks

diff --git a/Engine/Plugins/Runtime/nDisplay/Source/DisplayCluster/Private/Game/DisplayClusterGameManager.cpp b/Engine/Plugins/Runtime/nDisplay/Source/DisplayCluster/Private/Game/DisplayClusterGameManager.cpp
--- a/Engine/Plugins/Runtime/nDisplay/Source/DisplayCluster/Private/Game/DisplayClusterGameManager.cpp
+++ b/Engine/Plugins/Runtime/nDisplay/Source/DisplayCluster/Private/Game/DisplayClusterGameManager.cpp
@@ -259,15 +259,18 @@ TArray<UDisplayClusterSceneComponent*> FDisplayClusterGameManager::GetAllNodes()
 
 ADisplayClusterRootActor* FDisplayClusterGameManager::FindDisplayClusterRootActor(UWorld* InWorld)
 {
-	for (AActor* const Actor : InWorld->PersistentLevel->Actors)
+	if (InWorld && InWorld->PersistentLevel)
 	{
-		if (Actor && !Actor->IsPendingKill())
+		for (AActor* const Actor : InWorld->PersistentLevel->Actors)
 		{
-			ADisplayClusterRootActor* RootActor = Cast<ADisplayClusterRootActor>(Actor);
-			if (RootActor)
+			if (Actor && !Actor->IsPendingKill())
 			{
-				UE_LOG(LogDisplayClusterGame, Log, TEXT("Found root actor - %s"), *RootActor->GetName());
-				return RootActor;
+				ADisplayClusterRootActor* RootActor = Cast<ADisplayClusterRootActor>(Actor);
+				if (RootActor)
+				{
+					UE_LOG(LogDisplayClusterGame, Log, TEXT("Found root actor - %s"), *RootActor->GetName());
+					return RootActor;
+				}
 			}
 		}
 	}
diff --git a/Engine/Plugins/Runtime/nDisplay/Source/DisplayClusterEditor/Private/DisplayClusterEditorEngine.cpp b/Engine/Plugins/Runtime/nDisplay/Source/DisplayClusterEditor/Private/DisplayClusterEditorEngine.cpp
--- a/Engine/Plugins/Runtime/nDisplay/Source/DisplayClusterEditor/Private/DisplayClusterEditorEngine.cpp
+++ b/Engine/Plugins/Runtime/nDisplay/Source/DisplayClusterEditor/Private/DisplayClusterEditorEngine.cpp
@@ -50,15 +50,18 @@ void UDisplayClusterEditorEngine::PreExit()
 
 ADisplayClusterRootActor* UDisplayClusterEditorEngine::FindDisplayClusterRootActor(UWorld* InWorld)
 {
-	for (AActor* const Actor : InWorld->PersistentLevel->Actors)
+	if (InWorld && InWorld->PersistentLevel)
 	{
-		if (Actor && !Actor->IsPendingKill())
+		for (AActor* const Actor : InWorld->PersistentLevel->Actors)
 		{
-			ADisplayClusterRootActor* RootActor = Cast<ADisplayClusterRootActor>(Actor);
-			if (RootActor)
+			if (Actor && !Actor->IsPendingKill())
 			{
-				UE_LOG(LogDisplayClusterEditorEngine, Log, TEXT("Found root actor - %s"), *RootActor->GetName());
-				return RootActor;
+				ADisplayClusterRootActor* RootActor = Cast<ADisplayClusterRootActor>(Actor);
+				if (RootActor)
+				{
+					UE_LOG(LogDisplayClusterEditorEngine, Log, TEXT("Found root actor - %s"), *RootActor->GetName());
+					return RootActor;
+				}
 			}
 		}
 	}
@@ -76,14 +79,16 @@ void UDisplayClusterEditorEngine::StartPlayInEditorSession(FRequestPlaySessionPa
 	{
 		// Find nDisplay root actor
 		ADisplayClusterRootActor* RootActor = FindDisplayClusterRootActor(EditorWorldPreDup);
-		if (!RootActor)
+		if (!RootActor && EditorWorld)
 		{
 			// Also search inside streamed levels
 			const TArray<ULevelStreaming*>& StreamingLevels = EditorWorld->GetStreamingLevels();
 			for (ULevelStreaming* StreamingLevel : StreamingLevels)
 			{
-				switch (StreamingLevel->GetCurrentState())
+				if (StreamingLevel)
 				{
+					switch (StreamingLevel->GetCurrentState())
+					{
 					case ULevelStreaming::ECurrentState::LoadedVisible:
 					{
 						// Look for the actor in those sub-levels that have been loaded already
@@ -97,6 +102,7 @@ void UDisplayClusterEditorEngine::StartPlayInEditorSession(FRequestPlaySessionPa
 
 					default:
 						break;
+					}
 				}
 			}
 		}
