Set custom present on PreTick

diff --git a/Engine/Plugins/Runtime/nDisplay/Source/DisplayCluster/Private/Render/Device/DisplayClusterDeviceBase.cpp b/Engine/Plugins/Runtime/nDisplay/Source/DisplayCluster/Private/Render/Device/DisplayClusterDeviceBase.cpp
--- a/Engine/Plugins/Runtime/nDisplay/Source/DisplayCluster/Private/Render/Device/DisplayClusterDeviceBase.cpp
+++ b/Engine/Plugins/Runtime/nDisplay/Source/DisplayCluster/Private/Render/Device/DisplayClusterDeviceBase.cpp
@@ -166,6 +166,28 @@ void FDisplayClusterDeviceBase::EndScene()
 	bIsSceneOpen = false;
 }
 
+void FDisplayClusterDeviceBase::PreTick(float DeltaSeconds)
+{
+	if (!bIsCustomPresentSet)
+	{
+		// Set up our new present handler
+		if (MainViewport)
+		{
+			// Current sync policy
+			TSharedPtr<IDisplayClusterRenderSyncPolicy> SyncPolicy = GDisplayCluster->GetRenderMgr()->GetCurrentSynchronizationPolicy();
+			check(SyncPolicy.IsValid());
+
+			// Create present handler
+			FDisplayClusterPresentationBase* const CustomPresentHandler = CreatePresentationObject(MainViewport, SyncPolicy);
+			check(CustomPresentHandler);
+
+			MainViewport->GetViewportRHI()->SetCustomPresent(CustomPresentHandler);
+		}
+
+		bIsCustomPresentSet = true;
+	}
+}
+
 void FDisplayClusterDeviceBase::SetViewportCamera(const FString& InCameraId /* = FString() */, const FString& InViewportId /* = FString() */)
 {
 	DISPLAY_CLUSTER_FUNC_TRACE(LogDisplayClusterRender);
@@ -584,15 +606,6 @@ void FDisplayClusterDeviceBase::UpdateViewport(bool bUseSeparateRenderTarget, co
 	{
 		// UE viewport
 		MainViewport = (FViewport*)&Viewport;
-		// Current sync policy
-		TSharedPtr<IDisplayClusterRenderSyncPolicy> SyncPolicy = GDisplayCluster->GetRenderMgr()->GetCurrentSynchronizationPolicy();
-
-		// Create present handler
-		FDisplayClusterPresentationBase* const CustomPresentHandler = CreatePresentationObject(MainViewport, SyncPolicy);
-		check(CustomPresentHandler);
-
-		// Set up our new present handler
-		Viewport.GetViewportRHI()->SetCustomPresent(CustomPresentHandler);
 	}
 }
 
diff --git a/Engine/Plugins/Runtime/nDisplay/Source/DisplayCluster/Private/Render/Device/DisplayClusterDeviceBase.h b/Engine/Plugins/Runtime/nDisplay/Source/DisplayCluster/Private/Render/Device/DisplayClusterDeviceBase.h
--- a/Engine/Plugins/Runtime/nDisplay/Source/DisplayCluster/Private/Render/Device/DisplayClusterDeviceBase.h
+++ b/Engine/Plugins/Runtime/nDisplay/Source/DisplayCluster/Private/Render/Device/DisplayClusterDeviceBase.h
@@ -39,6 +39,7 @@ public:
 	virtual bool Initialize() override;
 	virtual void StartScene(UWorld* InWorld) override;
 	virtual void EndScene() override;
+	virtual void PreTick(float DeltaSeconds) override;
 	virtual void SetViewportCamera(const FString& InCameraId = FString(), const FString& InViewportId = FString()) override;
 	virtual void SetStartPostProcessingSettings(const FString& ViewportID, const FPostProcessSettings& StartPostProcessingSettings) override;
 	virtual void SetOverridePostProcessingSettings(const FString& ViewportID, const FPostProcessSettings& OverridePostProcessingSettings, float BlendWeight = 1.0f) override;
@@ -159,6 +160,7 @@ protected:
 	FViewport* MainViewport = nullptr;
 
 	bool bIsSceneOpen = false;
+	bool bIsCustomPresentSet = false;
 
 	// Per-eye regions
 	FIntRect EyeRegions[2];
diff --git a/Engine/Plugins/Runtime/nDisplay/Source/DisplayCluster/Private/Render/DisplayClusterRenderManager.cpp b/Engine/Plugins/Runtime/nDisplay/Source/DisplayCluster/Private/Render/DisplayClusterRenderManager.cpp
--- a/Engine/Plugins/Runtime/nDisplay/Source/DisplayCluster/Private/Render/DisplayClusterRenderManager.cpp
+++ b/Engine/Plugins/Runtime/nDisplay/Source/DisplayCluster/Private/Render/DisplayClusterRenderManager.cpp
@@ -198,6 +198,11 @@ void FDisplayClusterRenderManager::PreTick(float DeltaSeconds)
 			}
 		}
 	}
+
+	if (RenderDevicePtr)
+	{
+		RenderDevicePtr->PreTick(DeltaSeconds);
+	}
 }
 
 
@@ -750,7 +755,7 @@ void FDisplayClusterRenderManager::OnViewportCreatedHandler_CheckViewportClass()
 		UDisplayClusterViewportClient* const GameViewport = Cast<UDisplayClusterViewportClient>(GEngine->GameViewport);
 		if (!GameViewport)
 		{
-			UE_LOG(LogDisplayClusterRender, Warning, TEXT("DisplayClusterViewportClient is not set as default GameViewport class"));
+			UE_LOG(LogDisplayClusterRender, Warning, TEXT("DisplayClusterViewportClient is not set as a default GameViewport class"));
 		}
 	}
 }
diff --git a/Engine/Plugins/Runtime/nDisplay/Source/DisplayCluster/Public/Render/Device/IDisplayClusterRenderDevice.h b/Engine/Plugins/Runtime/nDisplay/Source/DisplayCluster/Public/Render/Device/IDisplayClusterRenderDevice.h
--- a/Engine/Plugins/Runtime/nDisplay/Source/DisplayCluster/Public/Render/Device/IDisplayClusterRenderDevice.h
+++ b/Engine/Plugins/Runtime/nDisplay/Source/DisplayCluster/Public/Render/Device/IDisplayClusterRenderDevice.h
@@ -36,6 +36,12 @@ public:
 	virtual void StartScene(UWorld* World)
 	{ }
 
+	/**
+	* Called before scene Tick
+	*/
+	virtual void PreTick(float DeltaSeconds)
+	{ }
+
 	/**
 	* Called before unload current level
 	*/
