Added missed include and copyright header to the nDisplay shaders

diff --git a/Engine/Plugins/Runtime/nDisplay/Shaders/Private/MPCDIUtils.ush b/Engine/Plugins/Runtime/nDisplay/Shaders/Private/MPCDIUtils.ush
--- a/Engine/Plugins/Runtime/nDisplay/Shaders/Private/MPCDIUtils.ush
+++ b/Engine/Plugins/Runtime/nDisplay/Shaders/Private/MPCDIUtils.ush
@@ -1,3 +1,6 @@
+// Copyright Epic Games, Inc. All Rights Reserved.
+
+#include "/engine/public/platform.ush"
 #include "/Engine/Private/GammaCorrectionCommon.ush"
 
 
diff --git a/Engine/Plugins/Runtime/nDisplay/Shaders/Private/OutputRemapShaders.usf b/Engine/Plugins/Runtime/nDisplay/Shaders/Private/OutputRemapShaders.usf
--- a/Engine/Plugins/Runtime/nDisplay/Shaders/Private/OutputRemapShaders.usf
+++ b/Engine/Plugins/Runtime/nDisplay/Shaders/Private/OutputRemapShaders.usf
@@ -1,5 +1,7 @@
 // Copyright Epic Games, Inc. All Rights Reserved.
 
+#include "/engine/public/platform.ush"
+
 Texture2D PostprocessInput0;
 SamplerState PostprocessInput0Sampler;
 
