nDisplay fixes

diff --git a/Engine/Plugins/Runtime/nDisplay/Source/DisplayCluster/Private/Game/Classes/Basics/DisplayClusterPlayerInput.cpp b/Engine/Plugins/Runtime/nDisplay/Source/DisplayCluster/Private/Game/Classes/Basics/DisplayClusterPlayerInput.cpp
--- a/Engine/Plugins/Runtime/nDisplay/Source/DisplayCluster/Private/Game/Classes/Basics/DisplayClusterPlayerInput.cpp
+++ b/Engine/Plugins/Runtime/nDisplay/Source/DisplayCluster/Private/Game/Classes/Basics/DisplayClusterPlayerInput.cpp
@@ -58,14 +58,14 @@ bool UDisplayClusterPlayerInput::SerializeKeyStateMap(TMap<FString, FString>& Ou
 		check(SendValue > 0.f);
 
 		StrKeyState = FString::Printf(TEXT("%s;%s;%s;%s;%s;%s;%s;%s;"),
-			*FDisplayClusterTypesConverter::template ToString(it->Value.RawValue),
-			*FDisplayClusterTypesConverter::template ToString(it->Value.Value),
-			*FDisplayClusterTypesConverter::template ToString(it->Value.LastUpDownTransitionTime),
+			*FDisplayClusterTypesConverter::template ToHexString(it->Value.RawValue),
+			*FDisplayClusterTypesConverter::template ToHexString(it->Value.Value),
+			*FDisplayClusterTypesConverter::template ToHexString(it->Value.LastUpDownTransitionTime),
 			*DisplayClusterHelpers::str::BoolToStr(it->Value.bDown ? true : false, false),
 			*DisplayClusterHelpers::str::BoolToStr(it->Value.bDownPrevious ? true : false, false),
 			*DisplayClusterHelpers::str::BoolToStr(it->Value.bConsumed ? true : false, false),
 			*FDisplayClusterTypesConverter::template ToString(it->Value.SampleCountAccumulator),
-			*FDisplayClusterTypesConverter::template ToString(it->Value.RawValueAccumulator));
+			*FDisplayClusterTypesConverter::template ToHexString(it->Value.RawValueAccumulator));
 
 		for (int i = 0; i < IE_MAX; ++i)
 		{
@@ -101,14 +101,14 @@ bool UDisplayClusterPlayerInput::DeserializeKeyStateMap(const TMap<FString, FStr
 
 		TArray<FString> Fields = DisplayClusterHelpers::str::StrToArray<FString>(it->Value, FString(";"), false);
 
-		KeyState.RawValue                 = FDisplayClusterTypesConverter::template FromString<FVector>(Fields[0]);
-		KeyState.Value                    = FDisplayClusterTypesConverter::template FromString<FVector>(Fields[1]);
-		KeyState.LastUpDownTransitionTime = FDisplayClusterTypesConverter::template FromString<float>(Fields[2]);
+		KeyState.RawValue                 = FDisplayClusterTypesConverter::template FromHexString<FVector>(Fields[0]);
+		KeyState.Value                    = FDisplayClusterTypesConverter::template FromHexString<FVector>(Fields[1]);
+		KeyState.LastUpDownTransitionTime = FDisplayClusterTypesConverter::template FromHexString<float>(Fields[2]);
 		KeyState.bDown                    = FDisplayClusterTypesConverter::template FromString<bool>(Fields[3]) ? 1 : 0;
 		KeyState.bDownPrevious            = FDisplayClusterTypesConverter::template FromString<bool>(Fields[4]) ? 1 : 0;
 		KeyState.bConsumed                = FDisplayClusterTypesConverter::template FromString<bool>(Fields[5]) ? 1 : 0;
 		KeyState.SampleCountAccumulator   = FDisplayClusterTypesConverter::template FromString<uint8>(Fields[6]);
-		KeyState.RawValueAccumulator      = FDisplayClusterTypesConverter::template FromString<FVector>(Fields[7]);
+		KeyState.RawValueAccumulator      = FDisplayClusterTypesConverter::template FromHexString<FVector>(Fields[7]);
 
 		for (int i = 0; i < IE_MAX; ++i)
 		{
diff --git a/Engine/Plugins/Runtime/nDisplay/Source/DisplayCluster/Private/Game/Classes/Scene/DisplayClusterCameraComponent.cpp b/Engine/Plugins/Runtime/nDisplay/Source/DisplayCluster/Private/Game/Classes/Scene/DisplayClusterCameraComponent.cpp
--- a/Engine/Plugins/Runtime/nDisplay/Source/DisplayCluster/Private/Game/Classes/Scene/DisplayClusterCameraComponent.cpp
+++ b/Engine/Plugins/Runtime/nDisplay/Source/DisplayCluster/Private/Game/Classes/Scene/DisplayClusterCameraComponent.cpp
@@ -14,7 +14,7 @@ UDisplayClusterCameraComponent::UDisplayClusterCameraComponent(const FObjectInit
 	, bEyeSwap(false)
 	, ForceEyeOffset(0)
 {
-	PrimaryComponentTick.bCanEverTick = false;
+	PrimaryComponentTick.bCanEverTick = true;
 }
 
 
diff --git a/Engine/Plugins/Runtime/nDisplay/Source/DisplayCluster/Private/Game/Classes/Scene/DisplayClusterRootActor.cpp b/Engine/Plugins/Runtime/nDisplay/Source/DisplayCluster/Private/Game/Classes/Scene/DisplayClusterRootActor.cpp
--- a/Engine/Plugins/Runtime/nDisplay/Source/DisplayCluster/Private/Game/Classes/Scene/DisplayClusterRootActor.cpp
+++ b/Engine/Plugins/Runtime/nDisplay/Source/DisplayCluster/Private/Game/Classes/Scene/DisplayClusterRootActor.cpp
@@ -34,6 +34,7 @@ ADisplayClusterRootActor::ADisplayClusterRootActor(const FObjectInitializer& Obj
 	DisplayClusterRootComponent->AttachToComponent(RootComponent, FAttachmentTransformRules(EAttachmentRule::KeepRelative, false));
 
 	PrimaryActorTick.bCanEverTick = true;
+	PrimaryActorTick.TickGroup = ETickingGroup::TG_PostUpdateWork;
 	bFindCameraComponentWhenViewTarget = false;
 	bReplicates = false;
 	SpawnCollisionHandlingMethod = ESpawnActorCollisionHandlingMethod::AdjustIfPossibleButAlwaysSpawn;
@@ -100,8 +101,7 @@ void ADisplayClusterRootActor::Tick(float DeltaSeconds)
 				APlayerCameraManager* const CurPlayerCameraManager = CurPlayerController->PlayerCameraManager;
 				if (CurPlayerCameraManager)
 				{
-					SetActorLocation(CurPlayerCameraManager->GetCameraLocation());
-					SetActorRotation(CurPlayerCameraManager->GetCameraRotation());
+					SetActorLocationAndRotation(CurPlayerCameraManager->GetCameraLocation(), CurPlayerCameraManager->GetCameraRotation());
 				}
 
 				if (bExitOnEsc)
diff --git a/Engine/Plugins/Runtime/nDisplay/Source/DisplayCluster/Private/Game/Classes/Scene/DisplayClusterRootComponent.cpp b/Engine/Plugins/Runtime/nDisplay/Source/DisplayCluster/Private/Game/Classes/Scene/DisplayClusterRootComponent.cpp
--- a/Engine/Plugins/Runtime/nDisplay/Source/DisplayCluster/Private/Game/Classes/Scene/DisplayClusterRootComponent.cpp
+++ b/Engine/Plugins/Runtime/nDisplay/Source/DisplayCluster/Private/Game/Classes/Scene/DisplayClusterRootComponent.cpp
@@ -26,7 +26,7 @@ UDisplayClusterRootComponent::UDisplayClusterRootComponent(const FObjectInitiali
 	DISPLAY_CLUSTER_FUNC_TRACE(LogDisplayClusterGame);
 
 	// This component settings
-	PrimaryComponentTick.bCanEverTick = false;
+	PrimaryComponentTick.bCanEverTick = true;
 }
 
 void UDisplayClusterRootComponent::BeginPlay()
diff --git a/Engine/Plugins/Runtime/nDisplay/Source/DisplayCluster/Private/Game/Classes/Scene/DisplayClusterSceneComponentSync.cpp b/Engine/Plugins/Runtime/nDisplay/Source/DisplayCluster/Private/Game/Classes/Scene/DisplayClusterSceneComponentSync.cpp
--- a/Engine/Plugins/Runtime/nDisplay/Source/DisplayCluster/Private/Game/Classes/Scene/DisplayClusterSceneComponentSync.cpp
+++ b/Engine/Plugins/Runtime/nDisplay/Source/DisplayCluster/Private/Game/Classes/Scene/DisplayClusterSceneComponentSync.cpp
@@ -7,6 +7,7 @@
 #include "Cluster/IPDisplayClusterClusterManager.h"
 #include "Game/IPDisplayClusterGameManager.h"
 
+#include "DisplayClusterUtils/DisplayClusterTypesConverter.h"
 #include "DisplayClusterGlobals.h"
 #include "DisplayClusterLog.h"
 
@@ -90,20 +91,14 @@ FString UDisplayClusterSceneComponentSync::GenerateSyncId()
 
 FString UDisplayClusterSceneComponentSync::SerializeToString() const
 {
-	return GetSyncTransform().ToString();
+	return FDisplayClusterTypesConverter::template ToHexString(GetSyncTransform());
 }
 
 bool UDisplayClusterSceneComponentSync::DeserializeFromString(const FString& data)
 {
-	FTransform t;
-	if (!t.InitFromString(data))
-	{
-		UE_LOG(LogDisplayClusterGame, Error, TEXT("Unable to deserialize transform data"));
-		return false;
-	}
-
-	UE_LOG(LogDisplayClusterGame, Verbose, TEXT("%s: applying transform data <%s>"), *SyncId, *t.ToHumanReadableString());
-	SetSyncTransform(t);
+	FTransform NewTransform = FDisplayClusterTypesConverter::template FromHexString<FTransform>(data);
+	UE_LOG(LogDisplayClusterGame, Verbose, TEXT("%s: applying transform data <%s>"), *SyncId, *NewTransform.ToHumanReadableString());
+	SetSyncTransform(NewTransform);
 
 	return true;
 }
diff --git a/Engine/Plugins/Runtime/nDisplay/Source/DisplayCluster/Private/Game/Classes/Scene/DisplayClusterSceneComponentSyncParent.cpp b/Engine/Plugins/Runtime/nDisplay/Source/DisplayCluster/Private/Game/Classes/Scene/DisplayClusterSceneComponentSyncParent.cpp
--- a/Engine/Plugins/Runtime/nDisplay/Source/DisplayCluster/Private/Game/Classes/Scene/DisplayClusterSceneComponentSyncParent.cpp
+++ b/Engine/Plugins/Runtime/nDisplay/Source/DisplayCluster/Private/Game/Classes/Scene/DisplayClusterSceneComponentSyncParent.cpp
@@ -10,7 +10,7 @@
 UDisplayClusterSceneComponentSyncParent::UDisplayClusterSceneComponentSyncParent(const FObjectInitializer& ObjectInitializer) :
 	UDisplayClusterSceneComponentSync(ObjectInitializer)
 {
-	PrimaryComponentTick.bCanEverTick = false;
+	PrimaryComponentTick.bCanEverTick = true;
 }
 
 void UDisplayClusterSceneComponentSyncParent::BeginPlay()
diff --git a/Engine/Plugins/Runtime/nDisplay/Source/DisplayCluster/Private/Game/Classes/Scene/DisplayClusterSceneComponentSyncThis.cpp b/Engine/Plugins/Runtime/nDisplay/Source/DisplayCluster/Private/Game/Classes/Scene/DisplayClusterSceneComponentSyncThis.cpp
--- a/Engine/Plugins/Runtime/nDisplay/Source/DisplayCluster/Private/Game/Classes/Scene/DisplayClusterSceneComponentSyncThis.cpp
+++ b/Engine/Plugins/Runtime/nDisplay/Source/DisplayCluster/Private/Game/Classes/Scene/DisplayClusterSceneComponentSyncThis.cpp
@@ -8,7 +8,7 @@
 UDisplayClusterSceneComponentSyncThis::UDisplayClusterSceneComponentSyncThis(const FObjectInitializer& ObjectInitializer) :
 	UDisplayClusterSceneComponentSync(ObjectInitializer)
 {
-	PrimaryComponentTick.bCanEverTick = false;
+	PrimaryComponentTick.bCanEverTick = true;
 }
 
 void UDisplayClusterSceneComponentSyncThis::BeginPlay()
diff --git a/Engine/Plugins/Runtime/nDisplay/Source/DisplayCluster/Private/Game/Classes/Scene/DisplayClusterScreenComponent.cpp b/Engine/Plugins/Runtime/nDisplay/Source/DisplayCluster/Private/Game/Classes/Scene/DisplayClusterScreenComponent.cpp
--- a/Engine/Plugins/Runtime/nDisplay/Source/DisplayCluster/Private/Game/Classes/Scene/DisplayClusterScreenComponent.cpp
+++ b/Engine/Plugins/Runtime/nDisplay/Source/DisplayCluster/Private/Game/Classes/Scene/DisplayClusterScreenComponent.cpp
@@ -28,6 +28,7 @@ UDisplayClusterScreenComponent::UDisplayClusterScreenComponent(const FObjectInit
 		static ConstructorHelpers::FObjectFinder<UStaticMesh> ScreenMesh(TEXT("StaticMesh'/Engine/BasicShapes/Cube.Cube'"));
 
 		ScreenGeometryComponent->AttachToComponent(this, FAttachmentTransformRules(EAttachmentRule::KeepRelative, false));
+		ScreenGeometryComponent->SetRelativeLocationAndRotation(FVector::ZeroVector, FRotator::ZeroRotator, false);
 		ScreenGeometryComponent->SetStaticMesh(ScreenMesh.Object);
 		ScreenGeometryComponent->SetMobility(EComponentMobility::Movable);
 		ScreenGeometryComponent->SetCollisionEnabled(ECollisionEnabled::NoCollision);
@@ -75,16 +76,11 @@ bool UDisplayClusterScreenComponent::ApplySettings()
 {
 	Super::ApplySettings();
 
-#if WITH_EDITOR
+	SetRelativeScale3D(FVector(0.005f, Size.X, Size.Y));
 	if (ScreenGeometryComponent)
 	{
 		ScreenGeometryComponent->RegisterComponent();
-		ScreenGeometryComponent->AttachToComponent(this, FAttachmentTransformRules(EAttachmentRule::KeepRelative, false));
-		ScreenGeometryComponent->SetRelativeLocationAndRotation(FVector::ZeroVector, FRotator::ZeroRotator, false);
 	}
-#endif
-
-	SetRelativeScale3D(FVector(0.0001f, Size.X, Size.Y));
 
 	return true;
 }
diff --git a/Engine/Plugins/Runtime/nDisplay/Source/DisplayCluster/Private/Input/Devices/VRPN/Analog/DisplayClusterVrpnAnalogInputDataHolder.cpp b/Engine/Plugins/Runtime/nDisplay/Source/DisplayCluster/Private/Input/Devices/VRPN/Analog/DisplayClusterVrpnAnalogInputDataHolder.cpp
--- a/Engine/Plugins/Runtime/nDisplay/Source/DisplayCluster/Private/Input/Devices/VRPN/Analog/DisplayClusterVrpnAnalogInputDataHolder.cpp
+++ b/Engine/Plugins/Runtime/nDisplay/Source/DisplayCluster/Private/Input/Devices/VRPN/Analog/DisplayClusterVrpnAnalogInputDataHolder.cpp
@@ -1,6 +1,7 @@
 // Copyright Epic Games, Inc. All Rights Reserved.
 
 #include "Input/Devices/VRPN/Analog/DisplayClusterVrpnAnalogInputDataHolder.h"
+#include "DisplayClusterUtils/DisplayClusterTypesConverter.h"
 #include "DisplayClusterLog.h"
 
 
@@ -33,7 +34,7 @@ FString FDisplayClusterVrpnAnalogInputDataHolder::SerializeToString() const
 
 	for (auto it = DeviceData.CreateConstIterator(); it; ++it)
 	{
-		result += FString::Printf(TEXT("%d%s%f%s"), it->Key, SerializationDelimiter, it->Value.axisValue, SerializationDelimiter);
+		result += FString::Printf(TEXT("%d%s%f%s"), it->Key, SerializationDelimiter, *FDisplayClusterTypesConverter::template ToHexString(it->Value.axisValue), SerializationDelimiter);
 	}
 
 	return result;
@@ -52,8 +53,8 @@ bool FDisplayClusterVrpnAnalogInputDataHolder::DeserializeFromString(const FStri
 
 	for (int i = 0; i < parsed.Num(); i += SerializationItems)
 	{
-		const int   ch = FCString::Atoi(*parsed[i]);
-		const float val = FCString::Atof(*parsed[i + 1]);
+		const int   ch  = FCString::Atoi(*parsed[i]);
+		const float val = FDisplayClusterTypesConverter::template FromString<float>(*parsed[i + 1]);
 		DeviceData.Add(ch, FDisplayClusterVrpnAnalogChannelData{ val });
 	}
 
diff --git a/Engine/Plugins/Runtime/nDisplay/Source/DisplayCluster/Private/Input/Devices/VRPN/Tracker/DisplayClusterVrpnTrackerInputDataHolder.cpp b/Engine/Plugins/Runtime/nDisplay/Source/DisplayCluster/Private/Input/Devices/VRPN/Tracker/DisplayClusterVrpnTrackerInputDataHolder.cpp
--- a/Engine/Plugins/Runtime/nDisplay/Source/DisplayCluster/Private/Input/Devices/VRPN/Tracker/DisplayClusterVrpnTrackerInputDataHolder.cpp
+++ b/Engine/Plugins/Runtime/nDisplay/Source/DisplayCluster/Private/Input/Devices/VRPN/Tracker/DisplayClusterVrpnTrackerInputDataHolder.cpp
@@ -1,36 +1,10 @@
 // Copyright Epic Games, Inc. All Rights Reserved.
 
 #include "Input/Devices/VRPN/Tracker/DisplayClusterVrpnTrackerInputDataHolder.h"
+#include "DisplayClusterUtils/DisplayClusterTypesConverter.h"
 #include "DisplayClusterLog.h"
 
 
-namespace
-{
-	// Create a FQuat from a string that is in the same format as generated by
-	// FQuat::ToString.
-	// FQuat is missing InitializeFromString member function.
-	FQuat QuatFromString(const FString& InSourceString)
-	{
-		FQuat Result;
-		const bool bSuccessful
-			=  FParse::Value( *InSourceString, TEXT("X="), Result.X )
-			&& FParse::Value( *InSourceString, TEXT("Y="), Result.Y )
-			&& FParse::Value( *InSourceString, TEXT("Z="), Result.Z )
-			&& FParse::Value( *InSourceString, TEXT("W="), Result.W );
-
-		if (!bSuccessful)
-		{
-			UE_LOG(LogDisplayClusterInputVRPN, Error, TEXT("Parsing FQuat from string '%s' failed!"), *InSourceString);
-			Result = FQuat::Identity;
-		}
-
-		return Result;
-	}
-
-} // namespace
-
-
-
 FDisplayClusterVrpnTrackerInputDataHolder::FDisplayClusterVrpnTrackerInputDataHolder(const FDisplayClusterConfigInput& config) :
 	FDisplayClusterInputDeviceBase<EDisplayClusterInputDeviceType::VrpnTracker>(config)
 {
@@ -61,7 +35,12 @@ FString FDisplayClusterVrpnTrackerInputDataHolder::SerializeToString() const
 	for (auto it = DeviceData.CreateConstIterator(); it; ++it)
 	{
 		result += FString::Printf(TEXT("%d%s%s%s%s%s"),
-			it->Key, SerializationDelimiter, *it->Value.trLoc.ToString(), SerializationDelimiter, *it->Value.trQuat.ToString(), SerializationDelimiter);
+			it->Key,
+			SerializationDelimiter,
+			*FDisplayClusterTypesConverter::template ToHexString(it->Value.trLoc),
+			SerializationDelimiter,
+			*FDisplayClusterTypesConverter::template ToHexString(it->Value.trQuat),
+			SerializationDelimiter);
 	}
 
 	return result;
@@ -81,9 +60,8 @@ bool FDisplayClusterVrpnTrackerInputDataHolder::DeserializeFromString(const FStr
 	for (int i = 0; i < parsed.Num(); i += SerializationItems)
 	{
 		const int  ch = FCString::Atoi(*parsed[i]);
-		FVector  loc;
-		FQuat    quat = QuatFromString(parsed[i + 2]);
-		loc.InitFromString(parsed[i + 1]);
+		const FVector  loc  = FDisplayClusterTypesConverter::template FromHexString<FVector>(parsed[i + 1]);
+		const FQuat    quat = FDisplayClusterTypesConverter::template FromHexString<FQuat>(parsed[i + 2]);
 
 		DeviceData.Add(ch, FDisplayClusterVrpnTrackerChannelData{ loc, quat });
 	}
diff --git a/Engine/Plugins/Runtime/nDisplay/Source/DisplayCluster/Public/DisplayClusterUtils/DisplayClusterTypesConverter.h b/Engine/Plugins/Runtime/nDisplay/Source/DisplayCluster/Public/DisplayClusterUtils/DisplayClusterTypesConverter.h
--- a/Engine/Plugins/Runtime/nDisplay/Source/DisplayCluster/Public/DisplayClusterUtils/DisplayClusterTypesConverter.h
+++ b/Engine/Plugins/Runtime/nDisplay/Source/DisplayCluster/Public/DisplayClusterUtils/DisplayClusterTypesConverter.h
@@ -9,6 +9,10 @@
 #include "DisplayClusterEnums.h"
 #include "DisplayClusterUtils/DisplayClusterCommonStrings.h"
 
+#include <iomanip>
+#include <iostream>
+#include <limits>
+#include <sstream>
 
 
 /**
@@ -28,20 +32,37 @@ public:
 	//////////////////////////////////////////////////////////////////////////////////////////////
 	template <typename ConvertTo>
 	static ConvertTo FromString(const FString& from);
+
+	//////////////////////////////////////////////////////////////////////////////////////////////
+	// TYPE --> HEX STRING
+	//////////////////////////////////////////////////////////////////////////////////////////////
+	template <typename ConvertFrom>
+	static FString ToHexString(const ConvertFrom& from);
+
+	//////////////////////////////////////////////////////////////////////////////////////////////
+	// HEX STRING --> TYPE
+	//////////////////////////////////////////////////////////////////////////////////////////////
+	template <typename ConvertTo>
+	static ConvertTo FromHexString(const FString& from);
 };
 
+
+//////////////////////////////////////////////////////////////////////////////////////////////
+// TYPE --> STRING
+//////////////////////////////////////////////////////////////////////////////////////////////
 template <> inline FString FDisplayClusterTypesConverter::ToString<>(const FString& from)    { return from; }
 template <> inline FString FDisplayClusterTypesConverter::ToString<>(const bool& from)       { return (from ? DisplayClusterStrings::cfg::spec::ValTrue : DisplayClusterStrings::cfg::spec::ValFalse); }
-template <> inline FString FDisplayClusterTypesConverter::ToString<>(const int8& from)      { return FString::FromInt(from); }
-template <> inline FString FDisplayClusterTypesConverter::ToString<>(const uint8& from)     { return ToString(static_cast<int8>(from)); }
-template <> inline FString FDisplayClusterTypesConverter::ToString<>(const int32& from)     { return FString::FromInt(from); }
-template <> inline FString FDisplayClusterTypesConverter::ToString<>(const uint32& from)    { return ToString(static_cast<int32>(from)); }
+template <> inline FString FDisplayClusterTypesConverter::ToString<>(const int8& from)       { return FString::FromInt(from); }
+template <> inline FString FDisplayClusterTypesConverter::ToString<>(const uint8& from)      { return ToString(static_cast<int8>(from)); }
+template <> inline FString FDisplayClusterTypesConverter::ToString<>(const int32& from)      { return FString::FromInt(from); }
+template <> inline FString FDisplayClusterTypesConverter::ToString<>(const uint32& from)     { return ToString(static_cast<int32>(from)); }
 template <> inline FString FDisplayClusterTypesConverter::ToString<>(const float& from)      { return FString::SanitizeFloat(from); }
 template <> inline FString FDisplayClusterTypesConverter::ToString<>(const double& from)     { return FString::Printf(TEXT("%lf"), from); }
 template <> inline FString FDisplayClusterTypesConverter::ToString<>(const FVector& from)    { return from.ToString(); }
 template <> inline FString FDisplayClusterTypesConverter::ToString<>(const FVector2D& from)  { return from.ToString(); }
 template <> inline FString FDisplayClusterTypesConverter::ToString<>(const FRotator& from)   { return from.ToString(); }
 template <> inline FString FDisplayClusterTypesConverter::ToString<>(const FMatrix& from)    { return from.ToString(); }
+template <> inline FString FDisplayClusterTypesConverter::ToString<>(const FQuat& from)      { return from.ToString(); }
 
 // We can't just use FTimecode ToString as that loses information.
 template <> inline FString FDisplayClusterTypesConverter::ToString<>(const FTimecode& from)  { return FString::Printf(TEXT("%d;%d;%d;%d;%d"), from.bDropFrameFormat ? 1 : 0, from.Hours, from.Minutes, from.Seconds, from.Frames); }
@@ -81,6 +102,9 @@ template <> inline FString FDisplayClusterTypesConverter::ToString<>(const EDisp
 }
 
 
+//////////////////////////////////////////////////////////////////////////////////////////////
+// STRING --> TYPE
+//////////////////////////////////////////////////////////////////////////////////////////////
 template <> inline FString   FDisplayClusterTypesConverter::FromString<> (const FString& from) { return from; }
 template <> inline bool      FDisplayClusterTypesConverter::FromString<> (const FString& from) { return (from == FString("1") || (from.Compare(DisplayClusterStrings::cfg::spec::ValTrue, ESearchCase::IgnoreCase) == 0)); }
 template <> inline int8      FDisplayClusterTypesConverter::FromString<> (const FString& from) { return FCString::Atoi(*from); }
@@ -141,6 +165,23 @@ template <> inline FMatrix   FDisplayClusterTypesConverter::FromString<>(const F
 	return ResultMatrix;
 }
 
+template <> inline FQuat FDisplayClusterTypesConverter::FromString<>(const FString& from)
+{
+	FQuat Result = FQuat::Identity;
+	const bool bSuccessful
+		=  FParse::Value(*from, TEXT("X="), Result.X)
+		&& FParse::Value(*from, TEXT("Y="), Result.Y)
+		&& FParse::Value(*from, TEXT("Z="), Result.Z)
+		&& FParse::Value(*from, TEXT("W="), Result.W);
+
+	if (!bSuccessful)
+	{
+		return FQuat::Identity;
+	}
+
+	return Result;
+}
+
 template <> inline FTimecode FDisplayClusterTypesConverter::FromString<> (const FString& from)
 {
 	FTimecode timecode;
@@ -216,3 +257,248 @@ template <> inline EDisplayClusterSyncGroup FDisplayClusterTypesConverter::FromS
 
 	return EDisplayClusterSyncGroup::Tick;
 }
+
+
+//////////////////////////////////////////////////////////////////////////////////////////////
+// TYPE --> HEX STRING
+//////////////////////////////////////////////////////////////////////////////////////////////
+template <> inline FString FDisplayClusterTypesConverter::ToHexString<>(const float& from)
+{
+	static_assert(std::numeric_limits<float>::is_iec559, "Native float must be an IEEE float");
+
+	union { float fval; std::uint32_t ival; };
+	fval = from;
+
+	std::ostringstream stm;
+	stm << std::hex << std::nouppercase << ival;
+
+	return FString(FString(TEXT("0x")) + FString(stm.str().c_str()));
+}
+
+template <> inline FString FDisplayClusterTypesConverter::ToHexString<>(const double& from)
+{
+	static_assert(std::numeric_limits<double>::is_iec559, "Native double must be an IEEE double");
+
+	union { double dval; std::uint64_t ival; };
+	dval = from;
+
+	std::ostringstream stm;
+	stm << std::hex << std::nouppercase << ival;
+
+	return FString(FString(TEXT("0x")) + FString(stm.str().c_str()));
+}
+
+template <> inline FString FDisplayClusterTypesConverter::ToHexString<>(const FVector& from)
+{
+	return FString::Printf(TEXT("X=%s Y=%s Z=%s"), 
+		*FDisplayClusterTypesConverter::template ToHexString<>(from.X),
+		*FDisplayClusterTypesConverter::template ToHexString<>(from.Y),
+		*FDisplayClusterTypesConverter::template ToHexString<>(from.Z));
+}
+
+template <> inline FString FDisplayClusterTypesConverter::ToHexString<>(const FVector2D& from)
+{
+	return FString::Printf(TEXT("X=%s Y=%s"),
+		*FDisplayClusterTypesConverter::template ToHexString<>(from.X),
+		*FDisplayClusterTypesConverter::template ToHexString<>(from.Y));
+}
+
+template <> inline FString FDisplayClusterTypesConverter::ToHexString<>(const FRotator& from)
+{
+	return FString::Printf(TEXT("P=%s Y=%s R=%s"),
+		*FDisplayClusterTypesConverter::template ToHexString<>(from.Pitch),
+		*FDisplayClusterTypesConverter::template ToHexString<>(from.Yaw),
+		*FDisplayClusterTypesConverter::template ToHexString<>(from.Roll));
+}
+
+template <> inline FString FDisplayClusterTypesConverter::ToHexString<>(const FMatrix& from)
+{
+	FString Result;
+
+	for (int i = 0; i < 4; ++i)
+	{
+		Result += FString::Printf(TEXT("[%s %s %s %s] "),
+			*FDisplayClusterTypesConverter::template ToHexString<>(from.M[i][0]),
+			*FDisplayClusterTypesConverter::template ToHexString<>(from.M[i][1]),
+			*FDisplayClusterTypesConverter::template ToHexString<>(from.M[i][2]),
+			*FDisplayClusterTypesConverter::template ToHexString<>(from.M[i][3]));
+	}
+
+	return Result;
+}
+
+template <> inline FString FDisplayClusterTypesConverter::ToHexString<>(const FTransform& from)
+{
+	return FString::Printf(TEXT("%s|%s|%s"),
+		*FDisplayClusterTypesConverter::template ToHexString<>(from.GetLocation()),
+		*FDisplayClusterTypesConverter::template ToHexString<>(from.GetRotation().Rotator()),
+		*FDisplayClusterTypesConverter::template ToHexString<>(from.GetScale3D()));
+}
+
+template <> inline FString FDisplayClusterTypesConverter::ToHexString<>(const FQuat& from)
+{
+	return FString::Printf(TEXT("X=%s Y=%s Z=%s W=%s"),
+		*FDisplayClusterTypesConverter::template ToHexString<>(from.X),
+		*FDisplayClusterTypesConverter::template ToHexString<>(from.Y),
+		*FDisplayClusterTypesConverter::template ToHexString<>(from.Z),
+		*FDisplayClusterTypesConverter::template ToHexString<>(from.W));
+}
+
+
+
+//////////////////////////////////////////////////////////////////////////////////////////////
+// HEX STRING --> TYPE
+//////////////////////////////////////////////////////////////////////////////////////////////
+template <> inline float FDisplayClusterTypesConverter::FromHexString<>(const FString& from)
+{
+	static_assert(std::numeric_limits<float>::is_iec559, "Native float must be an IEEE float");
+
+	union { float fval; std::uint32_t ival; };
+
+	std::string str(TCHAR_TO_UTF8(*from));
+	std::stringstream stm(str);
+	stm >> std::hex >> ival;
+	return fval;
+}
+
+template <> inline double FDisplayClusterTypesConverter::FromHexString<>(const FString& from)
+{
+	static_assert(std::numeric_limits<double>::is_iec559, "Native double must be an IEEE double");
+
+	union { double dval; std::uint64_t ival; };
+
+	std::string str(TCHAR_TO_UTF8(*from));
+	std::stringstream stm(str);
+	stm >> std::hex >> ival;
+	return dval;
+}
+
+template <> inline FVector FDisplayClusterTypesConverter::FromHexString<>(const FString& from)
+{
+	FString X, Y, Z;
+	FVector Result;
+
+	const bool bSuccessful = FParse::Value(*from, TEXT("X="), X) && FParse::Value(*from, TEXT("Y="), Y) && FParse::Value(*from, TEXT("Z="), Z);
+	if (bSuccessful)
+	{
+		Result.X = FDisplayClusterTypesConverter::template FromHexString<float>(X);
+		Result.Y = FDisplayClusterTypesConverter::template FromHexString<float>(Y);
+		Result.Z = FDisplayClusterTypesConverter::template FromHexString<float>(Z);
+	}
+
+	return Result;
+}
+
+template <> inline FVector2D FDisplayClusterTypesConverter::FromHexString<>(const FString& from)
+{
+	FString X, Y;
+	FVector2D Result;
+
+	const bool bSuccessful = FParse::Value(*from, TEXT("X="), X) && FParse::Value(*from, TEXT("Y="), Y);
+	if (bSuccessful)
+	{
+		Result.X = FDisplayClusterTypesConverter::template FromHexString<float>(X);
+		Result.Y = FDisplayClusterTypesConverter::template FromHexString<float>(Y);
+	}
+
+	return Result;
+}
+
+template <> inline FRotator FDisplayClusterTypesConverter::FromHexString<>(const FString& from)
+{
+	FString P, Y, R;
+	FRotator Result;
+
+	const bool bSuccessful = FParse::Value(*from, TEXT("P="), P) && FParse::Value(*from, TEXT("Y="), Y) && FParse::Value(*from, TEXT("R="), R);
+	if (bSuccessful)
+	{
+		Result.Pitch = FDisplayClusterTypesConverter::template FromHexString<float>(P);
+		Result.Yaw   = FDisplayClusterTypesConverter::template FromHexString<float>(Y);
+		Result.Roll  = FDisplayClusterTypesConverter::template FromHexString<float>(R);
+	}
+
+	return Result;
+}
+
+template <> inline FMatrix FDisplayClusterTypesConverter::FromHexString<>(const FString& from)
+{
+	FMatrix ResultMatrix = FMatrix::Identity;
+	FPlane  Planes[4];
+
+	int32 IdxStart = 0;
+	int32 IdxEnd = 0;
+	for (int PlaneNum = 0; PlaneNum < 4; ++PlaneNum)
+	{
+		IdxStart = from.Find(FString(TEXT("[")), ESearchCase::IgnoreCase, ESearchDir::FromStart, IdxEnd);
+		IdxEnd = from.Find(FString(TEXT("]")), ESearchCase::IgnoreCase, ESearchDir::FromStart, IdxStart);
+		if (IdxStart == INDEX_NONE || IdxEnd == INDEX_NONE || (IdxEnd - IdxStart) < 8)
+		{
+			return ResultMatrix;
+		}
+
+		FString StrPlane = from.Mid(IdxStart + 1, IdxEnd - IdxStart - 1);
+
+		int32 StrLen = 0;
+		int32 NewStrLen = -1;
+		while (NewStrLen < StrLen)
+		{
+			StrLen = StrPlane.Len();
+			StrPlane.ReplaceInline(TEXT("  "), TEXT(" "));
+			NewStrLen = StrPlane.Len();
+		}
+
+		TArray<FString> StrPlaneValues;
+		StrPlane.ParseIntoArray(StrPlaneValues, TEXT(" "));
+		if (StrPlaneValues.Num() != 4)
+		{
+			return ResultMatrix;
+		}
+
+		float PlaneValues[4] = { 0.f };
+		for (int i = 0; i < 4; ++i)
+		{
+			PlaneValues[i] = FDisplayClusterTypesConverter::template FromHexString<float>(StrPlaneValues[i]);
+		}
+
+		Planes[PlaneNum] = FPlane(PlaneValues[0], PlaneValues[1], PlaneValues[2], PlaneValues[3]);
+	}
+
+	ResultMatrix = FMatrix(Planes[0], Planes[1], Planes[2], Planes[3]);
+
+	return ResultMatrix;
+}
+
+template <> inline FTransform FDisplayClusterTypesConverter::FromHexString<>(const FString& from)
+{
+	TArray<FString> ComponentStrings;
+	from.ParseIntoArray(ComponentStrings, TEXT("|"), true);
+	if (ComponentStrings.Num() != 3)
+	{
+		return FTransform::Identity;
+	}
+
+	const FVector  ParsedTranslation = FDisplayClusterTypesConverter::template FromHexString<FVector>(ComponentStrings[0]);
+	const FRotator ParsedRotation    = FDisplayClusterTypesConverter::template FromHexString<FRotator>(ComponentStrings[1]);
+	const FVector  ParsedScale       = FDisplayClusterTypesConverter::template FromHexString<FVector>(ComponentStrings[2]);
+
+	const FTransform Result(ParsedRotation, ParsedTranslation, ParsedScale);
+
+	return Result;
+}
+
+template <> inline FQuat FDisplayClusterTypesConverter::FromHexString<>(const FString& from)
+{
+	FString X, Y, Z, W;
+	FQuat Result = FQuat::Identity;
+
+	const bool bSuccessful = FParse::Value(*from, TEXT("X="), X) && FParse::Value(*from, TEXT("Y="), Y) && FParse::Value(*from, TEXT("Z="), Z) && FParse::Value(*from, TEXT("W="), W);
+	if (bSuccessful)
+	{
+		Result.X = FDisplayClusterTypesConverter::template FromHexString<float>(X);
+		Result.Y = FDisplayClusterTypesConverter::template FromHexString<float>(Y);
+		Result.Z = FDisplayClusterTypesConverter::template FromHexString<float>(Z);
+		Result.W = FDisplayClusterTypesConverter::template FromHexString<float>(W);
+	}
+
+	return Result;
+}
diff --git a/Engine/Source/Programs/nDisplayLauncher/Cluster/Launcher/Launcher.StartApp.cs b/Engine/Source/Programs/nDisplayLauncher/Cluster/Launcher/Launcher.StartApp.cs
--- a/Engine/Source/Programs/nDisplayLauncher/Cluster/Launcher/Launcher.StartApp.cs
+++ b/Engine/Source/Programs/nDisplayLauncher/Cluster/Launcher/Launcher.StartApp.cs
@@ -77,7 +77,7 @@ namespace nDisplayLauncher.Cluster
 
 			if (!Config.Windows.ContainsKey(Node.Window))
 			{
-				throw new Exception("Node {0} has no windows property specified");
+				throw new Exception(string.Format("Node {0} has no windows property specified", Node.Id));
 			}
 
 			// Get window settings for the node
