Fixed crash when cluster event emits in PIE

diff --git a/Engine/Plugins/Runtime/nDisplay/Source/DisplayCluster/Private/Cluster/DisplayClusterClusterManager.cpp b/Engine/Plugins/Runtime/nDisplay/Source/DisplayCluster/Private/Cluster/DisplayClusterClusterManager.cpp
--- a/Engine/Plugins/Runtime/nDisplay/Source/DisplayCluster/Private/Cluster/DisplayClusterClusterManager.cpp
+++ b/Engine/Plugins/Runtime/nDisplay/Source/DisplayCluster/Private/Cluster/DisplayClusterClusterManager.cpp
@@ -377,7 +377,7 @@ void FDisplayClusterClusterManager::EmitClusterEvent(const FDisplayClusterCluste
 		else
 		{
 			// An event will be emitted from a slave node if it's explicitly specified by MasterOnly=false
-			if (!MasterOnly)
+			if (!MasterOnly && Controller)
 			{
 				Controller->EmitClusterEvent(Event);
 			}
