Renamed private 'Classes' directory to prevent considering it as a part of public API

diff --git a/Engine/Plugins/Runtime/nDisplay/Source/DisplayCluster/Private/Game/Classes/Basics/DisplayClusterGameEngine.cpp b/Engine/Plugins/Runtime/nDisplay/Source/DisplayCluster/Private/Game/Classes/Basics/DisplayClusterGameEngine.cpp
deleted file mode 100644
diff --git a/Engine/Plugins/Runtime/nDisplay/Source/DisplayCluster/Private/Game/Classes/Basics/DisplayClusterPlayerInput.cpp b/Engine/Plugins/Runtime/nDisplay/Source/DisplayCluster/Private/Game/Classes/Basics/DisplayClusterPlayerInput.cpp
deleted file mode 100644
diff --git a/Engine/Plugins/Runtime/nDisplay/Source/DisplayCluster/Private/Game/Classes/Basics/DisplayClusterViewportClient.cpp b/Engine/Plugins/Runtime/nDisplay/Source/DisplayCluster/Private/Game/Classes/Basics/DisplayClusterViewportClient.cpp
deleted file mode 100644
diff --git a/Engine/Plugins/Runtime/nDisplay/Source/DisplayCluster/Private/Game/Classes/Scene/DisplayClusterCameraComponent.cpp b/Engine/Plugins/Runtime/nDisplay/Source/DisplayCluster/Private/Game/Classes/Scene/DisplayClusterCameraComponent.cpp
deleted file mode 100644
diff --git a/Engine/Plugins/Runtime/nDisplay/Source/DisplayCluster/Private/Game/Classes/Scene/DisplayClusterRootActor.cpp b/Engine/Plugins/Runtime/nDisplay/Source/DisplayCluster/Private/Game/Classes/Scene/DisplayClusterRootActor.cpp
deleted file mode 100644
diff --git a/Engine/Plugins/Runtime/nDisplay/Source/DisplayCluster/Private/Game/Classes/Scene/DisplayClusterRootComponent.cpp b/Engine/Plugins/Runtime/nDisplay/Source/DisplayCluster/Private/Game/Classes/Scene/DisplayClusterRootComponent.cpp
deleted file mode 100644
diff --git a/Engine/Plugins/Runtime/nDisplay/Source/DisplayCluster/Private/Game/Classes/Scene/DisplayClusterSceneComponent.cpp b/Engine/Plugins/Runtime/nDisplay/Source/DisplayCluster/Private/Game/Classes/Scene/DisplayClusterSceneComponent.cpp
deleted file mode 100644
diff --git a/Engine/Plugins/Runtime/nDisplay/Source/DisplayCluster/Private/Game/Classes/Scene/DisplayClusterSceneComponentSync.cpp b/Engine/Plugins/Runtime/nDisplay/Source/DisplayCluster/Private/Game/Classes/Scene/DisplayClusterSceneComponentSync.cpp
deleted file mode 100644
diff --git a/Engine/Plugins/Runtime/nDisplay/Source/DisplayCluster/Private/Game/Classes/Scene/DisplayClusterSceneComponentSyncParent.cpp b/Engine/Plugins/Runtime/nDisplay/Source/DisplayCluster/Private/Game/Classes/Scene/DisplayClusterSceneComponentSyncParent.cpp
deleted file mode 100644
diff --git a/Engine/Plugins/Runtime/nDisplay/Source/DisplayCluster/Private/Game/Classes/Scene/DisplayClusterSceneComponentSyncThis.cpp b/Engine/Plugins/Runtime/nDisplay/Source/DisplayCluster/Private/Game/Classes/Scene/DisplayClusterSceneComponentSyncThis.cpp
deleted file mode 100644
diff --git a/Engine/Plugins/Runtime/nDisplay/Source/DisplayCluster/Private/Game/Classes/Scene/DisplayClusterScreenComponent.cpp b/Engine/Plugins/Runtime/nDisplay/Source/DisplayCluster/Private/Game/Classes/Scene/DisplayClusterScreenComponent.cpp
deleted file mode 100644
diff --git a/Engine/Plugins/Runtime/nDisplay/Source/DisplayCluster/Private/Game/EngineClasses/Basics/DisplayClusterGameEngine.cpp b/Engine/Plugins/Runtime/nDisplay/Source/DisplayCluster/Private/Game/EngineClasses/Basics/DisplayClusterGameEngine.cpp
new file mode 100644
--- /dev/null
+++ b/Engine/Plugins/Runtime/nDisplay/Source/DisplayCluster/Private/Game/EngineClasses/Basics/DisplayClusterGameEngine.cpp
@@ -0,0 +1,285 @@
+// Copyright Epic Games, Inc. All Rights Reserved.
+
+#include "DisplayClusterGameEngine.h"
+
+#include "Cluster/IPDisplayClusterClusterManager.h"
+#include "Cluster/Controller/IPDisplayClusterNodeController.h"
+#include "Config/IPDisplayClusterConfigManager.h"
+#include "Input/IPDisplayClusterInputManager.h"
+
+#include "Misc/App.h"
+#include "Misc/CommandLine.h"
+#include "Misc/DisplayClusterAppExit.h"
+#include "Misc/Parse.h"
+#include "Misc/QualifiedFrameTime.h"
+#include "DisplayClusterBuildConfig.h"
+#include "DisplayClusterEnums.h"
+#include "DisplayClusterGlobals.h"
+#include "DisplayClusterHelpers.h"
+#include "DisplayClusterLog.h"
+#include "DisplayClusterStrings.h"
+
+#include "Stats/Stats.h"
+
+#include "Misc/CoreDelegates.h"
+
+
+void UDisplayClusterGameEngine::Init(class IEngineLoop* InEngineLoop)
+{
+	DISPLAY_CLUSTER_FUNC_TRACE(LogDisplayClusterEngine);
+
+	// Detect requested operation mode
+	OperationMode = DetectOperationMode();
+
+	if (OperationMode == EDisplayClusterOperationMode::Cluster ||
+		OperationMode == EDisplayClusterOperationMode::Standalone)
+	{
+		// Instantiate the tickable helper
+		TickableHelper = NewObject<UDisplayClusterGameEngineTickableHelper>();
+	}
+
+	// Initialize Display Cluster
+	if (!GDisplayCluster->Init(OperationMode))
+	{
+		FDisplayClusterAppExit::ExitApplication(FDisplayClusterAppExit::ExitType::KillImmediately, FString("Couldn't initialize DisplayCluster module"));
+	}
+
+	FString cfgPath;
+	FString nodeId;
+
+	if (OperationMode == EDisplayClusterOperationMode::Cluster)
+	{
+		// Extract config path from command line
+		if (!FParse::Value(FCommandLine::Get(), DisplayClusterStrings::args::Config, cfgPath))
+		{
+			UE_LOG(LogDisplayClusterEngine, Error, TEXT("No config file specified"));
+			FDisplayClusterAppExit::ExitApplication(FDisplayClusterAppExit::ExitType::KillImmediately, FString("Cluster mode requires config file"));
+		}
+
+		// Extract node ID from command line
+		if (!FParse::Value(FCommandLine::Get(), DisplayClusterStrings::args::Node, nodeId))
+		{
+#ifdef DISPLAY_CLUSTER_USE_AUTOMATIC_NODE_ID_RESOLVE
+			UE_LOG(LogDisplayClusterEngine, Log, TEXT("Node ID is not specified"));
+#else
+			UE_LOG(LogDisplayClusterEngine, Warning, TEXT("Node ID is not specified"));
+			FDisplayClusterAppExit::ExitApplication(FDisplayClusterAppExit::ExitType::KillImmediately, FString("Cluster mode requires node ID"));
+#endif
+		}
+	}
+	else if (OperationMode == EDisplayClusterOperationMode::Standalone)
+	{
+#ifdef DISPLAY_CLUSTER_USE_DEBUG_STANDALONE_CONFIG
+		// Save config path from command line
+		cfgPath = DisplayClusterStrings::misc::DbgStubConfig;
+		nodeId  = DisplayClusterStrings::misc::DbgStubNodeId;
+#endif
+	}
+
+	if (OperationMode == EDisplayClusterOperationMode::Cluster ||
+		OperationMode == EDisplayClusterOperationMode::Standalone)
+	{
+		DisplayClusterHelpers::str::TrimStringValue(cfgPath);
+		DisplayClusterHelpers::str::TrimStringValue(nodeId);
+
+		// Start game session
+		if (!GDisplayCluster->StartSession(cfgPath, nodeId))
+		{
+			FDisplayClusterAppExit::ExitApplication(FDisplayClusterAppExit::ExitType::KillImmediately, FString("Couldn't start DisplayCluster session"));
+		}
+
+		// Initialize internals
+		InitializeInternals();
+	}
+
+	// Initialize base stuff.
+	UGameEngine::Init(InEngineLoop);
+}
+
+EDisplayClusterOperationMode UDisplayClusterGameEngine::DetectOperationMode()
+{
+	DISPLAY_CLUSTER_FUNC_TRACE(LogDisplayClusterEngine);
+
+	EDisplayClusterOperationMode OpMode = EDisplayClusterOperationMode::Disabled;
+	if (FParse::Param(FCommandLine::Get(), DisplayClusterStrings::args::Cluster))
+	{
+		OpMode = EDisplayClusterOperationMode::Cluster;
+	}
+	else if (FParse::Param(FCommandLine::Get(), DisplayClusterStrings::args::Standalone))
+	{
+		OpMode = EDisplayClusterOperationMode::Standalone;
+	}
+
+	UE_LOG(LogDisplayClusterEngine, Log, TEXT("Detected operation mode: %s"), *FDisplayClusterTypesConverter::template ToString(OpMode));
+
+	return OpMode;
+}
+
+bool UDisplayClusterGameEngine::InitializeInternals()
+{
+	DISPLAY_CLUSTER_FUNC_TRACE(LogDisplayClusterEngine);
+
+	// Store debug settings locally
+	CfgDebug = GDisplayCluster->GetPrivateConfigMgr()->GetConfigDebug();
+	
+	InputMgr       = GDisplayCluster->GetPrivateInputMgr();
+	ClusterMgr     = GDisplayCluster->GetPrivateClusterMgr();
+	NodeController = ClusterMgr->GetController();
+
+	FDisplayClusterConfigClusterNode LocalClusterNode;
+	if (DisplayClusterHelpers::config::GetLocalClusterNode(LocalClusterNode))
+	{
+		UE_LOG(LogDisplayClusterEngine, Log, TEXT("Configuring sound enabled: %s"), *FDisplayClusterTypesConverter::template ToString(LocalClusterNode.SoundEnabled));
+		bUseSound = LocalClusterNode.SoundEnabled;
+	}
+
+	check(ClusterMgr);
+	check(InputMgr);
+		
+	return true;
+}
+
+void UDisplayClusterGameEngine::PreExit()
+{
+	DISPLAY_CLUSTER_FUNC_TRACE(LogDisplayClusterEngine);
+
+	if (OperationMode == EDisplayClusterOperationMode::Cluster ||
+		OperationMode == EDisplayClusterOperationMode::Standalone)
+	{
+		// Close current DisplayCluster session
+		GDisplayCluster->EndSession();
+	}
+
+	// Release the engine
+	UGameEngine::PreExit();
+}
+
+bool UDisplayClusterGameEngine::LoadMap(FWorldContext& WorldContext, FURL URL, class UPendingNetGame* Pending, FString& Error)
+{
+	DISPLAY_CLUSTER_FUNC_TRACE(LogDisplayClusterEngine);
+
+	if (OperationMode == EDisplayClusterOperationMode::Cluster ||
+		OperationMode == EDisplayClusterOperationMode::Standalone)
+	{
+		// Finish previous scene
+		GDisplayCluster->EndScene();
+
+		// Perform map loading
+		if (!Super::LoadMap(WorldContext, URL, Pending, Error))
+		{
+			return false;
+		}
+
+		// Start new scene
+		GDisplayCluster->StartScene(WorldContext.World());
+
+		// Game start barrier
+		if (NodeController)
+		{
+			NodeController->WaitForGameStart();
+		}
+	}
+	else
+	{
+		return Super::LoadMap(WorldContext, URL, Pending, Error);
+	}
+
+	return true;
+}
+
+void UDisplayClusterGameEngine::Tick(float DeltaSeconds, bool bIdleMode)
+{
+	DISPLAY_CLUSTER_FUNC_TRACE(LogDisplayClusterEngine);
+
+	if (OperationMode == EDisplayClusterOperationMode::Cluster ||
+		OperationMode == EDisplayClusterOperationMode::Standalone)
+	{
+		TOptional<FQualifiedFrameTime> FrameTime;
+
+		//////////////////////////////////////////////////////////////////////////////////////////////
+		// Frame start barrier
+		NodeController->WaitForFrameStart();
+		UE_LOG(LogDisplayClusterEngine, Verbose, TEXT("Sync frame start"));
+
+		// Perform StartFrame notification
+		GDisplayCluster->StartFrame(GFrameCounter);
+
+		// Sync DeltaSeconds
+		NodeController->GetDeltaTime(DeltaSeconds);
+		FApp::SetDeltaTime(DeltaSeconds);
+		UE_LOG(LogDisplayClusterEngine, Verbose, TEXT("DisplayCluster delta seconds: %f"), DeltaSeconds);
+
+		// Sync timecode and framerate
+		NodeController->GetFrameTime(FrameTime);
+
+		if (FrameTime.IsSet())
+		{
+			FApp::SetCurrentFrameTime(FrameTime.GetValue());
+			UE_LOG(LogDisplayClusterEngine, Verbose, TEXT("DisplayCluster timecode: %s | %s"), *FTimecode::FromFrameNumber(FrameTime->Time.GetFrame(), FrameTime->Rate).ToString(), *FrameTime->Rate.ToPrettyText().ToString());
+		}
+		else
+		{
+			FApp::InvalidateCurrentFrameTime();
+			UE_LOG(LogDisplayClusterEngine, Verbose, TEXT("DisplayCluster timecode: [Invalid]"));
+		}
+
+		// Perform PreTick for DisplayCluster module
+		UE_LOG(LogDisplayClusterEngine, Verbose, TEXT("Perform PreTick()"));
+		GDisplayCluster->PreTick(DeltaSeconds);
+
+		// Perform UGameEngine::Tick() calls for scene actors
+		UE_LOG(LogDisplayClusterEngine, Verbose, TEXT("Perform UGameEngine::Tick()"));
+		Super::Tick(DeltaSeconds, bIdleMode);
+
+		// Perform PostTick for DisplayCluster module
+		UE_LOG(LogDisplayClusterEngine, Verbose, TEXT("Perform PostTick()"));
+		GDisplayCluster->PostTick(DeltaSeconds);
+
+		if (CfgDebug.LagSimulateEnabled)
+		{
+			const float lag = CfgDebug.LagMaxTime;
+			UE_LOG(LogDisplayClusterEngine, Log, TEXT("Simulating lag: %f seconds"), lag);
+#if 1
+			FPlatformProcess::Sleep(FMath::RandRange(0.f, lag));
+#else
+			FPlatformProcess::Sleep(lag);
+#endif
+		}
+
+#if 0
+		//////////////////////////////////////////////////////////////////////////////////////////////
+		// Tick end barrier
+		NodeController->WaitForTickEnd();
+#endif
+
+		//////////////////////////////////////////////////////////////////////////////////////////////
+		// Frame end barrier
+		NodeController->WaitForFrameEnd();
+
+		// Perform EndFrame notification
+		GDisplayCluster->EndFrame(GFrameCounter);
+
+		UE_LOG(LogDisplayClusterEngine, Verbose, TEXT("Sync frame end"));
+	}
+	else
+	{
+		Super::Tick(DeltaSeconds, bIdleMode);
+	}
+}
+
+TStatId UDisplayClusterGameEngineTickableHelper::GetStatId() const
+{
+	RETURN_QUICK_DECLARE_CYCLE_STAT(UDisplayClusterGameEngineTickableHelper, STATGROUP_Tickables);
+}
+
+void UDisplayClusterGameEngineTickableHelper::Tick(float DeltaSeconds)
+{
+	static const EDisplayClusterOperationMode OperationMode = GDisplayCluster->GetOperationMode();
+
+	if (OperationMode == EDisplayClusterOperationMode::Cluster ||
+		OperationMode == EDisplayClusterOperationMode::Standalone)
+	{
+		UE_LOG(LogDisplayClusterEngine, Verbose, TEXT("Perform Tick()"));
+		GDisplayCluster->Tick(DeltaSeconds);
+	}
+}
diff --git a/Engine/Plugins/Runtime/nDisplay/Source/DisplayCluster/Private/Game/EngineClasses/Basics/DisplayClusterPlayerInput.cpp b/Engine/Plugins/Runtime/nDisplay/Source/DisplayCluster/Private/Game/EngineClasses/Basics/DisplayClusterPlayerInput.cpp
new file mode 100644
--- /dev/null
+++ b/Engine/Plugins/Runtime/nDisplay/Source/DisplayCluster/Private/Game/EngineClasses/Basics/DisplayClusterPlayerInput.cpp
@@ -0,0 +1,133 @@
+// Copyright Epic Games, Inc. All Rights Reserved.
+
+#include "DisplayClusterPlayerInput.h"
+#include "Components/InputComponent.h"
+
+#include "Cluster/IPDisplayClusterClusterManager.h"
+#include "DisplayClusterGlobals.h"
+#include "DisplayClusterLog.h"
+#include "DisplayClusterUtils/DisplayClusterTypesConverter.h"
+#include "DisplayClusterHelpers.h"
+
+#include "Engine/World.h"
+
+
+UDisplayClusterPlayerInput::UDisplayClusterPlayerInput()
+	: Super()
+{
+}
+
+void UDisplayClusterPlayerInput::ProcessInputStack(const TArray<UInputComponent*>& InputComponentStack, const float DeltaTime, const bool bGamePaused)
+{
+	UE_LOG(LogDisplayClusterGame, Verbose, TEXT("Processing input stack..."));
+
+	static IPDisplayClusterClusterManager* const ClusterMgr = GDisplayCluster->GetPrivateClusterMgr();
+
+	if (ClusterMgr && ClusterMgr->IsCluster())
+	{
+		TMap<FString, FString> KeyStates;
+
+		if (ClusterMgr->IsMaster())
+		{
+			// Export key states data
+			SerializeKeyStateMap(KeyStates);
+			ClusterMgr->ProvideNativeInputData(KeyStates);
+		}
+		else
+		{
+			// Import key states data
+			ClusterMgr->SyncNativeInput(KeyStates);
+			DeserializeKeyStateMap(KeyStates);
+		}
+	}
+
+	Super::ProcessInputStack(InputComponentStack, DeltaTime, bGamePaused);
+}
+
+bool UDisplayClusterPlayerInput::SerializeKeyStateMap(TMap<FString, FString>& OutKeyStateMap)
+{
+	TMap<FKey, FKeyState>& StateMap = GetKeyStateMap();
+	for (auto it = StateMap.CreateConstIterator(); it; ++it)
+	{
+		const FString KeyName = it->Key.ToString();
+
+		FString StrKeyState;
+		StrKeyState.Reserve(2048);
+
+		float SendValue = GetWorld()->GetRealTimeSeconds() - it->Value.LastUpDownTransitionTime;
+		check(SendValue > 0.f);
+
+		StrKeyState = FString::Printf(TEXT("%s;%s;%s;%s;%s;%s;%s;%s;"),
+			*FDisplayClusterTypesConverter::template ToHexString(it->Value.RawValue),
+			*FDisplayClusterTypesConverter::template ToHexString(it->Value.Value),
+			*FDisplayClusterTypesConverter::template ToHexString(it->Value.LastUpDownTransitionTime),
+			*DisplayClusterHelpers::str::BoolToStr(it->Value.bDown ? true : false, false),
+			*DisplayClusterHelpers::str::BoolToStr(it->Value.bDownPrevious ? true : false, false),
+			*DisplayClusterHelpers::str::BoolToStr(it->Value.bConsumed ? true : false, false),
+			*FDisplayClusterTypesConverter::template ToString(it->Value.SampleCountAccumulator),
+			*FDisplayClusterTypesConverter::template ToHexString(it->Value.RawValueAccumulator));
+
+		for (int i = 0; i < IE_MAX; ++i)
+		{
+			StrKeyState += FString::Printf(TEXT("%s;"), *DisplayClusterHelpers::str::ArrayToStr(it->Value.EventCounts[i], FString(","), false));
+		}
+
+		for (int i = 0; i < IE_MAX; ++i)
+		{
+			StrKeyState += FString::Printf(TEXT("%s;"), *DisplayClusterHelpers::str::ArrayToStr(it->Value.EventAccumulator[i], FString(","), false));
+		}
+
+		OutKeyStateMap.Emplace(KeyName, StrKeyState);
+	}
+
+	return true;
+}
+
+bool UDisplayClusterPlayerInput::DeserializeKeyStateMap(const TMap<FString, FString>& InKeyStateMap)
+{
+	TMap<FKey, FKeyState>& StateMap = GetKeyStateMap();
+	
+	// Reset local key state map
+	StateMap.Reset();
+
+	int idx = 0;
+	for (auto it = InKeyStateMap.CreateConstIterator(); it; ++it)
+	{
+		UE_LOG(LogDisplayClusterGame, VeryVerbose, TEXT("Input data [%d]: %s = %s"), idx, *it->Key, *it->Value);
+		++idx;
+
+		FKey Key(*it->Key);
+		FKeyState KeyState;
+
+		TArray<FString> Fields = DisplayClusterHelpers::str::StrToArray<FString>(it->Value, FString(";"), false);
+
+		KeyState.RawValue                 = FDisplayClusterTypesConverter::template FromHexString<FVector>(Fields[0]);
+		KeyState.Value                    = FDisplayClusterTypesConverter::template FromHexString<FVector>(Fields[1]);
+		KeyState.LastUpDownTransitionTime = FDisplayClusterTypesConverter::template FromHexString<float>(Fields[2]);
+		KeyState.bDown                    = FDisplayClusterTypesConverter::template FromString<bool>(Fields[3]) ? 1 : 0;
+		KeyState.bDownPrevious            = FDisplayClusterTypesConverter::template FromString<bool>(Fields[4]) ? 1 : 0;
+		KeyState.bConsumed                = FDisplayClusterTypesConverter::template FromString<bool>(Fields[5]) ? 1 : 0;
+		KeyState.SampleCountAccumulator   = FDisplayClusterTypesConverter::template FromString<uint8>(Fields[6]);
+		KeyState.RawValueAccumulator      = FDisplayClusterTypesConverter::template FromHexString<FVector>(Fields[7]);
+
+		for (int i = 0; i < IE_MAX; ++i)
+		{
+			const FString& EventCountsStr = Fields[8 + i];
+			if (!EventCountsStr.IsEmpty())
+			{
+				KeyState.EventCounts[i] = DisplayClusterHelpers::str::StrToArray<uint32>(EventCountsStr, FString(","));
+			}
+
+			const FString& EventAccumulatorStr = Fields[8 + IE_MAX + i];
+			if (!EventAccumulatorStr.IsEmpty())
+			{
+				KeyState.EventAccumulator[i] = DisplayClusterHelpers::str::StrToArray<uint32>(Fields[8 + IE_MAX + i], FString(","));
+			}
+		}
+
+		// Add incoming data to the local map
+		StateMap.Add(Key, KeyState);
+	}
+
+	return true;
+}
diff --git a/Engine/Plugins/Runtime/nDisplay/Source/DisplayCluster/Private/Game/EngineClasses/Basics/DisplayClusterViewportClient.cpp b/Engine/Plugins/Runtime/nDisplay/Source/DisplayCluster/Private/Game/EngineClasses/Basics/DisplayClusterViewportClient.cpp
new file mode 100644
--- /dev/null
+++ b/Engine/Plugins/Runtime/nDisplay/Source/DisplayCluster/Private/Game/EngineClasses/Basics/DisplayClusterViewportClient.cpp
@@ -0,0 +1,696 @@
+// Copyright Epic Games, Inc. All Rights Reserved.
+
+#include "DisplayClusterViewportClient.h"
+
+#include "SceneView.h"
+#include "Engine/Canvas.h"
+#include "SceneViewExtension.h"
+#include "IHeadMountedDisplay.h"
+#include "IXRTrackingSystem.h"
+#include "BufferVisualizationData.h"
+#include "Engine/Engine.h"
+#include "Engine/Console.h"
+#include "Engine/LocalPlayer.h"
+#include "UnrealEngine.h"
+#include "EngineUtils.h"
+
+#include "AudioDeviceManager.h"
+#include "AudioDevice.h"
+#include "Audio/AudioDebug.h"
+
+#include "GameFramework/PlayerController.h"
+#include "Debug/DebugDrawService.h"
+#include "ContentStreaming.h"
+#include "EngineModule.h"
+#include "FXSystem.h"
+#include "GameFramework/HUD.h"
+#include "SubtitleManager.h"
+#include "Components/LineBatchComponent.h"
+#include "GameFramework/GameUserSettings.h"
+#include "Framework/Application/SlateApplication.h"
+
+#include "LegacyScreenPercentageDriver.h"
+#include "DynamicResolutionState.h"
+#include "EngineStats.h"
+
+#include "Render\Device\IDisplayClusterRenderDevice.h"
+
+#include "DisplayClusterEnums.h"
+#include "DisplayClusterGlobals.h"
+
+
+//DECLARE_CYCLE_STAT(TEXT("UI Drawing Time"), STAT_UIDrawingTime, STATGROUP_UI);
+
+UDisplayClusterViewportClient::UDisplayClusterViewportClient(FVTableHelper& Helper) : Super(Helper)
+{
+	
+}
+
+UDisplayClusterViewportClient::~UDisplayClusterViewportClient()
+{
+
+}
+
+/** Util to find named canvas in transient package, and create if not found */
+static UCanvas* GetCanvasByName(FName CanvasName)
+{
+	// Cache to avoid FString/FName conversions/compares
+	static TMap<FName, UCanvas*> CanvasMap;
+	UCanvas** FoundCanvas = CanvasMap.Find(CanvasName);
+	if (!FoundCanvas)
+	{
+		UCanvas* CanvasObject = FindObject<UCanvas>(GetTransientPackage(), *CanvasName.ToString());
+		if (!CanvasObject)
+		{
+			CanvasObject = NewObject<UCanvas>(GetTransientPackage(), CanvasName);
+			CanvasObject->AddToRoot();
+		}
+
+		CanvasMap.Add(CanvasName, CanvasObject);
+		return CanvasObject;
+	}
+
+	return *FoundCanvas;
+}
+
+void UDisplayClusterViewportClient::Init(struct FWorldContext& WorldContext, UGameInstance* OwningGameInstance, bool bCreateNewAudioDevice)
+{
+	IConsoleVariable* const ForceLoadCVar = IConsoleManager::Get().FindConsoleVariable(TEXT("r.CompositionForceRenderTargetLoad"));
+	if (ForceLoadCVar)
+	{
+		ForceLoadCVar->Set(int32(1));
+	}
+
+	IConsoleVariable* const RTResizeCVar = IConsoleManager::Get().FindConsoleVariable(TEXT("r.SceneRenderTargetResizeMethod"));
+	if (RTResizeCVar)
+	{
+		RTResizeCVar->Set(int32(2));
+	}
+
+	Super::Init(WorldContext, OwningGameInstance, bCreateNewAudioDevice);
+}
+
+void UDisplayClusterViewportClient::Draw(FViewport* InViewport, FCanvas* SceneCanvas)
+{
+	//Valid SceneCanvas is required.  Make this explicit.
+	check(SceneCanvas);
+
+	OnBeginDraw().Broadcast();
+	
+	const bool bStereoRendering = GEngine->IsStereoscopic3D(InViewport);
+	FCanvas* DebugCanvas = InViewport->GetDebugCanvas();
+
+	// Create a temporary canvas if there isn't already one.
+	static FName CanvasObjectName(TEXT("CanvasObject"));
+	UCanvas* CanvasObject = GetCanvasByName(CanvasObjectName);
+	CanvasObject->Canvas = SceneCanvas;
+
+	// Create temp debug canvas object
+	FIntPoint DebugCanvasSize = InViewport->GetSizeXY();
+	if (bStereoRendering && GEngine->XRSystem.IsValid() && GEngine->XRSystem->GetHMDDevice())
+	{
+		DebugCanvasSize = GEngine->XRSystem->GetHMDDevice()->GetIdealDebugCanvasRenderTargetSize();
+	}
+
+	static FName DebugCanvasObjectName(TEXT("DebugCanvasObject"));
+	UCanvas* DebugCanvasObject = GetCanvasByName(DebugCanvasObjectName);
+	DebugCanvasObject->Init(DebugCanvasSize.X, DebugCanvasSize.Y, NULL, DebugCanvas);
+
+	if (DebugCanvas)
+	{
+		DebugCanvas->SetScaledToRenderTarget(bStereoRendering);
+		DebugCanvas->SetStereoRendering(bStereoRendering);
+	}
+	if (SceneCanvas)
+	{
+		SceneCanvas->SetScaledToRenderTarget(bStereoRendering);
+		SceneCanvas->SetStereoRendering(bStereoRendering);
+	}
+
+	UWorld* MyWorld = GetWorld();
+
+	// Force path tracing view mode, and extern code set path tracer show flags
+	const bool bForcePathTracing = InViewport->GetClient()->GetEngineShowFlags()->PathTracing;
+	if (bForcePathTracing)
+	{
+		EngineShowFlags.SetPathTracing(true);
+		ViewModeIndex = VMI_PathTracing;
+	}
+
+	const bool bEnableStereo = GEngine->IsStereoscopic3D(InViewport);
+	int32 NumFamilies = 1;
+	int32 NumViews = bStereoRendering ? GEngine->StereoRenderingDevice->GetDesiredNumberOfViews(bStereoRendering) : 1;
+	int32 StartViewIndex = 0;
+	bool nDisplay = false;
+	// check if HMD or nDisplay connected 
+	if (!GEngine->XRSystem.IsValid())
+	{
+		if (GEngine->StereoRenderingDevice.IsValid())
+		{
+			NumFamilies = GEngine->StereoRenderingDevice->GetDesiredNumberOfViews(bStereoRendering);
+			nDisplay = (GDisplayCluster->GetOperationMode() == EDisplayClusterOperationMode::Cluster);
+		}
+	}
+
+	for (int32 viewFamily = 0; viewFamily < NumFamilies; ++viewFamily)
+	{
+		float CustomBufferRatio = 1.f;
+
+		if (nDisplay)
+		{
+			StartViewIndex = viewFamily;
+			NumViews = StartViewIndex + 1;
+
+			IDisplayClusterRenderDevice* DisplayClusterRenderDevice = static_cast<IDisplayClusterRenderDevice*>(GEngine->StereoRenderingDevice.Get());
+			if (DisplayClusterRenderDevice)
+			{
+				DisplayClusterRenderDevice->GetBufferRatio(viewFamily, CustomBufferRatio);
+			}
+		}
+
+		// create the view family for rendering the world scene to the viewport's render target
+		FSceneViewFamilyContext ViewFamily(FSceneViewFamily::ConstructionValues(
+			InViewport,
+			MyWorld->Scene,
+			EngineShowFlags)
+			.SetRealtimeUpdate(true)
+			.SetAdditionalViewFamily(viewFamily > 0));
+
+#if WITH_EDITOR
+		if (GIsEditor)
+		{
+			// Force enable view family show flag for HighDPI derived's screen percentage.
+			ViewFamily.EngineShowFlags.ScreenPercentage = true;
+		}
+
+		UpdateDebugViewModeShaders();
+#endif
+
+		ViewFamily.ViewExtensions = GEngine->ViewExtensions->GatherActiveExtensions(InViewport);
+
+		for (auto ViewExt : ViewFamily.ViewExtensions)
+		{
+			ViewExt->SetupViewFamily(ViewFamily);
+		}
+
+		if (bStereoRendering && GEngine->XRSystem.IsValid() && GEngine->XRSystem->GetHMDDevice())
+		{
+			// Allow HMD to modify screen settings
+			GEngine->XRSystem->GetHMDDevice()->UpdateScreenSettings(Viewport);
+		}
+
+		ESplitScreenType::Type SplitScreenConfig = GetCurrentSplitscreenConfiguration();
+		ViewFamily.ViewMode = EViewModeIndex(ViewModeIndex);
+		EngineShowFlagOverride(ESFIM_Game, ViewFamily.ViewMode, ViewFamily.EngineShowFlags, false);
+
+		if (ViewFamily.EngineShowFlags.VisualizeBuffer && AllowDebugViewmodes())
+		{
+			// Process the buffer visualization console command
+			FName NewBufferVisualizationMode = NAME_None;
+			static IConsoleVariable* ICVar = IConsoleManager::Get().FindConsoleVariable(FBufferVisualizationData::GetVisualizationTargetConsoleCommandName());
+			if (ICVar)
+			{
+				static const FName OverviewName = TEXT("Overview");
+				FString ModeNameString = ICVar->GetString();
+				FName ModeName = *ModeNameString;
+				if (ModeNameString.IsEmpty() || ModeName == OverviewName || ModeName == NAME_None)
+				{
+					NewBufferVisualizationMode = NAME_None;
+				}
+				else
+				{
+					if (GetBufferVisualizationData().GetMaterial(ModeName) == NULL)
+					{
+						// Mode is out of range, so display a message to the user, and reset the mode back to the previous valid one
+						UE_LOG(LogConsoleResponse, Warning, TEXT("Buffer visualization mode '%s' does not exist"), *ModeNameString);
+						NewBufferVisualizationMode = GetCurrentBufferVisualizationMode();
+						// todo: cvars are user settings, here the cvar state is used to avoid log spam and to auto correct for the user (likely not what the user wants)
+						ICVar->Set(*NewBufferVisualizationMode.GetPlainNameString(), ECVF_SetByCode);
+					}
+					else
+					{
+						NewBufferVisualizationMode = ModeName;
+					}
+				}
+			}
+
+			if (NewBufferVisualizationMode != GetCurrentBufferVisualizationMode())
+			{
+				SetCurrentBufferVisualizationMode(NewBufferVisualizationMode);
+			}
+		}
+
+		TMap<ULocalPlayer*, FSceneView*> PlayerViewMap;
+
+		FAudioDeviceHandle RetrievedAudioDevice = MyWorld->GetAudioDevice();
+		TArray<FSceneView*> Views;
+
+		for (FLocalPlayerIterator Iterator(GEngine, MyWorld); Iterator; ++Iterator)
+		{
+			ULocalPlayer* LocalPlayer = *Iterator;
+			if (LocalPlayer)
+			{
+				APlayerController* PlayerController = LocalPlayer->PlayerController;
+
+				for (int32 viewIndex = StartViewIndex; viewIndex < NumViews; ++viewIndex)
+				{
+					// Calculate the player's view information.
+					FVector		ViewLocation;
+					FRotator	ViewRotation;
+
+					EStereoscopicPass PassType = bStereoRendering ? GEngine->StereoRenderingDevice->GetViewPassForIndex(bStereoRendering, viewIndex) : eSSP_FULL;
+
+					FSceneView* View = LocalPlayer->CalcSceneView(&ViewFamily, ViewLocation, ViewRotation, InViewport, nullptr, PassType);
+
+					if (View)
+					{
+						Views.Add(View);
+
+						if (View->Family->EngineShowFlags.Wireframe)
+						{
+							// Wireframe color is emissive-only, and mesh-modifying materials do not use material substitution, hence...
+							View->DiffuseOverrideParameter = FVector4(0.f, 0.f, 0.f, 0.f);
+							View->SpecularOverrideParameter = FVector4(0.f, 0.f, 0.f, 0.f);
+						}
+						else if (View->Family->EngineShowFlags.OverrideDiffuseAndSpecular)
+						{
+							View->DiffuseOverrideParameter = FVector4(GEngine->LightingOnlyBrightness.R, GEngine->LightingOnlyBrightness.G, GEngine->LightingOnlyBrightness.B, 0.0f);
+							View->SpecularOverrideParameter = FVector4(.1f, .1f, .1f, 0.0f);
+						}
+						else if (View->Family->EngineShowFlags.LightingOnlyOverride)
+						{
+							View->DiffuseOverrideParameter = FVector4(GEngine->LightingOnlyBrightness.R, GEngine->LightingOnlyBrightness.G, GEngine->LightingOnlyBrightness.B, 0.0f);
+							View->SpecularOverrideParameter = FVector4(0.f, 0.f, 0.f, 0.f);
+						}
+						else if (View->Family->EngineShowFlags.ReflectionOverride)
+						{
+							View->DiffuseOverrideParameter = FVector4(0.f, 0.f, 0.f, 0.f);
+							View->SpecularOverrideParameter = FVector4(1, 1, 1, 0.0f);
+							View->NormalOverrideParameter = FVector4(0, 0, 1, 0.0f);
+							View->RoughnessOverrideParameter = FVector2D(0.0f, 0.0f);
+						}
+
+						if (!View->Family->EngineShowFlags.Diffuse)
+						{
+							View->DiffuseOverrideParameter = FVector4(0.f, 0.f, 0.f, 0.f);
+						}
+
+						if (!View->Family->EngineShowFlags.Specular)
+						{
+							View->SpecularOverrideParameter = FVector4(0.f, 0.f, 0.f, 0.f);
+						}
+
+						View->CurrentBufferVisualizationMode = GetCurrentBufferVisualizationMode();
+
+						View->CameraConstrainedViewRect = View->UnscaledViewRect;
+
+
+
+						// If this is the primary drawing pass, update things that depend on the view location
+						//if (viewIndex == 0)
+						{
+							// Save the location of the view.
+							LocalPlayer->LastViewLocation = ViewLocation;
+
+							PlayerViewMap.Add(LocalPlayer, View);
+
+							// Update the listener.
+							if (RetrievedAudioDevice && PlayerController != NULL)
+							{
+								bool bUpdateListenerPosition = true;
+
+								// If the main audio device is used for multiple PIE viewport clients, we only
+								// want to update the main audio device listener position if it is in focus
+								if (GEngine)
+								{
+									FAudioDeviceManager* AudioDeviceManager = GEngine->GetAudioDeviceManager();
+
+									// If there is more than one world referencing the main audio device
+									if (AudioDeviceManager->GetNumMainAudioDeviceWorlds() > 1)
+									{
+										uint32 MainAudioDeviceID = GEngine->GetMainAudioDeviceID();
+										if (AudioDevice->DeviceID == MainAudioDeviceID && !HasAudioFocus())
+										{
+											bUpdateListenerPosition = false;
+										}
+									}
+								}
+
+								if (bUpdateListenerPosition)
+								{
+									FVector Location;
+									FVector ProjFront;
+									FVector ProjRight;
+									PlayerController->GetAudioListenerPosition(Location, ProjFront, ProjRight);
+
+									FTransform ListenerTransform(FRotationMatrix::MakeFromXY(ProjFront, ProjRight));
+
+									// Allow the HMD to adjust based on the head position of the player, as opposed to the view location
+									if (GEngine->XRSystem.IsValid() && GEngine->StereoRenderingDevice.IsValid() && GEngine->StereoRenderingDevice->IsStereoEnabled())
+									{
+										const FVector Offset = GEngine->XRSystem->GetAudioListenerOffset();
+										Location += ListenerTransform.TransformPositionNoScale(Offset);
+									}
+
+									ListenerTransform.SetTranslation(Location);
+									ListenerTransform.NormalizeRotation();
+
+									uint32 ViewportIndex = PlayerViewMap.Num() - 1;
+									RetrievedAudioDevice->SetListener(MyWorld, ViewportIndex, ListenerTransform, (View->bCameraCut ? 0.f : MyWorld->GetDeltaSeconds()));
+
+									FVector OverrideAttenuation;
+									if (PlayerController->GetAudioListenerAttenuationOverridePosition(OverrideAttenuation))
+									{
+										RetrievedAudioDevice->SetListenerAttenuationOverride(ViewportIndex, OverrideAttenuation);
+									}
+									else
+									{
+										RetrievedAudioDevice->ClearListenerAttenuationOverride(ViewportIndex);
+									}
+								}
+							}
+
+#if RHI_RAYTRACING
+							View->SetupRayTracedRendering();
+#endif
+						}
+
+						// Add view information for resource streaming. Allow up to 5X boost for small FOV.
+						const float StreamingScale = 1.f / FMath::Clamp<float>(View->LODDistanceFactor, .2f, 1.f);
+						IStreamingManager::Get().AddViewInformation(View->ViewMatrices.GetViewOrigin(), View->UnscaledViewRect.Width(), View->UnscaledViewRect.Width() * View->ViewMatrices.GetProjectionMatrix().M[0][0], StreamingScale);
+						MyWorld->ViewLocationsRenderedLastFrame.Add(View->ViewMatrices.GetViewOrigin());
+					}
+				}
+			}
+		}
+
+		FinalizeViews(&ViewFamily, PlayerViewMap);
+
+		// Update level streaming.
+		MyWorld->UpdateLevelStreaming();
+
+		// Find largest rectangle bounded by all rendered views.
+		uint32 MinX = InViewport->GetSizeXY().X, MinY = InViewport->GetSizeXY().Y, MaxX = 0, MaxY = 0;
+		uint32 TotalArea = 0;
+		{
+			for (int32 ViewIndex = 0; ViewIndex < ViewFamily.Views.Num(); ++ViewIndex)
+			{
+				const FSceneView* View = ViewFamily.Views[ViewIndex];
+
+				FIntRect UpscaledViewRect = View->UnscaledViewRect;
+
+				MinX = FMath::Min<uint32>(UpscaledViewRect.Min.X, MinX);
+				MinY = FMath::Min<uint32>(UpscaledViewRect.Min.Y, MinY);
+				MaxX = FMath::Max<uint32>(UpscaledViewRect.Max.X, MaxX);
+				MaxY = FMath::Max<uint32>(UpscaledViewRect.Max.Y, MaxY);
+				TotalArea += FMath::TruncToInt(UpscaledViewRect.Width()) * FMath::TruncToInt(UpscaledViewRect.Height());
+			}
+
+			// To draw black borders around the rendered image (prevents artifacts from post processing passes that read outside of the image e.g. PostProcessAA)
+			{
+				int32 BlackBordersCVar = IConsoleManager::Get().FindConsoleVariable(TEXT("r.BlackBorders"), false)->GetInt();
+				int32 BlackBorders = FMath::Clamp(BlackBordersCVar, 0, 10);
+
+				if (ViewFamily.Views.Num() == 1 && BlackBorders)
+				{
+					MinX += BlackBorders;
+					MinY += BlackBorders;
+					MaxX -= BlackBorders;
+					MaxY -= BlackBorders;
+					TotalArea = (MaxX - MinX) * (MaxY - MinY);
+				}
+			}
+		}
+
+		// If the views don't cover the entire bounding rectangle, clear the entire buffer.
+		bool bBufferCleared = false;
+		bool bStereoscopicPass = (ViewFamily.Views.Num() != 0 && ViewFamily.Views[0]->StereoPass != eSSP_FULL);
+		if (ViewFamily.Views.Num() == 0 || TotalArea != (MaxX - MinX)*(MaxY - MinY) || bDisableWorldRendering || bStereoscopicPass)
+		{
+			if (bDisableWorldRendering || !bStereoscopicPass) // TotalArea computation does not work correctly for stereoscopic views
+			{
+				//SceneCanvas->Clear(FLinearColor::Transparent);
+			}
+
+			bBufferCleared = true;
+		}
+
+		// Force screen percentage show flag to be turned off if not supported.
+		if (!ViewFamily.SupportsScreenPercentage())
+		{
+			ViewFamily.EngineShowFlags.ScreenPercentage = false;
+		}
+		// Update level streaming.
+		MyWorld->UpdateLevelStreaming();
+
+
+		// Set up secondary resolution fraction for the view family.
+		if (!bStereoRendering && ViewFamily.SupportsScreenPercentage())
+		{
+			float BlackBordersCVar = IConsoleManager::Get().FindConsoleVariable(TEXT("r.SecondaryScreenPercentage.GameViewport"), false)->GetFloat();
+			float CustomSecondaruScreenPercentage = BlackBordersCVar;
+
+			if (CustomSecondaruScreenPercentage > 0.0)
+			{
+				// Override secondary resolution fraction with CVar.
+				ViewFamily.SecondaryViewFraction = FMath::Min(CustomSecondaruScreenPercentage / 100.0f, 1.0f);
+			}
+			else
+			{
+				// Automatically compute secondary resolution fraction from DPI.
+				ViewFamily.SecondaryViewFraction = GetDPIDerivedResolutionFraction();
+			}
+
+			check(ViewFamily.SecondaryViewFraction > 0.0f);
+		}
+
+		checkf(ViewFamily.GetScreenPercentageInterface() == nullptr,
+			TEXT("Some code has tried to set up an alien screen percentage driver, that could be wrong if not supported very well by the RHI."));
+
+		// Setup main view family with screen percentage interface by dynamic resolution if screen percentage is enabled.
+#if WITH_DYNAMIC_RESOLUTION
+		if (ViewFamily.EngineShowFlags.ScreenPercentage)
+		{
+			FDynamicResolutionStateInfos DynamicResolutionStateInfos;
+			GEngine->GetDynamicResolutionCurrentStateInfos( DynamicResolutionStateInfos);
+
+			// Do not allow dynamic resolution to touch the view family if not supported to ensure there is no possibility to ruin
+			// game play experience on platforms that does not support it, but have it enabled by mistake.
+			if (DynamicResolutionStateInfos.Status == EDynamicResolutionStatus::Enabled)
+			{
+				GEngine->EmitDynamicResolutionEvent(EDynamicResolutionStateEvent::BeginDynamicResolutionRendering);
+				GEngine->GetDynamicResolutionState()->SetupMainViewFamily(ViewFamily);
+			}
+			else if (DynamicResolutionStateInfos.Status == EDynamicResolutionStatus::DebugForceEnabled)
+			{
+				GEngine->EmitDynamicResolutionEvent(EDynamicResolutionStateEvent::BeginDynamicResolutionRendering);
+				ViewFamily.SetScreenPercentageInterface(new FLegacyScreenPercentageDriver(
+					ViewFamily,
+					DynamicResolutionStateInfos.ResolutionFractionApproximation,
+					false,
+					DynamicResolutionStateInfos.ResolutionFractionUpperBound));
+			}
+
+#if CSV_PROFILER
+			if (DynamicResolutionStateInfos.ResolutionFractionApproximation >= 0.0f)
+			{
+				CSV_CUSTOM_STAT_GLOBAL(DynamicResolutionPercentage, DynamicResolutionStateInfos.ResolutionFractionApproximation * 100.0f, ECsvCustomStatOp::Set);
+			}
+#endif
+		}
+#endif
+
+		// If a screen percentage interface was not set by dynamic resolution, then create one matching legacy behavior.
+		if (ViewFamily.GetScreenPercentageInterface() == nullptr)
+		{
+			bool AllowPostProcessSettingsScreenPercentage = false;
+			float GlobalResolutionFraction = 1.0f;
+
+			if (ViewFamily.EngineShowFlags.ScreenPercentage)
+			{
+				// Allow FPostProcessSettings::ScreenPercentage.
+				AllowPostProcessSettingsScreenPercentage = true;
+
+				// Get global view fraction set by r.ScreenPercentage.
+				GlobalResolutionFraction = FLegacyScreenPercentageDriver::GetCVarResolutionFraction() * CustomBufferRatio;
+			}
+
+			ViewFamily.SetScreenPercentageInterface(new FLegacyScreenPercentageDriver(
+				ViewFamily, GlobalResolutionFraction, AllowPostProcessSettingsScreenPercentage));
+		}
+		else if (bStereoRendering)
+		{
+			// Change screen percentage method to raw output when doing dynamic resolution with VR if not using TAA upsample.
+			for (FSceneView* View : Views)
+			{
+				if (View->PrimaryScreenPercentageMethod == EPrimaryScreenPercentageMethod::SpatialUpscale)
+				{
+					View->PrimaryScreenPercentageMethod = EPrimaryScreenPercentageMethod::RawOutput;
+				}
+			}
+		}
+
+		// Draw the player views.
+		if (!bDisableWorldRendering && PlayerViewMap.Num() > 0 && FSlateApplication::Get().GetPlatformApplication()->IsAllowedToRender()) //-V560
+		{
+			GetRendererModule().BeginRenderingViewFamily(SceneCanvas, &ViewFamily);
+		}
+		else
+		{
+			// Make sure RHI resources get flushed if we're not using a renderer
+			ENQUEUE_RENDER_COMMAND(UGameViewportClient_FlushRHIResources)(
+				[](FRHICommandListImmediate& RHICmdList)
+			{
+				RHICmdList.ImmediateFlush(EImmediateFlushType::FlushRHIThreadFlushResources);
+			});
+		}
+
+		//SCOPE_CYCLE_COUNTER(STAT_UIDrawingTime);
+		//CSV_SCOPED_TIMING_STAT_EXCLUSIVE(UI);
+
+		// render HUD
+		bool bDisplayedSubtitles = false;
+		for (FConstPlayerControllerIterator Iterator = MyWorld->GetPlayerControllerIterator(); Iterator; ++Iterator)
+		{
+			APlayerController* PlayerController = Iterator->Get();
+			if (PlayerController)
+			{
+				ULocalPlayer* LocalPlayer = Cast<ULocalPlayer>(PlayerController->Player);
+				if (LocalPlayer)
+				{
+					FSceneView* View = PlayerViewMap.FindRef(LocalPlayer);
+					if (View != NULL)
+					{
+						// rendering to directly to viewport target
+						FVector CanvasOrigin(FMath::TruncToFloat(View->UnscaledViewRect.Min.X), FMath::TruncToInt(View->UnscaledViewRect.Min.Y), 0.f);
+
+						CanvasObject->Init(View->UnscaledViewRect.Width(), View->UnscaledViewRect.Height(), View, SceneCanvas);
+
+						// Set the canvas transform for the player's view rectangle.
+						check(SceneCanvas);
+						SceneCanvas->PushAbsoluteTransform(FTranslationMatrix(CanvasOrigin));
+						CanvasObject->ApplySafeZoneTransform();
+
+						// Render the player's HUD.
+						if (PlayerController->MyHUD)
+						{
+							//SCOPE_CYCLE_COUNTER(STAT_HudTime);
+
+							DebugCanvasObject->SceneView = View;
+							PlayerController->MyHUD->SetCanvas(CanvasObject, DebugCanvasObject);
+
+							PlayerController->MyHUD->PostRender();
+
+							// Put these pointers back as if a blueprint breakpoint hits during HUD PostRender they can
+							// have been changed
+							CanvasObject->Canvas = SceneCanvas;
+							DebugCanvasObject->Canvas = DebugCanvas;
+
+							// A side effect of PostRender is that the playercontroller could be destroyed
+							if (!PlayerController->IsPendingKill())
+							{
+								PlayerController->MyHUD->SetCanvas(NULL, NULL);
+							}
+						}
+
+						if (DebugCanvas != NULL)
+						{
+							DebugCanvas->PushAbsoluteTransform(FTranslationMatrix(CanvasOrigin));
+							UDebugDrawService::Draw(ViewFamily.EngineShowFlags, InViewport, View, DebugCanvas, DebugCanvasObject);
+							DebugCanvas->PopTransform();
+						}
+
+						CanvasObject->PopSafeZoneTransform();
+						SceneCanvas->PopTransform();
+
+						// draw subtitles
+						if (!bDisplayedSubtitles)
+						{
+							FVector2D MinPos(0.f, 0.f);
+							FVector2D MaxPos(1.f, 1.f);
+							GetSubtitleRegion(MinPos, MaxPos);
+
+							const uint32 SizeX = SceneCanvas->GetRenderTarget()->GetSizeXY().X;
+							const uint32 SizeY = SceneCanvas->GetRenderTarget()->GetSizeXY().Y;
+							FIntRect SubtitleRegion(FMath::TruncToInt(SizeX * MinPos.X), FMath::TruncToInt(SizeY * MinPos.Y), FMath::TruncToInt(SizeX * MaxPos.X), FMath::TruncToInt(SizeY * MaxPos.Y));
+							FSubtitleManager::GetSubtitleManager()->DisplaySubtitles(SceneCanvas, SubtitleRegion, MyWorld->GetAudioTimeSeconds());
+							bDisplayedSubtitles = true;
+						}
+					}
+				}
+			}
+		}
+	}
+
+	// Beyond this point, only UI rendering independent from dynamc resolution.
+	GEngine->EmitDynamicResolutionEvent(EDynamicResolutionStateEvent::EndDynamicResolutionRendering);
+
+	// Remove temporary debug lines.
+	if (MyWorld->LineBatcher != nullptr)
+	{
+		MyWorld->LineBatcher->Flush();
+	}
+
+	if (MyWorld->ForegroundLineBatcher != nullptr)
+	{
+		MyWorld->ForegroundLineBatcher->Flush();
+	}
+
+	// Draw FX debug information.
+	if (MyWorld->FXSystem)
+	{
+		MyWorld->FXSystem->DrawDebug(SceneCanvas);
+	}
+
+	{
+		//ensure canvas has been flushed before rendering UI
+		SceneCanvas->Flush_GameThread();
+
+		OnDrawn().Broadcast();
+
+		// Allow the viewport to render additional stuff
+		PostRender(DebugCanvasObject);
+	}
+
+
+	// Grab the player camera location and orientation so we can pass that along to the stats drawing code.
+	FVector PlayerCameraLocation = FVector::ZeroVector;
+	FRotator PlayerCameraRotation = FRotator::ZeroRotator;
+	{
+		for (FConstPlayerControllerIterator Iterator = MyWorld->GetPlayerControllerIterator(); Iterator; ++Iterator)
+		{
+			if (APlayerController* PC = Iterator->Get())
+			{
+				PC->GetPlayerViewPoint(PlayerCameraLocation, PlayerCameraRotation);
+			}
+		}
+	}
+
+	if (DebugCanvas)
+	{
+		// Reset the debug canvas to be full-screen before drawing the console
+		// (the debug draw service above has messed with the viewport size to fit it to a single player's subregion)
+		DebugCanvasObject->Init(DebugCanvasSize.X, DebugCanvasSize.Y, NULL, DebugCanvas);
+
+		DrawStatsHUD(MyWorld, InViewport, DebugCanvas, DebugCanvasObject, DebugProperties, PlayerCameraLocation, PlayerCameraRotation);
+
+		if (GEngine->IsStereoscopic3D(InViewport))
+		{
+#if 0 //!UE_BUILD_SHIPPING
+			// TODO: replace implementation in OculusHMD with a debug renderer
+			if (GEngine->XRSystem.IsValid())
+			{
+				GEngine->XRSystem->DrawDebug(DebugCanvasObject);
+			}
+#endif
+		}
+
+		// Render the console absolutely last because developer input is was matter the most.
+		if (ViewportConsole)
+		{
+			ViewportConsole->PostRender_Console(DebugCanvasObject);
+		}
+	}
+
+	OnEndDraw().Broadcast();
+}
diff --git a/Engine/Plugins/Runtime/nDisplay/Source/DisplayCluster/Private/Game/EngineClasses/Scene/DisplayClusterCameraComponent.cpp b/Engine/Plugins/Runtime/nDisplay/Source/DisplayCluster/Private/Game/EngineClasses/Scene/DisplayClusterCameraComponent.cpp
new file mode 100644
--- /dev/null
+++ b/Engine/Plugins/Runtime/nDisplay/Source/DisplayCluster/Private/Game/EngineClasses/Scene/DisplayClusterCameraComponent.cpp
@@ -0,0 +1,45 @@
+// Copyright Epic Games, Inc. All Rights Reserved.
+
+#include "DisplayClusterCameraComponent.h"
+#include "Config/DisplayClusterConfigTypes.h"
+
+#include "Camera/CameraComponent.h"
+
+#include "CoreGlobals.h"
+
+
+UDisplayClusterCameraComponent::UDisplayClusterCameraComponent(const FObjectInitializer& ObjectInitializer)
+	: UDisplayClusterSceneComponent(ObjectInitializer)
+	, EyeDist(0.064f)
+	, bEyeSwap(false)
+	, ForceEyeOffset(0)
+{
+	PrimaryComponentTick.bCanEverTick = true;
+}
+
+
+void UDisplayClusterCameraComponent::BeginPlay()
+{
+	Super::BeginPlay();
+}
+
+void UDisplayClusterCameraComponent::SetSettings(const FDisplayClusterConfigSceneNode* ConfigData)
+{
+	check(ConfigData);
+	if (ConfigData)
+	{
+		const FDisplayClusterConfigCamera* const CameraCfg = static_cast<const FDisplayClusterConfigCamera*>(ConfigData);
+
+		EyeDist        = CameraCfg->EyeDist;
+		bEyeSwap       = CameraCfg->EyeSwap;
+		ForceEyeOffset = CameraCfg->ForceOffset;
+	}
+	
+
+	Super::SetSettings(ConfigData);
+}
+
+bool UDisplayClusterCameraComponent::ApplySettings()
+{
+	return Super::ApplySettings();
+}
diff --git a/Engine/Plugins/Runtime/nDisplay/Source/DisplayCluster/Private/Game/EngineClasses/Scene/DisplayClusterRootActor.cpp b/Engine/Plugins/Runtime/nDisplay/Source/DisplayCluster/Private/Game/EngineClasses/Scene/DisplayClusterRootActor.cpp
new file mode 100644
--- /dev/null
+++ b/Engine/Plugins/Runtime/nDisplay/Source/DisplayCluster/Private/Game/EngineClasses/Scene/DisplayClusterRootActor.cpp
@@ -0,0 +1,119 @@
+// Copyright Epic Games, Inc. All Rights Reserved.
+
+#include "DisplayClusterRootActor.h"
+
+#include "Components/SceneComponent.h"
+#include "GameFramework/PlayerController.h"
+#include "Engine/World.h"
+
+#include "Config/IPDisplayClusterConfigManager.h"
+#include "Config/DisplayClusterConfigTypes.h"
+#include "Misc/DisplayClusterAppExit.h"
+
+#include "DisplayClusterRootComponent.h"
+#include "DisplayClusterSceneComponentSyncParent.h"
+#include "DisplayClusterPlayerInput.h"
+
+#include "DisplayClusterGlobals.h"
+#include "DisplayClusterLog.h"
+
+
+ADisplayClusterRootActor::ADisplayClusterRootActor(const FObjectInitializer& ObjectInitializer)
+	: Super(ObjectInitializer)
+	, bExitOnEsc(true)
+	, bShowProjectionScreens(false)
+	, ProjectionScreensMaterial(nullptr)
+{
+	DISPLAY_CLUSTER_FUNC_TRACE(LogDisplayClusterGame);
+
+	// Root component
+	RootComponent = CreateDefaultSubobject<USceneComponent>(TEXT("RootComponent"));
+
+	// DisplayCluster sync
+	DisplayClusterRootComponent = CreateDefaultSubobject<UDisplayClusterRootComponent>(TEXT("DisplayClusterRoot"));
+	DisplayClusterRootComponent->AttachToComponent(RootComponent, FAttachmentTransformRules(EAttachmentRule::KeepRelative, false));
+
+	PrimaryActorTick.bCanEverTick = true;
+	PrimaryActorTick.TickGroup = ETickingGroup::TG_PostUpdateWork;
+	bFindCameraComponentWhenViewTarget = false;
+	bReplicates = false;
+	SpawnCollisionHandlingMethod = ESpawnActorCollisionHandlingMethod::AdjustIfPossibleButAlwaysSpawn;
+}
+
+void ADisplayClusterRootActor::BeginPlay()
+{
+	DISPLAY_CLUSTER_FUNC_TRACE(LogDisplayClusterGame);
+
+	int32 NativeInputSyncPolicy = 0;
+
+	// Store current operation mode
+	OperationMode = GDisplayCluster->GetOperationMode();
+
+	if (OperationMode == EDisplayClusterOperationMode::Cluster ||
+		OperationMode == EDisplayClusterOperationMode::Editor)
+	{
+		// Read native input synchronization settings
+		IPDisplayClusterConfigManager* const ConfigMgr = GDisplayCluster->GetPrivateConfigMgr();
+		if (ConfigMgr)
+		{
+			FDisplayClusterConfigGeneral CfgGeneral = ConfigMgr->GetConfigGeneral();
+			NativeInputSyncPolicy = CfgGeneral.NativeInputSyncPolicy;
+			UE_LOG(LogDisplayClusterGame, Log, TEXT("Native input sync policy: %d"), NativeInputSyncPolicy);
+		}
+	}
+
+	if (OperationMode == EDisplayClusterOperationMode::Cluster)
+	{
+		// Optionally activate native input synchronization
+		if (NativeInputSyncPolicy == 1)
+		{
+			APlayerController* const PlayerController = GetWorld()->GetFirstPlayerController();
+			if (PlayerController)
+			{
+				PlayerController->PlayerInput = NewObject<UDisplayClusterPlayerInput>(PlayerController);
+			}
+		}
+	}
+
+	Super::BeginPlay();
+}
+
+void ADisplayClusterRootActor::BeginDestroy()
+{
+	DISPLAY_CLUSTER_FUNC_TRACE(LogDisplayClusterGame);
+
+	Super::BeginDestroy();
+}
+
+void ADisplayClusterRootActor::Tick(float DeltaSeconds)
+{
+	DISPLAY_CLUSTER_FUNC_TRACE(LogDisplayClusterGame);
+
+	if (OperationMode == EDisplayClusterOperationMode::Cluster ||
+		OperationMode == EDisplayClusterOperationMode::Editor)
+	{
+		UWorld* const CurWorld = GetWorld();
+		if (CurWorld)
+		{
+			APlayerController* const CurPlayerController = CurWorld->GetFirstPlayerController();
+			if (CurPlayerController)
+			{
+				APlayerCameraManager* const CurPlayerCameraManager = CurPlayerController->PlayerCameraManager;
+				if (CurPlayerCameraManager)
+				{
+					SetActorLocationAndRotation(CurPlayerCameraManager->GetCameraLocation(), CurPlayerCameraManager->GetCameraRotation());
+				}
+
+				if (bExitOnEsc)
+				{
+					if (CurPlayerController->WasInputKeyJustPressed(EKeys::Escape))
+					{
+						FDisplayClusterAppExit::ExitApplication(FDisplayClusterAppExit::ExitType::NormalSoft, FString("Exit on ESC requested"));
+					}
+				}
+			}
+		}
+	}
+
+	Super::Tick(DeltaSeconds);
+}
diff --git a/Engine/Plugins/Runtime/nDisplay/Source/DisplayCluster/Private/Game/EngineClasses/Scene/DisplayClusterRootComponent.cpp b/Engine/Plugins/Runtime/nDisplay/Source/DisplayCluster/Private/Game/EngineClasses/Scene/DisplayClusterRootComponent.cpp
new file mode 100644
--- /dev/null
+++ b/Engine/Plugins/Runtime/nDisplay/Source/DisplayCluster/Private/Game/EngineClasses/Scene/DisplayClusterRootComponent.cpp
@@ -0,0 +1,269 @@
+// Copyright Epic Games, Inc. All Rights Reserved.
+
+#include "DisplayClusterRootComponent.h"
+
+#include "Config/DisplayClusterConfigTypes.h"
+#include "Cluster/IPDisplayClusterClusterManager.h"
+#include "Config/IPDisplayClusterConfigManager.h"
+#include "Config/DisplayClusterConfigTypes.h"
+#include "Game/IPDisplayClusterGameManager.h"
+
+#include "Camera/CameraComponent.h"
+
+#include "DisplayClusterCameraComponent.h"
+#include "DisplayClusterSceneComponent.h"
+#include "DisplayClusterScreenComponent.h"
+
+#include "DisplayClusterGlobals.h"
+#include "DisplayClusterHelpers.h"
+#include "DisplayClusterLog.h"
+#include "DisplayClusterStrings.h"
+
+
+UDisplayClusterRootComponent::UDisplayClusterRootComponent(const FObjectInitializer& ObjectInitializer)
+	: USceneComponent(ObjectInitializer)
+{
+	DISPLAY_CLUSTER_FUNC_TRACE(LogDisplayClusterGame);
+
+	// This component settings
+	PrimaryComponentTick.bCanEverTick = true;
+}
+
+void UDisplayClusterRootComponent::BeginPlay()
+{
+	DISPLAY_CLUSTER_FUNC_TRACE(LogDisplayClusterGame);
+
+	// Store current operation mode on BeginPlay (not in the constructor)
+	const EDisplayClusterOperationMode OperationMode = GDisplayCluster->GetOperationMode();
+
+	if (OperationMode == EDisplayClusterOperationMode::Cluster ||
+		OperationMode == EDisplayClusterOperationMode::Editor)
+	{
+		// Build nDisplay hierarchy
+		InitializeHierarchy();
+	}
+
+	Super::BeginPlay();
+}
+
+void UDisplayClusterRootComponent::BeginDestroy()
+{
+	DISPLAY_CLUSTER_FUNC_TRACE(LogDisplayClusterGame);
+
+	{
+		FScopeLock lock(&InternalsSyncScope);
+
+		// Clean containers. We store only pointers so there is no need to do any additional
+		// operations. All components will be destroyed by the engine.
+		ScreenComponents.Reset();
+		CameraComponents.Reset();
+		SceneNodeComponents.Reset();
+	}
+
+	Super::BeginDestroy();
+}
+
+bool UDisplayClusterRootComponent::InitializeHierarchy()
+{
+	DISPLAY_CLUSTER_FUNC_TRACE(LogDisplayClusterGame);
+
+	// Now create all child components
+	if (!(CreateCameras() && CreateScreens() && CreateNodes()))
+	{
+		UE_LOG(LogDisplayClusterGame, Error, TEXT("An error occurred during DisplayCluster root initialization"));
+		return false;
+	}
+
+	// Let components initialize ourselves
+	for (auto it = SceneNodeComponents.CreateIterator(); it; ++it)
+	{
+		if (it->Value->ApplySettings() == false)
+		{
+			UE_LOG(LogDisplayClusterGame, Warning, TEXT("Coulnd't initialize DisplayCluster node: ID=%s"), *it->Key);
+		}
+	}
+
+	// Check if default camera was specified in command line arguments
+	FString DefaultCamId;
+	if (FParse::Value(FCommandLine::Get(), DisplayClusterStrings::args::Camera, DefaultCamId))
+	{
+		DisplayClusterHelpers::str::TrimStringValue(DefaultCamId);
+		UE_LOG(LogDisplayClusterGame, Log, TEXT("Default camera from command line arguments: %s"), *DefaultCamId);
+		if (CameraComponents.Contains(DefaultCamId))
+		{
+			SetDefaultCamera(DefaultCamId);
+		}
+	}
+
+	return true;
+}
+
+bool UDisplayClusterRootComponent::CreateScreens()
+{
+	// Create screens
+	const IPDisplayClusterConfigManager* const ConfigMgr = GDisplayCluster->GetPrivateConfigMgr();
+	if (!ConfigMgr)
+	{
+		UE_LOG(LogDisplayClusterGame, Error, TEXT("Couldn't get config manager interface"));
+		return false;
+	}
+
+	const TArray<FDisplayClusterConfigScreen> Screens = ConfigMgr->GetScreens();
+	for (const FDisplayClusterConfigScreen& Screen : Screens)
+	{
+		// Create screen component
+		UDisplayClusterScreenComponent* NewScreenComp = NewObject<UDisplayClusterScreenComponent>(GetOwner(), FName(*Screen.Id), RF_Transient);
+		check(NewScreenComp);
+
+		NewScreenComp->AttachToComponent(this, FAttachmentTransformRules(EAttachmentRule::KeepRelative, false));
+		NewScreenComp->RegisterComponent();
+		NewScreenComp->SetSettings(&Screen);
+
+		// Store the screen
+		ScreenComponents.Add(Screen.Id, NewScreenComp);
+		SceneNodeComponents.Add(Screen.Id, NewScreenComp);
+	}
+
+	return true;
+}
+
+bool UDisplayClusterRootComponent::CreateNodes()
+{
+	// Create other nodes
+	const TArray<FDisplayClusterConfigSceneNode> SceneNodes = GDisplayCluster->GetPrivateConfigMgr()->GetSceneNodes();
+	for (const FDisplayClusterConfigSceneNode& Node : SceneNodes)
+	{
+		UDisplayClusterSceneComponent* NewSceneNode = NewObject<UDisplayClusterSceneComponent>(GetOwner(), FName(*Node.Id), RF_Transient);
+		check(NewSceneNode);
+
+		NewSceneNode->AttachToComponent(this, FAttachmentTransformRules(EAttachmentRule::KeepRelative, false));
+		NewSceneNode->RegisterComponent();
+		NewSceneNode->SetSettings(&Node);
+
+		SceneNodeComponents.Add(Node.Id, NewSceneNode);
+	}
+
+	return true;
+}
+
+bool UDisplayClusterRootComponent::CreateCameras()
+{
+	// Get all cameras listed in the config file
+	const TArray<FDisplayClusterConfigCamera> Cameras = GDisplayCluster->GetPrivateConfigMgr()->GetCameras();
+	
+	// Instantiate each camera within hierarchy of this root
+	for (const FDisplayClusterConfigCamera& Camera : Cameras)
+	{
+		UDisplayClusterCameraComponent* NewCameraComponent = NewObject<UDisplayClusterCameraComponent>(GetOwner(), FName(*Camera.Id), RF_Transient);
+		check(NewCameraComponent);
+
+		NewCameraComponent->AttachToComponent(this, FAttachmentTransformRules(EAttachmentRule::KeepRelative, false));
+		NewCameraComponent->RegisterComponent();
+		NewCameraComponent->SetSettings(&Camera);
+
+		CameraComponents.Add(Camera.Id, NewCameraComponent);
+		SceneNodeComponents.Add(Camera.Id, NewCameraComponent);
+
+		if (!DefaultCameraComponent)
+		{
+			DefaultCameraComponent = NewCameraComponent;
+		}
+	}
+
+	// At least one camera must be set up
+	if (CameraComponents.Num() == 0)
+	{
+		UE_LOG(LogDisplayClusterGame, Warning, TEXT("No cameras found"));
+		return false;
+	}
+
+	return CameraComponents.Num() > 0;
+}
+
+
+TArray<UDisplayClusterScreenComponent*> UDisplayClusterRootComponent::GetAllScreens() const
+{
+	FScopeLock lock(&InternalsSyncScope);
+	return GetMapValues<UDisplayClusterScreenComponent>(ScreenComponents);
+}
+
+UDisplayClusterScreenComponent* UDisplayClusterRootComponent::GetScreenById(const FString& id) const
+{
+	FScopeLock lock(&InternalsSyncScope);
+	return GetItem<UDisplayClusterScreenComponent>(ScreenComponents, id, FString("GetScreenById"));
+}
+
+int32 UDisplayClusterRootComponent::GetScreensAmount() const
+{
+	FScopeLock lock(&InternalsSyncScope);
+	return ScreenComponents.Num();
+}
+
+UDisplayClusterCameraComponent* UDisplayClusterRootComponent::GetCameraById(const FString& id) const
+{
+	FScopeLock lock(&InternalsSyncScope);
+	return GetItem<UDisplayClusterCameraComponent>(CameraComponents, id, FString("GetCameraById"));
+}
+
+TArray<UDisplayClusterCameraComponent*> UDisplayClusterRootComponent::GetAllCameras() const
+{
+	FScopeLock lock(&InternalsSyncScope);
+	return GetMapValues<UDisplayClusterCameraComponent>(CameraComponents);
+}
+
+int32 UDisplayClusterRootComponent::GetCamerasAmount() const
+{
+	FScopeLock lock(&InternalsSyncScope);
+	return CameraComponents.Num();
+}
+
+UDisplayClusterCameraComponent* UDisplayClusterRootComponent::GetDefaultCamera() const
+{
+	FScopeLock lock(&InternalsSyncScope);
+	return DefaultCameraComponent;
+}
+
+void UDisplayClusterRootComponent::SetDefaultCamera(const FString& id)
+{
+	FScopeLock lock(&InternalsSyncScope);
+
+	UDisplayClusterCameraComponent* NewDefaultCamera = GetCameraById(id);
+	if (NewDefaultCamera)
+	{
+		DefaultCameraComponent = NewDefaultCamera;
+	}
+}
+
+UDisplayClusterSceneComponent* UDisplayClusterRootComponent::GetNodeById(const FString& id) const
+{
+	FScopeLock lock(&InternalsSyncScope);
+	return GetItem<UDisplayClusterSceneComponent>(SceneNodeComponents, id, FString("GetNodeById"));
+}
+
+TArray<UDisplayClusterSceneComponent*> UDisplayClusterRootComponent::GetAllNodes() const
+{
+	FScopeLock lock(&InternalsSyncScope);
+	return GetMapValues<UDisplayClusterSceneComponent>(SceneNodeComponents);
+}
+
+// Extracts array of values from a map
+template <typename ObjType>
+TArray<ObjType*> UDisplayClusterRootComponent::GetMapValues(const TMap<FString, ObjType*>& container) const
+{
+	TArray<ObjType*> items;
+	container.GenerateValueArray(items);
+	return items;
+}
+
+// Gets item by id. Performs checks and logging.
+template <typename DataType>
+DataType* UDisplayClusterRootComponent::GetItem(const TMap<FString, DataType*>& container, const FString& id, const FString& logHeader) const
+{
+	if (container.Contains(id))
+	{
+		return container[id];
+	}
+
+	UE_LOG(LogDisplayClusterGame, Warning, TEXT("%s: ID not found <%s>. Return nullptr."), *logHeader, *id);
+	return nullptr;
+}
diff --git a/Engine/Plugins/Runtime/nDisplay/Source/DisplayCluster/Private/Game/EngineClasses/Scene/DisplayClusterSceneComponent.cpp b/Engine/Plugins/Runtime/nDisplay/Source/DisplayCluster/Private/Game/EngineClasses/Scene/DisplayClusterSceneComponent.cpp
new file mode 100644
--- /dev/null
+++ b/Engine/Plugins/Runtime/nDisplay/Source/DisplayCluster/Private/Game/EngineClasses/Scene/DisplayClusterSceneComponent.cpp
@@ -0,0 +1,87 @@
+// Copyright Epic Games, Inc. All Rights Reserved.
+
+#include "DisplayClusterSceneComponent.h"
+
+#include "DisplayClusterRootComponent.h"
+
+#include "Config/DisplayClusterConfigTypes.h"
+#include "Input/IPDisplayClusterInputManager.h"
+
+#include "DisplayClusterGlobals.h"
+#include "DisplayClusterLog.h"
+
+
+UDisplayClusterSceneComponent::UDisplayClusterSceneComponent(const FObjectInitializer& ObjectInitializer) :
+	USceneComponent(ObjectInitializer)
+{
+	PrimaryComponentTick.bCanEverTick = true;
+}
+
+void UDisplayClusterSceneComponent::BeginPlay()
+{
+	Super::BeginPlay();
+}
+
+void UDisplayClusterSceneComponent::BeginDestroy()
+{
+	Super::BeginDestroy();
+}
+
+void UDisplayClusterSceneComponent::TickComponent( float DeltaTime, ELevelTick TickType, FActorComponentTickFunction* ThisTickFunction )
+{
+	Super::TickComponent(DeltaTime, TickType, ThisTickFunction);
+
+	// Update transform if attached to a tracker
+	if (!Config.TrackerId.IsEmpty())
+	{
+		const IPDisplayClusterInputManager* const InputMgr = GDisplayCluster->GetPrivateInputMgr();
+		if (InputMgr)
+		{
+			FVector loc;
+			FQuat rot;
+			const bool bLocAvail = InputMgr->GetTrackerLocation(Config.TrackerId, Config.TrackerCh, loc);
+			const bool bRotAvail = InputMgr->GetTrackerQuat(Config.TrackerId, Config.TrackerCh, rot);
+
+			if (bLocAvail && bRotAvail)
+			{
+				UE_LOG(LogDisplayClusterGame, Verbose, TEXT("%s[%s] update from tracker %s:%d - {loc %s} {quat %s}"),
+					*GetName(), *GetId(), *Config.TrackerId, Config.TrackerCh, *loc.ToString(), *rot.ToString());
+
+				// Update transform
+				this->SetRelativeLocationAndRotation(loc, rot);
+				// Force child transforms update
+				UpdateChildTransforms(EUpdateTransformFlags::PropagateFromParent);
+			}
+		}
+	}
+}
+
+void UDisplayClusterSceneComponent::SetSettings(const FDisplayClusterConfigSceneNode* pConfig)
+{
+	check(pConfig);
+
+	Config = *pConfig;
+
+	// Convert m to cm
+	Config.Loc *= 100.f;
+}
+
+bool UDisplayClusterSceneComponent::ApplySettings()
+{
+	// Take place in hierarchy
+	if (!GetParentId().IsEmpty())
+	{
+		UDisplayClusterRootComponent* const RootComp = Cast<UDisplayClusterRootComponent>(GetAttachParent());
+		if (RootComp)
+		{
+			UE_LOG(LogDisplayClusterGame, Log, TEXT("Attaching %s to %s"), *GetId(), *GetParentId());
+			UDisplayClusterSceneComponent* const pComp = RootComp->GetNodeById(GetParentId());
+			AttachToComponent(pComp, FAttachmentTransformRules(EAttachmentRule::KeepRelative, false));
+		}
+	}
+
+	// Set up location and rotation
+	this->SetRelativeLocationAndRotation(Config.Loc, Config.Rot);
+
+	return true;
+}
diff --git a/Engine/Plugins/Runtime/nDisplay/Source/DisplayCluster/Private/Game/EngineClasses/Scene/DisplayClusterSceneComponentSync.cpp b/Engine/Plugins/Runtime/nDisplay/Source/DisplayCluster/Private/Game/EngineClasses/Scene/DisplayClusterSceneComponentSync.cpp
new file mode 100644
--- /dev/null
+++ b/Engine/Plugins/Runtime/nDisplay/Source/DisplayCluster/Private/Game/EngineClasses/Scene/DisplayClusterSceneComponentSync.cpp
@@ -0,0 +1,104 @@
+// Copyright Epic Games, Inc. All Rights Reserved.
+
+#include "DisplayClusterSceneComponentSync.h"
+
+#include "GameFramework/Actor.h"
+
+#include "Cluster/IPDisplayClusterClusterManager.h"
+#include "Game/IPDisplayClusterGameManager.h"
+
+#include "DisplayClusterUtils/DisplayClusterTypesConverter.h"
+#include "DisplayClusterGlobals.h"
+#include "DisplayClusterLog.h"
+
+
+UDisplayClusterSceneComponentSync::UDisplayClusterSceneComponentSync(const FObjectInitializer& ObjectInitializer) :
+	USceneComponent(ObjectInitializer)
+{
+	PrimaryComponentTick.bCanEverTick = true;
+}
+
+void UDisplayClusterSceneComponentSync::BeginPlay()
+{
+	Super::BeginPlay();
+
+	if (!GDisplayCluster->IsModuleInitialized())
+	{
+		return;
+	}
+
+	// Generate unique sync id
+	SyncId = GenerateSyncId();
+
+	GameMgr = GDisplayCluster->GetPrivateGameMgr();
+	if (GameMgr && GameMgr->IsDisplayClusterActive())
+	{
+		// Register sync object
+		ClusterMgr = GDisplayCluster->GetPrivateClusterMgr();
+		if (ClusterMgr)
+		{
+			UE_LOG(LogDisplayClusterGame, Log, TEXT("Registering sync object %s..."), *SyncId);
+			ClusterMgr->RegisterSyncObject(this, EDisplayClusterSyncGroup::Tick);
+		}
+		else
+		{
+			UE_LOG(LogDisplayClusterGame, Warning, TEXT("Couldn't register %s scene component sync. Looks like we're in non-DisplayCluster mode."), *SyncId);
+		}
+	}
+}
+
+void UDisplayClusterSceneComponentSync::EndPlay(const EEndPlayReason::Type EndPlayReason)
+{
+	if (GDisplayCluster->IsModuleInitialized())
+	{
+		if (GameMgr && GameMgr->IsDisplayClusterActive())
+		{
+			if (ClusterMgr)
+			{
+				UE_LOG(LogDisplayClusterGame, Log, TEXT("Unregistering sync object %s..."), *SyncId);
+				ClusterMgr->UnregisterSyncObject(this);
+			}
+		}
+	}
+
+	Super::EndPlay(EndPlayReason);
+}
+
+void UDisplayClusterSceneComponentSync::TickComponent( float DeltaTime, ELevelTick TickType, FActorComponentTickFunction* ThisTickFunction )
+{
+	Super::TickComponent( DeltaTime, TickType, ThisTickFunction );
+
+	// ...
+}
+
+
+//////////////////////////////////////////////////////////////////////////////////////////////
+// IDisplayClusterClusterSyncObject
+//////////////////////////////////////////////////////////////////////////////////////////////
+bool UDisplayClusterSceneComponentSync::IsActive() const
+{
+	return !this->IsPendingKill();
+}
+
+
+//////////////////////////////////////////////////////////////////////////////////////////////
+// IDisplayClusterStringSerializable
+//////////////////////////////////////////////////////////////////////////////////////////////
+FString UDisplayClusterSceneComponentSync::GenerateSyncId()
+{
+	return FString::Printf(TEXT("S_%s"), *GetOwner()->GetName());
+}
+
+FString UDisplayClusterSceneComponentSync::SerializeToString() const
+{
+	return FDisplayClusterTypesConverter::template ToHexString(GetSyncTransform());
+}
+
+bool UDisplayClusterSceneComponentSync::DeserializeFromString(const FString& data)
+{
+	FTransform NewTransform = FDisplayClusterTypesConverter::template FromHexString<FTransform>(data);
+	UE_LOG(LogDisplayClusterGame, Verbose, TEXT("%s: applying transform data <%s>"), *SyncId, *NewTransform.ToHumanReadableString());
+	SetSyncTransform(NewTransform);
+
+	return true;
+}
diff --git a/Engine/Plugins/Runtime/nDisplay/Source/DisplayCluster/Private/Game/EngineClasses/Scene/DisplayClusterSceneComponentSyncParent.cpp b/Engine/Plugins/Runtime/nDisplay/Source/DisplayCluster/Private/Game/EngineClasses/Scene/DisplayClusterSceneComponentSyncParent.cpp
new file mode 100644
--- /dev/null
+++ b/Engine/Plugins/Runtime/nDisplay/Source/DisplayCluster/Private/Game/EngineClasses/Scene/DisplayClusterSceneComponentSyncParent.cpp
@@ -0,0 +1,76 @@
+// Copyright Epic Games, Inc. All Rights Reserved.
+
+#include "DisplayClusterSceneComponentSyncParent.h"
+#include "GameFramework/Actor.h"
+
+#include "DisplayClusterHelpers.h"
+#include "DisplayClusterLog.h"
+
+
+UDisplayClusterSceneComponentSyncParent::UDisplayClusterSceneComponentSyncParent(const FObjectInitializer& ObjectInitializer) :
+	UDisplayClusterSceneComponentSync(ObjectInitializer)
+{
+	PrimaryComponentTick.bCanEverTick = true;
+}
+
+void UDisplayClusterSceneComponentSyncParent::BeginPlay()
+{
+	Super::BeginPlay();
+
+	// ...
+}
+
+
+void UDisplayClusterSceneComponentSyncParent::TickComponent( float DeltaTime, ELevelTick TickType, FActorComponentTickFunction* ThisTickFunction )
+{
+	Super::TickComponent( DeltaTime, TickType, ThisTickFunction );
+
+	// ...
+}
+
+
+//////////////////////////////////////////////////////////////////////////////////////////////
+// IDisplayClusterClusterSyncObject
+//////////////////////////////////////////////////////////////////////////////////////////////
+bool UDisplayClusterSceneComponentSyncParent::IsDirty() const
+{
+	USceneComponent* const pParent = GetAttachParent();
+	if (pParent && !pParent->IsPendingKill())
+	{
+		const bool bIsDirty = (LastSyncLoc != pParent->GetRelativeLocation() || LastSyncRot != pParent->GetRelativeRotation() || LastSyncScale != pParent->GetRelativeScale3D());
+		UE_LOG(LogDisplayClusterGame, Verbose, TEXT("SYNC_PARENT: %s dirty state is %s"), *GetSyncId(), *DisplayClusterHelpers::str::BoolToStr(bIsDirty));
+		return bIsDirty;
+	}
+
+	return false;
+}
+
+void UDisplayClusterSceneComponentSyncParent::ClearDirty()
+{
+	USceneComponent* const pParent = GetAttachParent();
+	if (pParent && !pParent->IsPendingKill())
+	{
+		LastSyncLoc   = pParent->GetRelativeLocation();
+		LastSyncRot   = pParent->GetRelativeRotation();
+		LastSyncScale = pParent->GetRelativeScale3D();
+	}
+}
+
+
+//////////////////////////////////////////////////////////////////////////////////////////////
+// UDisplayClusterSceneComponentSync
+//////////////////////////////////////////////////////////////////////////////////////////////
+FString UDisplayClusterSceneComponentSyncParent::GenerateSyncId()
+{
+	return FString::Printf(TEXT("SP_%s.%s"), *GetOwner()->GetName(), *GetAttachParent()->GetName());
+}
+
+FTransform UDisplayClusterSceneComponentSyncParent::GetSyncTransform() const
+{
+	return GetAttachParent()->GetRelativeTransform();
+}
+
+void UDisplayClusterSceneComponentSyncParent::SetSyncTransform(const FTransform& t)
+{
+	GetAttachParent()->SetRelativeTransform(t);
+}
diff --git a/Engine/Plugins/Runtime/nDisplay/Source/DisplayCluster/Private/Game/EngineClasses/Scene/DisplayClusterSceneComponentSyncThis.cpp b/Engine/Plugins/Runtime/nDisplay/Source/DisplayCluster/Private/Game/EngineClasses/Scene/DisplayClusterSceneComponentSyncThis.cpp
new file mode 100644
--- /dev/null
+++ b/Engine/Plugins/Runtime/nDisplay/Source/DisplayCluster/Private/Game/EngineClasses/Scene/DisplayClusterSceneComponentSyncThis.cpp
@@ -0,0 +1,62 @@
+// Copyright Epic Games, Inc. All Rights Reserved.
+
+#include "DisplayClusterSceneComponentSyncThis.h"
+
+#include "GameFramework/Actor.h"
+
+
+UDisplayClusterSceneComponentSyncThis::UDisplayClusterSceneComponentSyncThis(const FObjectInitializer& ObjectInitializer) :
+	UDisplayClusterSceneComponentSync(ObjectInitializer)
+{
+	PrimaryComponentTick.bCanEverTick = true;
+}
+
+void UDisplayClusterSceneComponentSyncThis::BeginPlay()
+{
+	Super::BeginPlay();
+
+	// ...
+}
+
+
+void UDisplayClusterSceneComponentSyncThis::TickComponent( float DeltaTime, ELevelTick TickType, FActorComponentTickFunction* ThisTickFunction )
+{
+	Super::TickComponent( DeltaTime, TickType, ThisTickFunction );
+
+	// ...
+}
+
+
+//////////////////////////////////////////////////////////////////////////////////////////////
+// IDisplayClusterClusterSyncObject
+//////////////////////////////////////////////////////////////////////////////////////////////
+bool UDisplayClusterSceneComponentSyncThis::IsDirty() const
+{
+	return (LastSyncLoc != GetRelativeLocation() || LastSyncRot != GetRelativeRotation() || LastSyncScale != GetRelativeScale3D());
+}
+
+void UDisplayClusterSceneComponentSyncThis::ClearDirty()
+{
+	LastSyncLoc = GetRelativeLocation();
+	LastSyncRot = GetRelativeRotation();
+	LastSyncScale = GetRelativeScale3D();
+}
+
+
+//////////////////////////////////////////////////////////////////////////////////////////////
+// UDisplayClusterSceneComponentSync
+//////////////////////////////////////////////////////////////////////////////////////////////
+FString UDisplayClusterSceneComponentSyncThis::GenerateSyncId()
+{
+	return FString::Printf(TEXT("ST_%s"), *GetOwner()->GetName());
+}
+
+FTransform UDisplayClusterSceneComponentSyncThis::GetSyncTransform() const
+{
+	return GetRelativeTransform();
+}
+
+void UDisplayClusterSceneComponentSyncThis::SetSyncTransform(const FTransform& t)
+{
+	SetRelativeTransform(t);
+}
diff --git a/Engine/Plugins/Runtime/nDisplay/Source/DisplayCluster/Private/Game/EngineClasses/Scene/DisplayClusterScreenComponent.cpp b/Engine/Plugins/Runtime/nDisplay/Source/DisplayCluster/Private/Game/EngineClasses/Scene/DisplayClusterScreenComponent.cpp
new file mode 100644
--- /dev/null
+++ b/Engine/Plugins/Runtime/nDisplay/Source/DisplayCluster/Private/Game/EngineClasses/Scene/DisplayClusterScreenComponent.cpp
@@ -0,0 +1,86 @@
+// Copyright Epic Games, Inc. All Rights Reserved.
+
+#include "DisplayClusterScreenComponent.h"
+#include "Components/StaticMeshComponent.h"
+#include "Engine/GameEngine.h"
+#include "Engine/StaticMesh.h"
+#include "Materials/MaterialInterface.h"
+#include "Materials/Material.h"
+#include "UObject/ConstructorHelpers.h"
+
+#include "DisplayClusterRootActor.h"
+
+#include "Game/IPDisplayClusterGameManager.h"
+#include "DisplayClusterGlobals.h"
+#include "EngineDefines.h"
+
+
+UDisplayClusterScreenComponent::UDisplayClusterScreenComponent(const FObjectInitializer& ObjectInitializer) :
+	UDisplayClusterSceneComponent(ObjectInitializer)
+{
+	PrimaryComponentTick.bCanEverTick = true;
+
+	ScreenGeometryComponent = CreateDefaultSubobject<UStaticMeshComponent>(FName(*(GetName() + FString("_impl"))));
+	check(ScreenGeometryComponent);
+
+	if (ScreenGeometryComponent)
+	{
+		static ConstructorHelpers::FObjectFinder<UStaticMesh> ScreenMesh(TEXT("StaticMesh'/Engine/BasicShapes/Cube.Cube'"));
+
+		ScreenGeometryComponent->AttachToComponent(this, FAttachmentTransformRules(EAttachmentRule::KeepRelative, false));
+		ScreenGeometryComponent->SetRelativeLocationAndRotation(FVector::ZeroVector, FRotator::ZeroRotator, false);
+		ScreenGeometryComponent->SetStaticMesh(ScreenMesh.Object);
+		ScreenGeometryComponent->SetMobility(EComponentMobility::Movable);
+		ScreenGeometryComponent->SetCollisionEnabled(ECollisionEnabled::NoCollision);
+		ScreenGeometryComponent->SetVisibility(false);
+	}
+}
+
+
+void UDisplayClusterScreenComponent::BeginPlay()
+{
+	Super::BeginPlay();
+}
+
+
+void UDisplayClusterScreenComponent::TickComponent( float DeltaTime, ELevelTick TickType, FActorComponentTickFunction* ThisTickFunction )
+{
+	static const IPDisplayClusterGameManager* const GameMgr = GDisplayCluster->GetPrivateGameMgr();
+	if (GameMgr)
+	{
+		ADisplayClusterRootActor* const RootActor = GameMgr->GetRootActor();
+		if (RootActor)
+		{
+			if (RootActor->GetProjectionScreenMaterial())
+			{
+				ScreenGeometryComponent->SetMaterial(0, RootActor->GetProjectionScreenMaterial());
+			}
+
+			const bool bScreenVisible = RootActor->GetShowProjectionScreens();
+			ScreenGeometryComponent->SetVisibility(bScreenVisible);
+		}
+	}
+
+	Super::TickComponent(DeltaTime, TickType, ThisTickFunction);
+}
+
+void UDisplayClusterScreenComponent::SetSettings(const FDisplayClusterConfigSceneNode* pConfig)
+{
+	const FDisplayClusterConfigScreen* pScreenCfg = static_cast<const FDisplayClusterConfigScreen*>(pConfig);
+	Size = pScreenCfg->Size;
+
+	Super::SetSettings(pConfig);
+}
+
+bool UDisplayClusterScreenComponent::ApplySettings()
+{
+	Super::ApplySettings();
+
+	SetRelativeScale3D(FVector(0.005f, Size.X, Size.Y));
+	if (ScreenGeometryComponent)
+	{
+		ScreenGeometryComponent->RegisterComponent();
+	}
+
+	return true;
+}
