- Fixed cluster event packing on the slave side
- Fixed crash when cluster events are sent to 2nd PIE session
- Added some cleaning on EndScene and EndSession

diff --git a/Engine/Plugins/Runtime/nDisplay/Source/DisplayCluster/Private/Cluster/DisplayClusterClusterManager.cpp b/Engine/Plugins/Runtime/nDisplay/Source/DisplayCluster/Private/Cluster/DisplayClusterClusterManager.cpp
--- a/Engine/Plugins/Runtime/nDisplay/Source/DisplayCluster/Private/Cluster/DisplayClusterClusterManager.cpp
+++ b/Engine/Plugins/Runtime/nDisplay/Source/DisplayCluster/Private/Cluster/DisplayClusterClusterManager.cpp
@@ -165,6 +165,10 @@ void FDisplayClusterClusterManager::EndSession()
 			Controller->Release();
 			Controller.Reset();
 		}
+
+		NodesAmount = 0;
+		ConfigPath.Empty();
+		ClusterNodeId.Empty();
 	}
 }
 
@@ -189,6 +193,16 @@ void FDisplayClusterClusterManager::EndScene()
 			SyncGroupPair.Value.Reset();
 		}
 	}
+
+	{
+		FScopeLock lock(&ClusterEventsCritSec);
+		ClusterEventListeners.Reset(ClusterEventListeners.Num() | 0x7);
+		ClusterEventsPoolMain.Reset();
+		ClusterEventsPoolOut.Reset();
+	}
+
+	NativeInputDataCache.Reset();
+	CurrentWorld = nullptr;
 }
 
 void FDisplayClusterClusterManager::EndFrame(uint64 FrameNum)
diff --git a/Engine/Plugins/Runtime/nDisplay/Source/DisplayCluster/Private/Config/DisplayClusterConfigManager.cpp b/Engine/Plugins/Runtime/nDisplay/Source/DisplayCluster/Private/Config/DisplayClusterConfigManager.cpp
--- a/Engine/Plugins/Runtime/nDisplay/Source/DisplayCluster/Private/Config/DisplayClusterConfigManager.cpp
+++ b/Engine/Plugins/Runtime/nDisplay/Source/DisplayCluster/Private/Config/DisplayClusterConfigManager.cpp
@@ -62,6 +62,9 @@ void FDisplayClusterConfigManager::EndSession()
 {
 	DISPLAY_CLUSTER_FUNC_TRACE(LogDisplayClusterConfig);
 
+	ConfigPath.Empty();
+	ClusterNodeId.Empty();
+
 	ResetConfigData();
 }
 
@@ -493,16 +496,19 @@ void FDisplayClusterConfigManager::ResetConfigData()
 	CfgWindows.Reset();
 	CfgScreens.Reset();
 	CfgViewports.Reset();
+	CfgPostprocess.Reset();
 	CfgCameras.Reset();
 	CfgSceneNodes.Reset();
 	CfgInputDevices.Reset();
 	CfgInputSetupRecords.Reset();
+	CfgProjections.Reset();
 
 	CfgInfo    = FDisplayClusterConfigInfo();
 	CfgGeneral = FDisplayClusterConfigGeneral();
 	CfgStereo  = FDisplayClusterConfigStereo();
-	CfgNetwork = FDisplayClusterConfigNetwork();
 	CfgRender  = FDisplayClusterConfigRender();
+	CfgNvidia  = FDisplayClusterConfigNvidia();
+	CfgNetwork = FDisplayClusterConfigNetwork();
 	CfgDebug   = FDisplayClusterConfigDebug();
 	CfgCustom  = FDisplayClusterConfigCustom();
 }
diff --git a/Engine/Plugins/Runtime/nDisplay/Source/DisplayCluster/Private/Game/DisplayClusterGameManager.cpp b/Engine/Plugins/Runtime/nDisplay/Source/DisplayCluster/Private/Game/DisplayClusterGameManager.cpp
--- a/Engine/Plugins/Runtime/nDisplay/Source/DisplayCluster/Private/Game/DisplayClusterGameManager.cpp
+++ b/Engine/Plugins/Runtime/nDisplay/Source/DisplayCluster/Private/Game/DisplayClusterGameManager.cpp
@@ -65,6 +65,9 @@ bool FDisplayClusterGameManager::StartSession(const FString& configPath, const F
 void FDisplayClusterGameManager::EndSession()
 {
 	DISPLAY_CLUSTER_FUNC_TRACE(LogDisplayClusterGame);
+
+	ConfigPath.Reset();
+	ClusterNodeId.Reset();
 }
 
 bool FDisplayClusterGameManager::StartScene(UWorld* InWorld)
@@ -119,6 +122,7 @@ void FDisplayClusterGameManager::EndScene()
 	FScopeLock lock(&InternalsSyncScope);
 
 	DisplayClusterRootActor = nullptr;
+	CurrentWorld = nullptr;
 }
 
 
diff --git a/Engine/Plugins/Runtime/nDisplay/Source/DisplayCluster/Private/Input/DisplayClusterInputManager.cpp b/Engine/Plugins/Runtime/nDisplay/Source/DisplayCluster/Private/Input/DisplayClusterInputManager.cpp
--- a/Engine/Plugins/Runtime/nDisplay/Source/DisplayCluster/Private/Input/DisplayClusterInputManager.cpp
+++ b/Engine/Plugins/Runtime/nDisplay/Source/DisplayCluster/Private/Input/DisplayClusterInputManager.cpp
@@ -61,6 +61,9 @@ void FDisplayClusterInputManager::EndSession()
 {
 	DISPLAY_CLUSTER_FUNC_TRACE(LogDisplayClusterInput);
 
+	ConfigPath.Reset();
+	ClusterNodeId.Reset();
+
 	ReleaseDevices();
 }
 
@@ -77,6 +80,8 @@ bool FDisplayClusterInputManager::StartScene(UWorld* pWorld)
 void FDisplayClusterInputManager::EndScene()
 {
 	DISPLAY_CLUSTER_FUNC_TRACE(LogDisplayClusterInput);
+
+	CurrentWorld = nullptr;
 }
 
 void FDisplayClusterInputManager::StartFrame(uint64 FrameNum)
diff --git a/Engine/Plugins/Runtime/nDisplay/Source/DisplayCluster/Private/Network/Service/ClusterEvents/DisplayClusterClusterEventsClient.cpp b/Engine/Plugins/Runtime/nDisplay/Source/DisplayCluster/Private/Network/Service/ClusterEvents/DisplayClusterClusterEventsClient.cpp
--- a/Engine/Plugins/Runtime/nDisplay/Source/DisplayCluster/Private/Network/Service/ClusterEvents/DisplayClusterClusterEventsClient.cpp
+++ b/Engine/Plugins/Runtime/nDisplay/Source/DisplayCluster/Private/Network/Service/ClusterEvents/DisplayClusterClusterEventsClient.cpp
@@ -27,15 +27,22 @@ void FDisplayClusterClusterEventsClient::EmitClusterEvent(const FDisplayClusterC
 {
 	TSharedPtr<FJsonObject> Request = MakeShareable(new FJsonObject);
 
+	// Main parameters
 	Request->SetStringField(FString(FDisplayClusterClusterEventsMsg::ArgName), Event.Name);
 	Request->SetStringField(FString(FDisplayClusterClusterEventsMsg::ArgType), Event.Type);
 	Request->SetStringField(FString(FDisplayClusterClusterEventsMsg::ArgCategory), Event.Category);
 	
+	// Prapare custom parameters object
+	TSharedPtr<FJsonObject> CustomParams = MakeShareable(new FJsonObject);
 	for (const auto& Parameter : Event.Parameters)
 	{
-		Request->SetStringField(Parameter.Key, Parameter.Value);
+		CustomParams->SetStringField(Parameter.Key, Parameter.Value);
 	}
 
+	// Add custom parameter objects to the main json
+	Request->SetObjectField(FString(FDisplayClusterClusterEventsMsg::ArgParameters), CustomParams);
+
+	// Send event
 	const bool Result = SendJson(Request);
 	if (!Result)
 	{
diff --git a/Engine/Plugins/Runtime/nDisplay/Source/DisplayCluster/Private/Render/DisplayClusterRenderManager.cpp b/Engine/Plugins/Runtime/nDisplay/Source/DisplayCluster/Private/Render/DisplayClusterRenderManager.cpp
--- a/Engine/Plugins/Runtime/nDisplay/Source/DisplayCluster/Private/Render/DisplayClusterRenderManager.cpp
+++ b/Engine/Plugins/Runtime/nDisplay/Source/DisplayCluster/Private/Render/DisplayClusterRenderManager.cpp
@@ -132,6 +132,17 @@ void FDisplayClusterRenderManager::EndSession()
 		RenderDevicePtr = nullptr;
 	}
 #endif
+
+	ConfigPath.Reset();
+	ClusterNodeId.Reset();
+	
+	bWindowAdjusted = false;
+
+	RenderDeviceFactories.Reset();
+	SyncPolicyFactories.Reset();
+	SyncPolicy.Reset();
+	ProjectionPolicyFactories.Reset();
+	PostProcessOperations.Reset();
 }
 
 bool FDisplayClusterRenderManager::StartScene(UWorld* InWorld)
