Fixed issue where DisplayClusterRootActor could not be found if it's located in a sublevel

diff --git a/Engine/Plugins/Runtime/nDisplay/Source/DisplayClusterEditor/Private/DisplayClusterEditorEngine.cpp b/Engine/Plugins/Runtime/nDisplay/Source/DisplayClusterEditor/Private/DisplayClusterEditorEngine.cpp
--- a/Engine/Plugins/Runtime/nDisplay/Source/DisplayClusterEditor/Private/DisplayClusterEditorEngine.cpp
+++ b/Engine/Plugins/Runtime/nDisplay/Source/DisplayClusterEditor/Private/DisplayClusterEditorEngine.cpp
@@ -79,10 +79,10 @@ void UDisplayClusterEditorEngine::StartPlayInEditorSession(FRequestPlaySessionPa
 	{
 		// Find nDisplay root actor
 		ADisplayClusterRootActor* RootActor = FindDisplayClusterRootActor(EditorWorldPreDup);
-		if (!RootActor && EditorWorld)
+		if (!RootActor && EditorWorldPreDup)
 		{
 			// Also search inside streamed levels
-			const TArray<ULevelStreaming*>& StreamingLevels = EditorWorld->GetStreamingLevels();
+			const TArray<ULevelStreaming*>& StreamingLevels = EditorWorldPreDup->GetStreamingLevels();
 			for (ULevelStreaming* StreamingLevel : StreamingLevels)
 			{
 				if (StreamingLevel)
