Fixed VRPN analog data serialization/deserialization

diff --git a/Engine/Plugins/Runtime/nDisplay/Source/DisplayCluster/Private/Input/Devices/VRPN/Analog/DisplayClusterVrpnAnalogInputDataHolder.cpp b/Engine/Plugins/Runtime/nDisplay/Source/DisplayCluster/Private/Input/Devices/VRPN/Analog/DisplayClusterVrpnAnalogInputDataHolder.cpp
--- a/Engine/Plugins/Runtime/nDisplay/Source/DisplayCluster/Private/Input/Devices/VRPN/Analog/DisplayClusterVrpnAnalogInputDataHolder.cpp
+++ b/Engine/Plugins/Runtime/nDisplay/Source/DisplayCluster/Private/Input/Devices/VRPN/Analog/DisplayClusterVrpnAnalogInputDataHolder.cpp
@@ -34,7 +34,7 @@ FString FDisplayClusterVrpnAnalogInputDataHolder::SerializeToString() const
 
 	for (auto it = DeviceData.CreateConstIterator(); it; ++it)
 	{
-		result += FString::Printf(TEXT("%d%s%f%s"), it->Key, SerializationDelimiter, *FDisplayClusterTypesConverter::template ToHexString(it->Value.axisValue), SerializationDelimiter);
+		result += FString::Printf(TEXT("%d%s%s%s"), it->Key, SerializationDelimiter, *FDisplayClusterTypesConverter::template ToHexString(it->Value.axisValue), SerializationDelimiter);
 	}
 
 	return result;
@@ -54,7 +54,7 @@ bool FDisplayClusterVrpnAnalogInputDataHolder::DeserializeFromString(const FStri
 	for (int i = 0; i < parsed.Num(); i += SerializationItems)
 	{
 		const int   ch  = FCString::Atoi(*parsed[i]);
-		const float val = FDisplayClusterTypesConverter::template FromString<float>(*parsed[i + 1]);
+		const float val = FDisplayClusterTypesConverter::template FromHexString<float>(*parsed[i + 1]);
 		DeviceData.Add(ch, FDisplayClusterVrpnAnalogChannelData{ val });
 	}
 
