1. Fixed crash when nDisplay is used with instanced stereo enabled
UDN: https://udn.unrealengine.com/questions/558405/ndisplay-with-instanced-stereo-causes-a-crash-and.html
2. RHI.MaximumFrameLatency=1
3. UDisplayClusterViewportClient::Draw optimization

diff --git a/Engine/Plugins/Runtime/nDisplay/Source/DisplayCluster/Private/Game/EngineClasses/Basics/DisplayClusterViewportClient.cpp b/Engine/Plugins/Runtime/nDisplay/Source/DisplayCluster/Private/Game/EngineClasses/Basics/DisplayClusterViewportClient.cpp
--- a/Engine/Plugins/Runtime/nDisplay/Source/DisplayCluster/Private/Game/EngineClasses/Basics/DisplayClusterViewportClient.cpp
+++ b/Engine/Plugins/Runtime/nDisplay/Source/DisplayCluster/Private/Game/EngineClasses/Basics/DisplayClusterViewportClient.cpp
@@ -39,8 +39,6 @@
 #include "DisplayClusterGlobals.h"
 
 
-//DECLARE_CYCLE_STAT(TEXT("UI Drawing Time"), STAT_UIDrawingTime, STATGROUP_UI);
-
 UDisplayClusterViewportClient::UDisplayClusterViewportClient(FVTableHelper& Helper) : Super(Helper)
 {
 	
@@ -75,16 +73,29 @@ static UCanvas* GetCanvasByName(FName CanvasName)
 
 void UDisplayClusterViewportClient::Init(struct FWorldContext& WorldContext, UGameInstance* OwningGameInstance, bool bCreateNewAudioDevice)
 {
-	IConsoleVariable* const ForceLoadCVar = IConsoleManager::Get().FindConsoleVariable(TEXT("r.CompositionForceRenderTargetLoad"));
-	if (ForceLoadCVar)
+	const bool bIsNDisplayClusterMode = (!GEngine->XRSystem.IsValid() && GEngine->StereoRenderingDevice.IsValid() && GDisplayCluster->GetOperationMode() == EDisplayClusterOperationMode::Cluster);
+	if (bIsNDisplayClusterMode)
 	{
-		ForceLoadCVar->Set(int32(1));
-	}
+		// r.CompositionForceRenderTargetLoad
+		IConsoleVariable* const ForceLoadCVar = IConsoleManager::Get().FindConsoleVariable(TEXT("r.CompositionForceRenderTargetLoad"));
+		if (ForceLoadCVar)
+		{
+			ForceLoadCVar->Set(int32(1));
+		}
 
-	IConsoleVariable* const RTResizeCVar = IConsoleManager::Get().FindConsoleVariable(TEXT("r.SceneRenderTargetResizeMethod"));
-	if (RTResizeCVar)
-	{
-		RTResizeCVar->Set(int32(2));
+		// r.SceneRenderTargetResizeMethod
+		IConsoleVariable* const RTResizeCVar = IConsoleManager::Get().FindConsoleVariable(TEXT("r.SceneRenderTargetResizeMethod"));
+		if (RTResizeCVar)
+		{
+			RTResizeCVar->Set(int32(2));
+		}
+
+		// RHI.MaximumFrameLatency
+		IConsoleVariable* const MaximumFrameLatencyCVar = IConsoleManager::Get().FindConsoleVariable(TEXT("RHI.MaximumFrameLatency"));
+		if (MaximumFrameLatencyCVar)
+		{
+			MaximumFrameLatencyCVar->Set(int32(1));
+		}
 	}
 
 	Super::Init(WorldContext, OwningGameInstance, bCreateNewAudioDevice);
@@ -92,11 +103,30 @@ void UDisplayClusterViewportClient::Init(struct FWorldContext& WorldContext, UGa
 
 void UDisplayClusterViewportClient::Draw(FViewport* InViewport, FCanvas* SceneCanvas)
 {
-	//Valid SceneCanvas is required.  Make this explicit.
+	////////////////////////////////
+	// For any operation mode other than 'Cluster' we use default UGameViewportClient::Draw pipeline
+	const bool bIsNDisplayClusterMode = (!GEngine->XRSystem.IsValid() && GEngine->StereoRenderingDevice.IsValid() && GDisplayCluster->GetOperationMode() == EDisplayClusterOperationMode::Cluster);
+	if (!bIsNDisplayClusterMode)
+	{
+		return UGameViewportClient::Draw(InViewport, SceneCanvas);
+	}
+
+	// Get nDisplay stereo device
+	IDisplayClusterRenderDevice* const DCRenderDevice = static_cast<IDisplayClusterRenderDevice* const>(GEngine->StereoRenderingDevice.Get());
+	if (!DCRenderDevice)
+	{
+		return UGameViewportClient::Draw(InViewport, SceneCanvas);
+	}
+
+	////////////////////////////////
+	// Otherwise we use our own version of the UGameViewportClient::Draw which is basically
+	// a simpler version of the original one but with multiple ViewFamilies support
+
+	// Valid SceneCanvas is required.  Make this explicit.
 	check(SceneCanvas);
 
 	OnBeginDraw().Broadcast();
-	
+
 	const bool bStereoRendering = GEngine->IsStereoscopic3D(InViewport);
 	FCanvas* DebugCanvas = InViewport->GetDebugCanvas();
 
@@ -137,44 +167,19 @@ void UDisplayClusterViewportClient::Draw(FViewport* InViewport, FCanvas* SceneCa
 		ViewModeIndex = VMI_PathTracing;
 	}
 
-	const bool bEnableStereo = GEngine->IsStereoscopic3D(InViewport);
-	int32 NumFamilies = 1;
-	int32 NumViews = bStereoRendering ? GEngine->StereoRenderingDevice->GetDesiredNumberOfViews(bStereoRendering) : 1;
-	int32 StartViewIndex = 0;
-	bool nDisplay = false;
-	// check if HMD or nDisplay connected 
-	if (!GEngine->XRSystem.IsValid())
-	{
-		if (GEngine->StereoRenderingDevice.IsValid())
-		{
-			NumFamilies = GEngine->StereoRenderingDevice->GetDesiredNumberOfViews(bStereoRendering);
-			nDisplay = (GDisplayCluster->GetOperationMode() == EDisplayClusterOperationMode::Cluster);
-		}
-	}
-
-	for (int32 viewFamily = 0; viewFamily < NumFamilies; ++viewFamily)
-	{
-		float CustomBufferRatio = 1.f;
-
-		if (nDisplay)
-		{
-			StartViewIndex = viewFamily;
-			NumViews = StartViewIndex + 1;
+	const int32 NumViews = DCRenderDevice->GetDesiredNumberOfViews(bStereoRendering);
+	const int32 NumViewsPerFamily = 1;
+	const int32 NumFamilies = NumViews / NumViewsPerFamily;
 
-			IDisplayClusterRenderDevice* DisplayClusterRenderDevice = static_cast<IDisplayClusterRenderDevice*>(GEngine->StereoRenderingDevice.Get());
-			if (DisplayClusterRenderDevice)
-			{
-				DisplayClusterRenderDevice->GetBufferRatio(viewFamily, CustomBufferRatio);
-			}
-		}
+	APlayerController* PlayerController = GEngine->GetFirstLocalPlayerController(GetWorld());
+	ULocalPlayer* LocalPlayer = PlayerController->GetLocalPlayer();
 
-		// create the view family for rendering the world scene to the viewport's render target
-		FSceneViewFamilyContext ViewFamily(FSceneViewFamily::ConstructionValues(
-			InViewport,
-			MyWorld->Scene,
-			EngineShowFlags)
+	for (int32 ViewFamilyIdx = 0; ViewFamilyIdx < NumFamilies; ++ViewFamilyIdx)
+	{
+		// Create the view family for rendering the world scene to the viewport's render target
+		FSceneViewFamilyContext ViewFamily(FSceneViewFamily::ConstructionValues(InViewport, MyWorld->Scene, EngineShowFlags)
 			.SetRealtimeUpdate(true)
-			.SetAdditionalViewFamily(viewFamily > 0));
+			.SetAdditionalViewFamily(ViewFamilyIdx > 0));
 
 #if WITH_EDITOR
 		if (GIsEditor)
@@ -193,13 +198,6 @@ void UDisplayClusterViewportClient::Draw(FViewport* InViewport, FCanvas* SceneCa
 			ViewExt->SetupViewFamily(ViewFamily);
 		}
 
-		if (bStereoRendering && GEngine->XRSystem.IsValid() && GEngine->XRSystem->GetHMDDevice())
-		{
-			// Allow HMD to modify screen settings
-			GEngine->XRSystem->GetHMDDevice()->UpdateScreenSettings(Viewport);
-		}
-
-		ESplitScreenType::Type SplitScreenConfig = GetCurrentSplitscreenConfiguration();
 		ViewFamily.ViewMode = EViewModeIndex(ViewModeIndex);
 		EngineShowFlagOverride(ESFIM_Game, ViewFamily.ViewMode, ViewFamily.EngineShowFlags, false);
 
@@ -245,215 +243,159 @@ void UDisplayClusterViewportClient::Draw(FViewport* InViewport, FCanvas* SceneCa
 		FAudioDeviceHandle RetrievedAudioDevice = MyWorld->GetAudioDevice();
 		TArray<FSceneView*> Views;
 
-		for (FLocalPlayerIterator Iterator(GEngine, MyWorld); Iterator; ++Iterator)
+		for (int32 ViewIdx = 0; ViewIdx < NumViewsPerFamily; ++ViewIdx)
 		{
-			ULocalPlayer* LocalPlayer = *Iterator;
-			if (LocalPlayer)
-			{
-				APlayerController* PlayerController = LocalPlayer->PlayerController;
+			const int32 DCViewIdx = ViewFamilyIdx * NumViewsPerFamily + ViewIdx;
 
-				for (int32 viewIndex = StartViewIndex; viewIndex < NumViews; ++viewIndex)
-				{
-					// Calculate the player's view information.
-					FVector		ViewLocation;
-					FRotator	ViewRotation;
+			// Calculate the player's view information.
+			FVector		ViewLocation;
+			FRotator	ViewRotation;
 
-					EStereoscopicPass PassType = bStereoRendering ? GEngine->StereoRenderingDevice->GetViewPassForIndex(bStereoRendering, viewIndex) : eSSP_FULL;
+			EStereoscopicPass PassType = bStereoRendering ? GEngine->StereoRenderingDevice->GetViewPassForIndex(bStereoRendering, DCViewIdx) : eSSP_FULL;
 
-					FSceneView* View = LocalPlayer->CalcSceneView(&ViewFamily, ViewLocation, ViewRotation, InViewport, nullptr, PassType);
+			FSceneView* View = LocalPlayer->CalcSceneView(&ViewFamily, ViewLocation, ViewRotation, InViewport, nullptr, PassType);
 
-					if (View)
-					{
-						Views.Add(View);
+			if (View)
+			{
+				Views.Add(View);
 
-						if (View->Family->EngineShowFlags.Wireframe)
-						{
-							// Wireframe color is emissive-only, and mesh-modifying materials do not use material substitution, hence...
-							View->DiffuseOverrideParameter = FVector4(0.f, 0.f, 0.f, 0.f);
-							View->SpecularOverrideParameter = FVector4(0.f, 0.f, 0.f, 0.f);
-						}
-						else if (View->Family->EngineShowFlags.OverrideDiffuseAndSpecular)
-						{
-							View->DiffuseOverrideParameter = FVector4(GEngine->LightingOnlyBrightness.R, GEngine->LightingOnlyBrightness.G, GEngine->LightingOnlyBrightness.B, 0.0f);
-							View->SpecularOverrideParameter = FVector4(.1f, .1f, .1f, 0.0f);
-						}
-						else if (View->Family->EngineShowFlags.LightingOnlyOverride)
-						{
-							View->DiffuseOverrideParameter = FVector4(GEngine->LightingOnlyBrightness.R, GEngine->LightingOnlyBrightness.G, GEngine->LightingOnlyBrightness.B, 0.0f);
-							View->SpecularOverrideParameter = FVector4(0.f, 0.f, 0.f, 0.f);
-						}
-						else if (View->Family->EngineShowFlags.ReflectionOverride)
-						{
-							View->DiffuseOverrideParameter = FVector4(0.f, 0.f, 0.f, 0.f);
-							View->SpecularOverrideParameter = FVector4(1, 1, 1, 0.0f);
-							View->NormalOverrideParameter = FVector4(0, 0, 1, 0.0f);
-							View->RoughnessOverrideParameter = FVector2D(0.0f, 0.0f);
-						}
+				// We don't allow instanced stereo currently
+				View->bIsInstancedStereoEnabled  = false;
+				View->bShouldBindInstancedViewUB = false;
 
-						if (!View->Family->EngineShowFlags.Diffuse)
-						{
-							View->DiffuseOverrideParameter = FVector4(0.f, 0.f, 0.f, 0.f);
-						}
+				if (View->Family->EngineShowFlags.Wireframe)
+				{
+					// Wireframe color is emissive-only, and mesh-modifying materials do not use material substitution, hence...
+					View->DiffuseOverrideParameter = FVector4(0.f, 0.f, 0.f, 0.f);
+					View->SpecularOverrideParameter = FVector4(0.f, 0.f, 0.f, 0.f);
+				}
+				else if (View->Family->EngineShowFlags.OverrideDiffuseAndSpecular)
+				{
+					View->DiffuseOverrideParameter = FVector4(GEngine->LightingOnlyBrightness.R, GEngine->LightingOnlyBrightness.G, GEngine->LightingOnlyBrightness.B, 0.0f);
+					View->SpecularOverrideParameter = FVector4(.1f, .1f, .1f, 0.0f);
+				}
+				else if (View->Family->EngineShowFlags.LightingOnlyOverride)
+				{
+					View->DiffuseOverrideParameter = FVector4(GEngine->LightingOnlyBrightness.R, GEngine->LightingOnlyBrightness.G, GEngine->LightingOnlyBrightness.B, 0.0f);
+					View->SpecularOverrideParameter = FVector4(0.f, 0.f, 0.f, 0.f);
+				}
+				else if (View->Family->EngineShowFlags.ReflectionOverride)
+				{
+					View->DiffuseOverrideParameter = FVector4(0.f, 0.f, 0.f, 0.f);
+					View->SpecularOverrideParameter = FVector4(1, 1, 1, 0.0f);
+					View->NormalOverrideParameter = FVector4(0, 0, 1, 0.0f);
+					View->RoughnessOverrideParameter = FVector2D(0.0f, 0.0f);
+				}
 
-						if (!View->Family->EngineShowFlags.Specular)
-						{
-							View->SpecularOverrideParameter = FVector4(0.f, 0.f, 0.f, 0.f);
-						}
+				if (!View->Family->EngineShowFlags.Diffuse)
+				{
+					View->DiffuseOverrideParameter = FVector4(0.f, 0.f, 0.f, 0.f);
+				}
 
-						View->CurrentBufferVisualizationMode = GetCurrentBufferVisualizationMode();
+				if (!View->Family->EngineShowFlags.Specular)
+				{
+					View->SpecularOverrideParameter = FVector4(0.f, 0.f, 0.f, 0.f);
+				}
 
-						View->CameraConstrainedViewRect = View->UnscaledViewRect;
+				View->CurrentBufferVisualizationMode = GetCurrentBufferVisualizationMode();
 
+				View->CameraConstrainedViewRect = View->UnscaledViewRect;
 
 
-						// If this is the primary drawing pass, update things that depend on the view location
-						//if (viewIndex == 0)
-						{
-							// Save the location of the view.
-							LocalPlayer->LastViewLocation = ViewLocation;
+				{
+					// Save the location of the view.
+					LocalPlayer->LastViewLocation = ViewLocation;
 
-							PlayerViewMap.Add(LocalPlayer, View);
+					PlayerViewMap.Add(LocalPlayer, View);
 
-							// Update the listener.
-							if (RetrievedAudioDevice && PlayerController != NULL)
-							{
-								bool bUpdateListenerPosition = true;
+					// Update the listener.
+					if (RetrievedAudioDevice && PlayerController != NULL)
+					{
+						bool bUpdateListenerPosition = true;
 
-								// If the main audio device is used for multiple PIE viewport clients, we only
-								// want to update the main audio device listener position if it is in focus
-								if (GEngine)
-								{
-									FAudioDeviceManager* AudioDeviceManager = GEngine->GetAudioDeviceManager();
-
-									// If there is more than one world referencing the main audio device
-									if (AudioDeviceManager->GetNumMainAudioDeviceWorlds() > 1)
-									{
-										uint32 MainAudioDeviceID = GEngine->GetMainAudioDeviceID();
-										if (AudioDevice->DeviceID == MainAudioDeviceID && !HasAudioFocus())
-										{
-											bUpdateListenerPosition = false;
-										}
-									}
-								}
+						// If the main audio device is used for multiple PIE viewport clients, we only
+						// want to update the main audio device listener position if it is in focus
+						if (GEngine)
+						{
+							FAudioDeviceManager* AudioDeviceManager = GEngine->GetAudioDeviceManager();
 
-								if (bUpdateListenerPosition)
+							// If there is more than one world referencing the main audio device
+							if (AudioDeviceManager->GetNumMainAudioDeviceWorlds() > 1)
+							{
+								uint32 MainAudioDeviceID = GEngine->GetMainAudioDeviceID();
+								if (AudioDevice->DeviceID == MainAudioDeviceID && !HasAudioFocus())
 								{
-									FVector Location;
-									FVector ProjFront;
-									FVector ProjRight;
-									PlayerController->GetAudioListenerPosition(Location, ProjFront, ProjRight);
-
-									FTransform ListenerTransform(FRotationMatrix::MakeFromXY(ProjFront, ProjRight));
-
-									// Allow the HMD to adjust based on the head position of the player, as opposed to the view location
-									if (GEngine->XRSystem.IsValid() && GEngine->StereoRenderingDevice.IsValid() && GEngine->StereoRenderingDevice->IsStereoEnabled())
-									{
-										const FVector Offset = GEngine->XRSystem->GetAudioListenerOffset();
-										Location += ListenerTransform.TransformPositionNoScale(Offset);
-									}
-
-									ListenerTransform.SetTranslation(Location);
-									ListenerTransform.NormalizeRotation();
-
-									uint32 ViewportIndex = PlayerViewMap.Num() - 1;
-									RetrievedAudioDevice->SetListener(MyWorld, ViewportIndex, ListenerTransform, (View->bCameraCut ? 0.f : MyWorld->GetDeltaSeconds()));
-
-									FVector OverrideAttenuation;
-									if (PlayerController->GetAudioListenerAttenuationOverridePosition(OverrideAttenuation))
-									{
-										RetrievedAudioDevice->SetListenerAttenuationOverride(ViewportIndex, OverrideAttenuation);
-									}
-									else
-									{
-										RetrievedAudioDevice->ClearListenerAttenuationOverride(ViewportIndex);
-									}
+									bUpdateListenerPosition = false;
 								}
 							}
-
-#if RHI_RAYTRACING
-							View->SetupRayTracedRendering();
-#endif
 						}
 
-						// Add view information for resource streaming. Allow up to 5X boost for small FOV.
-						const float StreamingScale = 1.f / FMath::Clamp<float>(View->LODDistanceFactor, .2f, 1.f);
-						IStreamingManager::Get().AddViewInformation(View->ViewMatrices.GetViewOrigin(), View->UnscaledViewRect.Width(), View->UnscaledViewRect.Width() * View->ViewMatrices.GetProjectionMatrix().M[0][0], StreamingScale);
-						MyWorld->ViewLocationsRenderedLastFrame.Add(View->ViewMatrices.GetViewOrigin());
-					}
-				}
-			}
-		}
+						if (bUpdateListenerPosition)
+						{
+							FVector Location;
+							FVector ProjFront;
+							FVector ProjRight;
+							PlayerController->GetAudioListenerPosition(Location, ProjFront, ProjRight);
 
-		FinalizeViews(&ViewFamily, PlayerViewMap);
+							FTransform ListenerTransform(FRotationMatrix::MakeFromXY(ProjFront, ProjRight));
 
-		// Update level streaming.
-		MyWorld->UpdateLevelStreaming();
+							// Allow the HMD to adjust based on the head position of the player, as opposed to the view location
+							if (GEngine->XRSystem.IsValid() && GEngine->StereoRenderingDevice.IsValid() && GEngine->StereoRenderingDevice->IsStereoEnabled())
+							{
+								const FVector Offset = GEngine->XRSystem->GetAudioListenerOffset();
+								Location += ListenerTransform.TransformPositionNoScale(Offset);
+							}
 
-		// Find largest rectangle bounded by all rendered views.
-		uint32 MinX = InViewport->GetSizeXY().X, MinY = InViewport->GetSizeXY().Y, MaxX = 0, MaxY = 0;
-		uint32 TotalArea = 0;
-		{
-			for (int32 ViewIndex = 0; ViewIndex < ViewFamily.Views.Num(); ++ViewIndex)
-			{
-				const FSceneView* View = ViewFamily.Views[ViewIndex];
+							ListenerTransform.SetTranslation(Location);
+							ListenerTransform.NormalizeRotation();
 
-				FIntRect UpscaledViewRect = View->UnscaledViewRect;
+							uint32 ViewportIndex = PlayerViewMap.Num() - 1;
+							RetrievedAudioDevice->SetListener(MyWorld, ViewportIndex, ListenerTransform, (View->bCameraCut ? 0.f : MyWorld->GetDeltaSeconds()));
 
-				MinX = FMath::Min<uint32>(UpscaledViewRect.Min.X, MinX);
-				MinY = FMath::Min<uint32>(UpscaledViewRect.Min.Y, MinY);
-				MaxX = FMath::Max<uint32>(UpscaledViewRect.Max.X, MaxX);
-				MaxY = FMath::Max<uint32>(UpscaledViewRect.Max.Y, MaxY);
-				TotalArea += FMath::TruncToInt(UpscaledViewRect.Width()) * FMath::TruncToInt(UpscaledViewRect.Height());
-			}
+							FVector OverrideAttenuation;
+							if (PlayerController->GetAudioListenerAttenuationOverridePosition(OverrideAttenuation))
+							{
+								RetrievedAudioDevice->SetListenerAttenuationOverride(ViewportIndex, OverrideAttenuation);
+							}
+							else
+							{
+								RetrievedAudioDevice->ClearListenerAttenuationOverride(ViewportIndex);
+							}
+						}
+					}
 
-			// To draw black borders around the rendered image (prevents artifacts from post processing passes that read outside of the image e.g. PostProcessAA)
-			{
-				int32 BlackBordersCVar = IConsoleManager::Get().FindConsoleVariable(TEXT("r.BlackBorders"), false)->GetInt();
-				int32 BlackBorders = FMath::Clamp(BlackBordersCVar, 0, 10);
+#if RHI_RAYTRACING
+					View->SetupRayTracedRendering();
+#endif
 
-				if (ViewFamily.Views.Num() == 1 && BlackBorders)
-				{
-					MinX += BlackBorders;
-					MinY += BlackBorders;
-					MaxX -= BlackBorders;
-					MaxY -= BlackBorders;
-					TotalArea = (MaxX - MinX) * (MaxY - MinY);
+#if CSV_PROFILER
+					UpdateCsvCameraStats(View);
+#endif
 				}
-			}
-		}
 
-		// If the views don't cover the entire bounding rectangle, clear the entire buffer.
-		bool bBufferCleared = false;
-		bool bStereoscopicPass = (ViewFamily.Views.Num() != 0 && ViewFamily.Views[0]->StereoPass != eSSP_FULL);
-		if (ViewFamily.Views.Num() == 0 || TotalArea != (MaxX - MinX)*(MaxY - MinY) || bDisableWorldRendering || bStereoscopicPass)
-		{
-			if (bDisableWorldRendering || !bStereoscopicPass) // TotalArea computation does not work correctly for stereoscopic views
-			{
-				//SceneCanvas->Clear(FLinearColor::Transparent);
+				// Add view information for resource streaming. Allow up to 5X boost for small FOV.
+				const float StreamingScale = 1.f / FMath::Clamp<float>(View->LODDistanceFactor, .2f, 1.f);
+				IStreamingManager::Get().AddViewInformation(View->ViewMatrices.GetViewOrigin(), View->UnscaledViewRect.Width(), View->UnscaledViewRect.Width() * View->ViewMatrices.GetProjectionMatrix().M[0][0], StreamingScale);
+				MyWorld->ViewLocationsRenderedLastFrame.Add(View->ViewMatrices.GetViewOrigin());
 			}
-
-			bBufferCleared = true;
 		}
 
+		FinalizeViews(&ViewFamily, PlayerViewMap);
+
 		// Force screen percentage show flag to be turned off if not supported.
 		if (!ViewFamily.SupportsScreenPercentage())
 		{
 			ViewFamily.EngineShowFlags.ScreenPercentage = false;
 		}
-		// Update level streaming.
-		MyWorld->UpdateLevelStreaming();
-
 
 		// Set up secondary resolution fraction for the view family.
 		if (!bStereoRendering && ViewFamily.SupportsScreenPercentage())
 		{
-			float BlackBordersCVar = IConsoleManager::Get().FindConsoleVariable(TEXT("r.SecondaryScreenPercentage.GameViewport"), false)->GetFloat();
-			float CustomSecondaruScreenPercentage = BlackBordersCVar;
-
-			if (CustomSecondaruScreenPercentage > 0.0)
+			float CustomSecondaryScreenPercentage = IConsoleManager::Get().FindConsoleVariable(TEXT("r.SecondaryScreenPercentage.GameViewport"), false)->GetFloat();
+			if (CustomSecondaryScreenPercentage > 0.0)
 			{
 				// Override secondary resolution fraction with CVar.
-				ViewFamily.SecondaryViewFraction = FMath::Min(CustomSecondaruScreenPercentage / 100.0f, 1.0f);
+				ViewFamily.SecondaryViewFraction = FMath::Min(CustomSecondaryScreenPercentage / 100.0f, 1.0f);
 			}
 			else
 			{
@@ -503,6 +445,10 @@ void UDisplayClusterViewportClient::Draw(FViewport* InViewport, FCanvas* SceneCa
 		// If a screen percentage interface was not set by dynamic resolution, then create one matching legacy behavior.
 		if (ViewFamily.GetScreenPercentageInterface() == nullptr)
 		{
+			// In case of stereo, we set the same buffer ratio for both left and right views (taken from left)
+			float CustomBufferRatio = 1.f;
+			DCRenderDevice->GetBufferRatio(ViewFamilyIdx * NumViewsPerFamily, CustomBufferRatio);
+
 			bool AllowPostProcessSettingsScreenPercentage = false;
 			float GlobalResolutionFraction = 1.0f;
 
@@ -530,6 +476,8 @@ void UDisplayClusterViewportClient::Draw(FViewport* InViewport, FCanvas* SceneCa
 			}
 		}
 
+		ViewFamily.bIsHDR = GetWindow().IsValid() ? GetWindow().Get()->GetIsHDR() : false;
+
 		// Draw the player views.
 		if (!bDisableWorldRendering && PlayerViewMap.Num() > 0 && FSlateApplication::Get().GetPlatformApplication()->IsAllowedToRender()) //-V560
 		{
@@ -537,6 +485,8 @@ void UDisplayClusterViewportClient::Draw(FViewport* InViewport, FCanvas* SceneCa
 		}
 		else
 		{
+			GetRendererModule().PerFrameCleanupIfSkipRenderer();
+
 			// Make sure RHI resources get flushed if we're not using a renderer
 			ENQUEUE_RENDER_COMMAND(UGameViewportClient_FlushRHIResources)(
 				[](FRHICommandListImmediate& RHICmdList)
@@ -544,87 +494,14 @@ void UDisplayClusterViewportClient::Draw(FViewport* InViewport, FCanvas* SceneCa
 				RHICmdList.ImmediateFlush(EImmediateFlushType::FlushRHIThreadFlushResources);
 			});
 		}
-
-		//SCOPE_CYCLE_COUNTER(STAT_UIDrawingTime);
-		//CSV_SCOPED_TIMING_STAT_EXCLUSIVE(UI);
-
-		// render HUD
-		bool bDisplayedSubtitles = false;
-		for (FConstPlayerControllerIterator Iterator = MyWorld->GetPlayerControllerIterator(); Iterator; ++Iterator)
-		{
-			APlayerController* PlayerController = Iterator->Get();
-			if (PlayerController)
-			{
-				ULocalPlayer* LocalPlayer = Cast<ULocalPlayer>(PlayerController->Player);
-				if (LocalPlayer)
-				{
-					FSceneView* View = PlayerViewMap.FindRef(LocalPlayer);
-					if (View != NULL)
-					{
-						// rendering to directly to viewport target
-						FVector CanvasOrigin(FMath::TruncToFloat(View->UnscaledViewRect.Min.X), FMath::TruncToInt(View->UnscaledViewRect.Min.Y), 0.f);
-
-						CanvasObject->Init(View->UnscaledViewRect.Width(), View->UnscaledViewRect.Height(), View, SceneCanvas);
-
-						// Set the canvas transform for the player's view rectangle.
-						check(SceneCanvas);
-						SceneCanvas->PushAbsoluteTransform(FTranslationMatrix(CanvasOrigin));
-						CanvasObject->ApplySafeZoneTransform();
-
-						// Render the player's HUD.
-						if (PlayerController->MyHUD)
-						{
-							//SCOPE_CYCLE_COUNTER(STAT_HudTime);
-
-							DebugCanvasObject->SceneView = View;
-							PlayerController->MyHUD->SetCanvas(CanvasObject, DebugCanvasObject);
-
-							PlayerController->MyHUD->PostRender();
-
-							// Put these pointers back as if a blueprint breakpoint hits during HUD PostRender they can
-							// have been changed
-							CanvasObject->Canvas = SceneCanvas;
-							DebugCanvasObject->Canvas = DebugCanvas;
-
-							// A side effect of PostRender is that the playercontroller could be destroyed
-							if (!PlayerController->IsPendingKill())
-							{
-								PlayerController->MyHUD->SetCanvas(NULL, NULL);
-							}
-						}
-
-						if (DebugCanvas != NULL)
-						{
-							DebugCanvas->PushAbsoluteTransform(FTranslationMatrix(CanvasOrigin));
-							UDebugDrawService::Draw(ViewFamily.EngineShowFlags, InViewport, View, DebugCanvas, DebugCanvasObject);
-							DebugCanvas->PopTransform();
-						}
-
-						CanvasObject->PopSafeZoneTransform();
-						SceneCanvas->PopTransform();
-
-						// draw subtitles
-						if (!bDisplayedSubtitles)
-						{
-							FVector2D MinPos(0.f, 0.f);
-							FVector2D MaxPos(1.f, 1.f);
-							GetSubtitleRegion(MinPos, MaxPos);
-
-							const uint32 SizeX = SceneCanvas->GetRenderTarget()->GetSizeXY().X;
-							const uint32 SizeY = SceneCanvas->GetRenderTarget()->GetSizeXY().Y;
-							FIntRect SubtitleRegion(FMath::TruncToInt(SizeX * MinPos.X), FMath::TruncToInt(SizeY * MinPos.Y), FMath::TruncToInt(SizeX * MaxPos.X), FMath::TruncToInt(SizeY * MaxPos.Y));
-							FSubtitleManager::GetSubtitleManager()->DisplaySubtitles(SceneCanvas, SubtitleRegion, MyWorld->GetAudioTimeSeconds());
-							bDisplayedSubtitles = true;
-						}
-					}
-				}
-			}
-		}
 	}
 
-	// Beyond this point, only UI rendering independent from dynamc resolution.
+	// Beyond this point, only UI rendering independent from dynamic resolution.
 	GEngine->EmitDynamicResolutionEvent(EDynamicResolutionStateEvent::EndDynamicResolutionRendering);
 
+	// Update level streaming.
+	MyWorld->UpdateLevelStreaming();
+
 	// Remove temporary debug lines.
 	if (MyWorld->LineBatcher != nullptr)
 	{
@@ -652,19 +529,10 @@ void UDisplayClusterViewportClient::Draw(FViewport* InViewport, FCanvas* SceneCa
 		PostRender(DebugCanvasObject);
 	}
 
-
 	// Grab the player camera location and orientation so we can pass that along to the stats drawing code.
 	FVector PlayerCameraLocation = FVector::ZeroVector;
 	FRotator PlayerCameraRotation = FRotator::ZeroRotator;
-	{
-		for (FConstPlayerControllerIterator Iterator = MyWorld->GetPlayerControllerIterator(); Iterator; ++Iterator)
-		{
-			if (APlayerController* PC = Iterator->Get())
-			{
-				PC->GetPlayerViewPoint(PlayerCameraLocation, PlayerCameraRotation);
-			}
-		}
-	}
+	PlayerController->GetPlayerViewPoint(PlayerCameraLocation, PlayerCameraRotation);
 
 	if (DebugCanvas)
 	{
diff --git a/Engine/Plugins/Runtime/nDisplay/Source/DisplayCluster/Private/Game/EngineClasses/Scene/DisplayClusterRootActor.cpp b/Engine/Plugins/Runtime/nDisplay/Source/DisplayCluster/Private/Game/EngineClasses/Scene/DisplayClusterRootActor.cpp
--- a/Engine/Plugins/Runtime/nDisplay/Source/DisplayCluster/Private/Game/EngineClasses/Scene/DisplayClusterRootActor.cpp
+++ b/Engine/Plugins/Runtime/nDisplay/Source/DisplayCluster/Private/Game/EngineClasses/Scene/DisplayClusterRootActor.cpp
@@ -4,8 +4,11 @@
 
 #include "Components/SceneComponent.h"
 #include "GameFramework/PlayerController.h"
+#include "Engine/Engine.h"
 #include "Engine/World.h"
 
+#include "HAL/IConsoleManager.h"
+
 #include "Config/IPDisplayClusterConfigManager.h"
 #include "Config/DisplayClusterConfigTypes.h"
 #include "Misc/DisplayClusterAppExit.h"
@@ -115,5 +118,20 @@ void ADisplayClusterRootActor::Tick(float DeltaSeconds)
 		}
 	}
 
+	// Show 'not supported' warning if instanced stereo is used
+	if (OperationMode != EDisplayClusterOperationMode::Disabled)
+	{
+		static const TConsoleVariableData<int32>* const InstancedStereoCVar = IConsoleManager::Get().FindTConsoleVariableDataInt(TEXT("vr.InstancedStereo"));
+		if (InstancedStereoCVar)
+		{
+			const bool bIsInstancedStereoRequested = (InstancedStereoCVar->GetValueOnGameThread() != 0);
+			if (bIsInstancedStereoRequested)
+			{
+				UE_LOG(LogDisplayClusterGame, Error, TEXT("Instanced stereo was requested. nDisplay doesn't support instanced stereo so far."));
+				GEngine->AddOnScreenDebugMessage(-1, 0.f, FColor::Red, TEXT("nDisplay doesn't support instanced stereo"));
+			}
+		}
+	}
+
 	Super::Tick(DeltaSeconds);
 }
diff --git a/Engine/Plugins/Runtime/nDisplay/Source/DisplayCluster/Private/Render/Device/DisplayClusterDeviceBase.h b/Engine/Plugins/Runtime/nDisplay/Source/DisplayCluster/Private/Render/Device/DisplayClusterDeviceBase.h
--- a/Engine/Plugins/Runtime/nDisplay/Source/DisplayCluster/Private/Render/Device/DisplayClusterDeviceBase.h
+++ b/Engine/Plugins/Runtime/nDisplay/Source/DisplayCluster/Private/Render/Device/DisplayClusterDeviceBase.h
@@ -17,7 +17,7 @@
 
 class IDisplayClusterPostProcess;
 class FDisplayClusterPresentationBase;
-
+class FSceneView;
 
 /**
  * Abstract render device
