Fixed DisplayClusterViewportClient.cpp build issue after merge to Main

diff --git a/Engine/Plugins/Runtime/nDisplay/Source/DisplayCluster/Private/Game/EngineClasses/Basics/DisplayClusterViewportClient.cpp b/Engine/Plugins/Runtime/nDisplay/Source/DisplayCluster/Private/Game/EngineClasses/Basics/DisplayClusterViewportClient.cpp
--- a/Engine/Plugins/Runtime/nDisplay/Source/DisplayCluster/Private/Game/EngineClasses/Basics/DisplayClusterViewportClient.cpp
+++ b/Engine/Plugins/Runtime/nDisplay/Source/DisplayCluster/Private/Game/EngineClasses/Basics/DisplayClusterViewportClient.cpp
@@ -368,9 +368,6 @@ void UDisplayClusterViewportClient::Draw(FViewport* InViewport, FCanvas* SceneCa
 					View->SetupRayTracedRendering();
 #endif
 
-#if CSV_PROFILER
-					UpdateCsvCameraStats(View);
-#endif
 				}
 
 				// Add view information for resource streaming. Allow up to 5X boost for small FOV.
@@ -380,6 +377,10 @@ void UDisplayClusterViewportClient::Draw(FViewport* InViewport, FCanvas* SceneCa
 			}
 		}
 
+#if CSV_PROFILER
+		UpdateCsvCameraStats(PlayerViewMap);
+#endif
+
 		FinalizeViews(&ViewFamily, PlayerViewMap);
 
 		// Force screen percentage show flag to be turned off if not supported.
