Use RootNode as default if parent component not found for simple projection policy

diff --git a/Engine/Plugins/Runtime/nDisplay/Source/DisplayClusterProjection/Private/Policy/Simple/DisplayClusterProjectionSimplePolicy.cpp b/Engine/Plugins/Runtime/nDisplay/Source/DisplayClusterProjection/Private/Policy/Simple/DisplayClusterProjectionSimplePolicy.cpp
--- a/Engine/Plugins/Runtime/nDisplay/Source/DisplayClusterProjection/Private/Policy/Simple/DisplayClusterProjectionSimplePolicy.cpp
+++ b/Engine/Plugins/Runtime/nDisplay/Source/DisplayClusterProjection/Private/Policy/Simple/DisplayClusterProjectionSimplePolicy.cpp
@@ -212,8 +212,8 @@ void FDisplayClusterProjectionSimplePolicy::InitializeMeshData()
 
 	if (!ParentComp)
 	{
-		UE_LOG(LogDisplayClusterProjectionSimple, Warning, TEXT("Couldn't find a parent component for the new screen component"));
-		return;
+		UE_LOG(LogDisplayClusterProjectionSimple, Warning, TEXT("Couldn't find parent component <%s>. Default root will be used."), *CfgScreen.ParentId);
+		ParentComp = Root;
 	}
 
 	// Finally, create the component
