Fixed missed argument in UE_LOG

diff --git a/Engine/Plugins/Runtime/nDisplay/Source/DisplayClusterProjection/Private/Policy/EasyBlend/DisplayClusterProjectionEasyBlendPolicyBase.cpp b/Engine/Plugins/Runtime/nDisplay/Source/DisplayClusterProjection/Private/Policy/EasyBlend/DisplayClusterProjectionEasyBlendPolicyBase.cpp
--- a/Engine/Plugins/Runtime/nDisplay/Source/DisplayClusterProjection/Private/Policy/EasyBlend/DisplayClusterProjectionEasyBlendPolicyBase.cpp
+++ b/Engine/Plugins/Runtime/nDisplay/Source/DisplayClusterProjection/Private/Policy/EasyBlend/DisplayClusterProjectionEasyBlendPolicyBase.cpp
@@ -180,7 +180,7 @@ bool FDisplayClusterProjectionEasyBlendPolicyBase::ReadConfigData(const FString&
 	// Origin node (optional)
 	if (DisplayClusterHelpers::str::ExtractValue(CfgProjection.Params, DisplayClusterStrings::cfg::data::projection::easyblend::Origin, OutOrigin))
 	{
-		UE_LOG(LogDisplayClusterProjectionEasyBlend, Log, TEXT("Found <%s> parameter for projection %s - %s"), *CfgProjection.Id, *OutOrigin);
+		UE_LOG(LogDisplayClusterProjectionEasyBlend, Log, TEXT("Found <%s> parameter for projection %s - %s"), DisplayClusterStrings::cfg::data::projection::easyblend::Origin, *CfgProjection.Id, *OutOrigin);
 	}
 	else
 	{
