CIS Fix - Nonunity build

IDisplayClusterRenderManager needed a fwd declare for  FPostProcessSettings.
Where as ADisplayClusterTestPatternsActor::ViewportPPSettings needed more than just a forward declaration, so moved 'Engine/Scene.h' include from the cpp to the h.

#rnx

diff --git a/Engine/Plugins/Runtime/nDisplay/Source/DisplayCluster/Private/Game/EngineClasses/Scene/DisplayClusterTestPatternsActor.cpp b/Engine/Plugins/Runtime/nDisplay/Source/DisplayCluster/Private/Game/EngineClasses/Scene/DisplayClusterTestPatternsActor.cpp
--- a/Engine/Plugins/Runtime/nDisplay/Source/DisplayCluster/Private/Game/EngineClasses/Scene/DisplayClusterTestPatternsActor.cpp
+++ b/Engine/Plugins/Runtime/nDisplay/Source/DisplayCluster/Private/Game/EngineClasses/Scene/DisplayClusterTestPatternsActor.cpp
@@ -5,7 +5,6 @@
 #include "Components/PostProcessComponent.h"
 #include "Components/SceneComponent.h"
 
-#include "Engine/Scene.h"
 #include "Materials/Material.h"
 #include "Materials/MaterialInstanceDynamic.h"
 
diff --git a/Engine/Plugins/Runtime/nDisplay/Source/DisplayCluster/Public/DisplayClusterTestPatternsActor.h b/Engine/Plugins/Runtime/nDisplay/Source/DisplayCluster/Public/DisplayClusterTestPatternsActor.h
--- a/Engine/Plugins/Runtime/nDisplay/Source/DisplayCluster/Public/DisplayClusterTestPatternsActor.h
+++ b/Engine/Plugins/Runtime/nDisplay/Source/DisplayCluster/Public/DisplayClusterTestPatternsActor.h
@@ -5,6 +5,8 @@
 #include "CoreMinimal.h"
 #include "GameFramework/Actor.h"
 
+#include "Engine/Scene.h"
+
 #include "Cluster/IDisplayClusterClusterManager.h"
 #include "DisplayClusterEnums.h"
 
@@ -16,7 +18,6 @@ class UMaterial;
 class UMaterialInstanceDynamic;
 class UPostProcessComponent;
 struct FDisplayClusterConfigViewport;
-struct FPostProcessSettings;
 struct FWeightedBlendable;
 
 struct FDisplayClusterClusterEvent;
diff --git a/Engine/Plugins/Runtime/nDisplay/Source/DisplayCluster/Public/Render/IDisplayClusterRenderManager.h b/Engine/Plugins/Runtime/nDisplay/Source/DisplayCluster/Public/Render/IDisplayClusterRenderManager.h
--- a/Engine/Plugins/Runtime/nDisplay/Source/DisplayCluster/Public/Render/IDisplayClusterRenderManager.h
+++ b/Engine/Plugins/Runtime/nDisplay/Source/DisplayCluster/Public/Render/IDisplayClusterRenderManager.h
@@ -10,7 +10,7 @@ class IDisplayClusterProjectionPolicyFactory;
 class IDisplayClusterRenderDeviceFactory;
 class IDisplayClusterRenderSyncPolicy;
 class IDisplayClusterRenderSyncPolicyFactory;
-
+struct FPostProcessSettings;
 
 /**
  * Public render manager interface
