Fixed issue when "Stop PIE" button is not available after an error during nDisplay initialization

diff --git a/Engine/Plugins/Runtime/nDisplay/Source/DisplayCluster/Private/Cluster/DisplayClusterClusterManager.cpp b/Engine/Plugins/Runtime/nDisplay/Source/DisplayCluster/Private/Cluster/DisplayClusterClusterManager.cpp
--- a/Engine/Plugins/Runtime/nDisplay/Source/DisplayCluster/Private/Cluster/DisplayClusterClusterManager.cpp
+++ b/Engine/Plugins/Runtime/nDisplay/Source/DisplayCluster/Private/Cluster/DisplayClusterClusterManager.cpp
@@ -516,17 +516,20 @@ void FDisplayClusterClusterManager::SyncObjects(EDisplayClusterSyncGroup SyncGro
 {
 	DISPLAY_CLUSTER_FUNC_TRACE(LogDisplayClusterCluster);
 
-	TMap<FString, FString> SyncData;
+	if (Controller)
+	{
+		TMap<FString, FString> SyncData;
 
-	UE_LOG(LogDisplayClusterCluster, Verbose, TEXT("Downloading synchronization data (objects)..."));
-	Controller->GetSyncData(SyncData, SyncGroup);
-	UE_LOG(LogDisplayClusterCluster, Verbose, TEXT("Downloading finished. Available %d records (objects)."), SyncData.Num());
+		UE_LOG(LogDisplayClusterCluster, Verbose, TEXT("Downloading synchronization data (objects)..."));
+		Controller->GetSyncData(SyncData, SyncGroup);
+		UE_LOG(LogDisplayClusterCluster, Verbose, TEXT("Downloading finished. Available %d records (objects)."), SyncData.Num());
 
-	// We don't have to import data here unless sync data provider is located on master node
-	if (IsSlave())
-	{
-		// Perform data load (objects state update)
-		ImportSyncData(SyncData, SyncGroup);
+		// We don't have to import data here unless sync data provider is located on master node
+		if (IsSlave())
+		{
+			// Perform data load (objects state update)
+			ImportSyncData(SyncData, SyncGroup);
+		}
 	}
 }
 
@@ -534,18 +537,21 @@ void FDisplayClusterClusterManager::SyncInput()
 {
 	DISPLAY_CLUSTER_FUNC_TRACE(LogDisplayClusterCluster);
 
-	TMap<FString, FString> InputData;
+	if (Controller)
+	{
+		TMap<FString, FString> InputData;
 
-	// Get input data from a provider
-	UE_LOG(LogDisplayClusterCluster, Verbose, TEXT("Downloading synchronization data (input)..."));
-	Controller->GetInputData(InputData);
-	UE_LOG(LogDisplayClusterCluster, Verbose, TEXT("Downloading finished. Available %d records (input)."), InputData.Num());
+		// Get input data from a provider
+		UE_LOG(LogDisplayClusterCluster, Verbose, TEXT("Downloading synchronization data (input)..."));
+		Controller->GetInputData(InputData);
+		UE_LOG(LogDisplayClusterCluster, Verbose, TEXT("Downloading finished. Available %d records (input)."), InputData.Num());
 
-	// We don't have to import data here unless input data provider is located on master node
-	if (IsSlave())
-	{
-		// Perform data load (objects state update)
-		GDisplayCluster->GetPrivateInputMgr()->ImportInputData(InputData);
+		// We don't have to import data here unless input data provider is located on master node
+		if (IsSlave())
+		{
+			// Perform data load (objects state update)
+			GDisplayCluster->GetPrivateInputMgr()->ImportInputData(InputData);
+		}
 	}
 }
 
@@ -553,15 +559,18 @@ void FDisplayClusterClusterManager::SyncEvents()
 {
 	DISPLAY_CLUSTER_FUNC_TRACE(LogDisplayClusterCluster);
 
-	TMap<FString, FString> EventsData;
+	if (Controller)
+	{
+		TMap<FString, FString> EventsData;
 
-	// Get events data from a provider
-	UE_LOG(LogDisplayClusterCluster, Verbose, TEXT("Downloading synchronization data (events)..."));
-	Controller->GetEventsData(EventsData);
-	UE_LOG(LogDisplayClusterCluster, Verbose, TEXT("Downloading finished. Available %d records (events)."), EventsData.Num());
+		// Get events data from a provider
+		UE_LOG(LogDisplayClusterCluster, Verbose, TEXT("Downloading synchronization data (events)..."));
+		Controller->GetEventsData(EventsData);
+		UE_LOG(LogDisplayClusterCluster, Verbose, TEXT("Downloading finished. Available %d records (events)."), EventsData.Num());
 
-	// Import and process them
-	ImportEventsData(EventsData);
+		// Import and process them
+		ImportEventsData(EventsData);
+	}
 }
 
 void FDisplayClusterClusterManager::ProvideNativeInputData(const TMap<FString, FString>& NativeInputData)
@@ -589,7 +598,10 @@ void FDisplayClusterClusterManager::SyncNativeInput(TMap<FString, FString>& Nati
 	else
 	{
 		UE_LOG(LogDisplayClusterCluster, Verbose, TEXT("Downloading native input data..."));
-		Controller->GetNativeInputData(NativeInputData);
+		if (Controller)
+		{
+			Controller->GetNativeInputData(NativeInputData);
+		}
 	}
 }
 
diff --git a/Engine/Plugins/Runtime/nDisplay/Source/DisplayClusterEditor/Private/DisplayClusterEditorEngine.cpp b/Engine/Plugins/Runtime/nDisplay/Source/DisplayClusterEditor/Private/DisplayClusterEditorEngine.cpp
--- a/Engine/Plugins/Runtime/nDisplay/Source/DisplayClusterEditor/Private/DisplayClusterEditorEngine.cpp
+++ b/Engine/Plugins/Runtime/nDisplay/Source/DisplayClusterEditor/Private/DisplayClusterEditorEngine.cpp
@@ -114,11 +114,7 @@ void UDisplayClusterEditorEngine::StartPlayInEditorSession(FRequestPlaySessionPa
 
 			if (!DisplayClusterModule->StartSession(RootActor->GetEditorConfigPath(), RootActor->GetEditorNodeId()))
 			{
-				UE_LOG(LogDisplayClusterEditorEngine, Error, TEXT("Couldn't start DisplayCluster session"));
-
-				// Couldn't start a new session
-				RequestEndPlayMap();
-				return;
+				UE_LOG(LogDisplayClusterEditorEngine, Error, TEXT("An error occurred during DisplayCluster session start"));
 			}
 		}
 	}
