Minor fix to viewport config data deserialization

diff --git a/Engine/Plugins/Runtime/nDisplay/Source/DisplayCluster/Private/Config/DisplayClusterConfigTypes.cpp b/Engine/Plugins/Runtime/nDisplay/Source/DisplayCluster/Private/Config/DisplayClusterConfigTypes.cpp
--- a/Engine/Plugins/Runtime/nDisplay/Source/DisplayCluster/Private/Config/DisplayClusterConfigTypes.cpp
+++ b/Engine/Plugins/Runtime/nDisplay/Source/DisplayCluster/Private/Config/DisplayClusterConfigTypes.cpp
@@ -98,7 +98,7 @@ bool FDisplayClusterConfigWindow::DeserializeFromString(const FString& line)
 //////////////////////////////////////////////////////////////////////////////////////////////
 FString FDisplayClusterConfigViewport::ToString() const
 {
-	return FString::Printf(TEXT("[%s + %s=%s, %s=%s, %s=%s, %s=%d, %s=%d, %s=%d, %s=%d, %s=%s, %s=%d]"),
+	return FString::Printf(TEXT("[%s + %s=%s, %s=%s, %s=%s, %s=%d, %s=%d, %s=%d, %s=%d, %s=%s, %s=%f]"),
 		*FDisplayClusterConfigBase::ToString(),
 		DisplayClusterStrings::cfg::data::Id, *Id,
 		DisplayClusterStrings::cfg::data::viewport::Projection,  *ProjectionId,
