Fix compiler error from removing bUseSound in nDisplay plugin.

diff --git a/Engine/Plugins/Runtime/nDisplay/Source/DisplayCluster/Private/Game/EngineClasses/Basics/DisplayClusterGameEngine.cpp b/Engine/Plugins/Runtime/nDisplay/Source/DisplayCluster/Private/Game/EngineClasses/Basics/DisplayClusterGameEngine.cpp
--- a/Engine/Plugins/Runtime/nDisplay/Source/DisplayCluster/Private/Game/EngineClasses/Basics/DisplayClusterGameEngine.cpp
+++ b/Engine/Plugins/Runtime/nDisplay/Source/DisplayCluster/Private/Game/EngineClasses/Basics/DisplayClusterGameEngine.cpp
@@ -130,7 +130,11 @@ bool UDisplayClusterGameEngine::InitializeInternals()
 	if (DisplayClusterHelpers::config::GetLocalClusterNode(LocalClusterNode))
 	{
 		UE_LOG(LogDisplayClusterEngine, Log, TEXT("Configuring sound enabled: %s"), *FDisplayClusterTypesConverter::template ToString(LocalClusterNode.SoundEnabled));
-		bUseSound = LocalClusterNode.SoundEnabled;
+		if (!LocalClusterNode.SoundEnabled)
+		{
+			AudioDeviceManager = nullptr;
+			FAudioDeviceManager::Shutdown();
+		}
 	}
 
 	check(ClusterMgr);
