Fixed LoadLevel crash if MPCDI projection policy used

[RN] Fixed OpenLevel crash that occurs sometime if MPCDI projection policy used

diff --git a/Engine/Plugins/Runtime/nDisplay/Source/DisplayClusterProjection/Private/Policy/DisplayClusterProjectionPolicyBase.cpp b/Engine/Plugins/Runtime/nDisplay/Source/DisplayClusterProjection/Private/Policy/DisplayClusterProjectionPolicyBase.cpp
--- a/Engine/Plugins/Runtime/nDisplay/Source/DisplayClusterProjection/Private/Policy/DisplayClusterProjectionPolicyBase.cpp
+++ b/Engine/Plugins/Runtime/nDisplay/Source/DisplayClusterProjection/Private/Policy/DisplayClusterProjectionPolicyBase.cpp
@@ -26,6 +26,9 @@ void FDisplayClusterProjectionPolicyBase::InitializeOriginComponent(const FStrin
 {
 	UE_LOG(LogDisplayClusterProjection, Log, TEXT("Looking for an origin component '%s'..."), *OriginCompId);
 
+	// Reset previous one
+	PolicyOriginComp = nullptr;
+
 	IDisplayClusterGameManager* const GameMgr = IDisplayCluster::Get().GetGameMgr();
 	if (!GameMgr)
 	{
