Fixed crash for LoadLevel in PIE

#rnx

diff --git a/Engine/Plugins/Runtime/nDisplay/Source/DisplayClusterEditor/Private/DisplayClusterEditorEngine.cpp b/Engine/Plugins/Runtime/nDisplay/Source/DisplayClusterEditor/Private/DisplayClusterEditorEngine.cpp
--- a/Engine/Plugins/Runtime/nDisplay/Source/DisplayClusterEditor/Private/DisplayClusterEditorEngine.cpp
+++ b/Engine/Plugins/Runtime/nDisplay/Source/DisplayClusterEditor/Private/DisplayClusterEditorEngine.cpp
@@ -136,6 +136,31 @@ void UDisplayClusterEditorEngine::StartPlayInEditorSession(FRequestPlaySessionPa
 	}
 }
 
+bool UDisplayClusterEditorEngine::LoadMap(FWorldContext& WorldContext, FURL URL, class UPendingNetGame* Pending, FString& Error)
+{
+	 
+	if (bIsNDisplayPIE)
+	{
+		// Finish previous scene
+		DisplayClusterModule->EndScene();
+
+		// Perform map loading
+		if (!Super::LoadMap(WorldContext, URL, Pending, Error))
+		{
+			return false;
+		}
+
+		// Start new scene
+		DisplayClusterModule->StartScene(WorldContext.World());
+	}
+	else
+	{
+		return Super::LoadMap(WorldContext, URL, Pending, Error);
+	}
+
+	return true;
+}
+
 void UDisplayClusterEditorEngine::Tick(float DeltaSeconds, bool bIdleMode)
 {
 	if (DisplayClusterModule && bIsActivePIE && bIsNDisplayPIE)
diff --git a/Engine/Plugins/Runtime/nDisplay/Source/DisplayClusterEditor/Private/DisplayClusterEditorEngine.h b/Engine/Plugins/Runtime/nDisplay/Source/DisplayClusterEditor/Private/DisplayClusterEditorEngine.h
--- a/Engine/Plugins/Runtime/nDisplay/Source/DisplayClusterEditor/Private/DisplayClusterEditorEngine.h
+++ b/Engine/Plugins/Runtime/nDisplay/Source/DisplayClusterEditor/Private/DisplayClusterEditorEngine.h
@@ -23,6 +23,7 @@ public:
 	virtual void Init(IEngineLoop* InEngineLoop) override;
 	virtual void PreExit() override;
 	virtual void StartPlayInEditorSession(FRequestPlaySessionParams& InRequestParams) override;
+	virtual bool LoadMap(FWorldContext& WorldContext, FURL URL, class UPendingNetGame* Pending, FString& Error) override;
 	virtual void Tick(float DeltaSeconds, bool bIdleMode) override;
 
 private:
