Fixed UDisplayClusterViewportClient crash when nDisplay application starts with a server address to connect

#rnx

diff --git a/Engine/Plugins/Runtime/nDisplay/Source/DisplayCluster/Private/Game/EngineClasses/Basics/DisplayClusterViewportClient.cpp b/Engine/Plugins/Runtime/nDisplay/Source/DisplayCluster/Private/Game/EngineClasses/Basics/DisplayClusterViewportClient.cpp
--- a/Engine/Plugins/Runtime/nDisplay/Source/DisplayCluster/Private/Game/EngineClasses/Basics/DisplayClusterViewportClient.cpp
+++ b/Engine/Plugins/Runtime/nDisplay/Source/DisplayCluster/Private/Game/EngineClasses/Basics/DisplayClusterViewportClient.cpp
@@ -157,8 +157,6 @@ void UDisplayClusterViewportClient::Draw(FViewport* InViewport, FCanvas* SceneCa
 		SceneCanvas->SetStereoRendering(bStereoRendering);
 	}
 
-	UWorld* MyWorld = GetWorld();
-
 	// Force path tracing view mode, and extern code set path tracer show flags
 	const bool bForcePathTracing = InViewport->GetClient()->GetEngineShowFlags()->PathTracing;
 	if (bForcePathTracing)
@@ -171,8 +169,18 @@ void UDisplayClusterViewportClient::Draw(FViewport* InViewport, FCanvas* SceneCa
 	const int32 NumViewsPerFamily = 1;
 	const int32 NumFamilies = NumViews / NumViewsPerFamily;
 
-	APlayerController* PlayerController = GEngine->GetFirstLocalPlayerController(GetWorld());
-	ULocalPlayer* LocalPlayer = PlayerController->GetLocalPlayer();
+	UWorld* const MyWorld = GetWorld();
+	APlayerController* const PlayerController = GEngine->GetFirstLocalPlayerController(GetWorld());
+	ULocalPlayer* LocalPlayer = nullptr;
+	if (PlayerController)
+	{
+		LocalPlayer = PlayerController->GetLocalPlayer();
+	}
+
+	if (!PlayerController || !LocalPlayer)
+	{
+		return Super::Draw(InViewport, SceneCanvas);
+	}
 
 	for (int32 ViewFamilyIdx = 0; ViewFamilyIdx < NumFamilies; ++ViewFamilyIdx)
 	{
