Fixed nDisplay deprecation warning

diff --git a/Engine/Plugins/Runtime/nDisplay/Source/DisplayCluster/Private/Render/DisplayClusterRenderManager.cpp b/Engine/Plugins/Runtime/nDisplay/Source/DisplayCluster/Private/Render/DisplayClusterRenderManager.cpp
--- a/Engine/Plugins/Runtime/nDisplay/Source/DisplayCluster/Private/Render/DisplayClusterRenderManager.cpp
+++ b/Engine/Plugins/Runtime/nDisplay/Source/DisplayCluster/Private/Render/DisplayClusterRenderManager.cpp
@@ -747,7 +747,7 @@ void FDisplayClusterRenderManager::ResizeWindow(int32 WinX, int32 WinY, int32 Re
 	window->ReshapeWindow(FVector2D(WinX, WinY), FVector2D(ResX, ResY));
 }
 
-void FDisplayClusterRenderManager::OnViewportCreatedHandler_SetCustomPresent()
+void FDisplayClusterRenderManager::OnViewportCreatedHandler_SetCustomPresent() const
 {
 	if (GEngine && GEngine->GameViewport)
 	{
@@ -758,7 +758,7 @@ void FDisplayClusterRenderManager::OnViewportCreatedHandler_SetCustomPresent()
 	}
 }
 
-void FDisplayClusterRenderManager::OnViewportCreatedHandler_CheckViewportClass()
+void FDisplayClusterRenderManager::OnViewportCreatedHandler_CheckViewportClass() const
 {
 	if (GEngine && GEngine->GameViewport)
 	{
@@ -770,7 +770,7 @@ void FDisplayClusterRenderManager::OnViewportCreatedHandler_CheckViewportClass()
 	}
 }
 
-void FDisplayClusterRenderManager::OnBeginDrawHandler()
+void FDisplayClusterRenderManager::OnBeginDrawHandler() const
 {
 	static bool initialized = false;
 	if (!initialized && GEngine->GameViewport->Viewport->GetViewportRHI().IsValid())
diff --git a/Engine/Plugins/Runtime/nDisplay/Source/DisplayCluster/Private/Render/DisplayClusterRenderManager.h b/Engine/Plugins/Runtime/nDisplay/Source/DisplayCluster/Private/Render/DisplayClusterRenderManager.h
--- a/Engine/Plugins/Runtime/nDisplay/Source/DisplayCluster/Private/Render/DisplayClusterRenderManager.h
+++ b/Engine/Plugins/Runtime/nDisplay/Source/DisplayCluster/Private/Render/DisplayClusterRenderManager.h
@@ -87,9 +87,9 @@ private:
 	// FDisplayClusterRenderManager
 	//////////////////////////////////////////////////////////////////////////////////////////////
 	void ResizeWindow(int32 WinX, int32 WinY, int32 ResX, int32 ResY);
-	void OnViewportCreatedHandler_SetCustomPresent();
-	void OnViewportCreatedHandler_CheckViewportClass();
-	void OnBeginDrawHandler();
+	void OnViewportCreatedHandler_SetCustomPresent() const;
+	void OnViewportCreatedHandler_CheckViewportClass() const;
+	void OnBeginDrawHandler() const;
 
 private:
 	EDisplayClusterOperationMode CurrentOperationMode;
@@ -110,7 +110,7 @@ private:
 	// Synchronization internals
 	TMap<FString, TSharedPtr<IDisplayClusterRenderSyncPolicyFactory>> SyncPolicyFactories;
 	TSharedPtr<IDisplayClusterRenderSyncPolicy> CreateRenderSyncPolicy() const;
-	TSharedPtr<IDisplayClusterRenderSyncPolicy> SyncPolicy;
+	mutable TSharedPtr<IDisplayClusterRenderSyncPolicy> SyncPolicy;
 
 private:
 	// Projection internals
