Trivial build fixes for building with clang 10 & 11.

#trivial
#rnx

diff --git a/Engine/Plugins/Runtime/nDisplay/Source/DisplayClusterConfiguration/Private/Formats/Text/DisplayClusterConfigurationTextParser.cpp b/Engine/Plugins/Runtime/nDisplay/Source/DisplayClusterConfiguration/Private/Formats/Text/DisplayClusterConfigurationTextParser.cpp
--- a/Engine/Plugins/Runtime/nDisplay/Source/DisplayClusterConfiguration/Private/Formats/Text/DisplayClusterConfigurationTextParser.cpp
+++ b/Engine/Plugins/Runtime/nDisplay/Source/DisplayClusterConfiguration/Private/Formats/Text/DisplayClusterConfigurationTextParser.cpp
@@ -454,7 +454,7 @@ UDisplayClusterConfigurationData* FDisplayClusterConfigurationTextParser::Conver
 	}
 
 	// Input bindings
-	for (const FDisplayClusterConfigurationTextInputSetup CfgInputBinding : CfgInputSetupRecords)
+	for (const FDisplayClusterConfigurationTextInputSetup& CfgInputBinding : CfgInputSetupRecords)
 	{
 		FDisplayClusterConfigurationInputBinding InputBinding;
 
diff --git a/Engine/Plugins/Runtime/nDisplay/Source/DisplayClusterConfigurator/Private/Views/Details/DisplayClusterConfiguratorDetailCustomization.cpp b/Engine/Plugins/Runtime/nDisplay/Source/DisplayClusterConfigurator/Private/Views/Details/DisplayClusterConfiguratorDetailCustomization.cpp
--- a/Engine/Plugins/Runtime/nDisplay/Source/DisplayClusterConfigurator/Private/Views/Details/DisplayClusterConfiguratorDetailCustomization.cpp
+++ b/Engine/Plugins/Runtime/nDisplay/Source/DisplayClusterConfigurator/Private/Views/Details/DisplayClusterConfiguratorDetailCustomization.cpp
@@ -331,7 +331,7 @@ void FDisplayClusterConfiguratorSceneComponentDetailCustomization::ResetTrackerI
 
 	TrackerIdOptions.Empty();
 
-	for (const TPair<FString, UDisplayClusterConfigurationInputDeviceTracker*> InputDeviceTrackerPair : ConfigurationData->Input->TrackerDevices)
+	for (const TPair<FString, UDisplayClusterConfigurationInputDeviceTracker*>& InputDeviceTrackerPair : ConfigurationData->Input->TrackerDevices)
 	{
 		if (!InputDeviceTrackerPair.Key.Equals(SceneComponent->TrackerId))
 		{
diff --git a/Engine/Source/Developer/DesktopPlatform/Private/Windows/DesktopPlatformWindows.cpp b/Engine/Source/Developer/DesktopPlatform/Private/Windows/DesktopPlatformWindows.cpp
--- a/Engine/Source/Developer/DesktopPlatform/Private/Windows/DesktopPlatformWindows.cpp
+++ b/Engine/Source/Developer/DesktopPlatform/Private/Windows/DesktopPlatformWindows.cpp
@@ -382,7 +382,7 @@ void FDesktopPlatformWindows::EnumerateEngineInstallations(TMap<FString, FString
 		}
 
 		// Remove all the keys which weren't valid
-		for(const FString InvalidKey: InvalidKeys)
+		for(const FString &InvalidKey: InvalidKeys)
 		{
 			RegDeleteValue(hKey, *InvalidKey);
 		}
diff --git a/Engine/Source/Runtime/D3D12RHI/Private/D3D12Viewport.cpp b/Engine/Source/Runtime/D3D12RHI/Private/D3D12Viewport.cpp
--- a/Engine/Source/Runtime/D3D12RHI/Private/D3D12Viewport.cpp
+++ b/Engine/Source/Runtime/D3D12RHI/Private/D3D12Viewport.cpp
@@ -629,7 +629,14 @@ bool FD3D12Viewport::PresentChecked(int32 SyncInterval)
 		// Present failed so current expected GPU index will not match anymore, so patch up expected back buffer index
 		// Warning: Present is skipped for this frame but could cause a black screen for the next frame as well
 		FScopeLock Lock(&ExpectedBackBufferIndexLock);
-		ExpectedBackBufferIndex_RenderThread = (ExpectedBackBufferIndex_RenderThread == 0) ? NumBackBuffers - 1 : ExpectedBackBufferIndex_RenderThread--;
+		if (ExpectedBackBufferIndex_RenderThread == 0)
+		{
+			ExpectedBackBufferIndex_RenderThread = NumBackBuffers - 1;
+		}
+		else
+		{
+			ExpectedBackBufferIndex_RenderThread--;
+		}
 #endif // WITH_MGPU
 		return false;
 	}
diff --git a/Engine/Source/Runtime/Slate/Public/Widgets/Input/NumericTypeInterface.h b/Engine/Source/Runtime/Slate/Public/Widgets/Input/NumericTypeInterface.h
--- a/Engine/Source/Runtime/Slate/Public/Widgets/Input/NumericTypeInterface.h
+++ b/Engine/Source/Runtime/Slate/Public/Widgets/Input/NumericTypeInterface.h
@@ -116,7 +116,9 @@ struct TDefaultNumericTypeInterface : INumericTypeInterface<NumericType>
 		TValueOrError<double, FExpressionError> Result = Parser.Evaluate(*InString, InExistingValue);
 		if (Result.IsValid())
 		{
-			long double ClampedValue = FMath::Clamp<long double>(Result.GetValue(), TNumericLimits<NumericType>::Lowest(), TNumericLimits<NumericType>::Max());
+			const double MinValue = static_cast<double>(TNumericLimits<NumericType>::Lowest());
+			const double MaxValue = static_cast<double>(TNumericLimits<NumericType>::Max());
+			double ClampedValue = FMath::Clamp<double>(Result.GetValue(), MinValue, MaxValue);
 			return NumericType(ClampedValue);
 		}
 
diff --git a/Engine/Source/Runtime/Slate/Public/Widgets/Input/NumericUnitTypeInterface.inl b/Engine/Source/Runtime/Slate/Public/Widgets/Input/NumericUnitTypeInterface.inl
--- a/Engine/Source/Runtime/Slate/Public/Widgets/Input/NumericUnitTypeInterface.inl
+++ b/Engine/Source/Runtime/Slate/Public/Widgets/Input/NumericUnitTypeInterface.inl
@@ -68,7 +68,9 @@ TOptional<NumericType> TNumericUnitTypeInterface<NumericType>::FromString(const
 			SourceUnits = FixedDisplayUnits.GetValue();
 		}
 		double ConvertedValue = FUnitConversion::Convert(NewValue.GetValue().Value, SourceUnits, UnderlyingUnits);
-		long double ClampedValue = FMath::Clamp<long double>(ConvertedValue, TNumericLimits<NumericType>::Lowest(), TNumericLimits<NumericType>::Max());
+		const double MinValue = static_cast<double>(TNumericLimits<NumericType>::Lowest());
+		const double MaxValue = static_cast<double>(TNumericLimits<NumericType>::Max());
+		double ClampedValue = FMath::Clamp<double>(ConvertedValue, MinValue, MaxValue);
 		return NumericType(ClampedValue);
 	}
 
