nDispay: fixes for non-unity build

#rnx

diff --git a/Engine/Plugins/Runtime/nDisplay/Source/DisplayCluster/Private/Render/Synchronization/DisplayClusterRenderSyncPolicyNvidiaDX11.cpp b/Engine/Plugins/Runtime/nDisplay/Source/DisplayCluster/Private/Render/Synchronization/DisplayClusterRenderSyncPolicyNvidiaDX11.cpp
--- a/Engine/Plugins/Runtime/nDisplay/Source/DisplayCluster/Private/Render/Synchronization/DisplayClusterRenderSyncPolicyNvidiaDX11.cpp
+++ b/Engine/Plugins/Runtime/nDisplay/Source/DisplayCluster/Private/Render/Synchronization/DisplayClusterRenderSyncPolicyNvidiaDX11.cpp
@@ -5,6 +5,7 @@
 #include "Engine/Engine.h"
 #include "Engine/GameViewportClient.h"
 
+#include "ShaderCore.h"
 #include "D3D11RHI/Private/Windows/D3D11RHIBasePrivate.h"
 #include "D3D11State.h"
 #include "D3D11Resources.h"
@@ -18,10 +19,8 @@ void FDisplayClusterRenderSyncPolicyNvidiaDX11::WaitForFrameCompletion()
 		FD3D11Viewport* const D3D11Viewport = static_cast<FD3D11Viewport*>(GEngine->GameViewport->Viewport->GetViewportRHI().GetReference());
 		if (D3D11Viewport)
 		{
-#if !WITH_EDITOR
 			D3D11Viewport->IssueFrameEvent();
 			D3D11Viewport->WaitForFrameEventCompletion();
-#endif
 		}
 	}
 }
diff --git a/Engine/Plugins/Runtime/nDisplay/Source/DisplayCluster/Private/Render/Synchronization/DisplayClusterRenderSyncPolicySoftwareBase.cpp b/Engine/Plugins/Runtime/nDisplay/Source/DisplayCluster/Private/Render/Synchronization/DisplayClusterRenderSyncPolicySoftwareBase.cpp
--- a/Engine/Plugins/Runtime/nDisplay/Source/DisplayCluster/Private/Render/Synchronization/DisplayClusterRenderSyncPolicySoftwareBase.cpp
+++ b/Engine/Plugins/Runtime/nDisplay/Source/DisplayCluster/Private/Render/Synchronization/DisplayClusterRenderSyncPolicySoftwareBase.cpp
@@ -361,6 +361,11 @@ void FDisplayClusterRenderSyncPolicySoftwareBase::Step_FinalizeFrameSynchronizat
 	}
 }
 
+void FDisplayClusterRenderSyncPolicySoftwareBase::DoSleep(float Seconds)
+{
+	FPlatformProcess::Sleep(Seconds);
+}
+
 double FDisplayClusterRenderSyncPolicySoftwareBase::GetVBlankTimestamp() const
 {
 	// Get timestamp right after the VBlank
diff --git a/Engine/Plugins/Runtime/nDisplay/Source/DisplayCluster/Private/Render/Synchronization/DisplayClusterRenderSyncPolicySoftwareBase.h b/Engine/Plugins/Runtime/nDisplay/Source/DisplayCluster/Private/Render/Synchronization/DisplayClusterRenderSyncPolicySoftwareBase.h
--- a/Engine/Plugins/Runtime/nDisplay/Source/DisplayCluster/Private/Render/Synchronization/DisplayClusterRenderSyncPolicySoftwareBase.h
+++ b/Engine/Plugins/Runtime/nDisplay/Source/DisplayCluster/Private/Render/Synchronization/DisplayClusterRenderSyncPolicySoftwareBase.h
@@ -34,10 +34,7 @@ protected:
 	virtual void Step_Present();
 	virtual void Step_FinalizeFrameSynchronization();
 
-	virtual void DoSleep(float Seconds)
-	{
-		FPlatformProcess::Sleep(Seconds);
-	}
+	virtual void DoSleep(float Seconds);
 
 protected:
 	double GetVBlankTimestamp() const;
diff --git a/Engine/Plugins/Runtime/nDisplay/Source/DisplayCluster/Private/Render/Synchronization/DisplayClusterRenderSyncPolicySoftwareDX11.cpp b/Engine/Plugins/Runtime/nDisplay/Source/DisplayCluster/Private/Render/Synchronization/DisplayClusterRenderSyncPolicySoftwareDX11.cpp
--- a/Engine/Plugins/Runtime/nDisplay/Source/DisplayCluster/Private/Render/Synchronization/DisplayClusterRenderSyncPolicySoftwareDX11.cpp
+++ b/Engine/Plugins/Runtime/nDisplay/Source/DisplayCluster/Private/Render/Synchronization/DisplayClusterRenderSyncPolicySoftwareDX11.cpp
@@ -5,6 +5,7 @@
 #include "Engine/Engine.h"
 #include "Engine/GameViewportClient.h"
 
+#include "ShaderCore.h"
 #include "D3D11RHI/Private/Windows/D3D11RHIBasePrivate.h"
 #include "D3D11State.h"
 #include "D3D11Resources.h"
@@ -18,10 +19,8 @@ void FDisplayClusterRenderSyncPolicySoftwareDX11::WaitForFrameCompletion()
 		FD3D11Viewport* const D3D11Viewport = static_cast<FD3D11Viewport*>(GEngine->GameViewport->Viewport->GetViewportRHI().GetReference());
 		if (D3D11Viewport)
 		{
-#if !WITH_EDITOR
 			D3D11Viewport->IssueFrameEvent();
 			D3D11Viewport->WaitForFrameEventCompletion();
-#endif
 		}
 	}
 }
