nDisplay - Fix CIS Compilation-Tests NonUnity Compile UnrealEditor Win64

diff --git a/Engine/Plugins/Runtime/nDisplay/Source/DisplayClusterConfigurator/Private/Views/OutputMapping/EdNodes/DisplayClusterConfiguratorViewportNode.h b/Engine/Plugins/Runtime/nDisplay/Source/DisplayClusterConfigurator/Private/Views/OutputMapping/EdNodes/DisplayClusterConfiguratorViewportNode.h
--- a/Engine/Plugins/Runtime/nDisplay/Source/DisplayClusterConfigurator/Private/Views/OutputMapping/EdNodes/DisplayClusterConfiguratorViewportNode.h
+++ b/Engine/Plugins/Runtime/nDisplay/Source/DisplayClusterConfigurator/Private/Views/OutputMapping/EdNodes/DisplayClusterConfiguratorViewportNode.h
@@ -4,6 +4,8 @@
 
 #include "Views/OutputMapping/EdNodes/DisplayClusterConfiguratorBaseNode.h"
 
+#include "Engine/Texture.h"
+
 #include "DisplayClusterConfiguratorViewportNode.generated.h"
 
 class UDisplayClusterConfigurationViewport;
