Fixing static analysis warnings by using D3D12_CPU_DESCRIPTOR_HANDLE instead of the helper type CD3DX12_CPU_DESCRIPTOR_HANDLE. (Missed a file)

diff --git a/Engine/Plugins/Runtime/nDisplay/Source/DisplayClusterProjection/Private/Policy/VIOSO/DisplayClusterProjectionVIOSOPolicy.cpp b/Engine/Plugins/Runtime/nDisplay/Source/DisplayClusterProjection/Private/Policy/VIOSO/DisplayClusterProjectionVIOSOPolicy.cpp
--- a/Engine/Plugins/Runtime/nDisplay/Source/DisplayClusterProjection/Private/Policy/VIOSO/DisplayClusterProjectionVIOSOPolicy.cpp
+++ b/Engine/Plugins/Runtime/nDisplay/Source/DisplayClusterProjection/Private/Policy/VIOSO/DisplayClusterProjectionVIOSOPolicy.cpp
@@ -446,7 +446,7 @@ bool FDisplayClusterProjectionVIOSOPolicy::FViewData::RenderVIOSO_RenderThread(F
 				FD3D12Texture2D* DestTexture2D = FD3D12DynamicRHI::ResourceCast(RenderTargetTexture);
 
 				FD3D12RenderTargetView* RTV = DestTexture2D->GetRenderTargetView(0, 0);
-				CD3DX12_CPU_DESCRIPTOR_HANDLE RTVHandle = RTV->GetView();
+				D3D12_CPU_DESCRIPTOR_HANDLE RTVHandle = RTV->GetView();
 
 				FD3D12Device* Device = DestTexture2D->GetParentDevice();
 				FD3D12CommandListHandle& hCommandList = Device->GetDefaultCommandContext().CommandListHandle;
