Fixing UE4 references.

#rnx

diff --git a/Engine/Plugins/Developer/RenderDocPlugin/Source/RenderDocPlugin/Private/RenderDocPluginModule.cpp b/Engine/Plugins/Developer/RenderDocPlugin/Source/RenderDocPlugin/Private/RenderDocPluginModule.cpp
--- a/Engine/Plugins/Developer/RenderDocPlugin/Source/RenderDocPlugin/Private/RenderDocPluginModule.cpp
+++ b/Engine/Plugins/Developer/RenderDocPlugin/Source/RenderDocPlugin/Private/RenderDocPluginModule.cpp
@@ -104,7 +104,7 @@ public:
 
 	static void BeginFrameCapture(HWND WindowHandle, FRenderDocPluginLoader::RENDERDOC_API_CONTEXT* RenderDocAPI)
 	{
-		UE4_GEmitDrawEvents_BeforeCapture = GetEmitDrawEvents();
+		UE_GEmitDrawEvents_BeforeCapture = GetEmitDrawEvents();
 		SetEmitDrawEvents(true);
 		RenderDocAPI->StartFrameCapture(GetRenderdocDevicePointer(), WindowHandle);
 	}
@@ -113,7 +113,7 @@ public:
 	{
 		FRHICommandListExecutor::GetImmediateCommandList().SubmitCommandsAndFlushGPU();
 		uint32 Result = RenderDocAPI->EndFrameCapture(GetRenderdocDevicePointer(), WindowHandle);
-		SetEmitDrawEvents(UE4_GEmitDrawEvents_BeforeCapture);
+		SetEmitDrawEvents(UE_GEmitDrawEvents_BeforeCapture);
 		return Result;
 	}
 
@@ -188,9 +188,9 @@ public:
 	}
 
 private:
-	static bool UE4_GEmitDrawEvents_BeforeCapture;
+	static bool UE_GEmitDrawEvents_BeforeCapture;
 };
-bool FRenderDocFrameCapturer::UE4_GEmitDrawEvents_BeforeCapture = false;
+bool FRenderDocFrameCapturer::UE_GEmitDrawEvents_BeforeCapture = false;
 
 class FRenderDocDummyInputDevice : public IInputDevice
 {
diff --git a/Engine/Plugins/Developer/RenderDocPlugin/Source/RenderDocPlugin/RenderDocPlugin.Build.cs b/Engine/Plugins/Developer/RenderDocPlugin/Source/RenderDocPlugin/RenderDocPlugin.Build.cs
--- a/Engine/Plugins/Developer/RenderDocPlugin/Source/RenderDocPlugin/RenderDocPlugin.Build.cs
+++ b/Engine/Plugins/Developer/RenderDocPlugin/Source/RenderDocPlugin/RenderDocPlugin.Build.cs
@@ -19,7 +19,7 @@ namespace UnrealBuildTool.Rules
 				, "Projects"
 				, "RenderCore"
 				, "InputDevice"
-				, "RHI"				// RHI module: required for accessing the UE4 flag GUsingNullRHI.
+				, "RHI"				// RHI module: required for accessing the UE flag GUsingNullRHI.
 				, "DeveloperSettings"
 			});
 
diff --git a/Engine/Plugins/Runtime/nDisplay/Source/DisplayClusterInput/Private/Controller/DisplayClusterInputControllerButton.cpp b/Engine/Plugins/Runtime/nDisplay/Source/DisplayClusterInput/Private/Controller/DisplayClusterInputControllerButton.cpp
--- a/Engine/Plugins/Runtime/nDisplay/Source/DisplayClusterInput/Private/Controller/DisplayClusterInputControllerButton.cpp
+++ b/Engine/Plugins/Runtime/nDisplay/Source/DisplayClusterInput/Private/Controller/DisplayClusterInputControllerButton.cpp
@@ -15,13 +15,13 @@
 #define LOCTEXT_NAMESPACE "DisplayClusterInput"
 
 
-// Add vrpn buttons to UE4 global name-space
+// Add vrpn buttons to UE global name-space
 void FButtonController::Initialize()
 {
 	static const FName nDisplayClusterInputCategoryName(TEXT("nDisplayButtons"));
 	EKeys::AddMenuCategoryDisplayInfo(nDisplayClusterInputCategoryName, LOCTEXT("nDisplayInputSubCateogry", "nDisplay"), TEXT("GraphEditor.KeyEvent_16x"));
 
-	// Register all names in UE4 namespace by macros. Purpose: easy add new channels and reduce code size
+	// Register all names in UE namespace by macros. Purpose: easy add new channels and reduce code size
 	for (int32 idx = 0; idx < FButtonKey::TotalCount; ++idx)
 	{
 		FText ButtonLocaleText = FText::Format(LOCTEXT("nDisplayButtonHintFmt", "nDisplay Button {0}"), idx);
diff --git a/Engine/Plugins/Runtime/nDisplay/Source/DisplayClusterInput/Private/Controller/DisplayClusterInputControllerKeyboard.cpp b/Engine/Plugins/Runtime/nDisplay/Source/DisplayClusterInput/Private/Controller/DisplayClusterInputControllerKeyboard.cpp
--- a/Engine/Plugins/Runtime/nDisplay/Source/DisplayClusterInput/Private/Controller/DisplayClusterInputControllerKeyboard.cpp
+++ b/Engine/Plugins/Runtime/nDisplay/Source/DisplayClusterInput/Private/Controller/DisplayClusterInputControllerKeyboard.cpp
@@ -265,16 +265,16 @@ void FKeyboardController::ConnectKey(FChannelBinds& KeyboardData, uint32 VrpnCha
 		KeyboardData[VrpnChannel].UnBindTarget(NDisplayTargetName);
 	}
 
-	FName UE4TargetName(*FString::Printf(TEXT("%s"), KeyName));
-	if (bReflectToUE4)
+	FName UETargetName(*FString::Printf(TEXT("%s"), KeyName));
+	if (bReflectToUE)
 	{
-		// Reflect link to UE4 keyboard
-		KeyboardData[VrpnChannel].BindTarget(UE4TargetName);
+		// Reflect link to UE keyboard
+		KeyboardData[VrpnChannel].BindTarget(UETargetName);
 	}
 	else
 	{
-		// Remove bind to UE4 keyboard  (runtime reflect purpose)
-		KeyboardData[VrpnChannel].UnBindTarget(UE4TargetName);
+		// Remove bind to UE keyboard  (runtime reflect purpose)
+		KeyboardData[VrpnChannel].UnBindTarget(UETargetName);
 	}
 }
 
@@ -292,17 +292,17 @@ void FKeyboardController::ReflectKeyboard(const FString& VrpnDeviceId, EDisplayC
 		break;
 
 	case EDisplayClusterInputKeyboardReflectionMode::Core:
-		bReflectToUE4 = true;
+		bReflectToUE = true;
 		break;
 
 	case EDisplayClusterInputKeyboardReflectionMode::All:
 		bReflectToNDisplayCluster = true;
-		bReflectToUE4 = true;
+		bReflectToUE = true;
 		break;
 
 	case EDisplayClusterInputKeyboardReflectionMode::None:
 	default:
-		bReflectToUE4 = false;
+		bReflectToUE = false;
 		bReflectToNDisplayCluster = false;
 	}
 
diff --git a/Engine/Plugins/Runtime/nDisplay/Source/DisplayClusterInput/Private/Controller/DisplayClusterInputControllerKeyboard.h b/Engine/Plugins/Runtime/nDisplay/Source/DisplayClusterInput/Private/Controller/DisplayClusterInputControllerKeyboard.h
--- a/Engine/Plugins/Runtime/nDisplay/Source/DisplayClusterInput/Private/Controller/DisplayClusterInputControllerKeyboard.h
+++ b/Engine/Plugins/Runtime/nDisplay/Source/DisplayClusterInput/Private/Controller/DisplayClusterInputControllerKeyboard.h
@@ -33,6 +33,6 @@ private:
 	void ConnectKey(FChannelBinds& KeyboardData, uint32 VrpnChannel, const TCHAR* KeyName);
 
 	//  Run-time flags for init
-	bool bReflectToUE4             : 1;  // Bind vrpn keyboard to UE4 at OnDisplayClusterStartSession pass
+	bool bReflectToUE              : 1;  // Bind vrpn keyboard to UE at OnDisplayClusterStartSession pass
 	bool bReflectToNDisplayCluster : 1;  // Bind vrpn keyboard to nDisplay keyboard at OnDisplayClusterStartSession pass
 };
diff --git a/Engine/Shaders/Private/DynamicLightingCommon.ush b/Engine/Shaders/Private/DynamicLightingCommon.ush
--- a/Engine/Shaders/Private/DynamicLightingCommon.ush
+++ b/Engine/Shaders/Private/DynamicLightingCommon.ush
@@ -15,10 +15,10 @@ float RadialAttenuationMask(float3 WorldLightVector)
 }
 float RadialAttenuation(float3 WorldLightVector, half FalloffExponent)
 {
-	// UE3 (fast, but now we not use the default of 2 which looks quite bad):
+	// Old (fast, but now we not use the default of 2 which looks quite bad):
 	return pow(RadialAttenuationMask(WorldLightVector), FalloffExponent);
 
-	// new UE4 (more physically correct but slower and has a more noticable cutoff ring in the dark):
+	// New (more physically correct but slower and has a more noticable cutoff ring in the dark):
 	// AttenFunc(x) = 1 / (x * x + 1)
 	// derived: InvAttenFunc(y) = sqrtf(1 / y - 1)
 	// FalloffExponent is ignored
diff --git a/Engine/Source/Runtime/ImageWrapper/Private/Formats/JpegImageWrapper.cpp b/Engine/Source/Runtime/ImageWrapper/Private/Formats/JpegImageWrapper.cpp
--- a/Engine/Source/Runtime/ImageWrapper/Private/Formats/JpegImageWrapper.cpp
+++ b/Engine/Source/Runtime/ImageWrapper/Private/Formats/JpegImageWrapper.cpp
@@ -176,7 +176,7 @@ void FJpegImageWrapper::Compress(int32 Quality)
 		CompressedData.AddUninitialized(RawData.Num());
 
 		// Note: OutBufferSize intentionally uses int64_t type as that's what jpge::compress_image_to_jpeg_file_in_memory expects.
-		// UE4 int64 type is not compatible with int64_t on all compilers (int64_t may be `long`, while int64 is `long long`).
+		// UE int64 type is not compatible with int64_t on all compilers (int64_t may be `long`, while int64 is `long long`).
 		int64_t OutBufferSize = CompressedData.Num();
 
 		jpge::params Parameters;
diff --git a/Engine/Source/Runtime/Renderer/Private/HaltonUtilities.cpp b/Engine/Source/Runtime/Renderer/Private/HaltonUtilities.cpp
--- a/Engine/Source/Runtime/Renderer/Private/HaltonUtilities.cpp
+++ b/Engine/Source/Runtime/Renderer/Private/HaltonUtilities.cpp
@@ -25,7 +25,7 @@
 // OUT OF OR IN CONNECTION WITH THE SOFTWARE OR THE USE OR OTHER DEALINGS IN THE
 // SOFTWARE.
 
-// This file has been adapted to UE4
+// This file has been adapted to UE
 
 #include "HaltonUtilities.h"
 
diff --git a/Engine/Source/Runtime/Renderer/Private/HaltonUtilities.h b/Engine/Source/Runtime/Renderer/Private/HaltonUtilities.h
--- a/Engine/Source/Runtime/Renderer/Private/HaltonUtilities.h
+++ b/Engine/Source/Runtime/Renderer/Private/HaltonUtilities.h
@@ -27,7 +27,7 @@
 // OUT OF OR IN CONNECTION WITH THE SOFTWARE OR THE USE OR OTHER DEALINGS IN THE
 // SOFTWARE.
 
-// This file has been adapted to UE4
+// This file has been adapted to UE
 
 #include "Containers/Array.h"
 
diff --git a/Engine/Source/Runtime/Renderer/Private/LightRendering.cpp b/Engine/Source/Runtime/Renderer/Private/LightRendering.cpp
--- a/Engine/Source/Runtime/Renderer/Private/LightRendering.cpp
+++ b/Engine/Source/Runtime/Renderer/Private/LightRendering.cpp
@@ -2316,7 +2316,7 @@ void FDeferredShadingSceneRenderer::RenderLight(
 				FarDepth = 0.0f;
 			}
 
-			// UE4 uses reversed depth, so far < near
+			// UE uses reversed depth, so far < near
 			RHICmdList.SetDepthBounds(FarDepth, NearDepth);
 		}
 
diff --git a/Engine/Source/Runtime/Renderer/Private/PostProcess/PostProcessWeightedSampleSum.cpp b/Engine/Source/Runtime/Renderer/Private/PostProcess/PostProcessWeightedSampleSum.cpp
--- a/Engine/Source/Runtime/Renderer/Private/PostProcess/PostProcessWeightedSampleSum.cpp
+++ b/Engine/Source/Runtime/Renderer/Private/PostProcess/PostProcessWeightedSampleSum.cpp
@@ -101,7 +101,7 @@ float NormalDistributionUnscaled(float X, float Sigma, float CrossCenterWeight)
 	}
 	else
 	{
-		// Constant is tweaked give a similar look to UE4 before we fix the scale bug (Some content tweaking might be needed).
+		// Constant is tweaked give a similar look to UE before we fixed the scale bug (Some content tweaking might be needed).
 		// The value defines how much of the Gaussian clipped by the sample window.
 		// r.Filter.SizeScale allows to tweak that for performance/quality.
 		const float LegacyCompatibilityConstant = -16.7f;
