nDisplay: fixed CIS issue

fyi Alejandro.Arango

diff --git a/Engine/Plugins/Runtime/nDisplay/Source/MPCDI/Private/MPCDIWarpTexture.cpp b/Engine/Plugins/Runtime/nDisplay/Source/MPCDI/Private/MPCDIWarpTexture.cpp
--- a/Engine/Plugins/Runtime/nDisplay/Source/MPCDI/Private/MPCDIWarpTexture.cpp
+++ b/Engine/Plugins/Runtime/nDisplay/Source/MPCDI/Private/MPCDIWarpTexture.cpp
@@ -186,17 +186,17 @@ bool FMPCDIWarpTexture::LoadCustom3DWarpMap(const TArray<FVector>& InPoints, int
 	for (int i = 0; i < InPoints.Num(); ++i)
 	{
 		const FVector& t = InPoints[i];
-		FVector4& Pts = data[i];
+		FVector4* Pts = &data[i];
 
 		if ((!(fabsf(t.X) < kEpsilon && fabsf(t.Y) < kEpsilon && fabsf(t.Z) < kEpsilon))
 			&& (!FMath::IsNaN(t.X) && !FMath::IsNaN(t.Y) && !FMath::IsNaN(t.Z)))
 		{
-			Pts = m.TransformPosition(t);
-			Pts.W = 1;
+			*Pts = m.TransformPosition(t);
+			Pts->W = 1;
 		}
 		else
 		{
-			Pts = FVector4(0.f, 0.f, 0.f, -1.f);
+			*Pts = FVector4(0.f, 0.f, 0.f, -1.f);
 		}
 	}
 
@@ -286,17 +286,17 @@ bool FMPCDIWarpTexture::LoadWarpMap(mpcdi::GeometryWarpFile *SourceWarpMap, IMPC
 			mpcdi::NODE &node = (*SourceWarpMap)(i, j);
 			FVector t(node.r, node.g, is2DData ? 0.f : node.b);
 
-			FVector4& Pts = data[i + j * WarpX];
+			FVector4* Pts = &data[i + j * WarpX];
 
 			if ((!(fabsf(t.X) < kEpsilon && fabsf(t.Y) < kEpsilon && fabsf(t.Z) < kEpsilon))
 				&& (!FMath::IsNaN(t.X) && !FMath::IsNaN(t.Y) && !FMath::IsNaN(t.Z)))
 			{
-				Pts = m.TransformPosition(t);
-				Pts.W = 1;
+				*Pts = m.TransformPosition(t);
+				Pts->W = 1;
 			}
 			else
 			{
-				Pts = FVector4(0.f, 0.f, 0.f, -1.f);
+				*Pts = FVector4(0.f, 0.f, 0.f, -1.f);
 			}
 		}
 	}
