re-deprecated `FPlatformProcess::CreateSynchEvent()`, now on all platforms. Replaced `FEvent*` by `FEventRef` where it made sense. Deprecated `FEvent::Create()` as another method to bypass the event pool

diff --git a/Engine/Plugins/Runtime/nDisplay/Source/DisplayCluster/Private/Cluster/Controller/DisplayClusterClusterNodeCtrlMaster.cpp b/Engine/Plugins/Runtime/nDisplay/Source/DisplayCluster/Private/Cluster/Controller/DisplayClusterClusterNodeCtrlMaster.cpp
--- a/Engine/Plugins/Runtime/nDisplay/Source/DisplayCluster/Private/Cluster/Controller/DisplayClusterClusterNodeCtrlMaster.cpp
+++ b/Engine/Plugins/Runtime/nDisplay/Source/DisplayCluster/Private/Cluster/Controller/DisplayClusterClusterNodeCtrlMaster.cpp
@@ -32,26 +32,16 @@ FDisplayClusterClusterNodeCtrlMaster::FDisplayClusterClusterNodeCtrlMaster(const
 	CachedSyncData.Emplace(EDisplayClusterSyncGroup::Tick);
 	CachedSyncData.Emplace(EDisplayClusterSyncGroup::PostTick);
 
-	CachedSyncDataEvents.Emplace(EDisplayClusterSyncGroup::PreTick, FPlatformProcess::CreateSynchEvent(true));
-	CachedSyncDataEvents.Emplace(EDisplayClusterSyncGroup::Tick, FPlatformProcess::CreateSynchEvent(true));
-	CachedSyncDataEvents.Emplace(EDisplayClusterSyncGroup::PostTick, FPlatformProcess::CreateSynchEvent(true));
-
-	CachedInputDataEvent  = FPlatformProcess::CreateSynchEvent(true);
-	CachedDeltaTimeEvent  = FPlatformProcess::CreateSynchEvent(true);
-	CachedEventsDataEvent = FPlatformProcess::CreateSynchEvent(true);
-	CachedFrameTimeEvent  = FPlatformProcess::CreateSynchEvent(true);
+	CachedSyncDataEvents.Emplace(EDisplayClusterSyncGroup::PreTick, FPlatformProcess::GetSynchEventFromPool(true));
+	CachedSyncDataEvents.Emplace(EDisplayClusterSyncGroup::Tick, FPlatformProcess::GetSynchEventFromPool(true));
+	CachedSyncDataEvents.Emplace(EDisplayClusterSyncGroup::PostTick, FPlatformProcess::GetSynchEventFromPool(true));
 }
 
 FDisplayClusterClusterNodeCtrlMaster::~FDisplayClusterClusterNodeCtrlMaster()
 {
-	delete CachedInputDataEvent;
-	delete CachedDeltaTimeEvent;
-	delete CachedEventsDataEvent;
-	delete CachedFrameTimeEvent;
-
 	for (auto& it : CachedSyncDataEvents)
 	{
-		delete it.Value;
+		FPlatformProcess::ReturnSynchEventToPool(it.Value);
 	}
 }
 
diff --git a/Engine/Plugins/Runtime/nDisplay/Source/DisplayCluster/Private/Cluster/Controller/DisplayClusterClusterNodeCtrlMaster.h b/Engine/Plugins/Runtime/nDisplay/Source/DisplayCluster/Private/Cluster/Controller/DisplayClusterClusterNodeCtrlMaster.h
--- a/Engine/Plugins/Runtime/nDisplay/Source/DisplayCluster/Private/Cluster/Controller/DisplayClusterClusterNodeCtrlMaster.h
+++ b/Engine/Plugins/Runtime/nDisplay/Source/DisplayCluster/Private/Cluster/Controller/DisplayClusterClusterNodeCtrlMaster.h
@@ -64,11 +64,11 @@ private:
 
 private:
 	// GetDeltaTime internals
-	FEvent* CachedDeltaTimeEvent = nullptr;
+	FEventRef CachedDeltaTimeEvent{ EEventMode::ManualReset };
 	float   CachedDeltaTime = 0.f;
 
 	// GetTimecode internals
-	FEvent* CachedFrameTimeEvent = nullptr;
+	FEventRef CachedFrameTimeEvent{ EEventMode::ManualReset };
 	TOptional<FQualifiedFrameTime>  CachedFrameTime;
 
 	// GetSyncData internals
@@ -76,11 +76,11 @@ private:
 	TMap<EDisplayClusterSyncGroup, TMap<FString, FString>> CachedSyncData;
 
 	// GetInputData internals
-	FEvent* CachedInputDataEvent = nullptr;
+	FEventRef CachedInputDataEvent{ EEventMode::ManualReset };
 	TMap<FString, FString> CachedInputData;
 
 	// GetEventsData internals
-	FEvent* CachedEventsDataEvent = nullptr;
+	FEventRef CachedEventsDataEvent{ EEventMode::ManualReset };
 	TArray<TSharedPtr<FDisplayClusterClusterEventJson,   ESPMode::ThreadSafe>>   CachedJsonEvents;
 	TArray<TSharedPtr<FDisplayClusterClusterEventBinary, ESPMode::ThreadSafe>> CachedBinaryEvents;
 
diff --git a/Engine/Plugins/Runtime/nDisplay/Source/DisplayCluster/Private/Cluster/DisplayClusterClusterManager.cpp b/Engine/Plugins/Runtime/nDisplay/Source/DisplayCluster/Private/Cluster/DisplayClusterClusterManager.cpp
--- a/Engine/Plugins/Runtime/nDisplay/Source/DisplayCluster/Private/Cluster/DisplayClusterClusterManager.cpp
+++ b/Engine/Plugins/Runtime/nDisplay/Source/DisplayCluster/Private/Cluster/DisplayClusterClusterManager.cpp
@@ -37,16 +37,8 @@ FDisplayClusterClusterManager::FDisplayClusterClusterManager()
 
 	// Set internal system events handler
 	OnClusterEventJson.Add(FDisplayClusterClusterEventHandler::Get().GetJsonListenerDelegate());
-
-	NativeInputDataAvailableEvent = FPlatformProcess::CreateSynchEvent(true);
 }
 
-FDisplayClusterClusterManager::~FDisplayClusterClusterManager()
-{
-	delete NativeInputDataAvailableEvent;
-}
-
-
 //////////////////////////////////////////////////////////////////////////////////////////////
 // IPDisplayClusterManager
 //////////////////////////////////////////////////////////////////////////////////////////////
diff --git a/Engine/Plugins/Runtime/nDisplay/Source/DisplayCluster/Private/Cluster/DisplayClusterClusterManager.h b/Engine/Plugins/Runtime/nDisplay/Source/DisplayCluster/Private/Cluster/DisplayClusterClusterManager.h
--- a/Engine/Plugins/Runtime/nDisplay/Source/DisplayCluster/Private/Cluster/DisplayClusterClusterManager.h
+++ b/Engine/Plugins/Runtime/nDisplay/Source/DisplayCluster/Private/Cluster/DisplayClusterClusterManager.h
@@ -5,11 +5,10 @@
 #include "Cluster/IPDisplayClusterClusterManager.h"
 #include "Cluster/DisplayClusterClusterEvent.h"
 #include "Misc/App.h"
+#include "HAL/Event.h"
 
 class ADisplayClusterSettings;
 class FJsonObject;
-class FEvent;
-
 
 /**
  * Cluster manager. Responsible for network communication and data replication.
@@ -19,7 +18,6 @@ class FDisplayClusterClusterManager
 {
 public:
 	FDisplayClusterClusterManager();
-	virtual ~FDisplayClusterClusterManager();
 
 public:
 	//////////////////////////////////////////////////////////////////////////////////////////////
@@ -132,7 +130,7 @@ private:
 	TArray<TScriptInterface<IDisplayClusterClusterEventListener>> ClusterEventListeners;
 
 	// Sync native input
-	FEvent* NativeInputDataAvailableEvent = nullptr;
+	FEventRef NativeInputDataAvailableEvent{ EEventMode::ManualReset };
 	TMap<FString, FString> NativeInputDataCache;
 
 	mutable FCriticalSection InternalsSyncScope;
diff --git a/Engine/Plugins/XGEController/Source/Private/XGEControllerModule.cpp b/Engine/Plugins/XGEController/Source/Private/XGEControllerModule.cpp
--- a/Engine/Plugins/XGEController/Source/Private/XGEControllerModule.cpp
+++ b/Engine/Plugins/XGEController/Source/Private/XGEControllerModule.cpp
@@ -102,18 +102,11 @@ FXGEControllerModule::FXGEControllerModule()
 	, TasksCS(new FCriticalSection)
 	, bShutdown(false)
 	, bRestartWorker(false)
-	, WriteOutThreadEvent(FPlatformProcess::CreateSynchEvent(false))
 	, LastEventTime(0)
 {}
 
 FXGEControllerModule::~FXGEControllerModule()
 {
-	if (WriteOutThreadEvent)
-	{
-		delete WriteOutThreadEvent;
-		WriteOutThreadEvent = nullptr;
-	}
-
 	if (TasksCS)
 	{
 		delete TasksCS;
diff --git a/Engine/Plugins/XGEController/Source/Public/XGEControllerModule.h b/Engine/Plugins/XGEController/Source/Public/XGEControllerModule.h
--- a/Engine/Plugins/XGEController/Source/Public/XGEControllerModule.h
+++ b/Engine/Plugins/XGEController/Source/Public/XGEControllerModule.h
@@ -40,7 +40,7 @@ class FXGEControllerModule : public IDistributedBuildController
 	TFuture<void> ReadBackThreadFuture;
 
 	/** event for the thread that communicates with the XGE controller */
-	FEvent* WriteOutThreadEvent;
+	FEventRef WriteOutThreadEvent;
 
 	// We need two pipes, as the named pipe API does not support simultaneous read/write on two threads.
 	FPlatformNamedPipe InputNamedPipe, OutputNamedPipe;
diff --git a/Engine/Source/Developer/DerivedDataCache/Private/HttpDerivedDataBackend.cpp b/Engine/Source/Developer/DerivedDataCache/Private/HttpDerivedDataBackend.cpp
--- a/Engine/Source/Developer/DerivedDataCache/Private/HttpDerivedDataBackend.cpp
+++ b/Engine/Source/Developer/DerivedDataCache/Private/HttpDerivedDataBackend.cpp
@@ -1015,7 +1015,6 @@ struct FDataRequestHelper
 		{
 			Batches[i].Reserved = 0;
 			Batches[i].Ready = 0;
-			Batches[i].Complete = FPlatformProcess::CreateSynchEvent(true); //todo: leaks
 		}
 		Initialized = true;
 	}
@@ -1054,7 +1053,7 @@ private:
 		std::atomic<uint32> Ready;
 		std::atomic<uint32> WeightHint;
 		FHttpRequest* Request;
-		FEvent* Complete;
+		FEventRef Complete{ EEventMode::ManualReset };
 	};
 
 	FHttpRequest* Request;
diff --git a/Engine/Source/Developer/XGEController/Private/XGEControllerModule.cpp b/Engine/Source/Developer/XGEController/Private/XGEControllerModule.cpp
--- a/Engine/Source/Developer/XGEController/Private/XGEControllerModule.cpp
+++ b/Engine/Source/Developer/XGEController/Private/XGEControllerModule.cpp
@@ -113,7 +113,7 @@ class FXGEControllerModule : public IXGEController
 	TFuture<void> WriteOutThreadFuture;
 	TFuture<void> ReadBackThreadFuture;
 
-	FEvent* WriteOutThreadEvent;
+	FEventRef WriteOutThreadEvent;
 
 	// We need two pipes, as the named pipe API does not support simultaneous read/write on two threads.
 	FPlatformNamedPipe InputNamedPipe, OutputNamedPipe;
@@ -160,18 +160,11 @@ FXGEControllerModule::FXGEControllerModule()
 	, TasksCS(new FCriticalSection)
 	, bShutdown(false)
 	, bRestartWorker(false)
-	, WriteOutThreadEvent(FPlatformProcess::CreateSynchEvent(false))
 	, LastEventTime(0)
 {}
 
 FXGEControllerModule::~FXGEControllerModule()
 {
-	if (WriteOutThreadEvent)
-	{
-		delete WriteOutThreadEvent;
-		WriteOutThreadEvent = nullptr;
-	}
-
 	if (TasksCS)
 	{
 		delete TasksCS;
diff --git a/Engine/Source/Runtime/Core/Private/HoloLens/HoloLensRunnableThread.h b/Engine/Source/Runtime/Core/Private/HoloLens/HoloLensRunnableThread.h
--- a/Engine/Source/Runtime/Core/Private/HoloLens/HoloLensRunnableThread.h
+++ b/Engine/Source/Runtime/Core/Private/HoloLens/HoloLensRunnableThread.h
@@ -23,7 +23,7 @@ class FRunnableThreadHoloLens
 	HANDLE Thread;
 
 	/** Sync event to make sure that Init() has been completed before allowing the main thread to continue. */
-	FEvent* ThreadInitSyncEvent;
+	FEventRef ThreadInitSyncEvent;
 
 	/** The priority to run the thread at. */
 	EThreadPriority ThreadPriority;
@@ -187,9 +187,6 @@ protected:
 		ThreadName = InThreadName ? InThreadName : TEXT("Unnamed UE4");
 		ThreadPriority = InThreadPri;
 
-		// Create a sync event to guarantee the Init() function is called first
-		ThreadInitSyncEvent = FPlatformProcess::CreateSynchEvent(true);
-
 		// Create the new thread
 		Thread = CreateThread(NULL, InStackSize, _ThreadProc, this, STACK_SIZE_PARAM_IS_A_RESERVATION, (DWORD *)&ThreadID);
 
@@ -208,9 +205,6 @@ protected:
 			SetThreadPriority(InThreadPri);
 		}
 
-		// Cleanup the sync event
-		delete ThreadInitSyncEvent;
-		ThreadInitSyncEvent = NULL;
 		return Thread != NULL;
 	}
 };
diff --git a/Engine/Source/Runtime/Core/Private/Misc/EventPool.h b/Engine/Source/Runtime/Core/Private/Misc/EventPool.h
--- a/Engine/Source/Runtime/Core/Private/Misc/EventPool.h
+++ b/Engine/Source/Runtime/Core/Private/Misc/EventPool.h
@@ -41,7 +41,9 @@ public:
 
 	virtual bool Create(bool bIsManualReset = false)
 	{
+		PRAGMA_DISABLE_DEPRECATION_WARNINGS
 		return InnerEvent->Create(bIsManualReset);
+		PRAGMA_ENABLE_DEPRECATION_WARNINGS
 	}
 
 	virtual bool IsManualReset()
@@ -103,7 +105,9 @@ public:
 		if (!Result)
 #endif
 		{
+			PRAGMA_DISABLE_DEPRECATION_WARNINGS
 			Result = FPlatformProcess::CreateSynchEvent((PoolType == EEventPoolTypes::ManualReset));
+			PRAGMA_ENABLE_DEPRECATION_WARNINGS
 		}
 		check(Result);
 		Result->AdvanceStats();
diff --git a/Engine/Source/Runtime/Core/Private/Windows/WindowsPlatformProcess.cpp b/Engine/Source/Runtime/Core/Private/Windows/WindowsPlatformProcess.cpp
--- a/Engine/Source/Runtime/Core/Private/Windows/WindowsPlatformProcess.cpp
+++ b/Engine/Source/Runtime/Core/Private/Windows/WindowsPlatformProcess.cpp
@@ -1374,7 +1374,9 @@ FEvent* FWindowsPlatformProcess::CreateSynchEvent(bool bIsManualReset)
 		Event = new FSingleThreadEvent();
 	}
 	// If the internal create fails, delete the instance and return NULL
+	PRAGMA_DISABLE_DEPRECATION_WARNINGS
 	if (!Event->Create(bIsManualReset))
+	PRAGMA_ENABLE_DEPRECATION_WARNINGS
 	{
 		delete Event;
 		Event = NULL;
diff --git a/Engine/Source/Runtime/Core/Public/GenericPlatform/GenericPlatformProcess.h b/Engine/Source/Runtime/Core/Public/GenericPlatform/GenericPlatformProcess.h
--- a/Engine/Source/Runtime/Core/Public/GenericPlatform/GenericPlatformProcess.h
+++ b/Engine/Source/Runtime/Core/Public/GenericPlatform/GenericPlatformProcess.h
@@ -525,6 +525,7 @@ struct CORE_API FGenericPlatformProcess
 	 * @see GetSynchEventFromPool, ReturnSynchEventToPool
 	 */
 	// Message to others in the future, don't try to delete this function as it isn't exactly deprecated, but it should only ever be called from FEventPool::GetEventFromPool()
+	UE_DEPRECATED(5.0, "Please use GetSynchEventFromPool to create a new event, and ReturnSynchEventToPool to release the event.")
 	static class FEvent* CreateSynchEvent(bool bIsManualReset = false);
 
 	/**
diff --git a/Engine/Source/Runtime/Core/Public/HAL/Event.h b/Engine/Source/Runtime/Core/Public/HAL/Event.h
--- a/Engine/Source/Runtime/Core/Public/HAL/Event.h
+++ b/Engine/Source/Runtime/Core/Public/HAL/Event.h
@@ -29,6 +29,7 @@ public:
 	 * @param bIsManualReset Whether the event requires manual reseting or not.
 	 * @return true if the event was created, false otherwise.
 	 */
+	UE_DEPRECATED(5.0, "Direct creation of FEvent is discouraged for performance reasons. Please use FPlatformProcess::GetSynchEventFromPool/ReturnSynchEventToPool.")
 	virtual bool Create( bool bIsManualReset = false ) = 0;
 
 	/**
diff --git a/Engine/Source/Runtime/Core/Public/HoloLens/HoloLensPlatformProcess.h b/Engine/Source/Runtime/Core/Public/HoloLens/HoloLensPlatformProcess.h
--- a/Engine/Source/Runtime/Core/Public/HoloLens/HoloLensPlatformProcess.h
+++ b/Engine/Source/Runtime/Core/Public/HoloLens/HoloLensPlatformProcess.h
@@ -32,6 +32,7 @@ struct CORE_API FHoloLensProcess : public FGenericPlatformProcess
 	static void SleepNoStats(float Seconds); 
 	static void SleepInfinite();
 	static void YieldThread();
+	UE_DEPRECATED(5.0, "Please use GetSynchEventFromPool to create a new event, and ReturnSynchEventToPool to release the event.")
 	static class FEvent* CreateSynchEvent(bool bIsManualReset = false);
 	static class FRunnableThread* CreateRunnableThread();
 	static const TCHAR* ExecutableName(bool bRemoveExtension = true);
diff --git a/Engine/Source/Runtime/Core/Public/Windows/WindowsPlatformProcess.h b/Engine/Source/Runtime/Core/Public/Windows/WindowsPlatformProcess.h
--- a/Engine/Source/Runtime/Core/Public/Windows/WindowsPlatformProcess.h
+++ b/Engine/Source/Runtime/Core/Public/Windows/WindowsPlatformProcess.h
@@ -176,6 +176,7 @@ public:
 	static void SleepNoStats(float Seconds);
 	static void SleepInfinite();
 	static void YieldThread();
+	UE_DEPRECATED(5.0, "Please use GetSynchEventFromPool to create a new event, and ReturnSynchEventToPool to release the event.")
 	static class FEvent* CreateSynchEvent(bool bIsManualReset = false);
 	static class FRunnableThread* CreateRunnableThread();
 	static void ClosePipe( void* ReadPipe, void* WritePipe );
