fixed merge error

diff --git a/Engine/Plugins/Runtime/nDisplay/Source/DisplayCluster/Private/Cluster/Controller/DisplayClusterClusterNodeCtrlMaster.h b/Engine/Plugins/Runtime/nDisplay/Source/DisplayCluster/Private/Cluster/Controller/DisplayClusterClusterNodeCtrlMaster.h
--- a/Engine/Plugins/Runtime/nDisplay/Source/DisplayCluster/Private/Cluster/Controller/DisplayClusterClusterNodeCtrlMaster.h
+++ b/Engine/Plugins/Runtime/nDisplay/Source/DisplayCluster/Private/Cluster/Controller/DisplayClusterClusterNodeCtrlMaster.h
@@ -62,7 +62,7 @@ private:
 
 private:
 	// GetTimeData internals
-	FEventRef CachedDeltaTimeEvent{ EEventMode::ManualReset };
+	FEventRef CachedTimeDataEvent{ EEventMode::ManualReset };
 	float   CachedDeltaTime = 0.f;
 	double  CachedGameTime = 0.f;
 	TOptional<FQualifiedFrameTime>  CachedFrameTime;
