Fix NonUnity UnrealGame build

diff --git a/Engine/Plugins/MovieScene/MovieRenderPipeline/Source/MovieRenderPipelineRenderPasses/Private/MoviePipelineImageSequenceOutput.cpp b/Engine/Plugins/MovieScene/MovieRenderPipeline/Source/MovieRenderPipelineRenderPasses/Private/MoviePipelineImageSequenceOutput.cpp
--- a/Engine/Plugins/MovieScene/MovieRenderPipeline/Source/MovieRenderPipelineRenderPasses/Private/MoviePipelineImageSequenceOutput.cpp
+++ b/Engine/Plugins/MovieScene/MovieRenderPipeline/Source/MovieRenderPipelineRenderPasses/Private/MoviePipelineImageSequenceOutput.cpp
@@ -21,6 +21,7 @@
 #include "MoviePipelineWidgetRenderSetting.h"
 #include "MoviePipelineUtils.h"
 #include "HAL/PlatformTime.h"
+#include "Misc/Paths.h"
 
 DECLARE_CYCLE_STAT(TEXT("ImgSeqOutput_RecieveImageData"), STAT_ImgSeqRecieveImageData, STATGROUP_MoviePipeline);
 
diff --git a/Engine/Plugins/Runtime/GeoReferencing/Source/GeoReferencing/Private/GeoReferencingSystem.cpp b/Engine/Plugins/Runtime/GeoReferencing/Source/GeoReferencing/Private/GeoReferencingSystem.cpp
--- a/Engine/Plugins/Runtime/GeoReferencing/Source/GeoReferencing/Private/GeoReferencingSystem.cpp
+++ b/Engine/Plugins/Runtime/GeoReferencing/Source/GeoReferencing/Private/GeoReferencingSystem.cpp
@@ -7,7 +7,8 @@
 #include "Kismet/GameplayStatics.h"
 #include "Widgets/Notifications/SNotificationList.h"
 #include "Framework/Notifications/NotificationManager.h"
-
+#include "Engine/Engine.h"
+#include "Engine/World.h"
 #include "proj.h"
 
 #if WITH_EDITOR
diff --git a/Engine/Plugins/Runtime/nDisplay/Source/DisplayCluster/Private/Render/Containers/DisplayClusterRender_MeshComponentProxy.cpp b/Engine/Plugins/Runtime/nDisplay/Source/DisplayCluster/Private/Render/Containers/DisplayClusterRender_MeshComponentProxy.cpp
--- a/Engine/Plugins/Runtime/nDisplay/Source/DisplayCluster/Private/Render/Containers/DisplayClusterRender_MeshComponentProxy.cpp
+++ b/Engine/Plugins/Runtime/nDisplay/Source/DisplayCluster/Private/Render/Containers/DisplayClusterRender_MeshComponentProxy.cpp
@@ -5,6 +5,9 @@
 #include "Render/Containers/DisplayClusterRender_MeshResources.h"
 #include "Render/Containers/DisplayClusterRender_MeshGeometry.h"
 
+#include "Engine/StaticMesh.h"
+#include "StaticMeshResources.h"
+
 TGlobalResource<FDisplayClusterMeshVertexDeclaration> GDisplayClusterMeshVertexDeclaration;
 
 //*************************************************************************
diff --git a/Engine/Plugins/Runtime/nDisplay/Source/DisplayCluster/Private/Render/Containers/DisplayClusterRender_MeshComponentRef.cpp b/Engine/Plugins/Runtime/nDisplay/Source/DisplayCluster/Private/Render/Containers/DisplayClusterRender_MeshComponentRef.cpp
--- a/Engine/Plugins/Runtime/nDisplay/Source/DisplayCluster/Private/Render/Containers/DisplayClusterRender_MeshComponentRef.cpp
+++ b/Engine/Plugins/Runtime/nDisplay/Source/DisplayCluster/Private/Render/Containers/DisplayClusterRender_MeshComponentRef.cpp
@@ -1,7 +1,7 @@
 // Copyright Epic Games, Inc. All Rights Reserved.
 
 #include "Render/Containers/DisplayClusterRender_MeshComponentRef.h"
-
+#include "Engine/StaticMesh.h"
 
 //*************************************************************************
 //* FDisplayClusterRender_MeshComponentRef
diff --git a/Engine/Plugins/Runtime/nDisplay/Source/DisplayClusterShaders/Private/WarpBlend/Math/DisplayClusterWarpBlendMath_Frustum.cpp b/Engine/Plugins/Runtime/nDisplay/Source/DisplayClusterShaders/Private/WarpBlend/Math/DisplayClusterWarpBlendMath_Frustum.cpp
--- a/Engine/Plugins/Runtime/nDisplay/Source/DisplayClusterShaders/Private/WarpBlend/Math/DisplayClusterWarpBlendMath_Frustum.cpp
+++ b/Engine/Plugins/Runtime/nDisplay/Source/DisplayClusterShaders/Private/WarpBlend/Math/DisplayClusterWarpBlendMath_Frustum.cpp
@@ -11,6 +11,8 @@
 
 #include "Render/Viewport/IDisplayClusterViewport.h"
 
+#include "StaticMeshResources.h"
+
 FMatrix GetProjectionMatrixAssymetric(IDisplayClusterViewport* InViewport, const uint32 InContextNum, const FDisplayClusterWarpEye& InEye, const FDisplayClusterWarpContext& InContext)
 {
 	const float Left   = InContext.ProjectionAngles.Left;
