Fix Deprecation warnings in DisplayClusterConfiguratorSCSEditorViewport.cpp

#rnx

diff --git a/Engine/Plugins/Runtime/nDisplay/Source/DisplayClusterConfigurator/Private/Views/Viewport/DisplayClusterConfiguratorSCSEditorViewport.cpp b/Engine/Plugins/Runtime/nDisplay/Source/DisplayClusterConfigurator/Private/Views/Viewport/DisplayClusterConfiguratorSCSEditorViewport.cpp
--- a/Engine/Plugins/Runtime/nDisplay/Source/DisplayClusterConfigurator/Private/Views/Viewport/DisplayClusterConfiguratorSCSEditorViewport.cpp
+++ b/Engine/Plugins/Runtime/nDisplay/Source/DisplayClusterConfigurator/Private/Views/Viewport/DisplayClusterConfiguratorSCSEditorViewport.cpp
@@ -15,7 +15,7 @@
 
 #include "EngineUtils.h"
 #include "BlueprintEditor.h"
-#include "SSCSEditor.h"
+#include "SSubobjectEditor.h"
 #include "EditorViewportCommands.h"
 #include "STransformViewportToolbar.h"
 #include "SEditorViewportToolBarMenu.h"
@@ -600,7 +600,7 @@ void SDisplayClusterConfiguratorSCSEditorViewport::PopulateViewportOverlays(TSha
 void SDisplayClusterConfiguratorSCSEditorViewport::BindCommands()
 {
 	TSharedPtr<FBlueprintEditor> BlueprintEditor = BlueprintEditorPtr.Pin();
-	CommandList->Append(BlueprintEditor->GetSCSEditor()->CommandList.ToSharedRef());
+	CommandList->Append(BlueprintEditor->GetSubobjectEditor()->GetCommandList().ToSharedRef());
 	CommandList->Append(BlueprintEditor->GetToolkitCommands());
 	
 	const FDisplayClusterConfiguratorCommands& Commands = FDisplayClusterConfiguratorCommands::Get();
