Fix deprecation warnings in DisplayClusterConfiguratorSCSEditorViewportClient.cpp

#rnx

diff --git a/Engine/Plugins/Runtime/nDisplay/Source/DisplayClusterConfigurator/Private/Views/Viewport/DisplayClusterConfiguratorSCSEditorViewportClient.cpp b/Engine/Plugins/Runtime/nDisplay/Source/DisplayClusterConfigurator/Private/Views/Viewport/DisplayClusterConfiguratorSCSEditorViewportClient.cpp
--- a/Engine/Plugins/Runtime/nDisplay/Source/DisplayClusterConfigurator/Private/Views/Viewport/DisplayClusterConfiguratorSCSEditorViewportClient.cpp
+++ b/Engine/Plugins/Runtime/nDisplay/Source/DisplayClusterConfigurator/Private/Views/Viewport/DisplayClusterConfiguratorSCSEditorViewportClient.cpp
@@ -20,7 +20,7 @@
 #include "ComponentVisualizer.h"
 #include "EngineUtils.h"
 #include "ISCSEditorCustomization.h"
-#include "SSCSEditor.h"
+#include "SSubobjectEditor.h"
 #include "UnrealEdGlobals.h"
 #include "UnrealWidget.h"
 #include "Components/PostProcessComponent.h"
@@ -210,13 +210,14 @@ namespace
 	}
 
 	// Determine whether or not the given node has a parent node that is not the root node, is movable and is selected
-	bool IsMovableParentNodeSelected(const FSCSEditorTreeNodePtrType& NodePtr, const TArray<FSCSEditorTreeNodePtrType>& SelectedNodes)
+	bool IsMovableParentNodeSelected(const FSubobjectEditorTreeNodePtrType& NodePtr, const TArray<FSubobjectEditorTreeNodePtrType>& SelectedNodes)
 	{
 		if (NodePtr.IsValid())
 		{
 			// Check for a valid parent node
-			FSCSEditorTreeNodePtrType ParentNodePtr = NodePtr->GetParent();
-			if (ParentNodePtr.IsValid() && !ParentNodePtr->IsRootComponent())
+			FSubobjectEditorTreeNodePtrType ParentNodePtr = NodePtr->GetParent();
+			FSubobjectData* ParentData = ParentNodePtr.IsValid() ? ParentNodePtr->GetDataSource() : nullptr;
+			if (ParentData && !ParentData->IsRootComponent())
 			{
 				if (SelectedNodes.Contains(ParentNodePtr))
 				{
@@ -240,7 +241,7 @@ void FDisplayClusterConfiguratorSCSEditorViewportClient::Tick(float DeltaSeconds
 	FEditorViewportClient::Tick(DeltaSeconds);
 
 	// Register the selection override delegate for the preview actor's components
-	TSharedPtr<SSCSEditor> SCSEditor = BlueprintEditorPtr.Pin()->GetSCSEditor();
+	TSharedPtr<SSubobjectEditor> SubobjectEditor = BlueprintEditorPtr.Pin()->GetSubobjectEditor();
 	AActor* PreviewActor = GetPreviewActor();
 	if (PreviewActor != nullptr)
 	{
@@ -250,7 +251,7 @@ void FDisplayClusterConfiguratorSCSEditorViewportClient::Tick(float DeltaSeconds
 			{
 				if (!PrimComponent->SelectionOverrideDelegate.IsBound())
 				{
-					SCSEditor->SetSelectionOverride(PrimComponent);
+					SubobjectEditor->SetSelectionOverride(PrimComponent);
 				}
 			}
 		}
@@ -293,12 +294,12 @@ void FDisplayClusterConfiguratorSCSEditorViewportClient::Draw(const FSceneView*
 	{
 		if (GUnrealEd != nullptr)
 		{
-			TArray<FSCSEditorTreeNodePtrType> SelectedNodes = BlueprintEditorPtr.Pin()->GetSelectedSCSEditorTreeNodes();
+			TArray<FSubobjectEditorTreeNodePtrType> SelectedNodes = BlueprintEditorPtr.Pin()->GetSelectedSubobjectEditorTreeNodes();
 			for (int32 SelectionIndex = 0; SelectionIndex < SelectedNodes.Num(); ++SelectionIndex)
 			{
-				FSCSEditorTreeNodePtrType SelectedNode = SelectedNodes[SelectionIndex];
+				FSubobjectEditorTreeNodePtrType SelectedNode = SelectedNodes[SelectionIndex];
 
-				UActorComponent* Comp = SelectedNode->FindComponentInstanceInActor(PreviewActor);
+				const UActorComponent* Comp = SelectedNode.IsValid() ? SelectedNode->GetDataSource()->FindComponentInstanceInActor(PreviewActor) : nullptr;
 				if (Comp != nullptr && Comp->IsRegistered())
 				{
 					// Try and find a visualizer
@@ -320,12 +321,12 @@ void FDisplayClusterConfiguratorSCSEditorViewportClient::DrawCanvas(FViewport& I
 	{
 		if (GUnrealEd != nullptr)
 		{
-			TArray<FSCSEditorTreeNodePtrType> SelectedNodes = BlueprintEditorPtr.Pin()->GetSelectedSCSEditorTreeNodes();
+			TArray<FSubobjectEditorTreeNodePtrType> SelectedNodes = BlueprintEditorPtr.Pin()->GetSelectedSubobjectEditorTreeNodes();
 			for (int32 SelectionIndex = 0; SelectionIndex < SelectedNodes.Num(); ++SelectionIndex)
 			{
-				FSCSEditorTreeNodePtrType SelectedNode = SelectedNodes[SelectionIndex];
-
-				UActorComponent* Comp = SelectedNode->FindComponentInstanceInActor(PreviewActor);
+				FSubobjectEditorTreeNodePtrType SelectedNode = SelectedNodes[SelectionIndex];
+				FSubobjectData* Data = SelectedNode.IsValid() ? SelectedNode->GetDataSource() : nullptr;
+				const UActorComponent* Comp = Data ? Data->FindComponentInstanceInActor(PreviewActor) : nullptr;
 				if (Comp != nullptr && Comp->IsRegistered())
 				{
 					// Try and find a visualizer
@@ -343,10 +344,11 @@ void FDisplayClusterConfiguratorSCSEditorViewportClient::DrawCanvas(FViewport& I
 		const int32 HalfX = 0.5f * Viewport->GetSizeXY().X;
 		const int32 HalfY = 0.5f * Viewport->GetSizeXY().Y;
 
-		TArray<TSharedPtr<FSCSEditorTreeNode>> SelectedNodes = BlueprintEditorPtr.Pin()->GetSelectedSCSEditorTreeNodes();
+		const TArray<FSubobjectEditorTreeNodePtrType>& SelectedNodes = BlueprintEditorPtr.Pin()->GetSelectedSubobjectEditorTreeNodes();
 		if (bIsManipulating && SelectedNodes.Num() > 0)
 		{
-			USceneComponent* SceneComp = Cast<USceneComponent>(SelectedNodes[0]->FindComponentInstanceInActor(PreviewActor));
+			FSubobjectData* Data = SelectedNodes[0]->GetDataSource();
+			const USceneComponent* SceneComp = Cast<USceneComponent>(Data->FindComponentInstanceInActor(PreviewActor));
 			if (SceneComp)
 			{
 				const FVector WidgetLocation = GetWidgetLocation();
@@ -473,14 +475,14 @@ void FDisplayClusterConfiguratorSCSEditorViewportClient::ProcessClick(class FSce
 				
 				if (SelectedCompInstance)
 				{
-					TSharedPtr<ISCSEditorCustomization> Customization = BlueprintEditorPtr.Pin()->CustomizeSCSEditor(SelectedCompInstance);
+					TSharedPtr<ISCSEditorCustomization> Customization = BlueprintEditorPtr.Pin()->CustomizeSubobjectEditor(SelectedCompInstance);
 					if (!(Customization.IsValid() && Customization->HandleViewportClick(AsShared(), View, HitProxy, Key, Event, HitX, HitY)))
 					{
 						const bool bIsCtrlKeyDown = Viewport->KeyState(EKeys::LeftControl) || Viewport->KeyState(EKeys::RightControl);
 						if (BlueprintEditorPtr.IsValid())
 						{
 							// Note: This will find and select any node associated with the component instance that's attached to the proxy (including visualizers)
-							BlueprintEditorPtr.Pin()->FindAndSelectSCSEditorTreeNode(SelectedCompInstance, bIsCtrlKeyDown);
+							BlueprintEditorPtr.Pin()->FindAndSelectSubobjectEditorTreeNode(SelectedCompInstance, bIsCtrlKeyDown);
 						}
 					}
 
@@ -516,7 +518,7 @@ bool FDisplayClusterConfiguratorSCSEditorViewportClient::InputWidgetDelta(FViewp
 		TSharedPtr<FBlueprintEditor> BlueprintEditor = BlueprintEditorPtr.Pin();
 		if (PreviewActor && BlueprintEditor.IsValid())
 		{
-			TArray<FSCSEditorTreeNodePtrType> SelectedNodes = BlueprintEditor->GetSelectedSCSEditorTreeNodes();
+			TArray<FSubobjectEditorTreeNodePtrType> SelectedNodes = BlueprintEditor->GetSelectedSubobjectEditorTreeNodes();
 			if (SelectedNodes.Num() > 0)
 			{
 				FVector ModifiedScale = Scale;
@@ -534,11 +536,12 @@ bool FDisplayClusterConfiguratorSCSEditorViewportClient::InputWidgetDelta(FViewp
 					ModifiedScale = FVector::ZeroVector;
 				}
 
-				for (const FSCSEditorTreeNodePtrType& SelectedNodePtr : SelectedNodes)
+				for (const FSubobjectEditorTreeNodePtrType& SelectedNodePtr : SelectedNodes)
 				{
+					FSubobjectData* Data = SelectedNodePtr->GetDataSource();
 					// Don't allow editing of a root node, inherited SCS node or child node that also has a movable (non-root) parent node selected
 					const bool bCanEdit = GUnrealEd->ComponentVisManager.IsActive() ||
-						(!SelectedNodePtr->IsRootComponent() && !IsMovableParentNodeSelected(SelectedNodePtr, SelectedNodes));
+						(!Data->IsRootComponent() && !IsMovableParentNodeSelected(SelectedNodePtr, SelectedNodes));
 
 					if (bCanEdit)
 					{
@@ -549,8 +552,8 @@ bool FDisplayClusterConfiguratorSCSEditorViewportClient::InputWidgetDelta(FViewp
 							return true;
 						}
 
-						USceneComponent* SceneComp = Cast<USceneComponent>(SelectedNodePtr->FindComponentInstanceInActor(PreviewActor));
-						USceneComponent* SelectedTemplate = Cast<USceneComponent>(SelectedNodePtr->GetOrCreateEditableComponentTemplate(BlueprintEditor->GetBlueprintObj()));
+						USceneComponent* SceneComp = const_cast<USceneComponent*>(Cast<USceneComponent>(Data->FindComponentInstanceInActor(PreviewActor)));
+						USceneComponent* SelectedTemplate = const_cast<USceneComponent*>(Data->GetObjectForBlueprint<USceneComponent>(BlueprintEditor->GetBlueprintObj()));
 						if (SceneComp != nullptr && SelectedTemplate != nullptr)
 						{
 							// Cache the current default values for propagation
@@ -561,7 +564,7 @@ bool FDisplayClusterConfiguratorSCSEditorViewportClient::InputWidgetDelta(FViewp
 							// Adjust the deltas as necessary
 							FComponentEditorUtils::AdjustComponentDelta(SceneComp, Drag, Rot);
 
-							TSharedPtr<ISCSEditorCustomization> Customization = BlueprintEditor->CustomizeSCSEditor(SceneComp);
+							TSharedPtr<ISCSEditorCustomization> Customization = BlueprintEditor->CustomizeSubobjectEditor(SceneComp);
 							if (Customization.IsValid() && Customization->HandleViewportDrag(SceneComp, SelectedTemplate, Drag, Rot, ModifiedScale, GetWidgetLocation()))
 							{
 								// Handled by SCS Editor customization
@@ -697,10 +700,10 @@ UE::Widget::EWidgetMode FDisplayClusterConfiguratorSCSEditorViewportClient::GetW
 		const TSharedPtr<FDisplayClusterConfiguratorBlueprintEditor> BluePrintEditor = BlueprintEditorPtr.Pin();
 		if (BluePrintEditor.IsValid())
 		{
-			TArray<FSCSEditorTreeNodePtrType> SelectedNodes = BluePrintEditor->GetSelectedSCSEditorTreeNodes();
-			if (BluePrintEditor->GetSCSEditor()->GetActorNode().IsValid())
+			TArray<FSubobjectEditorTreeNodePtrType> SelectedNodes = BluePrintEditor->GetSelectedSubobjectEditorTreeNodes();
+			if (BluePrintEditor->GetSubobjectEditor()->GetSceneRootNode().IsValid())
 			{
-				const TArray<FSCSEditorTreeNodePtrType>& RootNodes = BluePrintEditor->GetSCSEditor()->GetActorNode()->GetComponentNodes();
+				const TArray<FSubobjectEditorTreeNodePtrType>& RootNodes = BluePrintEditor->GetSubobjectEditor()->GetRootNodes();
 
 				if (GUnrealEd->ComponentVisManager.IsActive() &&
 					GUnrealEd->ComponentVisManager.IsVisualizingArchetype())
@@ -714,13 +717,14 @@ UE::Widget::EWidgetMode FDisplayClusterConfiguratorSCSEditorViewportClient::GetW
 					// root nodes array, or isn't visible in the preview actor, then don't display a transform widget
 					for (int32 CurrentNodeIndex = 0; CurrentNodeIndex < SelectedNodes.Num(); CurrentNodeIndex++)
 					{
-						FSCSEditorTreeNodePtrType CurrentNodePtr = SelectedNodes[CurrentNodeIndex];
-						if ((CurrentNodePtr.IsValid() &&
-							((!RootNodes.Contains(CurrentNodePtr) && !CurrentNodePtr->IsRootComponent()) ||
-								(CurrentNodePtr->GetObject<UInstancedStaticMeshComponent>() && // show widget if we are editing individual instances even if it is the root component
-									CastChecked<UInstancedStaticMeshComponent>(CurrentNodePtr->FindComponentInstanceInActor(GetPreviewActor()))->SelectedInstances.Contains(true))) &&
-							CurrentNodePtr->CanEdit() &&
-							CurrentNodePtr->FindComponentInstanceInActor(PreviewActor)))
+						FSubobjectEditorTreeNodePtrType CurrentNodePtr = SelectedNodes[CurrentNodeIndex];
+						FSubobjectData* Data = CurrentNodePtr.IsValid() ? CurrentNodePtr->GetDataSource() : nullptr;
+						if ((Data &&
+							((!RootNodes.Contains(CurrentNodePtr) && !Data->IsRootComponent()) ||
+								(Data->GetObject<UInstancedStaticMeshComponent>() && // show widget if we are editing individual instances even if it is the root component
+									CastChecked<UInstancedStaticMeshComponent>(Data->FindComponentInstanceInActor(GetPreviewActor()))->SelectedInstances.Contains(true))) &&
+							Data->CanEdit() &&
+							Data->FindComponentInstanceInActor(PreviewActor)))
 						{
 							// a non-nullptr, non-root item is selected, draw the widget
 							ReturnWidgetMode = WidgetMode;
@@ -760,14 +764,14 @@ FVector FDisplayClusterConfiguratorSCSEditorViewportClient::GetWidgetLocation()
 	AActor* PreviewActor = GetPreviewActor();
 	if (PreviewActor)
 	{
-		TArray<FSCSEditorTreeNodePtrType> SelectedNodes = BlueprintEditorPtr.Pin()->GetSelectedSCSEditorTreeNodes();
-		if (SelectedNodes.Num() > 0)
+		TArray<FSubobjectEditorTreeNodePtrType> SelectedNodes = BlueprintEditorPtr.Pin()->GetSelectedSubobjectEditorTreeNodes();
+		if (SelectedNodes.Num() > 0 && SelectedNodes.Last().IsValid())
 		{
 			// Use the last selected item for the widget location
-			USceneComponent* SceneComp = Cast<USceneComponent>(SelectedNodes.Last().Get()->FindComponentInstanceInActor(PreviewActor));
+			const USceneComponent* SceneComp = Cast<USceneComponent>(SelectedNodes.Last()->GetDataSource()->FindComponentInstanceInActor(PreviewActor));
 			if (SceneComp)
 			{
-				TSharedPtr<ISCSEditorCustomization> Customization = BlueprintEditorPtr.Pin()->CustomizeSCSEditor(SceneComp);
+				TSharedPtr<ISCSEditorCustomization> Customization = BlueprintEditorPtr.Pin()->CustomizeSubobjectEditor(SceneComp);
 				FVector CustomLocation;
 				if (Customization.IsValid() && Customization->HandleGetWidgetLocation(SceneComp, CustomLocation))
 				{
@@ -800,14 +804,15 @@ FMatrix FDisplayClusterConfiguratorSCSEditorViewportClient::GetWidgetCoordSystem
 		TSharedPtr<FBlueprintEditor> BlueprintEditor = BlueprintEditorPtr.Pin();
 		if (PreviewActor && BlueprintEditor.IsValid())
 		{
-			TArray<FSCSEditorTreeNodePtrType> SelectedNodes = BlueprintEditor->GetSelectedSCSEditorTreeNodes();
+			TArray<FSubobjectEditorTreeNodePtrType> SelectedNodes = BlueprintEditor->GetSelectedSubobjectEditorTreeNodes();
 			if (SelectedNodes.Num() > 0)
 			{
-				const FSCSEditorTreeNodePtrType SelectedNode = SelectedNodes.Last();
-				USceneComponent* SceneComp = SelectedNode.IsValid() ? Cast<USceneComponent>(SelectedNode->FindComponentInstanceInActor(PreviewActor)) : nullptr;
+				const FSubobjectEditorTreeNodePtrType SelectedNode = SelectedNodes.Last();
+				FSubobjectData* Data = SelectedNode.IsValid() ? SelectedNode->GetDataSource() : nullptr;
+				const USceneComponent* SceneComp = Data ? Cast<USceneComponent>(Data->FindComponentInstanceInActor(PreviewActor)) : nullptr;
 				if (SceneComp)
 				{
-					TSharedPtr<ISCSEditorCustomization> Customization = BlueprintEditor->CustomizeSCSEditor(SceneComp);
+					TSharedPtr<ISCSEditorCustomization> Customization = BlueprintEditor->CustomizeSubobjectEditor(SceneComp);
 					FMatrix CustomTransform;
 					if (Customization.IsValid() && Customization->HandleGetWidgetTransform(SceneComp, CustomTransform))
 					{
@@ -940,11 +945,11 @@ void FDisplayClusterConfiguratorSCSEditorViewportClient::FocusViewportToSelectio
 	AActor* PreviewActor = GetPreviewActor();
 	if (PreviewActor)
 	{
-		TArray<FSCSEditorTreeNodePtrType> SelectedNodes = BlueprintEditorPtr.Pin()->GetSelectedSCSEditorTreeNodes();
+		TArray<FSubobjectEditorTreeNodePtrType> SelectedNodes = BlueprintEditorPtr.Pin()->GetSelectedSubobjectEditorTreeNodes();
 		if (SelectedNodes.Num() > 0)
 		{
 			// Use the last selected item for the widget location
-			USceneComponent* SceneComp = Cast<USceneComponent>(SelectedNodes.Last()->FindComponentInstanceInActor(PreviewActor));
+			const USceneComponent* SceneComp = Cast<USceneComponent>(SelectedNodes.Last()->GetDataSource()->FindComponentInstanceInActor(PreviewActor));
 			if (SceneComp)
 			{
 				FocusViewportOnBox(SceneComp->Bounds.GetBox());
@@ -1197,12 +1202,13 @@ void FDisplayClusterConfiguratorSCSEditorViewportClient::BeginTransaction(const
 				FBlueprintEditorUtils::MarkBlueprintAsModified(PreviewBlueprint);
 			}
 
-			TArray<FSCSEditorTreeNodePtrType> SelectedNodes = BlueprintEditor->GetSelectedSCSEditorTreeNodes();
-			for (const FSCSEditorTreeNodePtrType& Node : SelectedNodes)
+			TArray<FSubobjectEditorTreeNodePtrType> SelectedNodes = BlueprintEditor->GetSelectedSubobjectEditorTreeNodes();
+			for (const FSubobjectEditorTreeNodePtrType& Node : SelectedNodes)
 			{
 				if (Node.IsValid())
 				{
-					if (USCS_Node* SCS_Node = Node->GetSCSNode())
+					FSubobjectData* Data = Node->GetDataSource();
+					if (USCS_Node* SCS_Node = const_cast<USCS_Node*>(Data->GetObject<USCS_Node>()))
 					{
 						USimpleConstructionScript* SCS = SCS_Node->GetSCS();
 						UBlueprint* Blueprint = SCS ? SCS->GetBlueprint() : nullptr;
@@ -1213,7 +1219,8 @@ void FDisplayClusterConfiguratorSCSEditorViewportClient::BeginTransaction(const
 					}
 
 					// Modify template, any instances will be reconstructed as part of PostUndo:
-					UActorComponent* ComponentTemplate = Node->GetOrCreateEditableComponentTemplate(PreviewBlueprint);
+					UActorComponent* ComponentTemplate = const_cast<UActorComponent*>(Data->GetObjectForBlueprint<UActorComponent>(PreviewBlueprint));
+					
 					if (ComponentTemplate != nullptr)
 					{
 						ComponentTemplate->SetFlags(RF_Transactional);
