Add SubobjectEditor and subsystem to the buildcs for DisplayClusterConfig

#rnx

diff --git a/Engine/Plugins/Runtime/nDisplay/Source/DisplayClusterConfigurator/DisplayClusterConfigurator.Build.cs b/Engine/Plugins/Runtime/nDisplay/Source/DisplayClusterConfigurator/DisplayClusterConfigurator.Build.cs
--- a/Engine/Plugins/Runtime/nDisplay/Source/DisplayClusterConfigurator/DisplayClusterConfigurator.Build.cs
+++ b/Engine/Plugins/Runtime/nDisplay/Source/DisplayClusterConfigurator/DisplayClusterConfigurator.Build.cs
@@ -44,6 +44,8 @@ public class DisplayClusterConfigurator : ModuleRules
 				"SlateCore",
 				"ToolMenus",
 				"UnrealEd",
+				"SubobjectEditor",
+				"SubobjectDataInterface"
 			});
 	}
 }
