Fix deprecation warnings in DisplayClusterConfiguratorBlueprintEditor.cpp

#rnx

diff --git a/Engine/Plugins/Runtime/nDisplay/Source/DisplayClusterConfigurator/Private/DisplayClusterConfiguratorBlueprintEditor.cpp b/Engine/Plugins/Runtime/nDisplay/Source/DisplayClusterConfigurator/Private/DisplayClusterConfiguratorBlueprintEditor.cpp
--- a/Engine/Plugins/Runtime/nDisplay/Source/DisplayClusterConfigurator/Private/DisplayClusterConfiguratorBlueprintEditor.cpp
+++ b/Engine/Plugins/Runtime/nDisplay/Source/DisplayClusterConfigurator/Private/DisplayClusterConfiguratorBlueprintEditor.cpp
@@ -45,6 +45,8 @@
 #include "ToolMenus.h"
 #include "ComponentAssetBroker.h"
 #include "DesktopPlatformModule.h"
+#include "SSubobjectEditor.h"
+#include "SubobjectDataSubsystem.h"
 #include "Tools/BaseAssetToolkit.h"
 #include "Engine/Selection.h"
 #include "Engine/TextureRenderTarget2D.h"
@@ -162,7 +164,7 @@ void FDisplayClusterConfiguratorBlueprintEditor::InitDisplayClusterBlueprintEdit
 	const TArray<UBlueprint*> Blueprints{ Blueprint };
 	CommonInitialization(Blueprints, false);
 
-	SetSCSEditorUICustomization(FDisplayClusterBlueprintEditorSCSEditorUICustomization::GetInstance());
+	SetSubobjectEditorUICustomization(FDisplayClusterBlueprintEditorSCSEditorUICustomization::GetInstance());
 	
 	CreateSCSEditorWrapper();
 
@@ -272,11 +274,11 @@ void FDisplayClusterConfiguratorBlueprintEditor::SelectObjects(TArray<UObject*>&
 	}
 	else
 	{
-		TSharedPtr<SSCSEditor> SCSEditorTreeView = GetSCSEditor();
-		if (SCSEditorTreeView.IsValid())
+		TSharedPtr<SSubobjectEditor> SubobjectEditorTreeView = GetSubobjectEditor();
+		if (SubobjectEditorTreeView.IsValid())
 		{
 			const FSelectionScope SelectionScope(this, ESelectionSource::Ancillary);
-			SCSEditorTreeView->ClearSelection();
+			SubobjectEditorTreeView->ClearSelection();
 		}
 		
 		SKismetInspector::FShowDetailsOptions Options;
@@ -305,7 +307,7 @@ void FDisplayClusterConfiguratorBlueprintEditor::SelectAncillaryComponents(const
 
 			if (FoundComponentPtr)
 			{
-				FindAndSelectSCSEditorTreeNode(*FoundComponentPtr, true);
+				FindAndSelectSubobjectEditorTreeNode(*FoundComponentPtr, true);
 			}
 		}
 	}
@@ -686,9 +688,9 @@ void FDisplayClusterConfiguratorBlueprintEditor::OnFocusChanged(const FFocusEven
 		// If the SCSEditor or the Cluster View are being focused, update the property inspector to show the items currently selected in the respective tree view.
 		// This ensures that if any item as selected as part of an ancillary selection, the user can still "select" it and have its properties show up in the
 		// inspector, since already selected items don't fire a OnSelectionChanged event where we would normally update the property inspector
-		if (NewFocusedWidgetPath.ContainsWidget(SCSEditor.ToSharedRef()))
+		if (NewFocusedWidgetPath.ContainsWidget(SubobjectEditor.ToSharedRef()))
 		{
-			TArray<FSCSEditorTreeNodePtrType> SelectedNodes = SCSEditor->GetSelectedNodes();
+			TArray<FSubobjectEditorTreeNodePtrType> SelectedNodes = SubobjectEditor->GetSelectedNodes();
 			if (Inspector.IsValid())
 			{
 				// Convert the selection set to an array of UObject* pointers
@@ -696,15 +698,16 @@ void FDisplayClusterConfiguratorBlueprintEditor::OnFocusChanged(const FFocusEven
 				TArray<UObject*> InspectorObjects;
 				bool bShowComponents = true;
 				InspectorObjects.Empty(SelectedNodes.Num());
-				for (FSCSEditorTreeNodePtrType NodePtr : SelectedNodes)
+				for (FSubobjectEditorTreeNodePtrType NodePtr : SelectedNodes)
 				{
 					if (NodePtr.IsValid())
 					{
-						if (NodePtr->IsActorNode())
+						FSubobjectData* Data = NodePtr->GetDataSource();
+						if (Data->IsActor())
 						{
-							if (AActor* DefaultActor = NodePtr->GetEditableObjectForBlueprint<AActor>(GetBlueprintObj()))
+							if (const AActor* DefaultActor = Data->GetObjectForBlueprint<AActor>(GetBlueprintObj()))
 							{
-								InspectorObjects.Add(DefaultActor);
+								InspectorObjects.Add(const_cast<AActor*>(DefaultActor));
 
 								FString Title;
 								DefaultActor->GetName(Title);
@@ -716,11 +719,11 @@ void FDisplayClusterConfiguratorBlueprintEditor::OnFocusChanged(const FFocusEven
 						}
 						else
 						{
-							UActorComponent* EditableComponent = NodePtr->GetOrCreateEditableComponentTemplate(GetBlueprintObj());
+							const UActorComponent* EditableComponent = Data->GetObjectForBlueprint<UActorComponent>(GetBlueprintObj());
 							if (EditableComponent)
 							{
 								InspectorTitle = FText::FromString(NodePtr->GetDisplayString());
-								InspectorObjects.Add(EditableComponent);
+								InspectorObjects.Add(const_cast<UActorComponent*>(EditableComponent));
 							}
 						}
 					}
@@ -1177,7 +1180,7 @@ TSharedRef<SWidget> FDisplayClusterConfiguratorBlueprintEditor::CreateSCSEditorE
 	{
 		if (USCSEditorExtensionContext* SCSEditorExtensionContext = Cast<USCSEditorExtensionContext>(ExtensionContext.Get()))
 		{
-			if (const TSharedPtr<SSCSEditor> SCSEditor = SCSEditorExtensionContext->GetSCSEditor().Pin())
+			if (const TSharedPtr<SSubobjectEditor> SubobjectEditor = SCSEditorExtensionContext->GetSubobjectEditor().Pin())
 			{
 				UClass* NewClass = ComponentClass;
 
@@ -1198,6 +1201,11 @@ TSharedRef<SWidget> FDisplayClusterConfiguratorBlueprintEditor::CreateSCSEditorE
 						AssetOverride = NewObject<UActorComponent>(GetTransientPackage(), NewClass, *DisplayName);
 					}
 					
+					USubobjectDataSubsystem* SubobjectSystem = USubobjectDataSubsystem::Get();
+					FAddNewSubobjectParams AddSubobjectParams;
+					FText FailReason;
+					AddSubobjectParams.NewClass = NewClass;
+										
 					// This adds components according to the type selected in the drop down. If the user
 					// has the appropriate objects selected in the content browser then those are added,
 					// else we go down the previous route of adding components by type.
@@ -1215,7 +1223,10 @@ TSharedRef<SWidget> FDisplayClusterConfiguratorBlueprintEditor::CreateSCSEditorE
 							{
 								if (ComponentClasses[ComponentIndex]->IsChildOf(NewClass))
 								{
-									NewComponent = SCSEditor->AddNewComponent(NewClass, Object);
+									AddSubobjectParams.AssetOverride = Object;
+									AddSubobjectParams.ParentHandle = SubobjectSystem->FindHandleForObject(SubobjectEditor->GetSceneRootNode()->GetDataHandle(), Object);
+									SubobjectSystem->AddNewSubobject(AddSubobjectParams, FailReason);
+									
 									bAddedComponent = true;
 									break;
 								}
@@ -1226,10 +1237,13 @@ TSharedRef<SWidget> FDisplayClusterConfiguratorBlueprintEditor::CreateSCSEditorE
 					if (!bAddedComponent)
 					{
 						// As the SCS splits up the scene and actor components, can now add directly
-						NewComponent = SCSEditor->AddNewComponent(NewClass, AssetOverride);
+						AddSubobjectParams.NewClass = NewClass;
+						AddSubobjectParams.AssetOverride = AssetOverride;
+						AddSubobjectParams.ParentHandle = SubobjectSystem->FindHandleForObject(SubobjectEditor->GetSceneRootNode()->GetDataHandle(), AssetOverride);
+						SubobjectSystem->AddNewSubobject(AddSubobjectParams, FailReason);
 					}
 					
-					SCSEditor->UpdateTree();
+					SubobjectEditor->UpdateTree();
 				}
 
 				return NewComponent;
@@ -1272,7 +1286,7 @@ void FDisplayClusterConfiguratorBlueprintEditor::CreateSCSEditorWrapper()
 	SCSEditorWrapper = SAssignNew(SCSEditorWrapper, SOverlay)
 		+SOverlay::Slot()
 		[
-			SCSEditor.ToSharedRef()
+			SubobjectEditor.ToSharedRef()
 		]
 
 		+SOverlay::Slot()
diff --git a/Engine/Source/Editor/Kismet/Public/BlueprintEditor.h b/Engine/Source/Editor/Kismet/Public/BlueprintEditor.h
--- a/Engine/Source/Editor/Kismet/Public/BlueprintEditor.h
+++ b/Engine/Source/Editor/Kismet/Public/BlueprintEditor.h
@@ -1274,10 +1274,6 @@ protected:
 	/** Macro/function libraries to keep loaded */
 	TArray<UBlueprint*> StandardLibraries;
 
-	/** SCS editor */
-	UE_DEPRECATED(5.0, "SCSEditor has been deprecated. Use SubobjectEditor instead.")
-	TSharedPtr<class SSCSEditor> SCSEditor;
-
 	/** Subobject Editor */
 	TSharedPtr<SSubobjectEditor> SubobjectEditor;
 
