fix for warnings about uninitialized struct members

diff --git a/Engine/Plugins/Runtime/nDisplay/Source/DisplayClusterConfiguration/Public/DisplayClusterConfigurationTypes_Viewport.h b/Engine/Plugins/Runtime/nDisplay/Source/DisplayClusterConfiguration/Public/DisplayClusterConfigurationTypes_Viewport.h
--- a/Engine/Plugins/Runtime/nDisplay/Source/DisplayClusterConfiguration/Public/DisplayClusterConfigurationTypes_Viewport.h
+++ b/Engine/Plugins/Runtime/nDisplay/Source/DisplayClusterConfiguration/Public/DisplayClusterConfigurationTypes_Viewport.h
@@ -155,6 +155,14 @@ struct DISPLAYCLUSTERCONFIGURATION_API FDisplayClusterConfigurationViewport_Colo
 {
 	GENERATED_BODY()
 
+	FDisplayClusterConfigurationViewport_ColorGradingSettings()
+		: bOverride_Saturation(0)
+		, bOverride_Contrast(0)
+		, bOverride_Gamma(0)
+		, bOverride_Gain(0)
+		, bOverride_Offset(0)
+	{}
+
 	UPROPERTY(EditAnywhere, Category = "Overrides", meta = (PinHiddenByDefault, InlineEditConditionToggle))
 	uint8 bOverride_Saturation:1;
 	UPROPERTY(EditAnywhere, Category = "Overrides", meta = (PinHiddenByDefault, InlineEditConditionToggle))
@@ -192,7 +200,11 @@ struct DISPLAYCLUSTERCONFIGURATION_API FDisplayClusterConfigurationViewport_PerV
 {
 	GENERATED_BODY()
 
-	FDisplayClusterConfigurationViewport_PerViewportSettings() {};
+	FDisplayClusterConfigurationViewport_PerViewportSettings()
+		: bOverride_WhiteTemp(0)
+		, bOverride_WhiteTint(0)
+		, bOverride_AutoExposureBias(0)
+	{}
 
 	UPROPERTY(EditAnywhere, Category = "Overrides", meta = (PinHiddenByDefault, InlineEditConditionToggle))
 	uint8 bOverride_WhiteTemp:1;
