Repalce uses of GetSCSEditor with the correct GetSubobjectEditor calls

#rnx

diff --git a/Engine/Plugins/Runtime/nDisplay/Source/DisplayClusterConfigurator/Private/DisplayClusterConfiguratorBlueprintEditor.cpp b/Engine/Plugins/Runtime/nDisplay/Source/DisplayClusterConfigurator/Private/DisplayClusterConfiguratorBlueprintEditor.cpp
--- a/Engine/Plugins/Runtime/nDisplay/Source/DisplayClusterConfigurator/Private/DisplayClusterConfiguratorBlueprintEditor.cpp
+++ b/Engine/Plugins/Runtime/nDisplay/Source/DisplayClusterConfigurator/Private/DisplayClusterConfiguratorBlueprintEditor.cpp
@@ -1264,9 +1264,9 @@ TSharedRef<SWidget> FDisplayClusterConfiguratorBlueprintEditor::CreateSCSEditorE
 			{
 				if(USCSEditorExtensionContext* SCSEditorExtensionContext = Cast<USCSEditorExtensionContext>(ExtensionContext.Get()))
 				{
-					if (SCSEditorExtensionContext->GetSCSEditor().IsValid())
+					if (SCSEditorExtensionContext->GetSubobjectEditor().IsValid())
 					{
-						if (ADisplayClusterRootActor* RootActor = Cast<ADisplayClusterRootActor>(SCSEditorExtensionContext->GetSCSEditor().Pin()->GetActorContext()))
+						if (const ADisplayClusterRootActor* RootActor = Cast<ADisplayClusterRootActor>(SCSEditorExtensionContext->GetSubobjectEditor().Pin()->GetObjectContext()))
 						{
 							// Check that this context is for our root actor and that we are editing just the CDO.
 							// Without this we could show up for unrelated blueprints or in the level editor when selecting our actor.
