Fixing instances where several FAssetEditorToolkit-derived editors treat a close request as a close command, prematurely cleaning up in anticipation of closure. This can leave an editor in a dysfunctional state if other code checks whether the editor *could* be closed without actually closing it.

For context, other editors commonly use OnRequestClose as an opportunity to see if their data has changed, and pop up a message box asking the user if they want to save or cancel the closure. Actual cleanup is performed in OnClose() or as part of their destructor.

diff --git a/Engine/Plugins/Media/MediaPlayerEditor/Source/MediaPlayerEditor/Private/Toolkits/MediaPlayerEditorToolkit.cpp b/Engine/Plugins/Media/MediaPlayerEditor/Source/MediaPlayerEditor/Private/Toolkits/MediaPlayerEditorToolkit.cpp
--- a/Engine/Plugins/Media/MediaPlayerEditor/Source/MediaPlayerEditor/Private/Toolkits/MediaPlayerEditorToolkit.cpp
+++ b/Engine/Plugins/Media/MediaPlayerEditor/Source/MediaPlayerEditor/Private/Toolkits/MediaPlayerEditorToolkit.cpp
@@ -148,14 +148,12 @@ FString FMediaPlayerEditorToolkit::GetDocumentationLink() const
 }
 
 
-bool FMediaPlayerEditorToolkit::OnRequestClose()
+void FMediaPlayerEditorToolkit::OnClose()
 {
 	if (MediaPlayer != nullptr)
 	{
 		MediaPlayer->Close();
 	}
-
-	return FAssetEditorToolkit::OnRequestClose();
 }
 
 
diff --git a/Engine/Plugins/Media/MediaPlayerEditor/Source/MediaPlayerEditor/Private/Toolkits/MediaPlayerEditorToolkit.h b/Engine/Plugins/Media/MediaPlayerEditor/Source/MediaPlayerEditor/Private/Toolkits/MediaPlayerEditorToolkit.h
--- a/Engine/Plugins/Media/MediaPlayerEditor/Source/MediaPlayerEditor/Private/Toolkits/MediaPlayerEditorToolkit.h
+++ b/Engine/Plugins/Media/MediaPlayerEditor/Source/MediaPlayerEditor/Private/Toolkits/MediaPlayerEditorToolkit.h
@@ -47,7 +47,7 @@ public:
 	//~ FAssetEditorToolkit interface
 
 	virtual FString GetDocumentationLink() const override;
-	virtual bool OnRequestClose() override;
+	virtual void OnClose() override;
 	virtual void RegisterTabSpawners(const TSharedRef<FTabManager>& TabManager) override;
 	virtual void UnregisterTabSpawners(const TSharedRef<FTabManager>& TabManager) override;
 
diff --git a/Engine/Plugins/MovieScene/LevelSequenceEditor/Source/LevelSequenceEditor/Private/LevelSequenceEditorToolkit.cpp b/Engine/Plugins/MovieScene/LevelSequenceEditor/Source/LevelSequenceEditor/Private/LevelSequenceEditorToolkit.cpp
--- a/Engine/Plugins/MovieScene/LevelSequenceEditor/Source/LevelSequenceEditor/Private/LevelSequenceEditorToolkit.cpp
+++ b/Engine/Plugins/MovieScene/LevelSequenceEditor/Source/LevelSequenceEditor/Private/LevelSequenceEditorToolkit.cpp
@@ -895,7 +895,7 @@ void FLevelSequenceEditorToolkit::BindAnimationInstance(USkeletalMeshComponent*
 	Sequencer->GetHandleToObject(AnimInstance ? AnimInstance : NewObject<UAnimInstance>(SkeletalComponent));
 }
 
-bool FLevelSequenceEditorToolkit::OnRequestClose()
+void FLevelSequenceEditorToolkit::OnClose()
 {
 	UWorld* World = PlaybackContext->GetPlaybackContext();
 	UVREditorMode* VRMode = Cast<UVREditorMode>(GEditor->GetEditorWorldExtensionsManager()->GetEditorWorldExtensions(World)->FindExtension(UVREditorMode::StaticClass()));
@@ -907,7 +907,6 @@ bool FLevelSequenceEditorToolkit::OnRequestClose()
 	OpenToolkits.Remove(this);
 
 	OnClosedEvent.Broadcast();
-	return true;
 }
 
 bool FLevelSequenceEditorToolkit::CanFindInContentBrowser() const
diff --git a/Engine/Plugins/MovieScene/LevelSequenceEditor/Source/LevelSequenceEditor/Private/LevelSequenceEditorToolkit.h b/Engine/Plugins/MovieScene/LevelSequenceEditor/Source/LevelSequenceEditor/Private/LevelSequenceEditorToolkit.h
--- a/Engine/Plugins/MovieScene/LevelSequenceEditor/Source/LevelSequenceEditor/Private/LevelSequenceEditorToolkit.h
+++ b/Engine/Plugins/MovieScene/LevelSequenceEditor/Source/LevelSequenceEditor/Private/LevelSequenceEditorToolkit.h
@@ -93,7 +93,7 @@ public:
 		return TEXT("FLevelSequenceEditorToolkit");
 	}
 
-	virtual bool OnRequestClose() override;
+	virtual void OnClose() override;
 	virtual bool CanFindInContentBrowser() const override;
 
 public:
diff --git a/Engine/Plugins/Runtime/HairStrands/Source/HairStrandsEditor/Private/GroomCustomAssetEditorToolkit.cpp b/Engine/Plugins/Runtime/HairStrands/Source/HairStrandsEditor/Private/GroomCustomAssetEditorToolkit.cpp
--- a/Engine/Plugins/Runtime/HairStrands/Source/HairStrandsEditor/Private/GroomCustomAssetEditorToolkit.cpp
+++ b/Engine/Plugins/Runtime/HairStrands/Source/HairStrandsEditor/Private/GroomCustomAssetEditorToolkit.cpp
@@ -263,7 +263,7 @@ void FGroomCustomAssetEditorToolkit::InitPreviewComponents()
 #endif
 }
 
-bool FGroomCustomAssetEditorToolkit::OnRequestClose() 
+void FGroomCustomAssetEditorToolkit::OnClose() 
 {
 	//FCoreUObjectDelegates::OnObjectPropertyChanged.Remove(PropertyListenDelegate);
 	if (GroomAsset.IsValid() && PropertyListenDelegate.IsValid())
@@ -284,8 +284,6 @@ bool FGroomCustomAssetEditorToolkit::OnRequestClose()
 #if GROOMEDITOR_ENABLE_COMPONENT_PANEL
 	DetailView_PreviewGroomComponent.Reset();
 #endif
-
-	return FAssetEditorToolkit::OnRequestClose();
 }
 
 void FGroomCustomAssetEditorToolkit::InitCustomAssetEditor(const EToolkitMode::Type Mode, const TSharedPtr<class IToolkitHost>& InitToolkitHost, UGroomAsset* InCustomAsset)
diff --git a/Engine/Plugins/Runtime/HairStrands/Source/HairStrandsEditor/Public/GroomCustomAssetEditorToolkit.h b/Engine/Plugins/Runtime/HairStrands/Source/HairStrandsEditor/Public/GroomCustomAssetEditorToolkit.h
--- a/Engine/Plugins/Runtime/HairStrands/Source/HairStrandsEditor/Public/GroomCustomAssetEditorToolkit.h
+++ b/Engine/Plugins/Runtime/HairStrands/Source/HairStrandsEditor/Public/GroomCustomAssetEditorToolkit.h
@@ -60,7 +60,7 @@ public:
 	virtual FString GetWorldCentricTabPrefix() const override;
 	virtual FLinearColor GetWorldCentricTabColorScale() const override;
 	virtual bool IsPrimaryEditor() const override { return true; }
-	virtual bool OnRequestClose() override;
+	virtual void OnClose() override;
 	/** End IToolkit interface */
 
 	/** Retrieves the current custom asset. */
diff --git a/Engine/Plugins/Runtime/nDisplay/Source/DisplayClusterConfigurator/Private/DisplayClusterConfiguratorBlueprintEditor.cpp b/Engine/Plugins/Runtime/nDisplay/Source/DisplayClusterConfigurator/Private/DisplayClusterConfiguratorBlueprintEditor.cpp
--- a/Engine/Plugins/Runtime/nDisplay/Source/DisplayClusterConfigurator/Private/DisplayClusterConfiguratorBlueprintEditor.cpp
+++ b/Engine/Plugins/Runtime/nDisplay/Source/DisplayClusterConfigurator/Private/DisplayClusterConfiguratorBlueprintEditor.cpp
@@ -976,26 +976,24 @@ TStatId FDisplayClusterConfiguratorBlueprintEditor::GetStatId() const
 
 bool FDisplayClusterConfiguratorBlueprintEditor::OnRequestClose()
 {
-	bool bShouldClose = FBlueprintEditor::OnRequestClose();
+	return FBlueprintEditor::OnRequestClose();
+}
 
-	// If we are closing the blueprint editor, take the time to remove any unused host display data objects.
-	if (bShouldClose)
+void FDisplayClusterConfiguratorBlueprintEditor::OnClose()
+{
+	// When closing the blueprint editor, take the time to remove any unused host display data objects.
+	if (UDisplayClusterConfigurationData* Config = GetConfig())
 	{
-		if (UDisplayClusterConfigurationData* Config = GetConfig())
-		{
-			bool bIsDirty = LoadedBlueprint->GetOutermost()->IsDirty();
-			bool bHostDataRemoved = FDisplayClusterConfiguratorClusterUtils::RemoveUnusedHostDisplayData(Config->Cluster);
+		bool bIsDirty = LoadedBlueprint->GetOutermost()->IsDirty();
+		bool bHostDataRemoved = FDisplayClusterConfiguratorClusterUtils::RemoveUnusedHostDisplayData(Config->Cluster);
 
-			// If the blueprint wasn't dirty before, removing the unused host display data will have make it dirty, which is confusing to the user.
-			// In this case, immediately save the host display data removal to the blueprint.
-			if (!bIsDirty && CanSaveAsset() && bHostDataRemoved)
-			{
-				SaveAsset_Execute();
-			}
+		// If the blueprint wasn't dirty before, removing the unused host display data will have make it dirty, which is confusing to the user.
+		// In this case, immediately save the host display data removal to the blueprint.
+		if (!bIsDirty && CanSaveAsset() && bHostDataRemoved)
+		{
+			SaveAsset_Execute();
 		}
 	}
-
-	return bShouldClose;
 }
 
 void FDisplayClusterConfiguratorBlueprintEditor::Compile()
diff --git a/Engine/Plugins/Runtime/nDisplay/Source/DisplayClusterConfigurator/Private/DisplayClusterConfiguratorBlueprintEditor.h b/Engine/Plugins/Runtime/nDisplay/Source/DisplayClusterConfigurator/Private/DisplayClusterConfiguratorBlueprintEditor.h
--- a/Engine/Plugins/Runtime/nDisplay/Source/DisplayClusterConfigurator/Private/DisplayClusterConfiguratorBlueprintEditor.h
+++ b/Engine/Plugins/Runtime/nDisplay/Source/DisplayClusterConfigurator/Private/DisplayClusterConfiguratorBlueprintEditor.h
@@ -128,6 +128,7 @@ protected:
 	virtual bool IsTickable() const override { return true; }
 	virtual TStatId GetStatId() const override;
 	virtual bool OnRequestClose() override;
+	virtual void OnClose() override;
 	virtual void Compile() override;
 	//~ End FBlueprintEditor Interface
 
diff --git a/Engine/Source/Editor/TranslationEditor/Private/ITranslationEditor.cpp b/Engine/Source/Editor/TranslationEditor/Private/ITranslationEditor.cpp
--- a/Engine/Source/Editor/TranslationEditor/Private/ITranslationEditor.cpp
+++ b/Engine/Source/Editor/TranslationEditor/Private/ITranslationEditor.cpp
@@ -86,11 +86,9 @@ void ITranslationEditor::UnregisterTranslationEditor()
 	OpenTranslationEditors.Remove(ArchiveFilePath);
 }
 
-bool ITranslationEditor::OnRequestClose()
+void ITranslationEditor::OnClose()
 {
 	ITranslationEditor::UnregisterTranslationEditor();
-
-	return true;
 }
 
 
diff --git a/Engine/Source/Editor/TranslationEditor/Public/ITranslationEditor.h b/Engine/Source/Editor/TranslationEditor/Public/ITranslationEditor.h
--- a/Engine/Source/Editor/TranslationEditor/Public/ITranslationEditor.h
+++ b/Engine/Source/Editor/TranslationEditor/Public/ITranslationEditor.h
@@ -21,7 +21,7 @@ public:
 		, AssociatedLocalizationTarget(InAssociatedLocalizationTarget)
 	{}
 
-	virtual bool OnRequestClose() override;
+	virtual void OnClose() override;
 
 protected:
 
