LWC: Clang compile fixes

diff --git a/Engine/Plugins/Runtime/nDisplay/Source/DisplayClusterShaders/Private/WarpBlend/Loader/DisplayClusterWarpBlendLoader_WarpMap.cpp b/Engine/Plugins/Runtime/nDisplay/Source/DisplayClusterShaders/Private/WarpBlend/Loader/DisplayClusterWarpBlendLoader_WarpMap.cpp
--- a/Engine/Plugins/Runtime/nDisplay/Source/DisplayClusterShaders/Private/WarpBlend/Loader/DisplayClusterWarpBlendLoader_WarpMap.cpp
+++ b/Engine/Plugins/Runtime/nDisplay/Source/DisplayClusterShaders/Private/WarpBlend/Loader/DisplayClusterWarpBlendLoader_WarpMap.cpp
@@ -3,6 +3,7 @@
 #include "DisplayClusterWarpBlendLoader_WarpMap.h"
 
 #include "DisplayClusterShadersLog.h"
+#include "RHI.h"
 
 THIRD_PARTY_INCLUDES_START
 #include "mpcdiProfile.h"
@@ -96,7 +97,7 @@ bool FDisplayClusterWarpBlendLoader_WarpMap::Load(FLoadedWarpMapData& OutWarpMap
 
 			FVector4& Pts = OutWarpMapData.WarpData[x + y * OutWarpMapData.Width];
 
-			if ((!(fabsf(t.X) < kEpsilon && fabsf(t.Y) < kEpsilon && fabsf(t.Z) < kEpsilon))
+			if ((!(FMath::Abs(t.X) < kEpsilon && FMath::Abs(t.Y) < kEpsilon && FMath::Abs(t.Z) < kEpsilon))
 				&& (!FMath::IsNaN(t.X) && !FMath::IsNaN(t.Y) && !FMath::IsNaN(t.Z)))
 			{
 				Pts = ConventionMatrix.TransformPosition(t);
@@ -221,7 +222,7 @@ void FLoadedWarpMapData::LoadGeometry(EDisplayClusterWarpProfileType ProfileType
 		FVector4& Pts = *DstPoint;
 		DstPoint++;
 
-		if ((!(fabsf(t.X) < kEpsilon && fabsf(t.Y) < kEpsilon && fabsf(t.Z) < kEpsilon))
+		if ((!(FMath::Abs(t.X) < kEpsilon && FMath::Abs(t.Y) < kEpsilon && FMath::Abs(t.Z) < kEpsilon))
 			&& (!FMath::IsNaN(t.X) && !FMath::IsNaN(t.Y) && !FMath::IsNaN(t.Z)))
 		{
 			Pts = ConventionMatrix.TransformPosition(t);
