Fix compiler warning/error on clang

diff --git a/Engine/Plugins/Runtime/nDisplay/Source/DisplayClusterShaders/Private/WarpBlend/DisplayClusterWarpBlendManager.cpp b/Engine/Plugins/Runtime/nDisplay/Source/DisplayClusterShaders/Private/WarpBlend/DisplayClusterWarpBlendManager.cpp
--- a/Engine/Plugins/Runtime/nDisplay/Source/DisplayClusterShaders/Private/WarpBlend/DisplayClusterWarpBlendManager.cpp
+++ b/Engine/Plugins/Runtime/nDisplay/Source/DisplayClusterShaders/Private/WarpBlend/DisplayClusterWarpBlendManager.cpp
@@ -1,7 +1,5 @@
 // Copyright Epic Games, Inc. All Rights Reserved.
 
-#pragma once
-
 #include "DisplayClusterWarpBlendManager.h"
 
 #include "WarpBlend/Loader/DisplayClusterWarpBlendLoader_MPCDI.h"
