Fix compiler warning/error on clang

diff --git a/Engine/Plugins/Runtime/nDisplay/Source/DisplayCluster/Private/Render/Viewport/RenderTarget/DisplayClusterRenderTargetResource.cpp b/Engine/Plugins/Runtime/nDisplay/Source/DisplayCluster/Private/Render/Viewport/RenderTarget/DisplayClusterRenderTargetResource.cpp
--- a/Engine/Plugins/Runtime/nDisplay/Source/DisplayCluster/Private/Render/Viewport/RenderTarget/DisplayClusterRenderTargetResource.cpp
+++ b/Engine/Plugins/Runtime/nDisplay/Source/DisplayCluster/Private/Render/Viewport/RenderTarget/DisplayClusterRenderTargetResource.cpp
@@ -1,7 +1,5 @@
 // Copyright Epic Games, Inc. All Rights Reserved.
 
-#pragma once
-
 #include "Render/Viewport/RenderTarget/DisplayClusterRenderTargetResource.h"
 
 #include "RHI.h"
diff --git a/Engine/Plugins/Runtime/nDisplay/Source/DisplayClusterShaders/Private/WarpBlend/Math/DisplayClusterWarpBlendMath_Frustum.cpp b/Engine/Plugins/Runtime/nDisplay/Source/DisplayClusterShaders/Private/WarpBlend/Math/DisplayClusterWarpBlendMath_Frustum.cpp
--- a/Engine/Plugins/Runtime/nDisplay/Source/DisplayClusterShaders/Private/WarpBlend/Math/DisplayClusterWarpBlendMath_Frustum.cpp
+++ b/Engine/Plugins/Runtime/nDisplay/Source/DisplayClusterShaders/Private/WarpBlend/Math/DisplayClusterWarpBlendMath_Frustum.cpp
@@ -1,7 +1,5 @@
 // Copyright Epic Games, Inc. All Rights Reserved.
 
-#pragma once
-
 #include "WarpBlend/Math/DisplayClusterWarpBlendMath_Frustum.h"
 
 #include "Render/IDisplayClusterRenderTexture.h"
