nDisplay: Fixed TextureSharing regression

diff --git a/Engine/Plugins/Runtime/TextureShare/Source/TextureShare/Private/TextureShareModule.cpp b/Engine/Plugins/Runtime/TextureShare/Source/TextureShare/Private/TextureShareModule.cpp
--- a/Engine/Plugins/Runtime/TextureShare/Source/TextureShare/Private/TextureShareModule.cpp
+++ b/Engine/Plugins/Runtime/TextureShare/Source/TextureShare/Private/TextureShareModule.cpp
@@ -12,6 +12,8 @@
 
 #include "TextureShareCoreContainers.h"
 
+#include "Blueprints/TextureShareContainers.h"
+
 #include "TextureShareStrings.h"
 
 FTextureShareModule::FTextureShareModule()
@@ -87,7 +89,7 @@ bool FTextureShareModule::LinkSceneContextToShare(const TSharedPtr<ITextureShare
 	return false;
 }
 
-bool FTextureShareModule::SetBackbufferRect(int StereoscopicPass, FIntRect* BackbufferRect)
+bool FTextureShareModule::SetBackbufferRect(int StereoscopicPass, const FIntRect* BackbufferRect)
 {
 	if (BackbufferRect == nullptr)
 	{
@@ -161,13 +163,16 @@ void FTextureShareModule::OnResolvedSceneColor_RenderThread(FRDGBuilder& GraphBu
 {
 	FScopeLock lock(&DataGuard);
 
-	EStereoscopicPass StereoscopicPass = ViewFamily.Views[0]->StereoPass;
-	// Send SceneContext callback for all registered shares:
-	for (auto& It : TextureShareSceneContextCallback)
+	if (ViewFamily.Views.Num() > 0)
 	{
-		if (It.Value == (int)StereoscopicPass)
+		EStereoscopicPass StereoscopicPass = ViewFamily.Views[0]->StereoPass;
+		// Send SceneContext callback for all registered shares:
+		for (auto& It : TextureShareSceneContextCallback)
 		{
-			SendSceneContext_RenderThread(GraphBuilder, It.Key, SceneTextures, ViewFamily);
+			if (It.Value == (int)StereoscopicPass)
+			{
+				SendSceneContext_RenderThread(GraphBuilder, It.Key, SceneTextures, ViewFamily);
+			}
 		}
 	}
 }
@@ -176,13 +181,16 @@ void FTextureShareModule::OnPostRenderViewFamily_RenderThread(FRHICommandListImm
 {
 	FScopeLock lock(&DataGuard);
 	
-	EStereoscopicPass StereoscopicPass = ViewFamily.Views[0]->StereoPass;
-	// Send PostRender callback for all registered shares:
-	for (auto& It : TextureShareSceneContextCallback)
+	if (ViewFamily.Views.Num() > 0)
 	{
-		if (It.Value == (int)StereoscopicPass)
+		EStereoscopicPass StereoscopicPass = ViewFamily.Views[0]->StereoPass;
+		// Send PostRender callback for all registered shares:
+		for (auto& It : TextureShareSceneContextCallback)
 		{
-			SendPostRender_RenderThread(RHICmdList, It.Key, ViewFamily);
+			if (It.Value == (int)StereoscopicPass)
+			{
+				SendPostRender_RenderThread(RHICmdList, It.Key, ViewFamily);
+			}
 		}
 	}
 }
@@ -231,9 +239,12 @@ bool FTextureShareModule::SendSceneContext_RenderThread(FRDGBuilder& GraphBuilde
 #if WITH_MGPU
 			if (ViewFamily.bMultiGPUForkAndJoin)
 			{
-				// Setup GPU index for all shared scene textures
-				const FSceneView* SceneView = ViewFamily.Views[0];
-				ShareItem->SetDefaultGPUIndex(SceneView->GPUMask.GetFirstIndex());
+				if (ViewFamily.Views.Num() > 0)
+				{
+					// Setup GPU index for all shared scene textures
+					const FSceneView* SceneView = ViewFamily.Views[0];
+					ShareItem->SetDefaultGPUIndex(SceneView->GPUMask.GetFirstIndex());
+				}
 			}
 #endif
 
@@ -283,7 +294,7 @@ bool FTextureShareModule::SendPostRender_RenderThread(FRHICommandListImmediate&
 	TSharedPtr<ITextureShareItem> ShareItem;
 	if (ShareCoreAPI.GetTextureShareItem(ShareName, ShareItem) && ShareItem.IsValid() && ShareItem->IsValid() && ShareItem->IsLocalFrameLocked())
 	{
-		if (ViewFamily.RenderTarget)
+		if (ViewFamily.RenderTarget && ViewFamily.Views.Num() > 0)
 		{
 			// Get backbuffer texture
 			FTexture2DRHIRef BackBufferTexture = ViewFamily.RenderTarget->GetRenderTargetTexture();
@@ -377,10 +388,15 @@ bool FTextureShareModule::ReadFromShare_RenderThread(FRHICommandListImmediate& R
 			bResult = FTextureShareRHI::ReadFromShareTexture_RenderThread(RHICmdList, SharedRHITexture, DstTexture, DstTextureRect, bIsFormatResampleRequired);
 			ShareItem->UnlockTexture_RenderThread(TextureName);
 		}
-	}
+	} 
 	return bResult;
 }
 
+void FTextureShareModule::CastTextureShareBPSyncPolicy(const FTextureShareBPSyncPolicy& InSyncPolicy, FTextureShareSyncPolicy& OutSyncPolicy)
+{
+	OutSyncPolicy = *InSyncPolicy;
+}
+
 //////////////////////////////////////////////////////////////////////////////////////////////
 // ITextureShare
 //////////////////////////////////////////////////////////////////////////////////////////////
diff --git a/Engine/Plugins/Runtime/TextureShare/Source/TextureShare/Private/TextureShareModule.h b/Engine/Plugins/Runtime/TextureShare/Source/TextureShare/Private/TextureShareModule.h
--- a/Engine/Plugins/Runtime/TextureShare/Source/TextureShare/Private/TextureShareModule.h
+++ b/Engine/Plugins/Runtime/TextureShare/Source/TextureShare/Private/TextureShareModule.h
@@ -31,7 +31,7 @@ public:
 	virtual bool GetShare(const FString& ShareName, TSharedPtr<ITextureShareItem>& OutShareItem) const override;
 
 	virtual bool LinkSceneContextToShare(const TSharedPtr<ITextureShareItem>& ShareItem, int StereoscopicPass, bool bIsEnabled) override;
-	virtual bool SetBackbufferRect(int StereoscopicPass, FIntRect* BackbufferRect) override;
+	virtual bool SetBackbufferRect(int StereoscopicPass, const FIntRect* BackbufferRect) override;
 
 	virtual bool RegisterTexture(const TSharedPtr<ITextureShareItem>& ShareItem, const FString& InTextureName, const FIntPoint& InSize, EPixelFormat InFormat, ETextureShareSurfaceOp OperationType) override;
 
@@ -41,6 +41,8 @@ public:
 	virtual bool SendTexture_RenderThread(FRHICommandListImmediate& RHICmdList, const TSharedPtr<ITextureShareItem>& ShareItem, const FString& TextureName, FRHITexture* SrcTexture, const FIntRect* SrcTextureRect = nullptr) override;
 	virtual bool ReceiveTexture_RenderThread(FRHICommandListImmediate& RHICmdList, const TSharedPtr<ITextureShareItem>& ShareItem, const FString& TextureName, FRHITexture* DstTexture, const FIntRect* DstTextureRect = nullptr) override;
 
+	virtual void CastTextureShareBPSyncPolicy(const struct FTextureShareBPSyncPolicy& InSyncPolicy, struct FTextureShareSyncPolicy& OutSyncPolicy) override;
+
 	// Rendered callback (capture scene textures)
 	void OnResolvedSceneColor_RenderThread(FRDGBuilder& GraphBuilder, const FSceneTextures& SceneTextures, class FSceneViewFamily& ViewFamily);
 	void OnPostRenderViewFamily_RenderThread(FRHICommandListImmediate& RHICmdList, class FSceneViewFamily& ViewFamily);
diff --git a/Engine/Plugins/Runtime/TextureShare/Source/TextureShare/Public/ITextureShare.h b/Engine/Plugins/Runtime/TextureShare/Source/TextureShare/Public/ITextureShare.h
--- a/Engine/Plugins/Runtime/TextureShare/Source/TextureShare/Public/ITextureShare.h
+++ b/Engine/Plugins/Runtime/TextureShare/Source/TextureShare/Public/ITextureShare.h
@@ -91,7 +91,7 @@ public:
 	 *
 	 * @return True if the success
 	 */
-	virtual bool SetBackbufferRect(int StereoscopicPass, FIntRect* BackbufferRect) = 0;
+	virtual bool SetBackbufferRect(int StereoscopicPass, const FIntRect* BackbufferRect) = 0;
 
 	/**
 	 * Register/Update texture info for share
@@ -158,4 +158,10 @@ public:
 	 * @return True if the success
 	 */
 	virtual bool ReceiveTexture_RenderThread(FRHICommandListImmediate& RHICmdList, const TSharedPtr<ITextureShareItem>& ShareItem, const FString& TextureName, FRHITexture* DstTexture, const FIntRect* DstTextureRect = nullptr) = 0;
+
+	/**
+	 * Cast internal structures bp<->cpp
+	 *
+	 */
+	virtual void CastTextureShareBPSyncPolicy(const struct FTextureShareBPSyncPolicy& InSyncPolicy, struct FTextureShareSyncPolicy& OutSyncPolicy) = 0;
 };
diff --git a/Engine/Plugins/Runtime/TextureShare/Source/TextureShareCore/Private/Core/TextureShareItemBase.cpp b/Engine/Plugins/Runtime/TextureShare/Source/TextureShareCore/Private/Core/TextureShareItemBase.cpp
--- a/Engine/Plugins/Runtime/TextureShare/Source/TextureShareCore/Private/Core/TextureShareItemBase.cpp
+++ b/Engine/Plugins/Runtime/TextureShare/Source/TextureShareCore/Private/Core/TextureShareItemBase.cpp
@@ -359,6 +359,12 @@ namespace TextureShareItem
 		return false;
 	}
 
+	bool FTextureShareItemBase::SetCustomProjectionData(const FTextureShareCustomProjectionData& InCustomProjectionData)
+	{
+		// NOT IMPLEMETED
+		return false;
+	}
+
 
 	bool FTextureShareItemBase::BeginTextureOp(FSharedResourceTexture& LocalTextureData)
 	{
diff --git a/Engine/Plugins/Runtime/TextureShare/Source/TextureShareCore/Private/Core/TextureShareItemBase.h b/Engine/Plugins/Runtime/TextureShare/Source/TextureShareCore/Private/Core/TextureShareItemBase.h
--- a/Engine/Plugins/Runtime/TextureShare/Source/TextureShareCore/Private/Core/TextureShareItemBase.h
+++ b/Engine/Plugins/Runtime/TextureShare/Source/TextureShareCore/Private/Core/TextureShareItemBase.h
@@ -51,6 +51,9 @@ namespace TextureShareItem
 		virtual bool SetLocalAdditionalData(const FTextureShareAdditionalData& InAdditionalData) override;
 		virtual bool GetRemoteAdditionalData(FTextureShareAdditionalData& OutAdditionalData) override;
 
+		virtual bool SetCustomProjectionData(const FTextureShareCustomProjectionData& InCustomProjectionData) override;
+
+
 		virtual ETextureShareDevice GetDeviceType() const override
 			{ return ETextureShareDevice::Undefined; }
 		virtual ITextureShareItemD3D11* GetD3D11() override
diff --git a/Engine/Plugins/Runtime/TextureShare/Source/TextureShareCore/Private/TextureShareCore.cpp b/Engine/Plugins/Runtime/TextureShare/Source/TextureShareCore/Private/TextureShareCore.cpp
--- a/Engine/Plugins/Runtime/TextureShare/Source/TextureShareCore/Private/TextureShareCore.cpp
+++ b/Engine/Plugins/Runtime/TextureShare/Source/TextureShareCore/Private/TextureShareCore.cpp
@@ -196,4 +196,16 @@ void FTextureShareCoreModule::SetSyncPolicySettings(ETextureShareProcess Process
 	TextureShareItem::FTextureShareItemBase::SetSyncPolicySettings(Process, InSyncPolicySettings);
 }
 
+bool FTextureShareCoreModule::BeginSyncFrame()
+{
+	// NOT IMPLEMENTED
+	return false;
+}
+
+bool FTextureShareCoreModule::EndSyncFrame()
+{
+	// NOT IMPLEMENTED
+	return false;
+}
+
 IMPLEMENT_MODULE(FTextureShareCoreModule, TextureShareCore);
diff --git a/Engine/Plugins/Runtime/TextureShare/Source/TextureShareCore/Private/TextureShareCore.h b/Engine/Plugins/Runtime/TextureShare/Source/TextureShareCore/Private/TextureShareCore.h
--- a/Engine/Plugins/Runtime/TextureShare/Source/TextureShareCore/Private/TextureShareCore.h
+++ b/Engine/Plugins/Runtime/TextureShare/Source/TextureShareCore/Private/TextureShareCore.h
@@ -26,6 +26,9 @@ public:
 	virtual bool ReleaseTextureShareItem(const FString& ShareName) override;
 	virtual bool GetTextureShareItem(const FString& ShareName, TSharedPtr<ITextureShareItem>& OutShareObject) const override;
 
+	virtual bool BeginSyncFrame() override;
+	virtual bool EndSyncFrame() override;
+
 protected:
 	void InitializeProcessMemory();
 	void ReleaseProcessMemory();
diff --git a/Engine/Plugins/Runtime/TextureShare/Source/TextureShareCore/Public/Containers/TextureShareCoreGenericContainers.h b/Engine/Plugins/Runtime/TextureShare/Source/TextureShareCore/Public/Containers/TextureShareCoreGenericContainers.h
--- a/Engine/Plugins/Runtime/TextureShare/Source/TextureShareCore/Public/Containers/TextureShareCoreGenericContainers.h
+++ b/Engine/Plugins/Runtime/TextureShare/Source/TextureShareCore/Public/Containers/TextureShareCoreGenericContainers.h
@@ -93,6 +93,13 @@ public:
 			break;
 		}
 	}
+
+	bool operator==(const FTextureShareSyncPolicy& InSyncPolicy) const
+	{
+		return ConnectionSync == InSyncPolicy.ConnectionSync
+			&& FrameSync == InSyncPolicy.FrameSync
+			&& TextureSync == InSyncPolicy.TextureSync;
+	}
 };
 
 struct FTextureShareTimeOut
diff --git a/Engine/Plugins/Runtime/TextureShare/Source/TextureShareCore/Public/ITextureShareCore.h b/Engine/Plugins/Runtime/TextureShare/Source/TextureShareCore/Public/ITextureShareCore.h
--- a/Engine/Plugins/Runtime/TextureShare/Source/TextureShareCore/Public/ITextureShareCore.h
+++ b/Engine/Plugins/Runtime/TextureShare/Source/TextureShareCore/Public/ITextureShareCore.h
@@ -88,4 +88,18 @@ class ITextureShareCore
 	 * Release all created ITextureShareItem objects
 	 */
 	virtual void ReleaseLib() = 0;
+
+	/** Frame sync scope */
+
+	/**
+	 * Begin global frame sync
+	 * NOT IMPLEMENTED
+	 */
+	virtual bool BeginSyncFrame() = 0;
+
+	/**
+	 * Finalize global frame sync
+	 * NOT IMPLEMENTED
+	 */
+	virtual bool EndSyncFrame() = 0;
 };
diff --git a/Engine/Plugins/Runtime/TextureShare/Source/TextureShareCore/Public/ITextureShareItem.h b/Engine/Plugins/Runtime/TextureShare/Source/TextureShareCore/Public/ITextureShareItem.h
--- a/Engine/Plugins/Runtime/TextureShare/Source/TextureShareCore/Public/ITextureShareItem.h
+++ b/Engine/Plugins/Runtime/TextureShare/Source/TextureShareCore/Public/ITextureShareItem.h
@@ -66,6 +66,10 @@ public:
 	virtual bool SetLocalAdditionalData(const FTextureShareAdditionalData& InAdditionalData) = 0;
 	virtual bool GetRemoteAdditionalData(FTextureShareAdditionalData& OutAdditionalData) = 0;
 
+	// NOT IMPLEMENTED
+	// Set custom projection data (override UE4 matrix)
+	virtual bool SetCustomProjectionData(const FTextureShareCustomProjectionData& InCustomProjectionData) = 0;
+
 	// Release this object
 	virtual void Release() = 0;
 };
diff --git a/Engine/Plugins/Runtime/TextureShare/Source/TextureShareCore/Public/TextureShareCoreContainers.h b/Engine/Plugins/Runtime/TextureShare/Source/TextureShareCore/Public/TextureShareCoreContainers.h
--- a/Engine/Plugins/Runtime/TextureShare/Source/TextureShareCore/Public/TextureShareCoreContainers.h
+++ b/Engine/Plugins/Runtime/TextureShare/Source/TextureShareCore/Public/TextureShareCoreContainers.h
@@ -23,3 +23,14 @@ public:
 
 	//@todo: add more frame data
 };
+
+// Custom render
+struct FTextureShareCustomProjectionData
+{
+	// Projection matrix
+	FMatrix PrjMatrix;
+
+	FVector  ViewLocation;
+	FRotator ViewRotation;
+	FVector  ViewScale;
+};
diff --git a/Engine/Plugins/Runtime/nDisplay/Source/DisplayCluster/DisplayCluster.Build.cs b/Engine/Plugins/Runtime/nDisplay/Source/DisplayCluster/DisplayCluster.Build.cs
--- a/Engine/Plugins/Runtime/nDisplay/Source/DisplayCluster/DisplayCluster.Build.cs
+++ b/Engine/Plugins/Runtime/nDisplay/Source/DisplayCluster/DisplayCluster.Build.cs
@@ -23,22 +23,29 @@ public class DisplayCluster : ModuleRules
 				"Engine"
 			});
 
-		PrivateDependencyModuleNames.AddRange(
+		if (Target.Platform == UnrealTargetPlatform.Win64)
+		{
+			PublicDependencyModuleNames.AddRange(
 			new string[] {
-				"HeadMountedDisplay",
-				"InputCore",
-				"Json",
-				"JsonUtilities",
-				"Networking",
-				"OpenColorIO",
-				"Renderer",
-				"RenderCore",
-				"RHI",
-				"RHICore",
-				"Slate",
-				"SlateCore",
-				"Sockets",
+				"TextureShare"
 			});
+		}
+
+		PrivateDependencyModuleNames.AddRange(
+		new string[] {
+			"HeadMountedDisplay",
+			"InputCore",
+			"Json",
+			"JsonUtilities",
+			"Networking",
+			"OpenColorIO",
+			"Renderer",
+			"RenderCore",
+			"RHI",
+			"Slate",
+			"SlateCore",
+			"Sockets",
+		});
 
 		if (Target.bBuildEditor == true)
 		{
diff --git a/Engine/Plugins/Runtime/nDisplay/Source/DisplayCluster/Private/Render/Device/DisplayClusterDeviceBase.cpp b/Engine/Plugins/Runtime/nDisplay/Source/DisplayCluster/Private/Render/Device/DisplayClusterDeviceBase.cpp
--- a/Engine/Plugins/Runtime/nDisplay/Source/DisplayCluster/Private/Render/Device/DisplayClusterDeviceBase.cpp
+++ b/Engine/Plugins/Runtime/nDisplay/Source/DisplayCluster/Private/Render/Device/DisplayClusterDeviceBase.cpp
@@ -30,10 +30,6 @@
 #include "Render/Projection/IDisplayClusterProjectionPolicyFactory.h"
 #include "Render/Synchronization/IDisplayClusterRenderSyncPolicy.h"
 
-#if PLATFORM_WINDOWS
-#include "ITextureShare.h"
-#include "ITextureShareItem.h"
-#endif
 
 #include "Render/Viewport/DisplayClusterViewportManager.h"
 #include "Render/Viewport/DisplayClusterViewportManagerProxy.h"
@@ -486,57 +482,6 @@ void FDisplayClusterDeviceBase::UpdateViewport(bool bUseSeparateRenderTarget, co
 	{
 		// UE viewport
 		MainViewport = (FViewport*)&Viewport;
-
-#if PLATFORM_WINDOWS
-		/*
-		// Create texture share for render viewports by config line flag
-		//@todo move to right place. add on\off
-		{
-			// Share viewports to external apps
-			for (int ViewportIndex = 0; ViewportIndex < RenderViewports.Num(); ViewportIndex++)
-			{
-				if (RenderViewports[ViewportIndex].IsShared())
-				{
-					static ITextureShare& TextureShareAPI = ITextureShare::Get();
-
-					//@todo: add custom sync setup
-					FTextureShareSyncPolicy SyncPolicy;
-
-					FString ShareName = RenderViewports[ViewportIndex].GetId();
-					EStereoscopicPass PassType = GetViewPassForIndex(IsStereoEnabled(), ViewportIndex);
-
-					// Create shared resource for external app
-					if (!TextureShareAPI.CreateShare(ShareName, SyncPolicy, ETextureShareProcess::Server))
-					{
-						UE_LOG(LogDisplayClusterRender, Error, TEXT("Failed create viewport share '%s'"), *ShareName);
-					}
-					else
-					{
-						// Find viewport stereoscopic pass
-						int ResourceViewportIndex = RenderViewports.Num() - 1;
-
-						// Initialize render callbacks
-						TSharedPtr<ITextureShareItem> ShareItem;
-						if (TextureShareAPI.GetShare(ShareName, ShareItem))
-						{
-							if(TextureShareAPI.LinkSceneContextToShare(ShareItem, PassType, true))
-							{
-								// Map viewport rect to stereoscopic pass
-								TextureShareAPI.SetBackbufferRect(PassType, &RenderViewports[ViewportIndex].GetRect());
-								// Begin share session
-								ShareItem->BeginSession();
-							}
-							else
-							{
-								TextureShareAPI.ReleaseShare(ShareName);
-								UE_LOG(LogDisplayClusterRender, Error, TEXT("failed link scene conext for share '%s'"), *ShareName);
-							}
-						}
-					}
-				}
-			}
-		}*/
-#endif
 	}
 }
 
diff --git a/Engine/Plugins/Runtime/nDisplay/Source/DisplayCluster/Private/Render/Viewport/Configuration/DisplayClusterViewportConfiguration.cpp b/Engine/Plugins/Runtime/nDisplay/Source/DisplayCluster/Private/Render/Viewport/Configuration/DisplayClusterViewportConfiguration.cpp
--- a/Engine/Plugins/Runtime/nDisplay/Source/DisplayCluster/Private/Render/Viewport/Configuration/DisplayClusterViewportConfiguration.cpp
+++ b/Engine/Plugins/Runtime/nDisplay/Source/DisplayCluster/Private/Render/Viewport/Configuration/DisplayClusterViewportConfiguration.cpp
@@ -90,6 +90,11 @@ bool FDisplayClusterViewportConfiguration::UpdateConfiguration(EDisplayClusterRe
 
 			ImplUpdateConfiguration_PostProcess(InClusterNodeId, *ConfigurationData);
 
+			if (!RenderFrameSettings.bIsRenderingInEditor)
+			{
+				// TextureShare not supported in Editor Preview
+				ConfigurationBase.UpdateTextureShare(InClusterNodeId);
+			}
 
 			return true;
 		}
diff --git a/Engine/Plugins/Runtime/nDisplay/Source/DisplayCluster/Private/Render/Viewport/Configuration/DisplayClusterViewportConfigurationBase.cpp b/Engine/Plugins/Runtime/nDisplay/Source/DisplayCluster/Private/Render/Viewport/Configuration/DisplayClusterViewportConfigurationBase.cpp
--- a/Engine/Plugins/Runtime/nDisplay/Source/DisplayCluster/Private/Render/Viewport/Configuration/DisplayClusterViewportConfigurationBase.cpp
+++ b/Engine/Plugins/Runtime/nDisplay/Source/DisplayCluster/Private/Render/Viewport/Configuration/DisplayClusterViewportConfigurationBase.cpp
@@ -81,3 +81,18 @@ bool FDisplayClusterViewportConfigurationBase::UpdateViewportConfiguration(FDisp
 	return true;
 }
 
+void FDisplayClusterViewportConfigurationBase::UpdateTextureShare(const FString& ClusterNodeId)
+{
+	for (FDisplayClusterViewport* ViewportIt : ViewportManager.ImplGetViewports())
+	{
+		if (ViewportIt)
+		{
+			const UDisplayClusterConfigurationViewport* ViewportCfg = ConfigurationData.GetViewport(ClusterNodeId, ViewportIt->GetId());
+			if (ViewportCfg)
+			{
+				ViewportIt->TextureShare.UpdateConfiguration(*ViewportIt, ViewportCfg->TextureShare);
+			}
+		}
+	}
+}
+
diff --git a/Engine/Plugins/Runtime/nDisplay/Source/DisplayCluster/Private/Render/Viewport/Configuration/DisplayClusterViewportConfigurationBase.h b/Engine/Plugins/Runtime/nDisplay/Source/DisplayCluster/Private/Render/Viewport/Configuration/DisplayClusterViewportConfigurationBase.h
--- a/Engine/Plugins/Runtime/nDisplay/Source/DisplayCluster/Private/Render/Viewport/Configuration/DisplayClusterViewportConfigurationBase.h
+++ b/Engine/Plugins/Runtime/nDisplay/Source/DisplayCluster/Private/Render/Viewport/Configuration/DisplayClusterViewportConfigurationBase.h
@@ -24,6 +24,7 @@ public:
 
 public:
 	void Update(const TArray<FString>& InClusterNodeIds);
+	void UpdateTextureShare(const FString& ClusterNodeId);
 
 	static bool UpdateViewportConfiguration(FDisplayClusterViewportManager& ViewportManager, ADisplayClusterRootActor& RootActor, FDisplayClusterViewport* DesiredViewport, const UDisplayClusterConfigurationViewport* ConfigurationViewport);
 
diff --git a/Engine/Plugins/Runtime/nDisplay/Source/DisplayCluster/Private/Render/Viewport/DisplayClusterViewport.h b/Engine/Plugins/Runtime/nDisplay/Source/DisplayCluster/Private/Render/Viewport/DisplayClusterViewport.h
--- a/Engine/Plugins/Runtime/nDisplay/Source/DisplayCluster/Private/Render/Viewport/DisplayClusterViewport.h
+++ b/Engine/Plugins/Runtime/nDisplay/Source/DisplayCluster/Private/Render/Viewport/DisplayClusterViewport.h
@@ -5,6 +5,7 @@
 #include "Render/Viewport/IDisplayClusterViewport.h"
 #include "Render/Viewport/DisplayClusterViewport_CustomPostProcessSettings.h"
 #include "Render/Viewport/DisplayClusterViewport_VisibilitySettings.h"
+#include "Render/Viewport/DisplayClusterViewport_TextureShare.h"
 
 #include "Render/Viewport/Containers/ImplDisplayClusterViewport_CameraMotionBlur.h"
 #include "Render/Viewport/Containers/ImplDisplayClusterViewport_Overscan.h"
@@ -167,6 +168,9 @@ public:
 	// OCIO wrapper
 	TSharedPtr<FOpenColorIODisplayExtension, ESPMode::ThreadSafe> OpenColorIODisplayExtension;
 
+	//TextureShare
+	FDisplayClusterViewport_TextureShare TextureShare;
+
 public:
 	// Projection policy instance that serves this viewport
 	TSharedPtr<IDisplayClusterProjectionPolicy, ESPMode::ThreadSafe> ProjectionPolicy;
diff --git a/Engine/Plugins/Runtime/nDisplay/Source/DisplayCluster/Private/Render/Viewport/DisplayClusterViewportManager.cpp b/Engine/Plugins/Runtime/nDisplay/Source/DisplayCluster/Private/Render/Viewport/DisplayClusterViewportManager.cpp
--- a/Engine/Plugins/Runtime/nDisplay/Source/DisplayCluster/Private/Render/Viewport/DisplayClusterViewportManager.cpp
+++ b/Engine/Plugins/Runtime/nDisplay/Source/DisplayCluster/Private/Render/Viewport/DisplayClusterViewportManager.cpp
@@ -208,6 +208,12 @@ bool FDisplayClusterViewportManager::BeginNewFrame(class FViewport* InViewport,
 		return false;
 	}
 
+	// Update TextureShare links
+	for (FDisplayClusterViewport* Viewport : Viewports)
+	{
+		Viewport->TextureShare.UpdateLinkSceneContextToShare(*Viewport);
+	}
+
 	// Update desired views number
 	OutRenderFrame.UpdateDesiredNumberOfViews();
 
diff --git a/Engine/Plugins/Runtime/nDisplay/Source/DisplayCluster/Private/Render/Viewport/DisplayClusterViewport_TextureShare.cpp b/Engine/Plugins/Runtime/nDisplay/Source/DisplayCluster/Private/Render/Viewport/DisplayClusterViewport_TextureShare.cpp
new file mode 100644
--- /dev/null
+++ b/Engine/Plugins/Runtime/nDisplay/Source/DisplayCluster/Private/Render/Viewport/DisplayClusterViewport_TextureShare.cpp
@@ -0,0 +1,195 @@
+// Copyright Epic Games, Inc. All Rights Reserved.
+
+#include "Render/Viewport/DisplayClusterViewport_TextureShare.h"
+#include "Render/Viewport/DisplayClusterViewport.h"
+
+#if PLATFORM_WINDOWS
+	#include "Containers/TextureShareCoreEnums.h"
+	#include "Blueprints/TextureShareContainers.h"
+	#include "ITextureShare.h"
+	#include "ITextureShareItem.h"
+#endif
+
+#include "DisplayClusterConfigurationTypes_TextureShare.h"
+#include "Misc/DisplayClusterLog.h"
+
+FDisplayClusterViewport_TextureShare::FDisplayClusterViewport_TextureShare()
+{
+
+}
+
+FDisplayClusterViewport_TextureShare::~FDisplayClusterViewport_TextureShare()
+{
+	ImplDelete();
+}
+
+bool FDisplayClusterViewport_TextureShare::ImplDelete()
+{
+#if PLATFORM_WINDOWS
+	if (ExistConfiguration.IsValid())
+	{
+		FString ExistTextureShareId = ExistConfiguration.TextureShareId;
+		ExistConfiguration.Empty();
+
+		if(ITextureShare::Get().ReleaseShare(ExistTextureShareId))
+		{
+			return true;
+		}
+
+		UE_LOG(LogDisplayClusterViewport, Error, TEXT("Failed delete viewport share '%s'"), *ExistTextureShareId);
+	}
+#endif
+
+	return false;
+}
+
+bool FDisplayClusterViewport_TextureShare::ImplCreate(const FTextureShareConfiguration& InConfiguration)
+{
+#if PLATFORM_WINDOWS
+	if (!ExistConfiguration.IsValid())
+	{
+		if (ITextureShare::Get().CreateShare(InConfiguration.TextureShareId, InConfiguration.SyncSettings, ETextureShareProcess::Server))
+		{
+			ExistConfiguration = InConfiguration;
+			return true;
+		}
+
+		UE_LOG(LogDisplayClusterViewport, Error, TEXT("Failed create viewport share '%s'"), *InConfiguration.TextureShareId);
+	}
+#endif
+
+	return false;
+}
+
+bool FDisplayClusterViewport_TextureShare::UpdateLinkSceneContextToShare(const FDisplayClusterViewport& InViewport)
+{
+#if PLATFORM_WINDOWS
+
+	const TArray<FDisplayClusterViewport_Context>& InContexts = InViewport.GetContexts();
+	// Now share only mono/left eye
+	if (InContexts.Num() > 0)
+	{
+		const FDisplayClusterViewport_Context& InContext = InContexts[0];
+		if (!InContext.bDisableRender)
+		{
+			// This viewport in render
+
+			// Update texture share configuration on fly
+			if (NewConfiguration.IsValid())
+			{
+				ImplDelete();
+				bool bIsTextureShareCreated = ImplCreate(NewConfiguration);
+
+				if(!bIsTextureShareCreated)
+				{
+					// MU issue: on recompile destructor not called for ViewportManager from RootActor
+					// fix it
+					ExistConfiguration = NewConfiguration;
+					ImplDelete();
+					bool bIsTextureShareCreatedFixed = ImplCreate(NewConfiguration);
+					if (!bIsTextureShareCreatedFixed)
+					{
+						return false;
+					}
+				}
+
+				NewConfiguration.Empty();
+			}
+
+			if(ExistConfiguration.IsValid())
+			{
+				int PassType = InContext.StereoscopicPass;
+
+				static ITextureShare& TextureShareAPI = ITextureShare::Get();
+				TSharedPtr<ITextureShareItem> ShareItem;
+				if (TextureShareAPI.GetShare(ExistConfiguration.TextureShareId, ShareItem))
+				{
+					if (TextureShareAPI.LinkSceneContextToShare(ShareItem, PassType, true))
+					{
+						// Map viewport rect to stereoscopic pass
+						TextureShareAPI.SetBackbufferRect(PassType, &InContext.RenderTargetRect);
+
+						// Begin share session
+						if (!ShareItem->IsSessionValid())
+						{
+							ShareItem->BeginSession();
+						}
+
+						return true;
+					}
+					else
+					{
+						UE_LOG(LogDisplayClusterViewport, Error, TEXT("failed link scene conext for share '%s'"), *ExistConfiguration.TextureShareId);
+					}
+				}
+				else
+				{
+					UE_LOG(LogDisplayClusterViewport, Error, TEXT("UpdateLinkSceneContextToShare: share '%s' not exist. Remove"), *ExistConfiguration.TextureShareId);
+				}
+			}
+		}
+	}
+
+	//
+	ImplDelete();
+	NewConfiguration.Empty();
+
+#endif
+
+	return false;
+}
+
+bool FDisplayClusterViewport_TextureShare::UpdateConfiguration(const FDisplayClusterViewport& InViewport, const FDisplayClusterConfigurationTextureShare_Viewport& InConfiguration)
+{
+#if PLATFORM_WINDOWS
+	if (InConfiguration.bIsEnabled)
+	{
+		FTextureShareSyncPolicy SyncPolicy;
+		FTextureShareBPSyncPolicy SyncSettings;
+
+		// sync settings
+		switch(InConfiguration.SyncSettings.Connection)
+		{
+			case ETextureShareSyncConnectDisplayCluster::Default: SyncSettings.Connection = ETextureShareBPSyncConnect::Default; break;
+			case ETextureShareSyncConnectDisplayCluster::None: SyncSettings.Connection = ETextureShareBPSyncConnect::None; break;
+			case ETextureShareSyncConnectDisplayCluster::SyncSession: SyncSettings.Connection = ETextureShareBPSyncConnect::SyncSession; break;
+			default: UE_LOG(LogDisplayClusterViewport, Error, TEXT("Unsupported texture share sync connection type"));  break;
+		}
+
+		switch (InConfiguration.SyncSettings.Frame)
+		{
+			case ETextureShareSyncFrameDisplayCluster::Default: SyncSettings.Frame = ETextureShareBPSyncFrame::Default; break;
+			case ETextureShareSyncFrameDisplayCluster::None: SyncSettings.Frame = ETextureShareBPSyncFrame::None; break;
+			case ETextureShareSyncFrameDisplayCluster::FrameSync: SyncSettings.Frame = ETextureShareBPSyncFrame::FrameSync; break;
+			default: UE_LOG(LogDisplayClusterViewport, Error, TEXT("Unsupported texture share sync frame type"));  break; break;
+		}
+
+		switch (InConfiguration.SyncSettings.Texture)
+		{
+			case ETextureShareSyncSurfaceDisplayCluster::Default: SyncSettings.Texture = ETextureShareBPSyncSurface::Default; break;
+			case ETextureShareSyncSurfaceDisplayCluster::None: SyncSettings.Texture = ETextureShareBPSyncSurface::None; break;
+			case ETextureShareSyncSurfaceDisplayCluster::SyncRead: SyncSettings.Texture = ETextureShareBPSyncSurface::SyncRead; break;
+			case ETextureShareSyncSurfaceDisplayCluster::SyncPairingRead: SyncSettings.Texture = ETextureShareBPSyncSurface::SyncPairingRead; break;
+			default: UE_LOG(LogDisplayClusterViewport, Error, TEXT("Unsupported texture share sync texture type"));  break;
+		}
+
+		static ITextureShare& TextureShareAPI = ITextureShare::Get();
+		TextureShareAPI.CastTextureShareBPSyncPolicy(SyncSettings, SyncPolicy);
+
+		FTextureShareConfiguration InCfg(InViewport.GetId(), SyncPolicy);
+		bool bIsConfigurationEqual = InCfg == ExistConfiguration;
+
+		if (!bIsConfigurationEqual || NewConfiguration.IsValid())
+		{
+			// Update share with new cfg delayed
+			NewConfiguration = InCfg;
+		}
+
+		return true;
+	}
+
+	// remove unused share
+	return ImplDelete();
+#endif
+	return false;
+}
diff --git a/Engine/Plugins/Runtime/nDisplay/Source/DisplayCluster/Private/Render/Viewport/DisplayClusterViewport_TextureShare.h b/Engine/Plugins/Runtime/nDisplay/Source/DisplayCluster/Private/Render/Viewport/DisplayClusterViewport_TextureShare.h
new file mode 100644
--- /dev/null
+++ b/Engine/Plugins/Runtime/nDisplay/Source/DisplayCluster/Private/Render/Viewport/DisplayClusterViewport_TextureShare.h
@@ -0,0 +1,65 @@
+// Copyright Epic Games, Inc. All Rights Reserved.
+
+#pragma once
+
+#include "CoreMinimal.h"
+
+#if PLATFORM_WINDOWS
+#include "Containers/TextureShareCoreGenericContainers.h"
+#endif
+
+struct FDisplayClusterConfigurationTextureShare_Viewport;
+class FDisplayClusterViewport;
+
+class FDisplayClusterViewport_TextureShare
+{
+	struct FTextureShareConfiguration
+	{
+	public:
+		FTextureShareConfiguration()
+		{ }
+
+#if PLATFORM_WINDOWS
+		FTextureShareConfiguration(const FString& InTextureShareId, const FTextureShareSyncPolicy& InSyncSettings)
+			: TextureShareId(InTextureShareId)
+			, SyncSettings(InSyncSettings)
+		{ }
+
+		bool operator==(const FTextureShareConfiguration& InCfg) const
+		{
+			return TextureShareId == InCfg.TextureShareId
+				&& SyncSettings == InCfg.SyncSettings;
+		}
+
+	public:
+		FString TextureShareId;
+		FTextureShareSyncPolicy SyncSettings;
+
+		void Empty()
+		{
+			TextureShareId.Empty();
+		}
+
+		bool IsValid() const
+		{
+			return !TextureShareId.IsEmpty();
+		}
+#endif
+	};
+
+public:
+	FDisplayClusterViewport_TextureShare();
+	~FDisplayClusterViewport_TextureShare();
+
+public:
+	bool UpdateConfiguration(const FDisplayClusterViewport& InViewport, const FDisplayClusterConfigurationTextureShare_Viewport& InConfiguration);
+	bool UpdateLinkSceneContextToShare(const FDisplayClusterViewport& InViewport);
+
+private:
+	bool ImplCreate(const FTextureShareConfiguration& InConfiguration);
+	bool ImplDelete();
+
+private:
+	FTextureShareConfiguration ExistConfiguration;
+	FTextureShareConfiguration NewConfiguration;
+};
diff --git a/Engine/Plugins/Runtime/nDisplay/Source/DisplayClusterConfiguration/DisplayClusterConfiguration.Build.cs b/Engine/Plugins/Runtime/nDisplay/Source/DisplayClusterConfiguration/DisplayClusterConfiguration.Build.cs
--- a/Engine/Plugins/Runtime/nDisplay/Source/DisplayClusterConfiguration/DisplayClusterConfiguration.Build.cs
+++ b/Engine/Plugins/Runtime/nDisplay/Source/DisplayClusterConfiguration/DisplayClusterConfiguration.Build.cs
@@ -20,9 +20,17 @@ public class DisplayClusterConfiguration : ModuleRules
 			new string[] {
 				"CinematicCamera",
 				"ActorLayerUtilities",
-				"OpenColorIO",
+				"OpenColorIO"
 			});
 
+		if (Target.Platform == UnrealTargetPlatform.Win64)
+		{
+			PublicDependencyModuleNames.AddRange(
+			new string[] {
+				"TextureShare"
+			});
+		}
+
 		PrivateDependencyModuleNames.AddRange(
 			new string[] {
 				"Core",
diff --git a/Engine/Plugins/Runtime/nDisplay/Source/DisplayClusterConfiguration/Private/Formats/JSON426/DisplayClusterConfigurationJsonParser_426.cpp b/Engine/Plugins/Runtime/nDisplay/Source/DisplayClusterConfiguration/Private/Formats/JSON426/DisplayClusterConfigurationJsonParser_426.cpp
--- a/Engine/Plugins/Runtime/nDisplay/Source/DisplayClusterConfiguration/Private/Formats/JSON426/DisplayClusterConfigurationJsonParser_426.cpp
+++ b/Engine/Plugins/Runtime/nDisplay/Source/DisplayClusterConfiguration/Private/Formats/JSON426/DisplayClusterConfigurationJsonParser_426.cpp
@@ -197,8 +197,6 @@ namespace JSON426
 					Viewport->Camera = CfgViewport.Value.Camera;
 					Viewport->Region = FDisplayClusterConfigurationRectangle(CfgViewport.Value.Region.X, CfgViewport.Value.Region.Y, CfgViewport.Value.Region.W, CfgViewport.Value.Region.H);
 					Viewport->GPUIndex = CfgViewport.Value.GPUIndex;
-					Viewport->bIsShared = CfgViewport.Value.IsShared;
-					//Viewport->bAllowCrossGPUTransfer = CfgViewport.Value.AllowCrossGPUTransfer;
 
 					// Projection policy
 					Viewport->ProjectionPolicy.Type = CfgViewport.Value.ProjectionPolicy.Type;
diff --git a/Engine/Plugins/Runtime/nDisplay/Source/DisplayClusterConfiguration/Private/Formats/JSON427/DisplayClusterConfigurationJsonParser_427.cpp b/Engine/Plugins/Runtime/nDisplay/Source/DisplayClusterConfiguration/Private/Formats/JSON427/DisplayClusterConfigurationJsonParser_427.cpp
--- a/Engine/Plugins/Runtime/nDisplay/Source/DisplayClusterConfiguration/Private/Formats/JSON427/DisplayClusterConfigurationJsonParser_427.cpp
+++ b/Engine/Plugins/Runtime/nDisplay/Source/DisplayClusterConfiguration/Private/Formats/JSON427/DisplayClusterConfigurationJsonParser_427.cpp
@@ -180,8 +180,6 @@ namespace JSON427
 					Viewport->Camera = CfgViewport.Value.Camera;
 					Viewport->Region = FDisplayClusterConfigurationRectangle(CfgViewport.Value.Region.X, CfgViewport.Value.Region.Y, CfgViewport.Value.Region.W, CfgViewport.Value.Region.H);
 					Viewport->GPUIndex = CfgViewport.Value.GPUIndex;
-					Viewport->bIsShared = CfgViewport.Value.IsShared;
-					//Viewport->bAllowCrossGPUTransfer = CfgViewport.Value.AllowCrossGPUTransfer;
 
 					// Projection policy
 					Viewport->ProjectionPolicy.Type = CfgViewport.Value.ProjectionPolicy.Type;
@@ -300,7 +298,7 @@ namespace JSON427
 					Viewport.Camera = CfgViewport.Value->Camera;
 					Viewport.Region = FDisplayClusterConfigurationJsonRectangle_427(CfgViewport.Value->Region.X, CfgViewport.Value->Region.Y, CfgViewport.Value->Region.W, CfgViewport.Value->Region.H);
 					Viewport.GPUIndex = CfgViewport.Value->GPUIndex;
-					Viewport.IsShared = CfgViewport.Value->bIsShared;
+					//Viewport.IsShared = CfgViewport.Value->bIsShared;
 					Viewport.BufferRatio = CfgViewport.Value->RenderSettings.BufferRatio;
 					//Viewport.AllowCrossGPUTransfer = CfgViewport.Value->bAllowCrossGPUTransfer;
 
diff --git a/Engine/Plugins/Runtime/nDisplay/Source/DisplayClusterConfiguration/Private/Formats/Text/DisplayClusterConfigurationTextParser.cpp b/Engine/Plugins/Runtime/nDisplay/Source/DisplayClusterConfiguration/Private/Formats/Text/DisplayClusterConfigurationTextParser.cpp
--- a/Engine/Plugins/Runtime/nDisplay/Source/DisplayClusterConfiguration/Private/Formats/Text/DisplayClusterConfigurationTextParser.cpp
+++ b/Engine/Plugins/Runtime/nDisplay/Source/DisplayClusterConfiguration/Private/Formats/Text/DisplayClusterConfigurationTextParser.cpp
@@ -225,7 +225,7 @@ UDisplayClusterConfigurationData* FDisplayClusterConfigurationTextParser::Conver
 						Viewport->Camera      = CfgViewport->CameraId;
 						Viewport->Region      = FDisplayClusterConfigurationRectangle(CfgViewport->Loc.X, CfgViewport->Loc.Y, CfgViewport->Size.X, CfgViewport->Size.Y);
 						Viewport->GPUIndex    = CfgViewport->GPUIndex;
-						Viewport->bIsShared   = CfgViewport->IsShared;
+						//Viewport->bIsShared   = CfgViewport->IsShared;
 						//Viewport->bAllowCrossGPUTransfer = CfgViewport->AllowCrossGPUTransfer;
 
 						const FDisplayClusterConfigurationTextProjection* const CfgProjection = CfgProjections.FindByPredicate([&](const FDisplayClusterConfigurationTextProjection& Item)
diff --git a/Engine/Plugins/Runtime/nDisplay/Source/DisplayClusterConfiguration/Public/DisplayClusterConfigurationStrings.h b/Engine/Plugins/Runtime/nDisplay/Source/DisplayClusterConfiguration/Public/DisplayClusterConfigurationStrings.h
--- a/Engine/Plugins/Runtime/nDisplay/Source/DisplayClusterConfiguration/Public/DisplayClusterConfigurationStrings.h
+++ b/Engine/Plugins/Runtime/nDisplay/Source/DisplayClusterConfiguration/Public/DisplayClusterConfigurationStrings.h
@@ -33,6 +33,7 @@ namespace DisplayClusterConfigurationStrings
 		static constexpr auto ConfigurationCategory        = TEXT("Configuration");
 		static constexpr auto PreviewCategory              = TEXT("Editor Preview");
 		static constexpr auto AdvancedCategory             = TEXT("Advanced");
+		static constexpr auto TextureShareCategory         = TEXT("Texture Share");
 	}
 
 	// Command line arguments
diff --git a/Engine/Plugins/Runtime/nDisplay/Source/DisplayClusterConfiguration/Public/DisplayClusterConfigurationTypes_TextureShare.h b/Engine/Plugins/Runtime/nDisplay/Source/DisplayClusterConfiguration/Public/DisplayClusterConfigurationTypes_TextureShare.h
new file mode 100644
--- /dev/null
+++ b/Engine/Plugins/Runtime/nDisplay/Source/DisplayClusterConfiguration/Public/DisplayClusterConfigurationTypes_TextureShare.h
@@ -0,0 +1,80 @@
+// Copyright Epic Games, Inc. All Rights Reserved.
+
+#pragma once
+
+#include "CoreMinimal.h"
+
+#include "UObject/ObjectMacros.h"
+
+#include "DisplayClusterConfigurationTypes_TextureShare.generated.h"
+
+/** Synchronize Session state events (BeginSession/EndSession) */
+UENUM(BlueprintType, Category = "TextureShare")
+enum class ETextureShareSyncConnectDisplayCluster : uint8
+{
+	/** [Default] - Use module global settings */
+	Default = 0,
+	/** [None] - do not wait for remote process */
+	None,
+	/** [SyncSession] - waiting until remote process not inside BeginSession()/EndSession() */
+	SyncSession
+};
+
+/** Synchronize frame events (BeginFrame/EndFrame) */
+UENUM(BlueprintType, Category = "TextureShare")
+enum class ETextureShareSyncFrameDisplayCluster : uint8
+{
+	/** [Default] - Use module global settings */
+	Default = 0,
+	/** [None] - Unordered frames */
+	None,
+	/** [FrameSync] - waiting until remote process frame index is equal */
+	FrameSync
+};
+
+/** Synchronize texture events (LockTexture/UnlockTexture) */
+UENUM(BlueprintType, Category = "TextureShare")
+enum class ETextureShareSyncSurfaceDisplayCluster : uint8
+{
+	/** [Default] - Use module global settings */
+	Default = 0,
+	/** [None] - Skip unpaired texture. Unordered read\write operations */
+	None,
+	/** [SyncRead] - Skip unpaired texture. Waiting until other process changed texture (readOP is wait for writeOP from remote process completed) */
+	SyncRead,
+	/** [SyncPairingRead] - Required texture pairing. Waiting until other process changed texture (readOP is wait for writeOP from remote process completed) */
+	SyncPairingRead
+};
+
+
+USTRUCT(BlueprintType)
+struct FTextureShareSyncPolicyDisplayCluster
+{
+	GENERATED_BODY()
+
+	// Synchronize Session state events (BeginSession/EndSession)
+	UPROPERTY(EditAnywhere, BlueprintReadWrite, Category = "TextureShare")
+	ETextureShareSyncConnectDisplayCluster Connection = ETextureShareSyncConnectDisplayCluster::Default;
+
+	// Synchronize frame events (BeginFrame/EndFrame)
+	UPROPERTY(EditAnywhere, BlueprintReadWrite, Category = "TextureShare")
+	ETextureShareSyncFrameDisplayCluster   Frame = ETextureShareSyncFrameDisplayCluster::Default;
+
+	// Synchronize texture events (LockTexture/UnlockTexture)
+	UPROPERTY(EditAnywhere, BlueprintReadWrite, Category = "TextureShare")
+	ETextureShareSyncSurfaceDisplayCluster Texture = ETextureShareSyncSurfaceDisplayCluster::Default;
+};
+
+USTRUCT(Blueprintable)
+struct FDisplayClusterConfigurationTextureShare_Viewport
+{
+	GENERATED_BODY()
+
+	// Allow this viewport to be shared through the TextureShare
+	UPROPERTY(EditAnywhere, BlueprintReadWrite, Category = "TextureShare", meta = (DisplayName = "Share Viewport"))
+	bool bIsEnabled = false;
+
+	// TextureShare synchronization settings, custom gui added in the configurator for Linux compatibility
+	UPROPERTY(EditAnywhere, BlueprintReadWrite, Category = "TextureShare", meta = (DisplayName = "Sync Settings"))
+	FTextureShareSyncPolicyDisplayCluster SyncSettings;
+};
diff --git a/Engine/Plugins/Runtime/nDisplay/Source/DisplayClusterConfiguration/Public/DisplayClusterConfigurationTypes_Viewport.h b/Engine/Plugins/Runtime/nDisplay/Source/DisplayClusterConfiguration/Public/DisplayClusterConfigurationTypes_Viewport.h
--- a/Engine/Plugins/Runtime/nDisplay/Source/DisplayClusterConfiguration/Public/DisplayClusterConfigurationTypes_Viewport.h
+++ b/Engine/Plugins/Runtime/nDisplay/Source/DisplayClusterConfiguration/Public/DisplayClusterConfigurationTypes_Viewport.h
@@ -11,6 +11,8 @@
 #include "DisplayClusterConfigurationTypes_PostRender.h"
 #include "DisplayClusterConfigurationTypes_Postprocess.h"
 
+#include "DisplayClusterConfigurationTypes_TextureShare.h"
+
 #include "OpenColorIOColorSpace.h"
 #include "Engine/Scene.h"
 
@@ -151,7 +153,7 @@ public:
 	FDisplayClusterConfigurationProjection ProjectionPolicy;
 
 	UPROPERTY(EditAnywhere, BlueprintReadOnly, Category = "Configuration", meta = (DisplayName = "Shared Texture"))
-	bool bIsShared = false;
+	FDisplayClusterConfigurationTextureShare_Viewport TextureShare;
 
 #if WITH_EDITORONLY_DATA
 	UPROPERTY(EditAnywhere, Category = "Configuration")
diff --git a/Engine/Plugins/Runtime/nDisplay/Source/DisplayClusterConfigurator/Private/Views/Details/DisplayClusterConfiguratorDetailCustomization.cpp b/Engine/Plugins/Runtime/nDisplay/Source/DisplayClusterConfigurator/Private/Views/Details/DisplayClusterConfiguratorDetailCustomization.cpp
--- a/Engine/Plugins/Runtime/nDisplay/Source/DisplayClusterConfigurator/Private/Views/Details/DisplayClusterConfiguratorDetailCustomization.cpp
+++ b/Engine/Plugins/Runtime/nDisplay/Source/DisplayClusterConfigurator/Private/Views/Details/DisplayClusterConfiguratorDetailCustomization.cpp
@@ -246,7 +246,6 @@ void FDisplayClusterConfiguratorViewportDetailCustomization::CustomizeDetails(ID
 		ADD_PROPERTY(UDisplayClusterConfigurationViewport, bAllowRendering)
 		REPLACE_PROPERTY_WITH_CUSTOM(UDisplayClusterConfigurationViewport, Camera, CreateCustomCameraWidget())
 		ADD_PROPERTY(UDisplayClusterConfigurationViewport, ProjectionPolicy)
-		ADD_PROPERTY(UDisplayClusterConfigurationViewport, bIsShared)
 		ADD_PROPERTY(UDisplayClusterConfigurationViewport, bFixedAspectRatio)
 		ADD_PROPERTY(UDisplayClusterConfigurationViewport, Region)
 		ADD_PROPERTY(UDisplayClusterConfigurationViewport, GPUIndex)
@@ -256,6 +255,9 @@ void FDisplayClusterConfiguratorViewportDetailCustomization::CustomizeDetails(ID
 		ADD_NESTED_PROPERTY(NestedPropertyHelper, UDisplayClusterConfigurationViewport, RenderSettings.RenderTargetRatio)
 		ADD_NESTED_PROPERTY(NestedPropertyHelper, UDisplayClusterConfigurationViewport, RenderSettings.BufferRatio)
 		ADD_NESTED_PROPERTY(NestedPropertyHelper, UDisplayClusterConfigurationViewport, RenderSettings.Overscan)
+#if PLATFORM_WINDOWS
+		ADD_NESTED_PROPERTY(NestedPropertyHelper, UDisplayClusterConfigurationViewport, TextureShare)
+#endif
 	END_CATEGORY();
 
 	// Update the metadata for the viewport's region. Must set this here instead of in the UPROPERTY specifier because
diff --git a/Engine/Source/Programs/TextureShare/TextureShareSDK/Private/TextureShareInterface.cpp b/Engine/Source/Programs/TextureShare/TextureShareSDK/Private/TextureShareInterface.cpp
--- a/Engine/Source/Programs/TextureShare/TextureShareSDK/Private/TextureShareInterface.cpp
+++ b/Engine/Source/Programs/TextureShare/TextureShareSDK/Private/TextureShareInterface.cpp
@@ -54,12 +54,48 @@ void CopyAdditionalData(SrcType& Src, DstType& Dst)
 	//@todo: add more frame data
 }
 
+
+template <typename SrcType, typename DstType>
+void CopyCustomProjectionData(SrcType& Src, DstType& Dst)
+{
+	// Projection matrix
+	CopyMatrix(Src.PrjMatrix, Dst.PrjMatrix);
+
+	CopyVector(Src.ViewLocation, Dst.ViewLocation);
+	CopyRotator(Src.ViewRotation, Dst.ViewRotation);
+	CopyVector(Src.ViewScale, Dst.ViewScale);
+}
+
 ITextureShareCore& ShareCoreAPI()
 {
 	static ITextureShareCore* Singleton = &ITextureShareCore::Get();
 	return *Singleton;
 }
 
+bool FTextureShareInterface::BeginSyncFrame()
+{
+	return ShareCoreAPI().BeginSyncFrame();
+}
+
+bool FTextureShareInterface::EndSyncFrame()
+{
+	return ShareCoreAPI().EndSyncFrame();
+}
+
+bool FTextureShareInterface::SetCustomProjectionData(const TCHAR* ShareName, const FTextureShareSDKCustomProjectionData& InData)
+{
+	TSharedPtr<ITextureShareItem> ShareItem;
+	if (ShareCoreAPI().GetTextureShareItem(FString(ShareName), ShareItem))
+	{
+		// Convert to UE
+		FTextureShareCustomProjectionData Data;
+		CopyCustomProjectionData(InData, Data);
+		return ShareItem->SetCustomProjectionData(Data);
+	}
+
+	return false;
+}
+
 bool FTextureShareInterface::SetLocalAdditionalData(const TCHAR* ShareName, const FTextureShareSDKAdditionalData& InData)
 {
 	TSharedPtr<ITextureShareItem> ShareItem;
diff --git a/Engine/Source/Programs/TextureShare/TextureShareSDK/Public/TextureShareContainers.h b/Engine/Source/Programs/TextureShare/TextureShareSDK/Public/TextureShareContainers.h
--- a/Engine/Source/Programs/TextureShare/TextureShareSDK/Public/TextureShareContainers.h
+++ b/Engine/Source/Programs/TextureShare/TextureShareSDK/Public/TextureShareContainers.h
@@ -64,3 +64,13 @@ struct FTextureShareSDKAdditionalData
 
 	//@todo: add more frame data
 };
+
+struct FTextureShareSDKCustomProjectionData
+{
+	// Projection matrix
+	FTextureShareSDKMatrix PrjMatrix;
+
+	FTextureShareSDKVector  ViewLocation;
+	FTextureShareSDKRotator ViewRotation;
+	FTextureShareSDKVector  ViewScale;
+};
diff --git a/Engine/Source/Programs/TextureShare/TextureShareSDK/Public/TextureShareInterface.h b/Engine/Source/Programs/TextureShare/TextureShareSDK/Public/TextureShareInterface.h
--- a/Engine/Source/Programs/TextureShare/TextureShareSDK/Public/TextureShareInterface.h
+++ b/Engine/Source/Programs/TextureShare/TextureShareSDK/Public/TextureShareInterface.h
@@ -49,5 +49,10 @@ public:
 	TEXTURE_SHARE_SDK_API static bool LockTextureD3D12_RenderThread(ID3D12Device* pD3D12Device, const TCHAR* ShareName, const TCHAR* TextureName, ID3D12Resource*& OutD3D12Resource);
 
 	TEXTURE_SHARE_SDK_API static bool UnlockTexture_RenderThread(const TCHAR* ShareName, const TCHAR* TextureName);
+
+	TEXTURE_SHARE_SDK_API static bool SetCustomProjectionData(const TCHAR* ShareName, const FTextureShareSDKCustomProjectionData& InData);
+
+	TEXTURE_SHARE_SDK_API static bool BeginSyncFrame();
+	TEXTURE_SHARE_SDK_API static bool EndSyncFrame();
 };
 
