Remove unnecessary SCSEditor.h include

#rnx
#p4v-preflight-copy 17139961

diff --git a/Engine/Plugins/Runtime/nDisplay/Source/DisplayClusterConfigurator/Private/DisplayClusterConfiguratorTabSpawners.cpp b/Engine/Plugins/Runtime/nDisplay/Source/DisplayClusterConfigurator/Private/DisplayClusterConfiguratorTabSpawners.cpp
--- a/Engine/Plugins/Runtime/nDisplay/Source/DisplayClusterConfigurator/Private/DisplayClusterConfiguratorTabSpawners.cpp
+++ b/Engine/Plugins/Runtime/nDisplay/Source/DisplayClusterConfigurator/Private/DisplayClusterConfiguratorTabSpawners.cpp
@@ -9,7 +9,6 @@
 #include "Interfaces/Views/TreeViews/IDisplayClusterConfiguratorViewTree.h"
 
 #include "BlueprintEditorTabs.h"
-#include "SSCSEditor.h"
 #include "Widgets/Docking/SDockTab.h"
 
 #define LOCTEXT_NAMESPACE "DisplayClusterConfiguratorTabSpawners"
diff --git a/Engine/Plugins/Runtime/nDisplay/Source/DisplayClusterConfigurator/Private/Views/SCSEditor/SDisplayClusterConfiguratorComponentCombo.cpp b/Engine/Plugins/Runtime/nDisplay/Source/DisplayClusterConfigurator/Private/Views/SCSEditor/SDisplayClusterConfiguratorComponentCombo.cpp
--- a/Engine/Plugins/Runtime/nDisplay/Source/DisplayClusterConfigurator/Private/Views/SCSEditor/SDisplayClusterConfiguratorComponentCombo.cpp
+++ b/Engine/Plugins/Runtime/nDisplay/Source/DisplayClusterConfigurator/Private/Views/SCSEditor/SDisplayClusterConfiguratorComponentCombo.cpp
@@ -31,7 +31,6 @@
 #include "EditorClassUtils.h"
 #include "SListViewSelectorDropdownMenu.h"
 #include "SComponentClassCombo.h"
-#include "SSCSEditor.h"
 
 #define LOCTEXT_NAMESPACE "DisplayClusterConfiguratorComponentClassCombo"
 
diff --git a/Engine/Source/Editor/Kismet/Private/SCSDiff.cpp b/Engine/Source/Editor/Kismet/Private/SCSDiff.cpp
--- a/Engine/Source/Editor/Kismet/Private/SCSDiff.cpp
+++ b/Engine/Source/Editor/Kismet/Private/SCSDiff.cpp
@@ -4,7 +4,6 @@
 #include "Widgets/Layout/SSplitter.h"
 #include "GameFramework/Actor.h"
 #include "SKismetInspector.h"
-#include "SSCSEditor.h"
 #include "SSubobjectEditor.h"
 #include "SSubobjectBlueprintEditor.h"
 #include "IDetailsView.h"
