Remove asset editor toolbar tab from layouts that still reference it (the toolbar is provided by default now and doesn't have a separate tab)

diff --git a/Engine/Plugins/Experimental/CommonConversation/Source/CommonConversationEditor/Private/ConversationEditorModes.cpp b/Engine/Plugins/Experimental/CommonConversation/Source/CommonConversationEditor/Private/ConversationEditorModes.cpp
--- a/Engine/Plugins/Experimental/CommonConversation/Source/CommonConversationEditor/Private/ConversationEditorModes.cpp
+++ b/Engine/Plugins/Experimental/CommonConversation/Source/CommonConversationEditor/Private/ConversationEditorModes.cpp
@@ -18,47 +18,35 @@ FConversationEditorApplicationMode_GraphView::FConversationEditorApplicationMode
 	ConversationEditorTabFactories.RegisterFactory(MakeShareable(new FConversationTreeEditorSummoner(InConversationEditor)));
 	//ConversationEditorTabFactories.RegisterFactory(MakeShareable(new FBlackboardSummoner(InConversationEditor)));
 
-	TabLayout = FTabManager::NewLayout( "Standalone_ConversationEditor_GraphView_Layout_v1" )
+	TabLayout = FTabManager::NewLayout( "Standalone_ConversationEditor_GraphView_Layout_v2" )
 	->AddArea
 	(
 		FTabManager::NewPrimaryArea()
-		->SetOrientation(Orient_Vertical)
+		->SetOrientation(Orient_Horizontal)
 		->Split
 		(
 			FTabManager::NewStack()
-			->SetSizeCoefficient(0.1f)
-			->AddTab(InConversationEditor->GetToolbarTabId(), ETabState::OpenedTab)
-			->SetHideTabWell(true) 
+			->SetSizeCoefficient(0.7f)
+			->AddTab(FConversationEditorTabs::GraphEditorID, ETabState::ClosedTab)
 		)
 		->Split
 		(
 			FTabManager::NewSplitter()
-			->SetOrientation(Orient_Horizontal)
+			->SetOrientation(Orient_Vertical)
+			->SetSizeCoefficient(0.3f)
 			->Split
 			(
 				FTabManager::NewStack()
-				->SetSizeCoefficient(0.7f)
-				->AddTab(FConversationEditorTabs::GraphEditorID, ETabState::ClosedTab)
-			)
-			->Split
-			(
-				FTabManager::NewSplitter()
-				->SetOrientation(Orient_Vertical)
-				->SetSizeCoefficient(0.3f)
-				->Split
-				(
-					FTabManager::NewStack()
-					->SetSizeCoefficient(0.6f)
-					->AddTab(FConversationEditorTabs::GraphDetailsID, ETabState::OpenedTab)
-					->AddTab(FConversationEditorTabs::SearchID, ETabState::ClosedTab)
-				)
- 				->Split
- 				(
- 					FTabManager::NewStack()
- 					->SetSizeCoefficient(0.4f)
- 					->AddTab(FConversationEditorTabs::TreeEditorID, ETabState::OpenedTab)
- 				)
+				->SetSizeCoefficient(0.6f)
+				->AddTab(FConversationEditorTabs::GraphDetailsID, ETabState::OpenedTab)
+				->AddTab(FConversationEditorTabs::SearchID, ETabState::ClosedTab)
 			)
+ 			->Split
+ 			(
+ 				FTabManager::NewStack()
+ 				->SetSizeCoefficient(0.4f)
+ 				->AddTab(FConversationEditorTabs::TreeEditorID, ETabState::OpenedTab)
+ 			)
 		)
 	);
 	
diff --git a/Engine/Plugins/FX/Niagara/Source/NiagaraEditor/Private/Toolkits/NiagaraParameterDefinitionsToolkit.cpp b/Engine/Plugins/FX/Niagara/Source/NiagaraEditor/Private/Toolkits/NiagaraParameterDefinitionsToolkit.cpp
--- a/Engine/Plugins/FX/Niagara/Source/NiagaraEditor/Private/Toolkits/NiagaraParameterDefinitionsToolkit.cpp
+++ b/Engine/Plugins/FX/Niagara/Source/NiagaraEditor/Private/Toolkits/NiagaraParameterDefinitionsToolkit.cpp
@@ -101,45 +101,34 @@ void FNiagaraParameterDefinitionsToolkit::Initialize(const EToolkitMode::Type Mo
 
 	ParameterPanelViewModel->Init(UIContext);
 
-	TSharedRef<FTabManager::FLayout> StandaloneDefaultLayout = FTabManager::NewLayout("Standalone_Niagara_ParameterDefinitionsLayout_v1")
+	TSharedRef<FTabManager::FLayout> StandaloneDefaultLayout = FTabManager::NewLayout("Standalone_Niagara_ParameterDefinitionsLayout_v2")
 		->AddArea
 		(
-			FTabManager::NewPrimaryArea()->SetOrientation(Orient_Vertical)
+			FTabManager::NewPrimaryArea()->SetOrientation(Orient_Horizontal)
 			->Split
 			(
-				FTabManager::NewStack()
-				->SetSizeCoefficient(0.1f)
-				->AddTab(GetToolbarTabId(), ETabState::OpenedTab)
-				->SetHideTabWell(true)
-			)
-			->Split
-			(
-				FTabManager::NewSplitter()->SetOrientation(Orient_Horizontal)->SetSizeCoefficient(0.9f)
+				FTabManager::NewSplitter()->SetOrientation(Orient_Vertical)
+				->SetSizeCoefficient(0.15f)
 				->Split
 				(
-					FTabManager::NewSplitter()->SetOrientation(Orient_Vertical)
-					->SetSizeCoefficient(0.15f)
-					->Split
-					(
-						FTabManager::NewStack()
-						->SetSizeCoefficient(.15f)
-						->AddTab(ParameterDefinitionsDetailsTabId, ETabState::OpenedTab)
-					)
-					->Split
-					(
-						FTabManager::NewStack()
-						->SetSizeCoefficient(.85f)
-						->AddTab(ParameterPanelTabId, ETabState::OpenedTab)
-						->SetForegroundTab(ParameterPanelTabId)
-					)
-
+					FTabManager::NewStack()
+					->SetSizeCoefficient(.15f)
+					->AddTab(ParameterDefinitionsDetailsTabId, ETabState::OpenedTab)
 				)
 				->Split
 				(
 					FTabManager::NewStack()
-					->SetSizeCoefficient(0.5f)
-					->AddTab(SelectedDetailsTabId, ETabState::OpenedTab)
+					->SetSizeCoefficient(.85f)
+					->AddTab(ParameterPanelTabId, ETabState::OpenedTab)
+					->SetForegroundTab(ParameterPanelTabId)
 				)
+
+			)
+			->Split
+			(
+				FTabManager::NewStack()
+				->SetSizeCoefficient(0.5f)
+				->AddTab(SelectedDetailsTabId, ETabState::OpenedTab)
 			)
 		);
 
diff --git a/Engine/Plugins/Media/BinkMedia/Source/BinkMediaPlayerEditor/Private/BinkMediaPlayerEditorToolkit.cpp b/Engine/Plugins/Media/BinkMedia/Source/BinkMediaPlayerEditor/Private/BinkMediaPlayerEditorToolkit.cpp
--- a/Engine/Plugins/Media/BinkMedia/Source/BinkMediaPlayerEditor/Private/BinkMediaPlayerEditorToolkit.cpp
+++ b/Engine/Plugins/Media/BinkMedia/Source/BinkMediaPlayerEditor/Private/BinkMediaPlayerEditorToolkit.cpp
@@ -41,31 +41,17 @@ void FBinkMediaPlayerEditorToolkit::Initialize( UBinkMediaPlayer* InMediaPlayer,
 		FExecuteAction::CreateSP(this, &FBinkMediaPlayerEditorToolkit::HandleRewindMediaActionExecute),
 		FCanExecuteAction::CreateSP(this, &FBinkMediaPlayerEditorToolkit::HandleRewindMediaActionCanExecute));
 
-	const TSharedRef<FTabManager::FLayout> Layout = FTabManager::NewLayout("BinkMediaPlayerEditor")
+	const TSharedRef<FTabManager::FLayout> Layout = FTabManager::NewLayout("BinkMediaPlayerEditor_Layout_v2")
 		->AddArea
 		(
 			FTabManager::NewPrimaryArea()
 				->SetOrientation(Orient_Horizontal)
 				->Split
 				(
-					FTabManager::NewSplitter()
-						->SetOrientation(Orient_Vertical)
+					FTabManager::NewStack()
+						->AddTab(ViewerTabId, ETabState::OpenedTab)
+						->SetHideTabWell(true)
 						->SetSizeCoefficient(0.66f)
-						->Split
-						(
-							FTabManager::NewStack()
-								->AddTab(GetToolbarTabId(), ETabState::OpenedTab)
-								->SetHideTabWell(true)
-								->SetSizeCoefficient(0.1f)
-								
-						)
-						->Split
-						(
-							FTabManager::NewStack()
-								->AddTab(ViewerTabId, ETabState::OpenedTab)
-								->SetHideTabWell(true)
-								->SetSizeCoefficient(0.9f)
-						)
 				)
 				->Split
 				(
diff --git a/Engine/Plugins/Runtime/nDisplay/Source/DisplayClusterConfigurator/Private/DisplayClusterConfiguratorEditorMode.cpp b/Engine/Plugins/Runtime/nDisplay/Source/DisplayClusterConfigurator/Private/DisplayClusterConfiguratorEditorMode.cpp
--- a/Engine/Plugins/Runtime/nDisplay/Source/DisplayClusterConfigurator/Private/DisplayClusterConfiguratorEditorMode.cpp
+++ b/Engine/Plugins/Runtime/nDisplay/Source/DisplayClusterConfigurator/Private/DisplayClusterConfiguratorEditorMode.cpp
@@ -37,7 +37,7 @@ FDisplayClusterConfiguratorEditorConfigurationMode::FDisplayClusterConfiguratorE
 	TSharedPtr<FDisplayClusterConfiguratorBlueprintEditor> EditorIn) : FDisplayClusterConfiguratorBlueprintModeBase(EditorIn,
 		FDisplayClusterEditorModes::DisplayClusterEditorConfigurationMode)
 {
-	TabLayout = BuildDefaultLayout(FString(TEXT("DisplayClusterConfigurator_v1.0")));
+	TabLayout = BuildDefaultLayout(FString(TEXT("DisplayClusterConfigurator_v1.1")));
 
 	EditorTabFactories.RegisterFactory(MakeShared<FDisplayClusterViewClusterSummoner>(EditorIn));
 	EditorTabFactories.RegisterFactory(MakeShared<FDisplayClusterViewOutputMappingSummoner>(EditorIn));
@@ -84,84 +84,70 @@ TSharedPtr<FTabManager::FLayout> FDisplayClusterConfiguratorEditorConfigurationM
 	return FTabManager::NewLayout(FName(LayoutName))
 		->AddArea
 		(
-			FTabManager::NewPrimaryArea()->SetOrientation(Orient_Vertical)
+			FTabManager::NewPrimaryArea()->SetOrientation(Orient_Horizontal)
+			// Tree
 			->Split
 			(
-				// Toolbar
-				FTabManager::NewStack()
-				->SetSizeCoefficient(0.1f)
-				->AddTab(Editor.Pin()->GetToolbarTabId(), ETabState::OpenedTab)
-				->SetHideTabWell(true)
+				FTabManager::NewSplitter()
+				->SetSizeCoefficient(.2f)
+				->SetOrientation(Orient_Vertical)
+				->Split
+				(
+					FTabManager::NewStack()
+					->SetSizeCoefficient(.5f)
+					->AddTab(TabID_Scene, ETabState::OpenedTab)
+					->AddTab(FBlueprintEditorTabs::MyBlueprintID, ETabState::OpenedTab)
+					->SetForegroundTab(TabID_Scene)
+				)
+				->Split
+				(
+					FTabManager::NewStack()
+					->SetSizeCoefficient(.5f)
+					->AddTab(TabID_Cluster, ETabState::OpenedTab)
+					->SetHideTabWell(false)
+				)
 			)
+			// Viewport and OutputMapping
 			->Split
 			(
-				// Main canvas
-				FTabManager::NewSplitter()->SetOrientation(Orient_Horizontal)
-
-				// Tree
+				FTabManager::NewSplitter()
+				->SetOrientation(Orient_Vertical)
+				->SetSizeCoefficient(.6f)
 				->Split
 				(
-					FTabManager::NewSplitter()
-					->SetSizeCoefficient(.2f)
-					->SetOrientation(Orient_Vertical)
-					->Split
-					(
-						FTabManager::NewStack()
-						->SetSizeCoefficient(.5f)
-						->AddTab(TabID_Scene, ETabState::OpenedTab)
-						->AddTab(FBlueprintEditorTabs::MyBlueprintID, ETabState::OpenedTab)
-						->SetForegroundTab(TabID_Scene)
-					)
-					->Split
-					(
-						FTabManager::NewStack()
-						->SetSizeCoefficient(.5f)
-						->AddTab(TabID_Cluster, ETabState::OpenedTab)
-						->SetHideTabWell(false)
-					)
+					FTabManager::NewStack()
+					->SetSizeCoefficient(0.5f)
+					->AddTab(TabID_Viewport, ETabState::OpenedTab)
+					->AddTab("Document", ETabState::ClosedTab)
+					->SetForegroundTab(TabID_Viewport)
 				)
-				// Viewport and OutputMapping
 				->Split
 				(
-					FTabManager::NewSplitter()
-					->SetOrientation(Orient_Vertical)
-					->SetSizeCoefficient(.6f)
-					->Split
-					(
-						FTabManager::NewStack()
-						->SetSizeCoefficient(0.5f)
-						->AddTab(TabID_Viewport, ETabState::OpenedTab)
-						->AddTab("Document", ETabState::ClosedTab)
-						->SetForegroundTab(TabID_Viewport)
-					)
-					->Split
-					(
-						FTabManager::NewStack()
-						->AddTab(TabID_OutputMapping, ETabState::OpenedTab)
-						->SetSizeCoefficient(0.5f)
-						->SetHideTabWell(false)
-					)
+					FTabManager::NewStack()
+					->AddTab(TabID_OutputMapping, ETabState::OpenedTab)
+					->SetSizeCoefficient(0.5f)
+					->SetHideTabWell(false)
+				)
+			)
+			// Details and Log
+			->Split
+			(
+				FTabManager::NewSplitter()
+				->SetOrientation(Orient_Vertical)
+				->SetSizeCoefficient(.2f)
+				->Split
+				(
+					FTabManager::NewStack()
+					->SetSizeCoefficient(.7f)
+					->AddTab(FBlueprintEditorTabs::DetailsID, ETabState::OpenedTab)
+					->SetForegroundTab(FBlueprintEditorTabs::DetailsID)
+					->SetHideTabWell(false)
 				)
-				// Details and Log
 				->Split
 				(
-					FTabManager::NewSplitter()
-					->SetOrientation(Orient_Vertical)
-					->SetSizeCoefficient(.2f)
-					->Split
-					(
-						FTabManager::NewStack()
-						->SetSizeCoefficient(.7f)
-						->AddTab(FBlueprintEditorTabs::DetailsID, ETabState::OpenedTab)
-						->SetForegroundTab(FBlueprintEditorTabs::DetailsID)
-						->SetHideTabWell(false)
-					)
-					->Split
-					(
-						FTabManager::NewStack()
-						->SetSizeCoefficient(.3f)
-						->AddTab(FBlueprintEditorTabs::CompilerResultsID, ETabState::OpenedTab)
-					)
+					FTabManager::NewStack()
+					->SetSizeCoefficient(.3f)
+					->AddTab(FBlueprintEditorTabs::CompilerResultsID, ETabState::OpenedTab)
 				)
 			)
 		);
diff --git a/Engine/Plugins/VirtualProduction/CameraCalibration/Source/CameraCalibrationEditor/Private/AssetEditor/CameraCalibrationToolkit.cpp b/Engine/Plugins/VirtualProduction/CameraCalibration/Source/CameraCalibrationEditor/Private/AssetEditor/CameraCalibrationToolkit.cpp
--- a/Engine/Plugins/VirtualProduction/CameraCalibration/Source/CameraCalibrationEditor/Private/AssetEditor/CameraCalibrationToolkit.cpp
+++ b/Engine/Plugins/VirtualProduction/CameraCalibration/Source/CameraCalibrationEditor/Private/AssetEditor/CameraCalibrationToolkit.cpp
@@ -56,19 +56,12 @@ void FCameraCalibrationToolkit::InitCameraCalibrationTool(const EToolkitMode::Ty
 	LensEditorTab = SNew(SLensFilePanel, LensFile, CalibrationStepsController.ToSharedRef())
 		.CachedFIZData(TAttribute<FCachedFIZData>::Create(TAttribute<FCachedFIZData>::FGetter::CreateSP(LensEvaluationWidget.ToSharedRef(), &SLensEvaluation::GetLastEvaluatedData)));
 		
-	TSharedRef<FTabManager::FLayout> NewLayout = FTabManager::NewLayout("CameraCalibrationToolLayout_v0.8")
+	TSharedRef<FTabManager::FLayout> NewLayout = FTabManager::NewLayout("CameraCalibrationToolLayout_v0.9")
 		->AddArea
 		(
 			FTabManager::NewPrimaryArea()
 			->SetOrientation(Orient_Vertical)
 			->Split
-			(
-				FTabManager::NewStack()
-				->SetSizeCoefficient(0.1f)
-				->SetHideTabWell(true)
-				->AddTab(GetToolbarTabId(), ETabState::OpenedTab)
-			)
-			->Split
 			(
 				FTabManager::NewSplitter()
 				->SetSizeCoefficient(0.85f)
diff --git a/Engine/Source/Editor/StaticMeshEditor/Private/StaticMeshEditor.cpp b/Engine/Source/Editor/StaticMeshEditor/Private/StaticMeshEditor.cpp
--- a/Engine/Source/Editor/StaticMeshEditor/Private/StaticMeshEditor.cpp
+++ b/Engine/Source/Editor/StaticMeshEditor/Private/StaticMeshEditor.cpp
@@ -330,9 +330,6 @@ void FStaticMeshEditor::GenerateSecondaryToolbar()
 	{
 		Tab->SetLabel(SecondaryToolbarDisplayName);
 	}
-
-	// But have the focus on the default toolbar
-	TabManager->TryInvokeTab(GetToolbarTabId());
 }
 
 void FStaticMeshEditor::AddSecondaryToolbarExtender(TSharedPtr<FExtender> Extender)
