nDisplay,TextureShare: Merge from RES.

#rnx

diff --git a/Engine/Plugins/Runtime/TextureShare/Source/TextureShare/Private/Blueprints/TextureShareBlueprintAPIImpl.cpp b/Engine/Plugins/Runtime/TextureShare/Source/TextureShare/Private/Blueprints/TextureShareBlueprintAPIImpl.cpp
--- a/Engine/Plugins/Runtime/TextureShare/Source/TextureShare/Private/Blueprints/TextureShareBlueprintAPIImpl.cpp
+++ b/Engine/Plugins/Runtime/TextureShare/Source/TextureShare/Private/Blueprints/TextureShareBlueprintAPIImpl.cpp
@@ -33,9 +33,9 @@ void UTextureShareAPIImpl::SetSyncPolicySettings(const FTextureShareBPSyncPolicy
 }
 
 
-bool UTextureShareAPIImpl::CreateTextureShare(const FString ShareName, FTextureShareBPSyncPolicy SyncMode, bool bIsServer)
+bool UTextureShareAPIImpl::CreateTextureShare(const FString ShareName, FTextureShareBPSyncPolicy SyncMode, bool bIsServer, float SyncWaitTime)
 {
-	return GetAPI().CreateShare(ShareName, *SyncMode, bIsServer ? ETextureShareProcess::Server : ETextureShareProcess::Client);
+	return GetAPI().CreateShare(ShareName, *SyncMode, bIsServer ? ETextureShareProcess::Server : ETextureShareProcess::Client, SyncWaitTime);
 }
 
 bool UTextureShareAPIImpl::ReleaseTextureShare(const FString ShareName)
diff --git a/Engine/Plugins/Runtime/TextureShare/Source/TextureShare/Private/Blueprints/TextureShareBlueprintAPIImpl.h b/Engine/Plugins/Runtime/TextureShare/Source/TextureShare/Private/Blueprints/TextureShareBlueprintAPIImpl.h
--- a/Engine/Plugins/Runtime/TextureShare/Source/TextureShare/Private/Blueprints/TextureShareBlueprintAPIImpl.h
+++ b/Engine/Plugins/Runtime/TextureShare/Source/TextureShare/Private/Blueprints/TextureShareBlueprintAPIImpl.h
@@ -20,7 +20,7 @@ class UTextureShareAPIImpl
 
 public:
 	UFUNCTION(BlueprintCallable, meta = (DisplayName = "Create TextureShare"), Category = "TextureShare")
-	virtual bool CreateTextureShare(const FString ShareName, FTextureShareBPSyncPolicy SyncMode, bool bIsServer = true) override;
+	virtual bool CreateTextureShare(const FString ShareName, FTextureShareBPSyncPolicy SyncMode, bool bIsServer = true, float SyncWaitTime = 0.03f) override;
 
 	UFUNCTION(BlueprintCallable, meta = (DisplayName = "Delete TextureShare"), Category = "TextureShare")
 	virtual bool ReleaseTextureShare(const FString ShareName) override;
diff --git a/Engine/Plugins/Runtime/TextureShare/Source/TextureShare/Private/TextureShareModule.cpp b/Engine/Plugins/Runtime/TextureShare/Source/TextureShare/Private/TextureShareModule.cpp
--- a/Engine/Plugins/Runtime/TextureShare/Source/TextureShare/Private/TextureShareModule.cpp
+++ b/Engine/Plugins/Runtime/TextureShare/Source/TextureShare/Private/TextureShareModule.cpp
@@ -126,14 +126,14 @@ ETextureShareDevice FTextureShareModule::GetTextureShareDeviceType() const
 	return ETextureShareDevice::Undefined;
 };
 
-bool FTextureShareModule::CreateShare(const FString& ShareName, const FTextureShareSyncPolicy& SyncMode, ETextureShareProcess Process)
+bool FTextureShareModule::CreateShare(const FString& ShareName, const FTextureShareSyncPolicy& SyncMode, ETextureShareProcess Process, float SyncWaitTime)
 {
 	FScopeLock lock(&DataGuard);
 
 	TSharedPtr<ITextureShareItem> ShareItem;
 	ETextureShareDevice ShareDevice = GetTextureShareDeviceType();
 
-	return ShareCoreAPI.CreateTextureShareItem(ShareName, Process, SyncMode, ShareDevice, ShareItem);
+	return ShareCoreAPI.CreateTextureShareItem(ShareName, Process, SyncMode, ShareDevice, ShareItem, SyncWaitTime);
 }
 
 bool FTextureShareModule::ReleaseShare(const FString& ShareName)
@@ -164,6 +164,7 @@ void FTextureShareModule::OnResolvedSceneColor_RenderThread(FRDGBuilder& GraphBu
 	FScopeLock lock(&DataGuard);
 
 	if (ViewFamily.Views.Num() > 0)
+	{
 	{
 		EStereoscopicPass StereoscopicPass = ViewFamily.Views[0]->StereoPass;
 		// Send SceneContext callback for all registered shares:
@@ -176,12 +177,14 @@ void FTextureShareModule::OnResolvedSceneColor_RenderThread(FRDGBuilder& GraphBu
 		}
 	}
 }
+}
 
 void FTextureShareModule::OnPostRenderViewFamily_RenderThread(FRHICommandListImmediate& RHICmdList, class FSceneViewFamily& ViewFamily)
 {
 	FScopeLock lock(&DataGuard);
 	
 	if (ViewFamily.Views.Num() > 0)
+	{
 	{
 		EStereoscopicPass StereoscopicPass = ViewFamily.Views[0]->StereoPass;
 		// Send PostRender callback for all registered shares:
@@ -194,6 +197,7 @@ void FTextureShareModule::OnPostRenderViewFamily_RenderThread(FRHICommandListImm
 		}
 	}
 }
+}
 
 bool FTextureShareModule::RegisterTexture(const TSharedPtr<ITextureShareItem>& ShareItem, const FString& InTextureName, const FIntPoint& InSize, EPixelFormat InFormat, ETextureShareSurfaceOp OperationType)
 {
diff --git a/Engine/Plugins/Runtime/TextureShare/Source/TextureShare/Private/TextureShareModule.h b/Engine/Plugins/Runtime/TextureShare/Source/TextureShare/Private/TextureShareModule.h
--- a/Engine/Plugins/Runtime/TextureShare/Source/TextureShare/Private/TextureShareModule.h
+++ b/Engine/Plugins/Runtime/TextureShare/Source/TextureShare/Private/TextureShareModule.h
@@ -26,7 +26,7 @@ public:
 	//////////////////////////////////////////////////////////////////////////////////////////////
 	// TextureShare
 	//////////////////////////////////////////////////////////////////////////////////////////////
-	virtual bool CreateShare(const FString& ShareName, const FTextureShareSyncPolicy& SyncMode, ETextureShareProcess Process = ETextureShareProcess::Server) override;
+	virtual bool CreateShare(const FString& ShareName, const FTextureShareSyncPolicy& SyncMode, ETextureShareProcess Process = ETextureShareProcess::Server, float SyncWaitTime = 0.03f) override;
 	virtual bool ReleaseShare(const FString& ShareName) override;
 	virtual bool GetShare(const FString& ShareName, TSharedPtr<ITextureShareItem>& OutShareItem) const override;
 
diff --git a/Engine/Plugins/Runtime/TextureShare/Source/TextureShare/Public/Blueprints/ITextureShareBlueprintAPI.h b/Engine/Plugins/Runtime/TextureShare/Source/TextureShare/Public/Blueprints/ITextureShareBlueprintAPI.h
--- a/Engine/Plugins/Runtime/TextureShare/Source/TextureShare/Public/Blueprints/ITextureShareBlueprintAPI.h
+++ b/Engine/Plugins/Runtime/TextureShare/Source/TextureShare/Public/Blueprints/ITextureShareBlueprintAPI.h
@@ -34,7 +34,7 @@ public:
 	 * @return True if the success
 	 */
 	UFUNCTION(BlueprintCallable, meta = (DisplayName = "Create TextureShare"), Category = "TextureShare")
-	virtual bool CreateTextureShare(const FString ShareName, FTextureShareBPSyncPolicy SyncMode, bool bIsServer = true) = 0;
+	virtual bool CreateTextureShare(const FString ShareName, FTextureShareBPSyncPolicy SyncMode, bool bIsServer = true, float SyncWaitTime = 0.03f) = 0;
 
 	/**
 	 * Release exist textureshare object
diff --git a/Engine/Plugins/Runtime/TextureShare/Source/TextureShare/Public/ITextureShare.h b/Engine/Plugins/Runtime/TextureShare/Source/TextureShare/Public/ITextureShare.h
--- a/Engine/Plugins/Runtime/TextureShare/Source/TextureShare/Public/ITextureShare.h
+++ b/Engine/Plugins/Runtime/TextureShare/Source/TextureShare/Public/ITextureShare.h
@@ -52,9 +52,9 @@ public:
 	 *
 	 * @return True if the success
 	 */
-	virtual bool CreateShare(const FString& ShareName, const FTextureShareSyncPolicy& SyncMode, ETextureShareProcess Process = ETextureShareProcess::Server) = 0;
+	virtual bool CreateShare(const FString& ShareName, const FTextureShareSyncPolicy& SyncMode, ETextureShareProcess Process = ETextureShareProcess::Server, float SyncWaitTime = 0.03f) = 0;
 
-	/**
+	/** 
 	 * Delete TextureShare object
 	 *
 	 * @param ShareName - Unique share name (case insensitive)
diff --git a/Engine/Plugins/Runtime/TextureShare/Source/TextureShareCore/Private/Core/TextureShareItemBase.cpp b/Engine/Plugins/Runtime/TextureShare/Source/TextureShareCore/Private/Core/TextureShareItemBase.cpp
--- a/Engine/Plugins/Runtime/TextureShare/Source/TextureShareCore/Private/Core/TextureShareItemBase.cpp
+++ b/Engine/Plugins/Runtime/TextureShare/Source/TextureShareCore/Private/Core/TextureShareItemBase.cpp
@@ -322,6 +322,11 @@ namespace TextureShareItem
 		return WriteLocalProcessData();
 	}
 
+	void FTextureShareItemBase::SetSyncWaitTime(float InSyncWaitTime)
+	{
+		SyncWaitTime = InSyncWaitTime;
+	}
+
 	/*ESharedResourceTextureState FTextureShareItemBase::GetLocalTextureState(const FString& TextureName)
 	{
 		FSharedResourceProcessData& Data = GetLocalProcessData();
@@ -422,7 +427,7 @@ namespace TextureShareItem
 		/** [Required] Waiting until remote process register texture (Required texture pairing) */
 		case ETextureShareSyncSurface::SyncPairingRead:
 		{
-			FSyncProcess SyncProcess(*this, SyncSettings.TimeOut.TexturePairingSync, ESyncProcessFailAction::ConnectionLost);
+			FSyncProcess SyncProcess(*this, SyncSettings.TimeOut.TexturePairingSync, ESyncProcessFailAction::ConnectionLost, SyncWaitTime);
 			while (RemoteTextureIndex < 0)
 			{
 				if (!SyncProcess.Tick())
@@ -458,7 +463,7 @@ namespace TextureShareItem
 						}
 					}
 #endif
-					FSyncProcess SyncProcess(*this, SyncSettings.TimeOut.TextureSync, ESyncProcessFailAction::ConnectionLost);
+					FSyncProcess SyncProcess(*this, SyncSettings.TimeOut.TextureSync, ESyncProcessFailAction::ConnectionLost, SyncWaitTime);
 					FSharedResourceProcessData& LocalData = GetLocalProcessData();
 					FSharedResourceProcessData& RemoteData = GetRemoteProcessData();
 
@@ -518,7 +523,7 @@ namespace TextureShareItem
 			case ETextureShareSyncFrame::FrameSync:
 			{
 				// Wait for remote handle initialized
-				FSyncProcess SyncProcess(*this, SyncSettings.TimeOut.TextureResourceSync, ESyncProcessFailAction::ConnectionLost);
+				FSyncProcess SyncProcess(*this, SyncSettings.TimeOut.TextureResourceSync, ESyncProcessFailAction::ConnectionLost, SyncWaitTime);
 				while (!GetRemoteProcessData().Textures[RemoteTextureIndex].SharingData.IsValid())
 				{
 					if (!SyncProcess.Tick())
@@ -559,7 +564,7 @@ namespace TextureShareItem
 			case ETextureShareSyncConnect::SyncSession:
 			{
 				// Wait for other process
-				FSyncProcess SyncProcess(*this, SyncSettings.TimeOut.ConnectionSync, ESyncProcessFailAction::ConnectionLost);
+				FSyncProcess SyncProcess(*this, SyncSettings.TimeOut.ConnectionSync, ESyncProcessFailAction::ConnectionLost, SyncWaitTime);
 				while (!IsConnectionValid())
 				{
 					if (!SyncProcess.Tick())
@@ -590,7 +595,7 @@ namespace TextureShareItem
 		{
 		case ETextureShareSyncFrame::FrameSync:
 		{
-			FSyncProcess SyncProcess(*this, SyncSettings.TimeOut.FrameSync, ESyncProcessFailAction::ConnectionLost);
+			FSyncProcess SyncProcess(*this, SyncSettings.TimeOut.FrameSync, ESyncProcessFailAction::ConnectionLost, SyncWaitTime);
 			while (!ResourceData.IsSyncFrameValid(IsClient()))
 			{
 				if (!SyncProcess.Tick())
diff --git a/Engine/Plugins/Runtime/TextureShare/Source/TextureShareCore/Private/Core/TextureShareItemBase.h b/Engine/Plugins/Runtime/TextureShare/Source/TextureShareCore/Private/Core/TextureShareItemBase.h
--- a/Engine/Plugins/Runtime/TextureShare/Source/TextureShareCore/Private/Core/TextureShareItemBase.h
+++ b/Engine/Plugins/Runtime/TextureShare/Source/TextureShareCore/Private/Core/TextureShareItemBase.h
@@ -34,6 +34,7 @@ namespace TextureShareItem
 
 		virtual bool SetTextureGPUIndex(const FString& TextureName, uint32 GPUIndex) override;
 		virtual bool SetDefaultGPUIndex(uint32 GPUIndex) override;
+		virtual void SetSyncWaitTime(float InSyncWaitTime) override;
 
 		virtual bool IsLocalTextureUsed(const FString& TextureName) const override
 			{ return IsTextureUsed(true, TextureName); }
@@ -132,6 +133,7 @@ namespace TextureShareItem
 		FSharedResource*           SharedResource = nullptr;
 		bool                       bIsSessionStarted = false;
 		bool                       bRemoteConnectionValid = false;
+		float                      SyncWaitTime = 0.03f;
 
 	private:
 		mutable FCriticalSection FrameLockGuard;
diff --git a/Engine/Plugins/Runtime/TextureShare/Source/TextureShareCore/Private/TextureShareCore.cpp b/Engine/Plugins/Runtime/TextureShare/Source/TextureShareCore/Private/TextureShareCore.cpp
--- a/Engine/Plugins/Runtime/TextureShare/Source/TextureShareCore/Private/TextureShareCore.cpp
+++ b/Engine/Plugins/Runtime/TextureShare/Source/TextureShareCore/Private/TextureShareCore.cpp
@@ -94,7 +94,7 @@ bool FTextureShareCoreModule::GetTextureShareItem(const FString& InShareName, TS
 }
 
 // Create shared resource object
-bool FTextureShareCoreModule::CreateTextureShareItem(const FString& InShareName, ETextureShareProcess Process, FTextureShareSyncPolicy SyncMode, ETextureShareDevice DeviceType, TSharedPtr<ITextureShareItem>& OutShareObject)
+bool FTextureShareCoreModule::CreateTextureShareItem(const FString& InShareName, ETextureShareProcess Process, FTextureShareSyncPolicy SyncMode, ETextureShareDevice DeviceType, TSharedPtr<ITextureShareItem>& OutShareObject, float SyncWaitTime)
 {
 	if (IsShareNameEmpty(InShareName))
 	{
@@ -145,6 +145,9 @@ bool FTextureShareCoreModule::CreateTextureShareItem(const FString& InShareName,
 			return false;
 		}
 
+		// set sync wait time
+		Resource->SetSyncWaitTime(SyncWaitTime);
+
 		// Save created object ptr
 		TSharedPtr<ITextureShareItem> NewShareObject = MakeShareable(Resource);
 		TextureShares.Add(LowerShareName, NewShareObject);
diff --git a/Engine/Plugins/Runtime/TextureShare/Source/TextureShareCore/Private/TextureShareCore.h b/Engine/Plugins/Runtime/TextureShare/Source/TextureShareCore/Private/TextureShareCore.h
--- a/Engine/Plugins/Runtime/TextureShare/Source/TextureShareCore/Private/TextureShareCore.h
+++ b/Engine/Plugins/Runtime/TextureShare/Source/TextureShareCore/Private/TextureShareCore.h
@@ -22,7 +22,7 @@ public:
 	virtual void ReleaseLib() override;
 
 	/** shared resource object api */
-	virtual bool CreateTextureShareItem(const FString& ShareName, ETextureShareProcess Process, FTextureShareSyncPolicy SyncMode, ETextureShareDevice DeviceType, TSharedPtr<ITextureShareItem>& OutShareObject) override;
+	virtual bool CreateTextureShareItem(const FString& ShareName, ETextureShareProcess Process, FTextureShareSyncPolicy SyncMode, ETextureShareDevice DeviceType, TSharedPtr<ITextureShareItem>& OutShareObject, float SyncWaitTime) override;
 	virtual bool ReleaseTextureShareItem(const FString& ShareName) override;
 	virtual bool GetTextureShareItem(const FString& ShareName, TSharedPtr<ITextureShareItem>& OutShareObject) const override;
 
diff --git a/Engine/Plugins/Runtime/TextureShare/Source/TextureShareCore/Public/ITextureShareCore.h b/Engine/Plugins/Runtime/TextureShare/Source/TextureShareCore/Public/ITextureShareCore.h
--- a/Engine/Plugins/Runtime/TextureShare/Source/TextureShareCore/Public/ITextureShareCore.h
+++ b/Engine/Plugins/Runtime/TextureShare/Source/TextureShareCore/Public/ITextureShareCore.h
@@ -64,7 +64,7 @@ class ITextureShareCore
 	 *
 	 * @return True if the success
 	 */
-	virtual bool CreateTextureShareItem(const FString& ShareName, ETextureShareProcess Process, FTextureShareSyncPolicy SyncMode, ETextureShareDevice DeviceType, TSharedPtr<ITextureShareItem>& OutShareObject) = 0;
+	virtual bool CreateTextureShareItem(const FString& ShareName, ETextureShareProcess Process, FTextureShareSyncPolicy SyncMode, ETextureShareDevice DeviceType, TSharedPtr<ITextureShareItem>& OutShareObject, float SyncWaitTime) = 0;
 
 	/**
 	 * Delete ITextureShareItem object
diff --git a/Engine/Plugins/Runtime/TextureShare/Source/TextureShareCore/Public/ITextureShareItem.h b/Engine/Plugins/Runtime/TextureShare/Source/TextureShareCore/Public/ITextureShareItem.h
--- a/Engine/Plugins/Runtime/TextureShare/Source/TextureShareCore/Public/ITextureShareItem.h
+++ b/Engine/Plugins/Runtime/TextureShare/Source/TextureShareCore/Public/ITextureShareItem.h
@@ -43,6 +43,9 @@ public:
 	virtual bool SetTextureGPUIndex(const FString& TextureName, uint32 GPUIndex) = 0;
 	virtual bool SetDefaultGPUIndex(uint32 GPUIndex) = 0;
 
+	/* Sync wait time */
+	virtual void SetSyncWaitTime(float InSyncWaitTime) = 0;
+
 	virtual bool GetRemoteTextureDesc(const FString& TextureName, FTextureShareSurfaceDesc& OutSharedTextureDesc) const = 0;
 
 	/** Session scope */
diff --git a/Engine/Plugins/Runtime/nDisplay/Source/DisplayCluster/Private/Blueprints/DisplayClusterBlueprint.cpp b/Engine/Plugins/Runtime/nDisplay/Source/DisplayCluster/Private/Blueprints/DisplayClusterBlueprint.cpp
--- a/Engine/Plugins/Runtime/nDisplay/Source/DisplayCluster/Private/Blueprints/DisplayClusterBlueprint.cpp
+++ b/Engine/Plugins/Runtime/nDisplay/Source/DisplayCluster/Private/Blueprints/DisplayClusterBlueprint.cpp
@@ -2,10 +2,16 @@
 
 #include "Blueprints/DisplayClusterBlueprint.h"
 #include "Blueprints/DisplayClusterBlueprintGeneratedClass.h"
-#include "Containers/Set.h"
 #include "DisplayClusterRootActor.h"
-#include "EngineAnalytics.h"
+#include "Components/DisplayClusterCameraComponent.h"
+#include "Components/DisplayClusterScreenComponent.h"
+
 #include "IDisplayClusterConfiguration.h"
+
+#include "EngineAnalytics.h"
+#include "Engine/SCS_Node.h"
+#include "Containers/Set.h"
+
 #include "Misc/DisplayClusterLog.h"
 #include "UObject/ObjectSaveContext.h"
 
@@ -36,8 +42,8 @@ void UDisplayClusterBlueprint::UpdateConfigExportProperty()
 
 	if (UDisplayClusterConfigurationData* Config = GetOrLoadConfig())
 	{
-		Config->Info.AssetPath = GetPathName();
-
+		PrepareConfigForExport();
+		
 		FString PrettyConfig;
 
 		bConfigExported = IDisplayClusterConfiguration::Get().ConfigAsString(Config, PrettyConfig);
@@ -177,3 +183,215 @@ void UDisplayClusterBlueprint::SetConfigPath(const FString& InPath)
 	}
 #endif
 }
+
+void UDisplayClusterBlueprint::PrepareConfigForExport()
+{
+	if (!ensure(GeneratedClass))
+	{
+		return;
+	}
+		
+	ADisplayClusterRootActor* CDO = CastChecked<ADisplayClusterRootActor>(GeneratedClass->ClassDefaultObject);
+
+	UDisplayClusterConfigurationData* Data = GetOrLoadConfig();
+	check(Data);
+	
+	// Components to export
+	TArray<UDisplayClusterCameraComponent*> CameraComponents;
+	TArray<UDisplayClusterScreenComponent*> ScreenComponents;
+	TArray<USceneComponent*>  XformComponents;
+	// Auxiliary map for building hierarchy
+	TMap<UActorComponent*, FString> ParentComponentsMap;
+
+	// Make sure the 'Scene' object is there. Otherwise instantiate it.
+	// Could be null on assets used during 4.27 development, before scene was added back in.
+	const EObjectFlags CommonFlags = RF_Public | RF_Transactional;
+	if (Data->Scene == nullptr)
+	{
+		Data->Scene = NewObject<UDisplayClusterConfigurationScene>(
+			this,
+			UDisplayClusterConfigurationScene::StaticClass(),
+			NAME_None,
+			IsTemplate() ? RF_ArchetypeObject | CommonFlags : CommonFlags);
+	}
+	
+	// Extract CDO cameras (no screens in the CDO)
+	// Get list of cameras
+	CDO->GetComponents<UDisplayClusterCameraComponent>(CameraComponents);
+
+	// Extract BP components
+
+	const TArray<USCS_Node*>& Nodes = SimpleConstructionScript->GetAllNodes();
+	for (const USCS_Node* const Node : Nodes)
+	{
+		// Fill ID info for all descendants
+		GatherParentComponentsInfo(Node, ParentComponentsMap);
+
+		// Cameras
+		if (Node->ComponentClass->IsChildOf(UDisplayClusterCameraComponent::StaticClass()))
+		{
+			UDisplayClusterCameraComponent* ComponentTemplate = CastChecked<UDisplayClusterCameraComponent>(Node->GetActualComponentTemplate(CastChecked<UBlueprintGeneratedClass>(GeneratedClass)));
+			CameraComponents.Add(ComponentTemplate);
+		}
+		// Screens
+		else if (Node->ComponentClass->IsChildOf(UDisplayClusterScreenComponent::StaticClass()))
+		{
+			UDisplayClusterScreenComponent* ComponentTemplate = CastChecked<UDisplayClusterScreenComponent>(Node->GetActualComponentTemplate(CastChecked<UBlueprintGeneratedClass>(GeneratedClass)));
+			ScreenComponents.Add(ComponentTemplate);
+		}
+		// All other components will be exported as Xforms
+		else if (Node->ComponentClass->IsChildOf(USceneComponent::StaticClass()))
+		{
+			USceneComponent* ComponentTemplate = CastChecked<USceneComponent>(Node->GetActualComponentTemplate(CastChecked<UBlueprintGeneratedClass>(GeneratedClass)));
+			XformComponents.Add(ComponentTemplate);
+		}
+	}
+
+	// Save asset path
+	Data->Info.AssetPath = GetPathName();
+
+	// Prepare the target containers
+	Data->Scene->Cameras.Empty(CameraComponents.Num());
+	Data->Scene->Screens.Empty(ScreenComponents.Num());
+	Data->Scene->Xforms.Empty(XformComponents.Num());
+
+	// Export cameras
+	for (const UDisplayClusterCameraComponent* const CfgComp : CameraComponents)
+	{
+		UDisplayClusterConfigurationSceneComponentCamera* SceneComp = NewObject<UDisplayClusterConfigurationSceneComponentCamera>(Data->Scene, CfgComp->GetFName(), RF_Public);
+
+		// Save the properties
+		SceneComp->bSwapEyes = CfgComp->GetSwapEyes();
+		SceneComp->InterpupillaryDistance = CfgComp->GetInterpupillaryDistance();
+		// Safe to cast -- values match.
+		SceneComp->StereoOffset = (EDisplayClusterConfigurationEyeStereoOffset)CfgComp->GetStereoOffset();
+
+		FString* ParentId = ParentComponentsMap.Find(CfgComp);
+		SceneComp->ParentId = (ParentId ? *ParentId : FString());
+		SceneComp->Location = CfgComp->GetRelativeLocation();
+		SceneComp->Rotation = CfgComp->GetRelativeRotation();
+
+		// Store the object
+		Data->Scene->Cameras.Emplace(GetObjectNameFromSCSNode(SceneComp), SceneComp);
+	}
+
+	// Export screens
+	for (const UDisplayClusterScreenComponent* const CfgComp : ScreenComponents)
+	{
+		UDisplayClusterConfigurationSceneComponentScreen* SceneComp = NewObject<UDisplayClusterConfigurationSceneComponentScreen>(Data->Scene, CfgComp->GetFName());
+
+		// Save the properties
+		FString* ParentId = ParentComponentsMap.Find(CfgComp);
+		SceneComp->ParentId = (ParentId ? *ParentId : FString());
+		SceneComp->Location = CfgComp->GetRelativeLocation();
+		SceneComp->Rotation = CfgComp->GetRelativeRotation();
+
+		const FVector RelativeCompScale = CfgComp->GetRelativeScale3D();
+		SceneComp->Size = FVector2D(RelativeCompScale.Y, RelativeCompScale.Z);
+
+		// Store the object
+		Data->Scene->Screens.Emplace(GetObjectNameFromSCSNode(SceneComp), SceneComp);
+	}
+
+	// Export xforms
+	for (const USceneComponent* const CfgComp : XformComponents)
+	{
+		UDisplayClusterConfigurationSceneComponentXform* SceneComp = NewObject<UDisplayClusterConfigurationSceneComponentXform>(Data->Scene, CfgComp->GetFName());
+
+		// Save the properties
+		FString* ParentId = ParentComponentsMap.Find(CfgComp);
+		SceneComp->ParentId = (ParentId ? *ParentId : FString());
+		SceneComp->Location = CfgComp->GetRelativeLocation();
+		SceneComp->Rotation = CfgComp->GetRelativeRotation();
+
+		// Store the object
+		Data->Scene->Xforms.Emplace(GetObjectNameFromSCSNode(SceneComp), SceneComp);
+	}
+
+	// Avoid empty string keys in the config data maps
+	CleanupConfigMaps(Data);
+}
+
+FString UDisplayClusterBlueprint::GetObjectNameFromSCSNode(const UObject* const Object) const
+{
+	FString OutCompName;
+
+	if (Object)
+	{
+		OutCompName = Object->GetName();
+		OutCompName.RemoveFromEnd(TEXT("_GEN_VARIABLE"));
+	}
+
+	return OutCompName;
+}
+
+void UDisplayClusterBlueprint::GatherParentComponentsInfo(const USCS_Node* const InNode,
+	TMap<UActorComponent*, FString>& OutParentsMap) const
+{
+	if (InNode && InNode->ComponentClass->IsChildOf(UActorComponent::StaticClass()))
+	{
+		// Save current node to the map
+		UActorComponent* ParentNode = InNode->GetActualComponentTemplate(CastChecked<UBlueprintGeneratedClass>(GeneratedClass));
+		if (!OutParentsMap.Contains(ParentNode))
+		{
+			OutParentsMap.Emplace(ParentNode);
+		}
+
+		// Now iterate through the children nodes
+		for (USCS_Node* ChildNode : InNode->ChildNodes)
+		{
+			UActorComponent* ChildComponentTemplate = CastChecked<UActorComponent>(ChildNode->GetActualComponentTemplate(CastChecked<UBlueprintGeneratedClass>(GeneratedClass)));
+			if (ChildComponentTemplate)
+			{
+				OutParentsMap.Emplace(ChildComponentTemplate, GetObjectNameFromSCSNode(InNode->ComponentTemplate));
+			}
+
+			GatherParentComponentsInfo(ChildNode, OutParentsMap);
+		}
+	}
+}
+
+void UDisplayClusterBlueprint::CleanupConfigMaps(UDisplayClusterConfigurationData* Data) const
+{
+	check(Data && Data->Cluster);
+
+	static const FString InvalidKey = FString();
+
+	// Set of the maps we're going to process
+	TSet<TMap<FString, FString>*> MapsToProcess;
+	// Pre-allocate some memory. Not a precise amount of elements, but should be enough in most cases
+	MapsToProcess.Reserve(3 + 4 * Data->Cluster->Nodes.Num());
+
+	// Add single instance maps
+	MapsToProcess.Add(&Data->CustomParameters);
+	MapsToProcess.Add(&Data->Cluster->Sync.InputSyncPolicy.Parameters);
+	MapsToProcess.Add(&Data->Cluster->Sync.RenderSyncPolicy.Parameters);
+
+	// Add per-node and per-viewport maps
+	Data->Cluster->Nodes.Remove(InvalidKey);
+	for (TPair<FString, UDisplayClusterConfigurationClusterNode*> Node : Data->Cluster->Nodes)
+	{
+		check(Node.Value);
+
+		// Per-node maps
+		Node.Value->Postprocess.Remove(InvalidKey);
+		for (TPair<FString, FDisplayClusterConfigurationPostprocess> PostOpIt : Node.Value->Postprocess)
+		{
+			MapsToProcess.Add(&PostOpIt.Value.Parameters);
+		}
+
+		// Per-viewport maps
+		Node.Value->Viewports.Remove(InvalidKey);
+		for (TPair<FString, UDisplayClusterConfigurationViewport*> ViewportIt : Node.Value->Viewports)
+		{
+			check(ViewportIt.Value);
+			MapsToProcess.Add(&ViewportIt.Value->ProjectionPolicy.Parameters);
+		}
+	}
+
+	// Finally, remove all the pairs with empty keys
+	for (TMap<FString, FString>* Map : MapsToProcess)
+	{
+		Map->Remove(InvalidKey);
+	}
+}
diff --git a/Engine/Plugins/Runtime/nDisplay/Source/DisplayCluster/Private/Blueprints/DisplayClusterBlueprintAPIImpl.cpp b/Engine/Plugins/Runtime/nDisplay/Source/DisplayCluster/Private/Blueprints/DisplayClusterBlueprintAPIImpl.cpp
--- a/Engine/Plugins/Runtime/nDisplay/Source/DisplayCluster/Private/Blueprints/DisplayClusterBlueprintAPIImpl.cpp
+++ b/Engine/Plugins/Runtime/nDisplay/Source/DisplayCluster/Private/Blueprints/DisplayClusterBlueprintAPIImpl.cpp
@@ -67,11 +67,31 @@ bool UDisplayClusterBlueprintAPIImpl::IsMaster() const
 
 bool UDisplayClusterBlueprintAPIImpl::IsSlave() const
 {
-	const bool bIsSlave = !IsMaster();
+	const bool bIsSlave = IDisplayCluster::Get().GetClusterMgr()->IsSlave();
 	UE_LOG(LogDisplayClusterBlueprint, Verbose, TEXT("IsSlave - %s"), *DisplayClusterTypesConverter::template ToString(bIsSlave));
 	return bIsSlave;
 }
 
+bool UDisplayClusterBlueprintAPIImpl::IsBackup() const
+{
+	const bool bIsBackup = IDisplayCluster::Get().GetClusterMgr()->IsBackup();
+	UE_LOG(LogDisplayClusterBlueprint, Verbose, TEXT("IsBackup - %s"), *DisplayClusterTypesConverter::template ToString(bIsBackup));
+	return bIsBackup;
+}
+
+EDisplayClusterNodeRole UDisplayClusterBlueprintAPIImpl::GetClusterRole() const
+{
+	const EDisplayClusterNodeRole ClusterRole = IDisplayCluster::Get().GetClusterMgr()->GetClusterRole();
+	UE_LOG(LogDisplayClusterBlueprint, Verbose, TEXT("GetClusterRole - %d"), static_cast<int32>(ClusterRole));
+	return ClusterRole;
+}
+
+void UDisplayClusterBlueprintAPIImpl::GetNodeIds(TArray<FString>& OutNodeIds) const
+{
+	UE_LOG(LogDisplayClusterBlueprint, Verbose, TEXT("GetNodeIds"));
+	IDisplayCluster::Get().GetClusterMgr()->GetNodeIds(OutNodeIds);
+}
+
 FString UDisplayClusterBlueprintAPIImpl::GetNodeId() const
 {
 	const FString NodeId = IDisplayCluster::Get().GetClusterMgr()->GetNodeId();
diff --git a/Engine/Plugins/Runtime/nDisplay/Source/DisplayCluster/Private/Blueprints/DisplayClusterBlueprintAPIImpl.h b/Engine/Plugins/Runtime/nDisplay/Source/DisplayCluster/Private/Blueprints/DisplayClusterBlueprintAPIImpl.h
--- a/Engine/Plugins/Runtime/nDisplay/Source/DisplayCluster/Private/Blueprints/DisplayClusterBlueprintAPIImpl.h
+++ b/Engine/Plugins/Runtime/nDisplay/Source/DisplayCluster/Private/Blueprints/DisplayClusterBlueprintAPIImpl.h
@@ -38,6 +38,16 @@ public:
 	UFUNCTION(BlueprintCallable, meta = (DisplayName = "Is slave node"), Category = "NDisplay|Cluster")
 	virtual bool IsSlave() const override;
 
+	/** Returns true if current node is a backup node in a cluster. */
+	UFUNCTION(BlueprintCallable, meta = (DisplayName = "Is backup node"), Category = "NDisplay|Cluster")
+	virtual bool IsBackup() const override;
+
+	UFUNCTION(BlueprintCallable, meta = (DisplayName = "Get cluster role"), Category = "NDisplay|Cluster")
+	virtual EDisplayClusterNodeRole GetClusterRole() const override;
+
+	UFUNCTION(BlueprintCallable, meta = (DisplayName = "Get cluster node IDs"), Category = "NDisplay|Cluster")
+	virtual void GetNodeIds(TArray<FString>& OutNodeIds) const override;
+
 	UFUNCTION(BlueprintCallable, meta = (DisplayName = "Get node ID"), Category = "NDisplay|Cluster")
 	virtual FString GetNodeId() const override;
 
diff --git a/Engine/Plugins/Runtime/nDisplay/Source/DisplayCluster/Private/Cluster/Controller/DisplayClusterClusterNodeCtrlMaster.h b/Engine/Plugins/Runtime/nDisplay/Source/DisplayCluster/Private/Cluster/Controller/DisplayClusterClusterNodeCtrlMaster.h
--- a/Engine/Plugins/Runtime/nDisplay/Source/DisplayCluster/Private/Cluster/Controller/DisplayClusterClusterNodeCtrlMaster.h
+++ b/Engine/Plugins/Runtime/nDisplay/Source/DisplayCluster/Private/Cluster/Controller/DisplayClusterClusterNodeCtrlMaster.h
@@ -36,9 +36,9 @@ public:
 	//////////////////////////////////////////////////////////////////////////////////////////////
 	virtual void ClearCache() override;
 
-	virtual bool IsSlave() const override final
+	virtual EDisplayClusterNodeRole GetClusterRole() const override
 	{
-		return false;
+		return EDisplayClusterNodeRole::Master;
 	}
 
 protected:
diff --git a/Engine/Plugins/Runtime/nDisplay/Source/DisplayCluster/Private/Cluster/Controller/DisplayClusterClusterNodeCtrlSlave.h b/Engine/Plugins/Runtime/nDisplay/Source/DisplayCluster/Private/Cluster/Controller/DisplayClusterClusterNodeCtrlSlave.h
--- a/Engine/Plugins/Runtime/nDisplay/Source/DisplayCluster/Private/Cluster/Controller/DisplayClusterClusterNodeCtrlSlave.h
+++ b/Engine/Plugins/Runtime/nDisplay/Source/DisplayCluster/Private/Cluster/Controller/DisplayClusterClusterNodeCtrlSlave.h
@@ -24,9 +24,9 @@ public:
 	//////////////////////////////////////////////////////////////////////////////////////////////
 	// IDisplayClusterNodeController
 	//////////////////////////////////////////////////////////////////////////////////////////////
-	virtual bool IsSlave() const override
+	virtual EDisplayClusterNodeRole GetClusterRole() const override
 	{
-		return true;
+		return EDisplayClusterNodeRole::Slave;
 	}
 
 public:
diff --git a/Engine/Plugins/Runtime/nDisplay/Source/DisplayCluster/Private/Cluster/Controller/DisplayClusterNodeCtrlBase.h b/Engine/Plugins/Runtime/nDisplay/Source/DisplayCluster/Private/Cluster/Controller/DisplayClusterNodeCtrlBase.h
--- a/Engine/Plugins/Runtime/nDisplay/Source/DisplayCluster/Private/Cluster/Controller/DisplayClusterNodeCtrlBase.h
+++ b/Engine/Plugins/Runtime/nDisplay/Source/DisplayCluster/Private/Cluster/Controller/DisplayClusterNodeCtrlBase.h
@@ -33,11 +33,6 @@ public:
 	virtual bool Initialize() override final;
 	virtual void Release() override final;
 
-	virtual bool IsMaster() const override final
-	{
-		return !IsSlave();
-	}
-
 	virtual FString GetNodeId() const override final
 	{
 		return NodeName;
diff --git a/Engine/Plugins/Runtime/nDisplay/Source/DisplayCluster/Private/Cluster/Controller/IDisplayClusterNodeController.h b/Engine/Plugins/Runtime/nDisplay/Source/DisplayCluster/Private/Cluster/Controller/IDisplayClusterNodeController.h
--- a/Engine/Plugins/Runtime/nDisplay/Source/DisplayCluster/Private/Cluster/Controller/IDisplayClusterNodeController.h
+++ b/Engine/Plugins/Runtime/nDisplay/Source/DisplayCluster/Private/Cluster/Controller/IDisplayClusterNodeController.h
@@ -9,6 +9,8 @@
 #include "Network/Protocol/IDisplayClusterProtocolClusterSync.h"
 #include "Network/Protocol/IDisplayClusterProtocolRenderSync.h"
 
+#include "DisplayClusterEnums.h"
+
 
 /**
  * Node controller interface
@@ -33,11 +35,8 @@ public:
 	virtual void ClearCache()
 	{ }
 
-	/** Check if this is master controller */
-	virtual bool IsMaster() const = 0;
-
-	/** Check if this is slave controller */
-	virtual bool IsSlave() const = 0;
+	/** Returns cluster role of the controller */
+	virtual EDisplayClusterNodeRole GetClusterRole() const = 0;
 
 	/** Return node ID */
 	virtual FString GetNodeId() const = 0;
diff --git a/Engine/Plugins/Runtime/nDisplay/Source/DisplayCluster/Private/Cluster/DisplayClusterClusterManager.cpp b/Engine/Plugins/Runtime/nDisplay/Source/DisplayCluster/Private/Cluster/DisplayClusterClusterManager.cpp
--- a/Engine/Plugins/Runtime/nDisplay/Source/DisplayCluster/Private/Cluster/DisplayClusterClusterManager.cpp
+++ b/Engine/Plugins/Runtime/nDisplay/Source/DisplayCluster/Private/Cluster/DisplayClusterClusterManager.cpp
@@ -78,8 +78,8 @@ bool FDisplayClusterClusterManager::StartSession(UDisplayClusterConfigurationDat
 		return false;
 	}
 
-	// Save nodes amount
-	NodesAmount = ConfigData->Cluster->Nodes.Num();
+	// Save node IDs
+	ConfigData->Cluster->Nodes.GenerateKeyArray(ClusterNodeIds);
 
 	// Instantiate node controller
 	Controller = CreateController();
@@ -113,7 +113,7 @@ void FDisplayClusterClusterManager::EndSession()
 		}
 	}
 
-	NodesAmount = 0;
+	ClusterNodeIds.Reset();
 	ClusterNodeId.Empty();
 }
 
@@ -216,12 +216,32 @@ IDisplayClusterNodeController* FDisplayClusterClusterManager::GetController() co
 
 bool FDisplayClusterClusterManager::IsMaster() const
 {
-	return Controller ? Controller->IsMaster() : false;
+	FScopeLock Lock(&InternalsSyncScope);
+	return Controller ? Controller->GetClusterRole() == EDisplayClusterNodeRole::Master : false;
 }
 
 bool FDisplayClusterClusterManager::IsSlave() const
 {
-	return Controller ? Controller->IsSlave() : false;
+	FScopeLock Lock(&InternalsSyncScope);
+	return Controller ? Controller->GetClusterRole() == EDisplayClusterNodeRole::Slave : false;
+}
+
+bool FDisplayClusterClusterManager::IsBackup() const
+{
+	FScopeLock Lock(&InternalsSyncScope);
+	return Controller ? Controller->GetClusterRole() == EDisplayClusterNodeRole::Backup : false;
+}
+
+EDisplayClusterNodeRole FDisplayClusterClusterManager::GetClusterRole() const
+{
+	FScopeLock Lock(&InternalsSyncScope);
+	return Controller ? Controller->GetClusterRole() : EDisplayClusterNodeRole::None;
+}
+
+void FDisplayClusterClusterManager::GetNodeIds(TArray<FString>& OutNodeIds) const
+{
+	FScopeLock Lock(&InternalsSyncScope);
+	OutNodeIds = ClusterNodeIds;
 }
 
 void FDisplayClusterClusterManager::RegisterSyncObject(IDisplayClusterClusterSyncObject* SyncObj, EDisplayClusterSyncGroup SyncGroup)
diff --git a/Engine/Plugins/Runtime/nDisplay/Source/DisplayCluster/Private/Cluster/DisplayClusterClusterManager.h b/Engine/Plugins/Runtime/nDisplay/Source/DisplayCluster/Private/Cluster/DisplayClusterClusterManager.h
--- a/Engine/Plugins/Runtime/nDisplay/Source/DisplayCluster/Private/Cluster/DisplayClusterClusterManager.h
+++ b/Engine/Plugins/Runtime/nDisplay/Source/DisplayCluster/Private/Cluster/DisplayClusterClusterManager.h
@@ -38,8 +38,10 @@ public:
 	//////////////////////////////////////////////////////////////////////////////////////////////
 	// IDisplayClusterClusterManager
 	//////////////////////////////////////////////////////////////////////////////////////////////
-	virtual bool IsMaster()     const override;
-	virtual bool IsSlave()      const override;
+	virtual bool IsMaster() const override;
+	virtual bool IsSlave()  const override;
+	virtual bool IsBackup() const override;
+	virtual EDisplayClusterNodeRole GetClusterRole() const override;
 
 	virtual FString GetNodeId() const override
 	{
@@ -48,9 +50,11 @@ public:
 
 	virtual uint32 GetNodesAmount() const override
 	{
-		return NodesAmount;
+		return static_cast<uint32>(ClusterNodeIds.Num());
 	}
 
+	virtual void GetNodeIds(TArray<FString>& OutNodeIds) const override;
+
 	virtual void RegisterSyncObject(IDisplayClusterClusterSyncObject* SyncObj, EDisplayClusterSyncGroup SyncGroup) override;
 	virtual void UnregisterSyncObject(IDisplayClusterClusterSyncObject* SyncObj) override;
 
@@ -97,8 +101,9 @@ private:
 private:
 	// Controller implementation
 	TUniquePtr<IDisplayClusterNodeController> Controller;
-	// Cluster/node props
-	uint32 NodesAmount = 0;
+
+	// Custom list of cluster nodes
+	TArray<FString> ClusterNodeIds;
 
 	// Current operation mode
 	EDisplayClusterOperationMode CurrentOperationMode;
diff --git a/Engine/Plugins/Runtime/nDisplay/Source/DisplayCluster/Private/Game/DisplayClusterGameManager.cpp b/Engine/Plugins/Runtime/nDisplay/Source/DisplayCluster/Private/Game/DisplayClusterGameManager.cpp
--- a/Engine/Plugins/Runtime/nDisplay/Source/DisplayCluster/Private/Game/DisplayClusterGameManager.cpp
+++ b/Engine/Plugins/Runtime/nDisplay/Source/DisplayCluster/Private/Game/DisplayClusterGameManager.cpp
@@ -71,21 +71,25 @@ bool FDisplayClusterGameManager::StartScene(UWorld* InWorld)
 {
 	FScopeLock Lock(&InternalsSyncScope);
 
-	check(InWorld);
-	check(ConfigData);
+	check(InWorld && ConfigData);
 
 	CurrentWorld = InWorld;
 
 	// Find the first DCRA instance that matches to the specified configuration
 	ADisplayClusterRootActor* RootActor = FindRootActor(InWorld, ConfigData);
 
-	// If no proper DCRA found,
-	// 1. Detect spawn location and rotation
-	// 2. Try to spawn a blueprint from a corresponding asset
-	// 3. If the asset not found, spawn an empty DCRA instance and initialize it with specified configuration data.
 	if (GDisplayCluster->GetOperationMode() == EDisplayClusterOperationMode::Cluster)
 	{
-		if (!RootActor)
+		// If a corresponding DCRA instance was found, overwrite its settings
+		if (RootActor)
+		{
+			RootActor->OverrideFromConfig(ConfigData);
+		}
+		// If no proper DCRA found,
+		// 1. Detect spawn location and rotation
+		// 2. Try to spawn a blueprint from a corresponding asset
+		// 3. If the asset not found, spawn an empty DCRA instance and initialize it with specified configuration data.
+		else
 		{
 			// 1. Detect spawn location and rotation
 			FVector  StartLocation = FVector::ZeroVector;
@@ -99,7 +103,7 @@ bool FDisplayClusterGameManager::StartScene(UWorld* InWorld)
 				StartRotation = (*It)->GetActorRotation();
 			}
 
-			// 2. Spawn the DCRA PB from a corresponding asset
+			// 2. Spawn the DCRA BP from a corresponding asset
 			UObject* ActorToSpawn = Cast<UObject>(StaticLoadObject(UObject::StaticClass(), NULL, *ConfigData->Info.AssetPath));
 			if (ActorToSpawn)
 			{
@@ -107,8 +111,12 @@ bool FDisplayClusterGameManager::StartScene(UWorld* InWorld)
 				UClass* ClassToSpawn = ActorToSpawn->StaticClass();
 				if (ClassToSpawn && GeneratedBP)
 				{
+					// Spawn an asset
 					AActor* NewActor = CurrentWorld->SpawnActor<AActor>(GeneratedBP->GeneratedClass, StartLocation, StartRotation, FActorSpawnParameters());
 					RootActor = Cast<ADisplayClusterRootActor>(NewActor);
+
+					// Override actor settings in case the config file contains some updates
+					RootActor->OverrideFromConfig(ConfigData);
 				}
 			}
 
diff --git a/Engine/Plugins/Runtime/nDisplay/Source/DisplayCluster/Private/Game/EngineClasses/Scene/DisplayClusterRootActor.cpp b/Engine/Plugins/Runtime/nDisplay/Source/DisplayCluster/Private/Game/EngineClasses/Scene/DisplayClusterRootActor.cpp
--- a/Engine/Plugins/Runtime/nDisplay/Source/DisplayCluster/Private/Game/EngineClasses/Scene/DisplayClusterRootActor.cpp
+++ b/Engine/Plugins/Runtime/nDisplay/Source/DisplayCluster/Private/Game/EngineClasses/Scene/DisplayClusterRootActor.cpp
@@ -135,6 +135,62 @@ void ADisplayClusterRootActor::InitializeFromConfig(UDisplayClusterConfiguration
 	}
 }
 
+void ADisplayClusterRootActor::OverrideFromConfig(UDisplayClusterConfigurationData* ConfigData)
+{
+	check(ConfigData);
+	check(ConfigData->Scene);
+	check(ConfigData->Cluster);
+
+	// Override base types and structures
+	CurrentConfigData->Meta = ConfigData->Meta;
+	CurrentConfigData->Info = ConfigData->Info;
+	CurrentConfigData->CustomParameters = ConfigData->CustomParameters;
+	CurrentConfigData->Diagnostics = ConfigData->Diagnostics;
+	CurrentConfigData->bFollowLocalPlayerCamera = ConfigData->bFollowLocalPlayerCamera;
+	CurrentConfigData->bExitOnEsc = ConfigData->bExitOnEsc;
+
+	// Override Scene but without changing its name
+	{
+		FName SceneName = NAME_None;
+
+		if (CurrentConfigData->Scene)
+		{
+			SceneName = CurrentConfigData->Scene->GetFName();
+
+			const FName DeadName = MakeUniqueObjectName(CurrentConfigData, UDisplayClusterConfigurationScene::StaticClass(), "DEAD_DisplayClusterConfigurationScene");
+			CurrentConfigData->Scene->Rename(*DeadName.ToString());
+		}
+
+		CurrentConfigData->Scene = DuplicateObject(ConfigData->Scene, CurrentConfigData, SceneName);
+	}
+
+	// Override Cluster but without changing its name
+	{
+		FName ClusterName = NAME_None;
+
+		if (CurrentConfigData->Cluster)
+		{
+			ClusterName = CurrentConfigData->Cluster->GetFName();
+
+			const FName DeadName = MakeUniqueObjectName(CurrentConfigData, UDisplayClusterConfigurationCluster::StaticClass(), "DEAD_DisplayClusterConfigurationCluster");
+			CurrentConfigData->Cluster->Rename(*DeadName.ToString());
+		}
+
+		CurrentConfigData->Cluster = DuplicateObject(ConfigData->Cluster, CurrentConfigData, ClusterName);
+	}
+
+	// There is no sense to call BuildHierarchy because it works for non-BP root actors.
+	// On the other hand, OverwriteFromConfig method is called for BP root actors only by nature.
+
+	// And update preview stuff in Editor
+#if WITH_EDITOR
+	if (GIsEditor && GetWorld())
+	{
+		UpdatePreviewComponents();
+	}
+#endif
+}
+
 UDisplayClusterConfigurationViewport* ADisplayClusterRootActor::GetViewportConfiguration(const FString& ClusterNodeID, const FString& ViewportID)
 {
 	if (CurrentConfigData)
@@ -426,16 +482,6 @@ void ADisplayClusterRootActor::InitializeRootActor()
 #endif
 			return;
 		}
-		
-		IPDisplayClusterConfigManager* const ConfigMgr = (GDisplayCluster ? GDisplayCluster->GetPrivateConfigMgr() : nullptr);
-		if (ConfigMgr)
-		{
-			UDisplayClusterConfigurationData* ConfigData = ConfigMgr->GetConfig();
-			if (ConfigData)
-			{
-				InitializeFromConfig(ConfigData);
-			}
-		}
 	}
 #if WITH_EDITOR
 	// Initialize from file property by default in Editor
diff --git a/Engine/Plugins/Runtime/nDisplay/Source/DisplayCluster/Private/Render/Containers/DisplayClusterRender_MeshComponent.cpp b/Engine/Plugins/Runtime/nDisplay/Source/DisplayCluster/Private/Render/Containers/DisplayClusterRender_MeshComponent.cpp
--- a/Engine/Plugins/Runtime/nDisplay/Source/DisplayCluster/Private/Render/Containers/DisplayClusterRender_MeshComponent.cpp
+++ b/Engine/Plugins/Runtime/nDisplay/Source/DisplayCluster/Private/Render/Containers/DisplayClusterRender_MeshComponent.cpp
@@ -102,6 +102,17 @@ void FDisplayClusterRender_MeshComponent::UpdateDefferedRef()
 	ImplUpdateMeshComponentProxyData(MeshComponentProxy, NewProxyData);
 }
 
+void FDisplayClusterRender_MeshComponent::UpdateDeffered(const UStaticMesh* InStaticMesh)
+{
+	FDisplayClusterRender_MeshComponentProxyData* NewProxyData = nullptr;
+	if (InStaticMesh)
+	{
+		NewProxyData = new FDisplayClusterRender_MeshComponentProxyData(DataFunc, *InStaticMesh);
+	}
+
+	ImplUpdateMeshComponentProxyData(MeshComponentProxy, NewProxyData);
+}
+
 void FDisplayClusterRender_MeshComponent::UpdateDeffered(const FDisplayClusterRender_MeshGeometry& InMeshGeometry)
 {
 	check(IsInGameThread());
diff --git a/Engine/Plugins/Runtime/nDisplay/Source/DisplayCluster/Private/Render/Containers/DisplayClusterRender_MeshComponentProxyData.cpp b/Engine/Plugins/Runtime/nDisplay/Source/DisplayCluster/Private/Render/Containers/DisplayClusterRender_MeshComponentProxyData.cpp
--- a/Engine/Plugins/Runtime/nDisplay/Source/DisplayCluster/Private/Render/Containers/DisplayClusterRender_MeshComponentProxyData.cpp
+++ b/Engine/Plugins/Runtime/nDisplay/Source/DisplayCluster/Private/Render/Containers/DisplayClusterRender_MeshComponentProxyData.cpp
@@ -11,39 +11,60 @@
 //*************************************************************************
 //* FDisaplyClusterMeshComponentProxyData
 //*************************************************************************
-FDisplayClusterRender_MeshComponentProxyData::FDisplayClusterRender_MeshComponentProxyData(const FDisplayClusterRender_MeshComponentProxyDataFunc InDataFunc, const UStaticMeshComponent& InMeshComponent, const int InUVChromakeyIndex)
+void FDisplayClusterRender_MeshComponentProxyData::ImplInitializeMesh(const FDisplayClusterRender_MeshComponentProxyDataFunc InDataFunc, const UStaticMesh& InStaticMesh, const int InUVChromakeyIndex)
 {
-	check(IsInGameThread());
-
-	UStaticMesh* StaticMesh = InMeshComponent.GetStaticMesh();
-	if (StaticMesh)
+	if (InStaticMesh.bAllowCPUAccess == false)
 	{
-		const FStaticMeshLODResources& SrcMeshResource = StaticMesh->GetLODForExport(0);
+		UE_LOG(LogDisplayClusterRender, Warning, TEXT("If packaging this project, static mesh '%s' requires its AllowCPUAccess flag to be enabled."), *InStaticMesh.GetName());
+#if !WITH_EDITOR
+		// Can't access to cooked data from CPU without this flag
+		return;
+#endif
+	}
 
-		const FPositionVertexBuffer& VertexPosition = SrcMeshResource.VertexBuffers.PositionVertexBuffer;
-		const FStaticMeshVertexBuffer& VertexBuffer = SrcMeshResource.VertexBuffers.StaticMeshVertexBuffer;
-		const FRawStaticIndexBuffer& IndexBuffer = SrcMeshResource.IndexBuffer;
+	const FStaticMeshLODResources& SrcMeshResource = InStaticMesh.GetLODForExport(0);
 
-		NumTriangles = IndexBuffer.GetNumIndices() / 3; // Now by default no triangle strip supported
-		NumVertices = VertexBuffer.GetNumVertices();
+	const FPositionVertexBuffer& VertexPosition = SrcMeshResource.VertexBuffers.PositionVertexBuffer;
+	const FStaticMeshVertexBuffer& VertexBuffer = SrcMeshResource.VertexBuffers.StaticMeshVertexBuffer;
+	const FRawStaticIndexBuffer& IndexBuffer = SrcMeshResource.IndexBuffer;
 
-		uint32 UVBaseIndex = 0;
-		uint32 UVChromakeyIndex = (InUVChromakeyIndex < SrcMeshResource.GetNumTexCoords()) ? InUVChromakeyIndex : 0;
+	NumTriangles = IndexBuffer.GetNumIndices() / 3; // Now by default no triangle strip supported
+	NumVertices = VertexBuffer.GetNumVertices();
 
-		IndexBuffer.GetCopy(IndexData);
+	uint32 UVBaseIndex = 0;
+	uint32 UVChromakeyIndex = (InUVChromakeyIndex < SrcMeshResource.GetNumTexCoords()) ? InUVChromakeyIndex : 0;
 
-		VertexData.AddZeroed(NumVertices);
-		for (uint32 i = 0; i < NumVertices; i++)
-		{
-			VertexData[i].Position = VertexPosition.VertexPosition(i);
-			VertexData[i].UV = VertexBuffer.GetVertexUV(i, UVBaseIndex);
-			VertexData[i].UV_Chromakey = VertexBuffer.GetVertexUV(i, UVChromakeyIndex);
-		}
+	IndexBuffer.GetCopy(IndexData);
+
+	VertexData.AddZeroed(NumVertices);
+	for (uint32 i = 0; i < NumVertices; i++)
+	{
+		VertexData[i].Position = VertexPosition.VertexPosition(i);
+		VertexData[i].UV = VertexBuffer.GetVertexUV(i, UVBaseIndex);
+		VertexData[i].UV_Chromakey = VertexBuffer.GetVertexUV(i, UVChromakeyIndex);
+	}
+
+	UpdateData(InDataFunc);
+}
 
-		UpdateData(InDataFunc);
+FDisplayClusterRender_MeshComponentProxyData::FDisplayClusterRender_MeshComponentProxyData(const FDisplayClusterRender_MeshComponentProxyDataFunc InDataFunc, const UStaticMeshComponent& InMeshComponent, const int InUVChromakeyIndex)
+{
+	check(IsInGameThread());
+
+	UStaticMesh* StaticMesh = InMeshComponent.GetStaticMesh();
+	if (StaticMesh)
+	{
+		ImplInitializeMesh(InDataFunc, *StaticMesh, InUVChromakeyIndex);
 	}
 }
 
+FDisplayClusterRender_MeshComponentProxyData::FDisplayClusterRender_MeshComponentProxyData(const FDisplayClusterRender_MeshComponentProxyDataFunc InDataFunc, const UStaticMesh& InStaticMesh, const int InUVChromakeyIndex)
+{
+	check(IsInGameThread());
+
+	ImplInitializeMesh(InDataFunc, InStaticMesh, InUVChromakeyIndex);
+}
+
 FDisplayClusterRender_MeshComponentProxyData::FDisplayClusterRender_MeshComponentProxyData(const FDisplayClusterRender_MeshComponentProxyDataFunc InDataFunc, const FDisplayClusterRender_MeshGeometry& InMeshGeometry)
 {
 	bool bUseChromakeyUV = InMeshGeometry.ChromakeyUV.Num() > 0;
diff --git a/Engine/Plugins/Runtime/nDisplay/Source/DisplayCluster/Private/Render/Viewport/Configuration/DisplayClusterViewportConfigurationBase.cpp b/Engine/Plugins/Runtime/nDisplay/Source/DisplayCluster/Private/Render/Viewport/Configuration/DisplayClusterViewportConfigurationBase.cpp
--- a/Engine/Plugins/Runtime/nDisplay/Source/DisplayCluster/Private/Render/Viewport/Configuration/DisplayClusterViewportConfigurationBase.cpp
+++ b/Engine/Plugins/Runtime/nDisplay/Source/DisplayCluster/Private/Render/Viewport/Configuration/DisplayClusterViewportConfigurationBase.cpp
@@ -5,10 +5,10 @@
 #include "Render/Viewport/DisplayClusterViewport.h"
 #include "Render/Viewport/DisplayClusterViewportManager.h"
 #include "Render/Viewport/Postprocess/DisplayClusterViewportPostProcessManager.h"
+#include "Render/Viewport/Postprocess/DisplayClusterViewportPostProcessOutputRemap.h"
 
 #include "DisplayClusterViewportConfigurationHelpers.h"
 #include "DisplayClusterConfigurationTypes.h"
-
 ///////////////////////////////////////////////////////////////////
 // FDisplayClusterViewportConfigurationBase
 ///////////////////////////////////////////////////////////////////
@@ -111,6 +111,31 @@ void FDisplayClusterViewportConfigurationBase::UpdateClusterNodePostProcess(cons
 					PPManager->CreatePostprocess(It.Key, &It.Value);
 				}
 			}
+
+			// Update OutputRemap PP
+			{
+				const struct FDisplayClusterConfigurationFramePostProcess_OutputRemap& OutputRemapCfg = ClusterNode->OutputRemap;
+				if (OutputRemapCfg.bEnable)
+				{
+					switch (OutputRemapCfg.DataSource)
+					{
+					case EDisplayClusterConfigurationFramePostProcess_OutputRemapSource::StaticMesh:
+						PPManager->GetOutputRemap()->UpdateConfiguration_StaticMesh(OutputRemapCfg.StaticMesh);
+						break;
+
+					case EDisplayClusterConfigurationFramePostProcess_OutputRemapSource::ExternalFile:
+						PPManager->GetOutputRemap()->UpdateConfiguration_ExternalFile(OutputRemapCfg.ExternalFile);
+						break;
+
+					default:
+						PPManager->GetOutputRemap()->UpdateConfiguration_Disabled();
+					}
+				}
+				else
+				{
+					PPManager->GetOutputRemap()->UpdateConfiguration_Disabled();
+				}
+			}
 		}
 	}
 }
diff --git a/Engine/Plugins/Runtime/nDisplay/Source/DisplayCluster/Private/Render/Viewport/DisplayClusterViewport.h b/Engine/Plugins/Runtime/nDisplay/Source/DisplayCluster/Private/Render/Viewport/DisplayClusterViewport.h
--- a/Engine/Plugins/Runtime/nDisplay/Source/DisplayCluster/Private/Render/Viewport/DisplayClusterViewport.h
+++ b/Engine/Plugins/Runtime/nDisplay/Source/DisplayCluster/Private/Render/Viewport/DisplayClusterViewport.h
@@ -232,4 +232,7 @@ protected:
 	TArray<FDisplayClusterTextureResource*> MipsShaderResources;
 
 	FDisplayClusterViewportManager& Owner;
+
+private:
+	bool bProjectionPolicyCalculateViewWarningOnce = false;
 };
diff --git a/Engine/Plugins/Runtime/nDisplay/Source/DisplayCluster/Private/Render/Viewport/DisplayClusterViewport_EditorPreview.cpp b/Engine/Plugins/Runtime/nDisplay/Source/DisplayCluster/Private/Render/Viewport/DisplayClusterViewport_EditorPreview.cpp
--- a/Engine/Plugins/Runtime/nDisplay/Source/DisplayCluster/Private/Render/Viewport/DisplayClusterViewport_EditorPreview.cpp
+++ b/Engine/Plugins/Runtime/nDisplay/Source/DisplayCluster/Private/Render/Viewport/DisplayClusterViewport_EditorPreview.cpp
@@ -265,10 +265,16 @@ bool FDisplayClusterViewport::ImplPreview_CalculateStereoViewOffset(const uint32
 	// Perform view calculations on a policy side
 	if (!CalculateView(InContextNum, ViewLocation, ViewRotation, ViewOffset, WorldToMeters, CfgNCP, CfgNCP))
 	{
-		UE_LOG(LogDisplayClusterViewport, Warning, TEXT("Couldn't compute view parameters for Viewport %s, ViewIdx: %d"), *GetId(), InContextNum);
+		if (!bProjectionPolicyCalculateViewWarningOnce)
+		{
+			UE_LOG(LogDisplayClusterViewport, Warning, TEXT("Couldn't compute view parameters for Viewport %s, ViewIdx: %d"), *GetId(), InContextNum);
+			bProjectionPolicyCalculateViewWarningOnce = true;
+		}
 		return false;
 	}
 
+	bProjectionPolicyCalculateViewWarningOnce = false;
+
 	UE_LOG(LogDisplayClusterViewport, VeryVerbose, TEXT("ViewLoc: %s, ViewRot: %s"), *ViewLocation.ToString(), *ViewRotation.ToString());
 
 	return true;
diff --git a/Engine/Plugins/Runtime/nDisplay/Source/DisplayCluster/Private/Render/Viewport/Postprocess/DisplayClusterViewportPostProcessManager.cpp b/Engine/Plugins/Runtime/nDisplay/Source/DisplayCluster/Private/Render/Viewport/Postprocess/DisplayClusterViewportPostProcessManager.cpp
--- a/Engine/Plugins/Runtime/nDisplay/Source/DisplayCluster/Private/Render/Viewport/Postprocess/DisplayClusterViewportPostProcessManager.cpp
+++ b/Engine/Plugins/Runtime/nDisplay/Source/DisplayCluster/Private/Render/Viewport/Postprocess/DisplayClusterViewportPostProcessManager.cpp
@@ -16,6 +16,8 @@
 #include "Render/Viewport/DisplayClusterViewportManagerProxy.h"
 #include "Render/Viewport/DisplayClusterViewportProxy.h"
 
+#include "Render/Viewport/Postprocess/DisplayClusterViewportPostProcessOutputRemap.h"
+
 #include "Misc/DisplayClusterGlobals.h"
 
 #include "DisplayClusterConfigurationTypes.h"
@@ -56,6 +58,12 @@ static TAutoConsoleVariable<int32> CVarPostprocessFrameAfterWarpBlend(
 	ECVF_RenderThreadSafe
 );
 
+FDisplayClusterViewportPostProcessManager::FDisplayClusterViewportPostProcessManager(FDisplayClusterViewportManager& InViewportManager)
+	: ViewportManager(InViewportManager)
+{
+	OutputRemap = MakeShared<FDisplayClusterViewportPostProcessOutputRemap, ESPMode::ThreadSafe>();
+}
+
 bool FDisplayClusterViewportPostProcessManager::HandleStartScene()
 {
 	check(IsInGameThread());;
@@ -220,6 +228,11 @@ bool FDisplayClusterViewportPostProcessManager::IsPostProcessFrameAfterWarpBlend
 		return (CVarPostprocessFrameAfterWarpBlend.GetValueOnAnyThread() != 0);
 	}
 
+	if (OutputRemap.IsValid() && OutputRemap->IsValid())
+	{
+		return true;
+	}
+
 	return false;
 }
 
@@ -240,6 +253,11 @@ bool FDisplayClusterViewportPostProcessManager::IsAnyPostProcessRequired(const T
 		return true;
 			}
 
+	if (OutputRemap.IsValid() && OutputRemap->IsValid())
+	{
+		return true;
+	}
+
 	return false;
 }
 
@@ -255,6 +273,11 @@ bool FDisplayClusterViewportPostProcessManager::ShouldUseAdditionalFrameTargetab
 		}
 	}
 
+	if (OutputRemap.IsValid() && OutputRemap->IsValid())
+	{
+		return true;
+	}
+
 	return false;
 }
 
@@ -381,6 +404,13 @@ void FDisplayClusterViewportPostProcessManager::ImplPerformPostProcessFrameAfter
 					It->PerformPostProcessFrameAfterWarpBlend_RenderThread(RHICmdList, &FrameResources, AdditionalResources);
 				}
 			}
+
+			// Apply OutputRemap after all postprocess
+			if (OutputRemap.IsValid() && OutputRemap->IsValid())
+			{
+				TArray<FRHITexture2D*>* AdditionalResources = (AdditionalFrameResources.Num() > 0) ? &AdditionalFrameResources : nullptr;
+				OutputRemap->PerformPostProcessFrame_RenderThread(RHICmdList, &FrameResources, AdditionalResources);
+			}
 		}
 	}
 }
diff --git a/Engine/Plugins/Runtime/nDisplay/Source/DisplayCluster/Private/Render/Viewport/Postprocess/DisplayClusterViewportPostProcessManager.h b/Engine/Plugins/Runtime/nDisplay/Source/DisplayCluster/Private/Render/Viewport/Postprocess/DisplayClusterViewportPostProcessManager.h
--- a/Engine/Plugins/Runtime/nDisplay/Source/DisplayCluster/Private/Render/Viewport/Postprocess/DisplayClusterViewportPostProcessManager.h
+++ b/Engine/Plugins/Runtime/nDisplay/Source/DisplayCluster/Private/Render/Viewport/Postprocess/DisplayClusterViewportPostProcessManager.h
@@ -7,6 +7,7 @@
 
 class FDisplayClusterViewportManager;
 class FDisplayClusterViewportManagerProxy;
+class FDisplayClusterViewportPostProcessOutputRemap;
 struct FDisplayClusterConfigurationPostprocess;
 
 /**
@@ -16,10 +17,7 @@ struct FDisplayClusterConfigurationPostprocess;
 class FDisplayClusterViewportPostProcessManager
 {
 public:
-	FDisplayClusterViewportPostProcessManager(FDisplayClusterViewportManager& InViewportManager)
-		: ViewportManager(InViewportManager)
-	{}
-
+	FDisplayClusterViewportPostProcessManager(FDisplayClusterViewportManager& InViewportManager);
 	virtual ~FDisplayClusterViewportPostProcessManager() = default;
 
 public:
@@ -42,6 +40,9 @@ public:
 	// Send data to render thread
 	void FinalizeNewFrame();
 
+	TSharedPtr<FDisplayClusterViewportPostProcessOutputRemap, ESPMode::ThreadSafe>& GetOutputRemap()
+	{ return OutputRemap; }
+
 public:
 	bool CreatePostprocess(const FString& InPostprocessId, const FDisplayClusterConfigurationPostprocess* InConfigurationPostprocess);
 	bool RemovePostprocess(const FString& InPostprocessId);
@@ -71,4 +72,5 @@ private:
 	TArray<TSharedPtr<IDisplayClusterPostProcess, ESPMode::ThreadSafe>> PostprocessProxy;
 
 	FDisplayClusterViewportManager& ViewportManager;
+	TSharedPtr<FDisplayClusterViewportPostProcessOutputRemap, ESPMode::ThreadSafe> OutputRemap;
 };
diff --git a/Engine/Plugins/Runtime/nDisplay/Source/DisplayCluster/Private/Render/Viewport/Postprocess/DisplayClusterViewportPostProcessOutputRemap.cpp b/Engine/Plugins/Runtime/nDisplay/Source/DisplayCluster/Private/Render/Viewport/Postprocess/DisplayClusterViewportPostProcessOutputRemap.cpp
new file mode 100644
--- /dev/null
+++ b/Engine/Plugins/Runtime/nDisplay/Source/DisplayCluster/Private/Render/Viewport/Postprocess/DisplayClusterViewportPostProcessOutputRemap.cpp
@@ -0,0 +1,149 @@
+// Copyright Epic Games, Inc. All Rights Reserved.
+
+#include "Render/Viewport/Postprocess/DisplayClusterViewportPostProcessOutputRemap.h"
+
+#include "DisplayClusterConfigurationTypes_OutputRemap.h"
+
+#include "Misc/DisplayClusterHelpers.h"
+#include "Misc/DisplayClusterLog.h"
+
+#include "DisplayClusterRootActor.h"
+
+#include "Render/Containers/DisplayClusterRender_MeshComponent.h"
+#include "Render/Containers/DisplayClusterRender_MeshComponentProxy.h"
+#include "Render/Containers/DisplayClusterRender_MeshComponentProxyData.h"
+#include "Render/Containers/DisplayClusterRender_MeshGeometry.h"
+
+#include "Render/Viewport/IDisplayClusterViewportManager.h"
+
+#include "IDisplayClusterShaders.h"
+
+#include "Engine/RendererSettings.h"
+#include "Engine/StaticMesh.h"
+#include "Misc/Paths.h"
+#
+
+//////////////////////////////////////////////////////////////////////////////////////////////
+// FDisplayClusterViewportPostProcessOutputRemap
+//////////////////////////////////////////////////////////////////////////////////////////////
+FDisplayClusterViewportPostProcessOutputRemap::FDisplayClusterViewportPostProcessOutputRemap()
+	: OutputRemapMesh(*(new FDisplayClusterRender_MeshComponent()))
+	, ShadersAPI(IDisplayClusterShaders::Get())
+{
+}
+
+FDisplayClusterViewportPostProcessOutputRemap::~FDisplayClusterViewportPostProcessOutputRemap()
+{
+	check(IsInRenderingThread());
+
+	delete (&OutputRemapMesh);
+}
+
+bool FDisplayClusterViewportPostProcessOutputRemap::UpdateConfiguration_ExternalFile(const FString& InExternalFile)
+{
+	if (ExternalFile == InExternalFile)
+	{
+		return true;
+	}
+
+	UpdateConfiguration_Disabled();
+
+	if (InExternalFile.IsEmpty())
+	{
+		return true;
+	}
+
+	// Support related paths
+	FString FullPathFileName = DisplayClusterHelpers::filesystem::GetFullPathForConfigResource(InExternalFile);
+	if (FullPathFileName.IsEmpty() || !FPaths::FileExists(FullPathFileName))
+	{
+		if (!bErrorCantFindFileOnce)
+		{
+			UE_LOG(LogDisplayClusterViewport, Error, TEXT("OutputRemap - Failed to find file '%s'"), *InExternalFile);
+			bErrorCantFindFileOnce = true;
+		}
+
+		return false;
+	}
+
+	bErrorCantFindFileOnce = false;
+
+	FDisplayClusterRender_MeshGeometry MeshGeometry;
+	if (!MeshGeometry.LoadFromFile(FullPathFileName))
+	{
+		if (!bErrorFailLoadFromFileOnce)
+		{
+			UE_LOG(LogDisplayClusterViewport, Error, TEXT("OutputRemap - Failed to load ext mesh from file '%s'"), *FullPathFileName);
+			bErrorFailLoadFromFileOnce = true;
+		}
+
+		return false;
+	}
+
+	bErrorFailLoadFromFileOnce = false;
+
+
+	OutputRemapMesh.DataFunc = FDisplayClusterRender_MeshComponentProxyDataFunc::OutputRemapScreenSpace;
+	OutputRemapMesh.UpdateDeffered(MeshGeometry);
+
+	bIsEnabled = true;
+	ExternalFile = InExternalFile;
+
+	return true;
+}
+
+bool FDisplayClusterViewportPostProcessOutputRemap::UpdateConfiguration_StaticMesh(UStaticMesh* InStaticMesh)
+{
+	if (StaticMesh == InStaticMesh)
+	{
+		return true;
+	}
+
+	UpdateConfiguration_Disabled();
+
+	if (InStaticMesh == nullptr)
+	{
+		return true;
+	}
+
+	OutputRemapMesh.DataFunc = FDisplayClusterRender_MeshComponentProxyDataFunc::OutputRemapScreenSpace;
+	OutputRemapMesh.UpdateDeffered(InStaticMesh);
+
+	bIsEnabled = true;
+	StaticMesh = InStaticMesh;
+
+	return true;
+}
+
+void FDisplayClusterViewportPostProcessOutputRemap::UpdateConfiguration_Disabled()
+{
+	StaticMesh = nullptr;
+	ExternalFile.Empty();
+	bIsEnabled = false;
+
+	OutputRemapMesh.UpdateDeffered(nullptr);
+}
+
+void FDisplayClusterViewportPostProcessOutputRemap::PerformPostProcessFrame_RenderThread(FRHICommandListImmediate& RHICmdList, const TArray<FRHITexture2D*>* InFrameTargets, const TArray<FRHITexture2D*>* InAdditionalFrameTargets) const
+{
+	check(IsInRenderingThread());
+
+	if (InFrameTargets && InAdditionalFrameTargets)
+	{
+		const FDisplayClusterRender_MeshComponentProxy* MeshProxy = OutputRemapMesh.GetProxy();
+		if (MeshProxy && MeshProxy->IsValid_RenderThread())
+		{
+			for (int Index = 0; Index < InFrameTargets->Num(); Index++)
+			{
+				FRHITexture2D* InOutTexture = (*InFrameTargets)[Index];
+				FRHITexture2D* TempTargetableTexture = (*InAdditionalFrameTargets)[Index];
+
+				if (ShadersAPI.RenderPostprocess_OutputRemap(RHICmdList, InOutTexture, TempTargetableTexture, *MeshProxy))
+				{
+					//Copy remapped result back
+					RHICmdList.CopyToResolveTarget(TempTargetableTexture, InOutTexture, FResolveParams());
+				}
+			}
+		}
+	}
+}
diff --git a/Engine/Plugins/Runtime/nDisplay/Source/DisplayCluster/Private/Render/Viewport/Postprocess/DisplayClusterViewportPostProcessOutputRemap.h b/Engine/Plugins/Runtime/nDisplay/Source/DisplayCluster/Private/Render/Viewport/Postprocess/DisplayClusterViewportPostProcessOutputRemap.h
new file mode 100644
--- /dev/null
+++ b/Engine/Plugins/Runtime/nDisplay/Source/DisplayCluster/Private/Render/Viewport/Postprocess/DisplayClusterViewportPostProcessOutputRemap.h
@@ -0,0 +1,45 @@
+// Copyright Epic Games, Inc. All Rights Reserved.
+
+#pragma once
+
+#include "RHI.h"
+#include "RHICommandList.h"
+#include "RHIResources.h"
+#include "RHIUtilities.h"
+
+class IDisplayClusterViewportManager;
+class FDisplayClusterRender_MeshComponent;
+class UStaticMesh;
+
+/**
+ * MPCDI projection policy
+ */
+class FDisplayClusterViewportPostProcessOutputRemap
+{
+public:
+	FDisplayClusterViewportPostProcessOutputRemap();
+	virtual ~FDisplayClusterViewportPostProcessOutputRemap();
+
+	// Game thread update calls
+	bool UpdateConfiguration_ExternalFile(const FString& InExternalFile);
+	bool UpdateConfiguration_StaticMesh(UStaticMesh* InStaticMesh);
+	void UpdateConfiguration_Disabled();
+
+	bool IsValid() const
+		{ return bIsEnabled; }
+
+public:
+	void PerformPostProcessFrame_RenderThread(FRHICommandListImmediate& RHICmdList, const TArray<FRHITexture2D*>* InFrameTargets = nullptr, const TArray<FRHITexture2D*>* InAdditionalFrameTargets = nullptr) const;
+
+private:
+	bool bIsEnabled = false;
+	UStaticMesh* StaticMesh = nullptr;
+	FString ExternalFile;
+
+	bool bErrorCantFindFileOnce = false;
+	bool bErrorFailLoadFromFileOnce = false;
+
+private:
+	FDisplayClusterRender_MeshComponent& OutputRemapMesh;
+	class IDisplayClusterShaders& ShadersAPI;
+};
diff --git a/Engine/Plugins/Runtime/nDisplay/Source/DisplayCluster/Public/Blueprints/DisplayClusterBlueprint.h b/Engine/Plugins/Runtime/nDisplay/Source/DisplayCluster/Public/Blueprints/DisplayClusterBlueprint.h
--- a/Engine/Plugins/Runtime/nDisplay/Source/DisplayCluster/Public/Blueprints/DisplayClusterBlueprint.h
+++ b/Engine/Plugins/Runtime/nDisplay/Source/DisplayCluster/Public/Blueprints/DisplayClusterBlueprint.h
@@ -5,6 +5,8 @@
 #include "Engine/Blueprint.h"
 #include "DisplayClusterBlueprint.generated.h"
 
+class USCS_Node;
+class UActorComponent;
 
 UCLASS(BlueprintType, DisplayName = "nDisplay Blueprint")
 class DISPLAYCLUSTER_API UDisplayClusterBlueprint : public UBlueprint
@@ -42,6 +44,9 @@ public:
 
 	const FString& GetConfigPath() const;
 	void SetConfigPath(const FString& InPath);
+
+	/** Configure the config with needed blueprint data for an export. */
+	void PrepareConfigForExport();
 	
 public:
 	// Holds the last saved config export. In the AssetRegistry to allow parsing without loading.
@@ -49,10 +54,18 @@ public:
 	FString ConfigExport;
 
 private:
-
 	//** Updates the ConfigExport property. Called when saving the asset.
 	void UpdateConfigExportProperty();
 
+	/** Returns normal component name. We get component name from a SCS component template which has _GEN_VARIABLE postfix. */
+	FString GetObjectNameFromSCSNode(const UObject* const Object) const;
+
+	/** Auxiliary recursive function that builds child-parentId components map */
+	void GatherParentComponentsInfo(const USCS_Node* const InNode, TMap<UActorComponent*, FString>& OutParentsMap) const;
+
+	/** Auxiliary function to remove invalid pairs from the TMap containers in the config*/
+	void CleanupConfigMaps(UDisplayClusterConfigurationData* Data) const;
+
 protected:
 	UPROPERTY()
 	UDisplayClusterConfigurationData* ConfigData;
diff --git a/Engine/Plugins/Runtime/nDisplay/Source/DisplayCluster/Public/Blueprints/IDisplayClusterBlueprintAPI.h b/Engine/Plugins/Runtime/nDisplay/Source/DisplayCluster/Public/Blueprints/IDisplayClusterBlueprintAPI.h
--- a/Engine/Plugins/Runtime/nDisplay/Source/DisplayCluster/Public/Blueprints/IDisplayClusterBlueprintAPI.h
+++ b/Engine/Plugins/Runtime/nDisplay/Source/DisplayCluster/Public/Blueprints/IDisplayClusterBlueprintAPI.h
@@ -56,14 +56,26 @@ public:
 	//////////////////////////////////////////////////////////////////////////////////////////////
 	// Cluster API
 	//////////////////////////////////////////////////////////////////////////////////////////////
-	/** Returns true if current node is a master computer in a cluster. */
+	/** Returns true if current node is a master node in a cluster. */
 	UFUNCTION(BlueprintCallable, meta = (DisplayName = "Is master node"), Category = "NDisplay|Cluster")
 	virtual bool IsMaster() const = 0;
 	
-	/** Returns true if current node is a slave computer in a cluster. */
+	/** Returns true if current node is a slave node in a cluster. */
 	UFUNCTION(BlueprintCallable, meta = (DisplayName = "Is slave node"), Category = "NDisplay|Cluster")
 	virtual bool IsSlave() const = 0;
 
+	/** Returns true if current node is a backup node in a cluster. */
+	UFUNCTION(BlueprintCallable, meta = (DisplayName = "Is backup node"), Category = "NDisplay|Cluster")
+	virtual bool IsBackup() const = 0;
+
+	/** Returns the role of the current cluster node. */
+	UFUNCTION(BlueprintCallable, meta = (DisplayName = "Get cluster role"), Category = "NDisplay|Cluster")
+	virtual EDisplayClusterNodeRole GetClusterRole() const = 0;
+
+	/** Returns the role of the current cluster node. */
+	UFUNCTION(BlueprintCallable, meta = (DisplayName = "Get cluster node IDs"), Category = "NDisplay|Cluster")
+	virtual void GetNodeIds(TArray<FString>& OutNodeIds) const = 0;
+
 	/** Returns cluster node name of the current application instance. */
 	UFUNCTION(BlueprintCallable, meta = (DisplayName = "Get node ID"), Category = "NDisplay|Cluster")
 	virtual FString GetNodeId() const = 0;
diff --git a/Engine/Plugins/Runtime/nDisplay/Source/DisplayCluster/Public/Cluster/IDisplayClusterClusterManager.h b/Engine/Plugins/Runtime/nDisplay/Source/DisplayCluster/Public/Cluster/IDisplayClusterClusterManager.h
--- a/Engine/Plugins/Runtime/nDisplay/Source/DisplayCluster/Public/Cluster/IDisplayClusterClusterManager.h
+++ b/Engine/Plugins/Runtime/nDisplay/Source/DisplayCluster/Public/Cluster/IDisplayClusterClusterManager.h
@@ -30,15 +30,21 @@ public:
 	virtual ~IDisplayClusterClusterManager() = default;
 
 public:
-	/** Returns true if current node is master. */
+	/** Returns true if current node has master role. */
 	virtual bool IsMaster() const = 0;
-	/** Returns true if current node is slave. */
+	/** Returns true if current node has slave role. */
 	virtual bool IsSlave()  const = 0;
+	/** Returns true if current node has backup role. */
+	virtual bool IsBackup() const = 0;
+	/** Returns cluster node role. */
+	virtual EDisplayClusterNodeRole GetClusterRole() const = 0;
 
 	/** Returns current cluster node ID. */
 	virtual FString GetNodeId() const = 0;
 	/** Returns amount of cluster nodes in the cluster. */
 	virtual uint32 GetNodesAmount() const = 0;
+	/** Returns IDs of available cluster nodes. */
+	virtual void GetNodeIds(TArray<FString>& OutNodeIds) const = 0;
 
 
 	/** Registers object to synchronize. */
diff --git a/Engine/Plugins/Runtime/nDisplay/Source/DisplayCluster/Public/DisplayClusterEnums.h b/Engine/Plugins/Runtime/nDisplay/Source/DisplayCluster/Public/DisplayClusterEnums.h
--- a/Engine/Plugins/Runtime/nDisplay/Source/DisplayCluster/Public/DisplayClusterEnums.h
+++ b/Engine/Plugins/Runtime/nDisplay/Source/DisplayCluster/Public/DisplayClusterEnums.h
@@ -19,6 +19,18 @@ enum class EDisplayClusterOperationMode : uint8
 	Disabled
 };
 
+/**
+ * Display cluster node role
+ */
+UENUM(BlueprintType)
+enum class EDisplayClusterNodeRole : uint8
+{
+	None = 0,
+	Master,
+	Slave,
+	Backup,
+};
+
 
 /**
  * Display cluster synchronization groups
diff --git a/Engine/Plugins/Runtime/nDisplay/Source/DisplayCluster/Public/DisplayClusterRootActor.h b/Engine/Plugins/Runtime/nDisplay/Source/DisplayCluster/Public/DisplayClusterRootActor.h
--- a/Engine/Plugins/Runtime/nDisplay/Source/DisplayCluster/Public/DisplayClusterRootActor.h
+++ b/Engine/Plugins/Runtime/nDisplay/Source/DisplayCluster/Public/DisplayClusterRootActor.h
@@ -51,8 +51,21 @@ public:
 	~ADisplayClusterRootActor();
 
 public:
+
+	/**
+	 * Initialializes the instance with specified config data
+	 *
+	 * @param ConfigData - Configuration data
+	 */
 	void InitializeFromConfig(UDisplayClusterConfigurationData* ConfigData);
 
+	/**
+	 * Cherry picking settings from a specified config data
+	 *
+	 * @param ConfigData - Configuration data
+	 */
+	void OverrideFromConfig(UDisplayClusterConfigurationData* ConfigData);
+
 	/**
 	 * Update or create the config data object. The config sub object is only instantiated once.
 	 * Subsequent calls will only update ConfigDataName unless bForceRecreate is true.
diff --git a/Engine/Plugins/Runtime/nDisplay/Source/DisplayCluster/Public/Misc/DisplayClusterStrings.h b/Engine/Plugins/Runtime/nDisplay/Source/DisplayCluster/Public/Misc/DisplayClusterStrings.h
--- a/Engine/Plugins/Runtime/nDisplay/Source/DisplayCluster/Public/Misc/DisplayClusterStrings.h
+++ b/Engine/Plugins/Runtime/nDisplay/Source/DisplayCluster/Public/Misc/DisplayClusterStrings.h
@@ -21,6 +21,16 @@ namespace DisplayClusterStrings
 		static constexpr auto Cluster    = TEXT("dc_cluster");
 		static constexpr auto Node       = TEXT("dc_node");
 		static constexpr auto Config     = TEXT("dc_cfg");
+		static constexpr auto NodeList   = TEXT("dc_nodes");
+		static constexpr auto NodeRole   = TEXT("dc_role");
+
+		// Cluster node roles
+		namespace role
+		{
+			static constexpr auto Master = TEXT("master");
+			static constexpr auto Slave  = TEXT("slave");
+			static constexpr auto Backup = TEXT("backup");
+		}
 
 		// Stereo device types (command line values)
 		namespace dev
diff --git a/Engine/Plugins/Runtime/nDisplay/Source/DisplayCluster/Public/Render/Containers/DisplayClusterRender_MeshComponent.h b/Engine/Plugins/Runtime/nDisplay/Source/DisplayCluster/Public/Render/Containers/DisplayClusterRender_MeshComponent.h
--- a/Engine/Plugins/Runtime/nDisplay/Source/DisplayCluster/Public/Render/Containers/DisplayClusterRender_MeshComponent.h
+++ b/Engine/Plugins/Runtime/nDisplay/Source/DisplayCluster/Public/Render/Containers/DisplayClusterRender_MeshComponent.h
@@ -20,7 +20,8 @@ public:
 public:
 	void AssignMeshRefs(UStaticMeshComponent* MeshComponent, USceneComponent* OriginComponent = nullptr);
 	void UpdateDefferedRef();
-
+	
+	void UpdateDeffered(const UStaticMesh* InStaticMesh);
 	void UpdateDeffered(const FDisplayClusterRender_MeshGeometry& InMeshGeometry);
 
 	const FStaticMeshLODResources* GetStaticMeshLODResource(int LodIndex = 0) const;
diff --git a/Engine/Plugins/Runtime/nDisplay/Source/DisplayCluster/Public/Render/Containers/DisplayClusterRender_MeshComponentProxyData.h b/Engine/Plugins/Runtime/nDisplay/Source/DisplayCluster/Public/Render/Containers/DisplayClusterRender_MeshComponentProxyData.h
--- a/Engine/Plugins/Runtime/nDisplay/Source/DisplayCluster/Public/Render/Containers/DisplayClusterRender_MeshComponentProxyData.h
+++ b/Engine/Plugins/Runtime/nDisplay/Source/DisplayCluster/Public/Render/Containers/DisplayClusterRender_MeshComponentProxyData.h
@@ -5,6 +5,7 @@
 #include "CoreMinimal.h"
 #include "Render/Containers/DisplayClusterRender_MeshComponentTypes.h"
 
+class UStaticMesh;
 class UStaticMeshComponent;
 class FDisplayClusterRender_MeshGeometry;
 
@@ -13,6 +14,7 @@ class DISPLAYCLUSTER_API FDisplayClusterRender_MeshComponentProxyData
 public:
 	// Use channel 1 as default source for chromakey custom markers UV
 	FDisplayClusterRender_MeshComponentProxyData(const FDisplayClusterRender_MeshComponentProxyDataFunc InDataFunc, const UStaticMeshComponent& InMeshComponent, const int InUVChromakeyIndex = 1);
+	FDisplayClusterRender_MeshComponentProxyData(const FDisplayClusterRender_MeshComponentProxyDataFunc InDataFunc, const UStaticMesh& InStaticMesh, const int InUVChromakeyIndex = 1);
 	FDisplayClusterRender_MeshComponentProxyData(const FDisplayClusterRender_MeshComponentProxyDataFunc InDataFunc, const FDisplayClusterRender_MeshGeometry& InMeshGeometry);
 
 	const TArray<uint32>& GetIndexData() const
@@ -44,6 +46,8 @@ private:
 	bool IsFaceVisible(int32 Face);
 	bool IsUVVisible(int32 UVIndex);
 
+	void ImplInitializeMesh(const FDisplayClusterRender_MeshComponentProxyDataFunc InDataFunc, const UStaticMesh& InStaticMesh, const int InUVChromakeyIndex = 1);
+
 private:
 	TArray<uint32> IndexData;
 	TArray<FDisplayClusterMeshVertex> VertexData;
diff --git a/Engine/Plugins/Runtime/nDisplay/Source/DisplayClusterConfiguration/Private/DisplayClusterConfigurationTypes.cpp b/Engine/Plugins/Runtime/nDisplay/Source/DisplayClusterConfiguration/Private/DisplayClusterConfigurationTypes.cpp
--- a/Engine/Plugins/Runtime/nDisplay/Source/DisplayClusterConfiguration/Private/DisplayClusterConfigurationTypes.cpp
+++ b/Engine/Plugins/Runtime/nDisplay/Source/DisplayClusterConfiguration/Private/DisplayClusterConfigurationTypes.cpp
@@ -176,7 +176,7 @@ FDisplayClusterConfigurationProjection::FDisplayClusterConfigurationProjection()
 
 FDisplayClusterConfigurationPostprocess::FDisplayClusterConfigurationPostprocess()
 {
-	Type = TEXT("OutputRemap");
+	
 }
 
 const float UDisplayClusterConfigurationViewport::ViewportMinimumSize = 1.0f;
diff --git a/Engine/Plugins/Runtime/nDisplay/Source/DisplayClusterConfiguration/Private/Formats/JSON427/DisplayClusterConfigurationJsonParser_427.cpp b/Engine/Plugins/Runtime/nDisplay/Source/DisplayClusterConfiguration/Private/Formats/JSON427/DisplayClusterConfigurationJsonParser_427.cpp
--- a/Engine/Plugins/Runtime/nDisplay/Source/DisplayClusterConfiguration/Private/Formats/JSON427/DisplayClusterConfigurationJsonParser_427.cpp
+++ b/Engine/Plugins/Runtime/nDisplay/Source/DisplayClusterConfiguration/Private/Formats/JSON427/DisplayClusterConfigurationJsonParser_427.cpp
@@ -251,6 +251,16 @@ namespace JSON427
 					Node->Postprocess.Emplace(CfgPostprocess.Key, PostprocessOperation);
 				}
 
+				// Output remap
+				Node->OutputRemap.bEnable = CfgNode.Value.OutputRemap.bEnable;
+				Node->OutputRemap.DataSource = (CfgNode.Value.OutputRemap.DataSource == "file") ? EDisplayClusterConfigurationFramePostProcess_OutputRemapSource::ExternalFile : EDisplayClusterConfigurationFramePostProcess_OutputRemapSource::StaticMesh;
+				Node->OutputRemap.ExternalFile = CfgNode.Value.OutputRemap.ExternalFile;
+
+				if (!CfgNode.Value.OutputRemap.StaticMeshAsset.IsEmpty())
+				{
+					Node->OutputRemap.StaticMesh = Cast<UStaticMesh>(StaticLoadObject(UStaticMesh::StaticClass(), nullptr, *CfgNode.Value.OutputRemap.StaticMeshAsset));
+				}
+
 				// Store new cluster node
 				Config->Cluster->Nodes.Emplace(CfgNode.Key, Node);
 			}
@@ -397,9 +407,7 @@ namespace JSON427
 					Viewport.Camera = CfgViewport.Value->Camera;
 					Viewport.Region = FDisplayClusterConfigurationJsonRectangle_427(CfgViewport.Value->Region.X, CfgViewport.Value->Region.Y, CfgViewport.Value->Region.W, CfgViewport.Value->Region.H);
 					Viewport.GPUIndex = CfgViewport.Value->GPUIndex;
-					//Viewport.IsShared = CfgViewport.Value->bIsShared;
 					Viewport.BufferRatio = CfgViewport.Value->RenderSettings.BufferRatio;
-					//Viewport.AllowCrossGPUTransfer = CfgViewport.Value->bAllowCrossGPUTransfer;
 
 					// Projection policy
 					Viewport.ProjectionPolicy.Type = CfgViewport.Value->ProjectionPolicy.Type;
@@ -420,6 +428,18 @@ namespace JSON427
 					Node.Postprocess.Emplace(CfgPostprocess.Key, PostprocessOperation);
 				}
 
+				// OutputRemap
+				{
+					Node.OutputRemap.bEnable = CfgNode.Value->OutputRemap.bEnable;
+					Node.OutputRemap.DataSource = (CfgNode.Value->OutputRemap.DataSource == EDisplayClusterConfigurationFramePostProcess_OutputRemapSource::ExternalFile) ? "file" : "mesh";
+					Node.OutputRemap.ExternalFile = CfgNode.Value->OutputRemap.ExternalFile;
+
+					if (CfgNode.Value->OutputRemap.StaticMesh)
+					{
+						Node.OutputRemap.StaticMeshAsset = CfgNode.Value->OutputRemap.StaticMesh->GetPathName();
+					}
+				}
+
 				// Store new cluster node
 				Json.Cluster.Nodes.Emplace(CfgNode.Key, Node);
 			}
diff --git a/Engine/Plugins/Runtime/nDisplay/Source/DisplayClusterConfiguration/Private/Formats/JSON427/DisplayClusterConfigurationJsonTypes_427.h b/Engine/Plugins/Runtime/nDisplay/Source/DisplayClusterConfiguration/Private/Formats/JSON427/DisplayClusterConfigurationJsonTypes_427.h
--- a/Engine/Plugins/Runtime/nDisplay/Source/DisplayClusterConfiguration/Private/Formats/JSON427/DisplayClusterConfigurationJsonTypes_427.h
+++ b/Engine/Plugins/Runtime/nDisplay/Source/DisplayClusterConfiguration/Private/Formats/JSON427/DisplayClusterConfigurationJsonTypes_427.h
@@ -404,6 +404,25 @@ public:
 	FDisplayClusterConfigurationJsonProjectionPolicy_427 ProjectionPolicy;
 };
 
+USTRUCT()
+struct FDisplayClusterConfigurationFramePostProcess_OutputRemap_427
+{
+	GENERATED_BODY()
+
+public:
+	UPROPERTY()
+	bool bEnable = false;
+
+	UPROPERTY()
+	FString DataSource;
+
+	UPROPERTY()
+	FString StaticMeshAsset;
+
+	UPROPERTY()
+	FString ExternalFile;
+};
+
 USTRUCT()
 struct FDisplayClusterConfigurationJsonClusterNode_427
 {
@@ -427,6 +446,9 @@ public:
 
 	UPROPERTY()
 	TMap<FString, FDisplayClusterConfigurationJsonViewport_427> Viewports;
+
+	UPROPERTY()
+	FDisplayClusterConfigurationFramePostProcess_OutputRemap_427 OutputRemap;
 };
 
 USTRUCT()
diff --git a/Engine/Plugins/Runtime/nDisplay/Source/DisplayClusterConfiguration/Public/DisplayClusterConfigurationTypes.h b/Engine/Plugins/Runtime/nDisplay/Source/DisplayClusterConfiguration/Public/DisplayClusterConfigurationTypes.h
--- a/Engine/Plugins/Runtime/nDisplay/Source/DisplayClusterConfiguration/Public/DisplayClusterConfigurationTypes.h
+++ b/Engine/Plugins/Runtime/nDisplay/Source/DisplayClusterConfiguration/Public/DisplayClusterConfigurationTypes.h
@@ -10,6 +10,7 @@
 #include "DisplayClusterConfigurationTypes_Base.h"
 #include "DisplayClusterConfigurationTypes_ICVFX.h"
 #include "DisplayClusterConfigurationTypes_Viewport.h"
+#include "DisplayClusterConfigurationTypes_OutputRemap.h"
 
 #include "DisplayClusterConfigurationVersion.h"
 
@@ -314,6 +315,10 @@ public:
 	UPROPERTY(EditAnywhere, BlueprintReadWrite, Category = "Configuration", meta = (DisplayName = "Window"))
 	FDisplayClusterConfigurationRectangle WindowRect;
 
+	/** Output remapping settings for the selected cluster node */
+	UPROPERTY(EditAnywhere, BlueprintReadWrite, Category = "Configuration", meta = (DisplayName = "Output Remapping"))
+	FDisplayClusterConfigurationFramePostProcess_OutputRemap OutputRemap;
+
 #if WITH_EDITORONLY_DATA
 	/** Locks the application window aspect ratio for easier resizing */
 	UPROPERTY(EditAnywhere, Category = "Configuration")
@@ -509,7 +514,7 @@ public:
 	UPROPERTY(EditAnywhere, BlueprintReadWrite, Category = Advanced)
 	FDisplayClusterConfigurationInfo Info;
 
-	UPROPERTY(VisibleInstanceOnly, BlueprintReadOnly, Instanced, Category = Advanced)
+	UPROPERTY(VisibleInstanceOnly, BlueprintReadOnly, Transient, Category = Advanced)
 	UDisplayClusterConfigurationScene* Scene;
 
 	UPROPERTY(VisibleInstanceOnly, BlueprintReadOnly, Instanced, Category = Advanced, meta = (DisplayThumbnail = false, ShowInnerProperties))
diff --git a/Engine/Plugins/Runtime/nDisplay/Source/DisplayClusterConfiguration/Public/DisplayClusterConfigurationTypes_OutputRemap.h b/Engine/Plugins/Runtime/nDisplay/Source/DisplayClusterConfiguration/Public/DisplayClusterConfigurationTypes_OutputRemap.h
new file mode 100644
--- /dev/null
+++ b/Engine/Plugins/Runtime/nDisplay/Source/DisplayClusterConfiguration/Public/DisplayClusterConfigurationTypes_OutputRemap.h
@@ -0,0 +1,47 @@
+// Copyright Epic Games, Inc. All Rights Reserved.
+
+#pragma once
+
+#include "CoreMinimal.h"
+
+#include "UObject/ObjectMacros.h"
+#include "UObject/Object.h"
+
+#include "Engine/StaticMesh.h"
+
+#include "DisplayClusterConfigurationTypes_OutputRemap.generated.h"
+
+/* Source types of the output remapping */
+UENUM()
+enum class EDisplayClusterConfigurationFramePostProcess_OutputRemapSource : uint8
+{
+	/** Use a Static Mesh reference for output remapping when the Data Source is set to Mesh */
+	StaticMesh     UMETA(DisplayName = "Static Mesh"),
+
+	/** Use an external .obj file for output remapping when the Data Source is set to File */
+	ExternalFile   UMETA(DisplayName = "External File"),
+};
+
+/* Screen space remapping of the final backbuffer output. Applied at the whole window */
+USTRUCT(Blueprintable)
+struct DISPLAYCLUSTERCONFIGURATION_API FDisplayClusterConfigurationFramePostProcess_OutputRemap
+{
+	GENERATED_BODY()
+
+public:
+	/** Enables or disables output remapping */
+	UPROPERTY(EditAnywhere, BlueprintReadWrite, Category = "NDisplay OutputRemap", meta = (DisplayName = "Enable Output Remapping"))
+	bool bEnable = false;
+
+	/** Selects either the Static Mesh or External File setting as the source for output remapping */
+	UPROPERTY(EditAnywhere, BlueprintReadWrite, Category = "NDisplay OutputRemap", meta = (EditCondition = "bEnable"))
+	EDisplayClusterConfigurationFramePostProcess_OutputRemapSource DataSource = EDisplayClusterConfigurationFramePostProcess_OutputRemapSource::StaticMesh;
+
+	/** The Static Mesh reference to use for output remapping when the Data Source is set to Static Mesh */
+	UPROPERTY(EditAnywhere, BlueprintReadOnly, Category = "NDisplay OutputRemap", meta = (EditCondition = "DataSource == EDisplayClusterConfigurationFramePostProcess_OutputRemapSource::StaticMesh && bEnable"))
+	class UStaticMesh* StaticMesh = nullptr;
+
+	/** The external .obj file to use for output remapping when the Data Source is set to File */
+	UPROPERTY(EditAnywhere, BlueprintReadWrite, Category = "NDisplay OutputRemap", meta = (EditCondition = "DataSource == EDisplayClusterConfigurationFramePostProcess_OutputRemapSource::ExternalFile && bEnable"))
+	FString ExternalFile;
+};
diff --git a/Engine/Plugins/Runtime/nDisplay/Source/DisplayClusterConfiguration/Public/DisplayClusterConfigurationTypes_Postprocess.h b/Engine/Plugins/Runtime/nDisplay/Source/DisplayClusterConfiguration/Public/DisplayClusterConfigurationTypes_Postprocess.h
--- a/Engine/Plugins/Runtime/nDisplay/Source/DisplayClusterConfiguration/Public/DisplayClusterConfigurationTypes_Postprocess.h
+++ b/Engine/Plugins/Runtime/nDisplay/Source/DisplayClusterConfiguration/Public/DisplayClusterConfigurationTypes_Postprocess.h
@@ -274,7 +274,7 @@ struct FDisplayClusterConfigurationViewport_AllNodesColorGrading
 	bool bEnableEntireClusterColorGrading = true;
 
 	/** Color Grading */
-	UPROPERTY(EditAnywhere, BlueprintReadWrite, Category = "Viewport Settings", meta = (EditCondition = "bIsEnabled", DisplayName = "Color Grading"))
+	UPROPERTY(EditAnywhere, BlueprintReadWrite, Category = "Viewport Settings", meta = (DisplayName = "Color Grading"))
 	FDisplayClusterConfigurationViewport_ColorGradingRenderingSettings ColorGradingSettings;
 };
 
diff --git a/Engine/Plugins/Runtime/nDisplay/Source/DisplayClusterConfiguration/Public/DisplayClusterConfigurationTypes_Viewport.h b/Engine/Plugins/Runtime/nDisplay/Source/DisplayClusterConfiguration/Public/DisplayClusterConfigurationTypes_Viewport.h
--- a/Engine/Plugins/Runtime/nDisplay/Source/DisplayClusterConfiguration/Public/DisplayClusterConfigurationTypes_Viewport.h
+++ b/Engine/Plugins/Runtime/nDisplay/Source/DisplayClusterConfiguration/Public/DisplayClusterConfigurationTypes_Viewport.h
@@ -152,7 +152,7 @@ public:
 	UPROPERTY(EditAnywhere, BlueprintReadWrite, Category = "Configuration", meta = (DisplayName = "View Origin"))
 	FString Camera;
 
-	/** Selection of a Projection Policy */
+	/** Specify your Projection Policy Settings */
 	UPROPERTY(EditDefaultsOnly, Category = "Configuration")
 	FDisplayClusterConfigurationProjection ProjectionPolicy;
 
@@ -170,7 +170,7 @@ public:
 	UPROPERTY(EditAnywhere, BlueprintReadWrite, Category = "Configuration")
 	FDisplayClusterConfigurationRectangle Region;
 
-	/** Viewports can overlap, allow and sets Viewport overlaping order priority */
+	/** Allows Viewports to overlap and sets Viewport overlapping order priority */
 	UPROPERTY(EditAnywhere, BlueprintReadWrite, Category = "Configuration")
 	int OverlapOrder = 0;
 
diff --git a/Engine/Plugins/Runtime/nDisplay/Source/DisplayClusterConfigurator/Private/DisplayClusterConfiguratorBlueprintEditor.cpp b/Engine/Plugins/Runtime/nDisplay/Source/DisplayClusterConfigurator/Private/DisplayClusterConfiguratorBlueprintEditor.cpp
--- a/Engine/Plugins/Runtime/nDisplay/Source/DisplayClusterConfigurator/Private/DisplayClusterConfiguratorBlueprintEditor.cpp
+++ b/Engine/Plugins/Runtime/nDisplay/Source/DisplayClusterConfigurator/Private/DisplayClusterConfiguratorBlueprintEditor.cpp
@@ -226,6 +226,26 @@ void FDisplayClusterConfiguratorBlueprintEditor::PostRedo(bool bSuccess)
 	}
 }
 
+void FDisplayClusterConfiguratorBlueprintEditor::RefreshEditors(ERefreshBlueprintEditorReason::Type Reason)
+{
+	FBlueprintEditor::RefreshEditors(Reason);
+	
+	if (ViewportTabContent.IsValid())
+	{
+		TFunction<void(FName, TSharedPtr<IEditorViewportLayoutEntity>)> RefreshFunc =
+			[this](FName Name, TSharedPtr<IEditorViewportLayoutEntity> Entity)
+			{
+				if (Entity.IsValid())
+				{
+					TSharedRef<SDisplayClusterConfiguratorSCSEditorViewport> Viewport = StaticCastSharedRef<SDisplayClusterConfiguratorSCSEditorViewport>(Entity->AsWidget());
+					Viewport->RequestRefresh();
+				}
+			};
+
+		ViewportTabContent->PerformActionOnViewports(RefreshFunc);
+	}
+}
+
 UDisplayClusterConfigurationData* FDisplayClusterConfiguratorBlueprintEditor::GetEditorData() const
 {
 	if (!LoadedBlueprint.IsValid() || !LoadedBlueprint->GeneratedClass)
@@ -628,178 +648,24 @@ bool FDisplayClusterConfiguratorBlueprintEditor::CanExportConfig() const
 bool FDisplayClusterConfiguratorBlueprintEditor::SaveToFile(const FString& InFilePath)
 {
 	UDisplayClusterConfiguratorEditorSubsystem* EditorSubsystem = GEditor->GetEditorSubsystem<UDisplayClusterConfiguratorEditorSubsystem>();
-	UDisplayClusterConfigurationData* Data = GetEditorData();
-
-	if (EditorSubsystem && Data)
+	if (EditorSubsystem && LoadedBlueprint.IsValid())
 	{
-		// Components to export
-		TArray<UDisplayClusterCameraComponent*> CameraComponents;
-		TArray<UDisplayClusterScreenComponent*> ScreenComponents;
-		TArray<USceneComponent*>  XformComponents;
-		// Auxiliary map for building hierarchy
-		TMap<UActorComponent*, FString> ParentComponentsMap;
-
-		// Extract CDO cameras (no screens in the CDO)
-		if (ADisplayClusterRootActor* CDO = GetDefaultRootActor())
-		{
-			// Make sure the 'Scene' object is there. Otherwise instantiate it.
-			const EObjectFlags CommonFlags = RF_Public | RF_Transactional;
-			if (Data->Scene == nullptr)
-			{
-				Data->Scene = NewObject<UDisplayClusterConfigurationScene>(
-					Data,
-					UDisplayClusterConfigurationScene::StaticClass(),
-					NAME_None,
-					Data->IsTemplate() ? RF_ArchetypeObject | CommonFlags : CommonFlags);
-			}
-
-			// Get list of cameras
-			CDO->GetComponents<UDisplayClusterCameraComponent>(CameraComponents);
-		}
+		LoadedBlueprint->PrepareConfigForExport();
 
-		// Extract BP components
-		if (LoadedBlueprint.IsValid() && LoadedBlueprint->GeneratedClass.Get() != nullptr)
+		if (UDisplayClusterConfigurationData* Data = GetEditorData())
 		{
-			const TArray<USCS_Node*>& Nodes = LoadedBlueprint->SimpleConstructionScript->GetAllNodes();
-			for (const USCS_Node* const Node : Nodes)
+			// Finally, store data to the file
+			if (EditorSubsystem->SaveConfig(Data, InFilePath))
 			{
-				// Fill ID info for all descendants
-				GatherParentComponentsInfo(Node, ParentComponentsMap);
-
-				// Cameras
-				if (Node->ComponentClass->IsChildOf(UDisplayClusterCameraComponent::StaticClass()))
-				{
-					UDisplayClusterCameraComponent* ComponentTemplate = CastChecked<UDisplayClusterCameraComponent>(Node->GetActualComponentTemplate(CastChecked<UBlueprintGeneratedClass>(LoadedBlueprint->GeneratedClass)));
-					CameraComponents.Add(ComponentTemplate);
-				}
-				// Screens
-				else if (Node->ComponentClass->IsChildOf(UDisplayClusterScreenComponent::StaticClass()))
-				{
-					UDisplayClusterScreenComponent* ComponentTemplate = CastChecked<UDisplayClusterScreenComponent>(Node->GetActualComponentTemplate(CastChecked<UBlueprintGeneratedClass>(LoadedBlueprint->GeneratedClass)));
-					ScreenComponents.Add(ComponentTemplate);
-				}
-				// All other components will be exported as Xforms
-				else if (Node->ComponentClass->IsChildOf(USceneComponent::StaticClass()))
-				{
-					USceneComponent* ComponentTemplate = CastChecked<USceneComponent>(Node->GetActualComponentTemplate(CastChecked<UBlueprintGeneratedClass>(LoadedBlueprint->GeneratedClass)));
-					XformComponents.Add(ComponentTemplate);
-				}
+				LoadedBlueprint->SetConfigPath(Data->PathToConfig);
+				return true;
 			}
 		}
-
-		// Save asset path
-		Data->Info.AssetPath = LoadedBlueprint->GetPathName();
-
-		// Prepare the target containers
-		Data->Scene->Cameras.Empty(CameraComponents.Num());
-		Data->Scene->Screens.Empty(ScreenComponents.Num());
-		Data->Scene->Xforms.Empty(XformComponents.Num());
-
-		// Export cameras
-		for (const UDisplayClusterCameraComponent* const CfgComp : CameraComponents)
-		{
-			UDisplayClusterConfigurationSceneComponentCamera* SceneComp = NewObject<UDisplayClusterConfigurationSceneComponentCamera>(Data->Scene, CfgComp->GetFName());
-
-			// Save the properties
-			SceneComp->bSwapEyes = CfgComp->GetSwapEyes();
-			SceneComp->InterpupillaryDistance = CfgComp->GetInterpupillaryDistance();
-			// Safe to cast -- values match.
-			SceneComp->StereoOffset = (EDisplayClusterConfigurationEyeStereoOffset)CfgComp->GetStereoOffset();
-
-			FString* ParentId = ParentComponentsMap.Find(CfgComp);
-			SceneComp->ParentId = (ParentId ? *ParentId : FString());
-			SceneComp->Location = CfgComp->GetRelativeLocation();
-			SceneComp->Rotation = CfgComp->GetRelativeRotation();
-
-			// Store the object
-			Data->Scene->Cameras.Emplace(GetObjectNameFromSCSNode(SceneComp), SceneComp);
-		}
-
-		// Export screens
-		for (const UDisplayClusterScreenComponent* const CfgComp : ScreenComponents)
-		{
-			UDisplayClusterConfigurationSceneComponentScreen* SceneComp = NewObject<UDisplayClusterConfigurationSceneComponentScreen>(Data->Scene, CfgComp->GetFName());
-
-			// Save the properties
-			FString* ParentId = ParentComponentsMap.Find(CfgComp);
-			SceneComp->ParentId = (ParentId ? *ParentId : FString());
-			SceneComp->Location = CfgComp->GetRelativeLocation();
-			SceneComp->Rotation = CfgComp->GetRelativeRotation();
-
-			const FVector RelativeCompScale = CfgComp->GetRelativeScale3D();
-			SceneComp->Size = FVector2D(RelativeCompScale.Y, RelativeCompScale.Z);
-
-			// Store the object
-			Data->Scene->Screens.Emplace(GetObjectNameFromSCSNode(SceneComp), SceneComp);
-		}
-
-		// Export xforms
-		for (const USceneComponent* const CfgComp : XformComponents)
-		{
-			UDisplayClusterConfigurationSceneComponentXform* SceneComp = NewObject<UDisplayClusterConfigurationSceneComponentXform>(Data->Scene, CfgComp->GetFName());
-
-			// Save the properties
-			FString* ParentId = ParentComponentsMap.Find(CfgComp);
-			SceneComp->ParentId = (ParentId ? *ParentId : FString());
-			SceneComp->Location = CfgComp->GetRelativeLocation();
-			SceneComp->Rotation = CfgComp->GetRelativeRotation();
-
-			// Store the object
-			Data->Scene->Xforms.Emplace(GetObjectNameFromSCSNode(SceneComp), SceneComp);
-		}
-
-		// Finally, store data to the file
-		if (EditorSubsystem->SaveConfig(Data, InFilePath))
-		{
-			LoadedBlueprint->SetConfigPath(Data->PathToConfig);
-			return true;
-		}
 	}
 
 	return false;
 }
 
-FString FDisplayClusterConfiguratorBlueprintEditor::GetObjectNameFromSCSNode(const UObject* const Object) const
-{
-	FString OutCompName;
-
-	if (Object)
-	{
-		OutCompName = Object->GetName();
-		OutCompName.RemoveFromEnd(TEXT("_GEN_VARIABLE"));
-	}
-
-	return OutCompName;
-}
-
-void FDisplayClusterConfiguratorBlueprintEditor::GatherParentComponentsInfo(const USCS_Node* const InNode, TMap<UActorComponent*, FString>& OutParentsMap) const
-{
-	if (LoadedBlueprint.IsValid() && LoadedBlueprint->GeneratedClass.Get() != nullptr)
-	{
-		if (InNode && InNode->ComponentClass->IsChildOf(UActorComponent::StaticClass()))
-		{
-			// Save current node to the map
-			UActorComponent* ParentNode = InNode->GetActualComponentTemplate(CastChecked<UBlueprintGeneratedClass>(LoadedBlueprint->GeneratedClass));
-			if (!OutParentsMap.Contains(ParentNode))
-			{
-				OutParentsMap.Emplace(ParentNode);
-			}
-
-			// Now iterate through the children nodes
-			for (USCS_Node* ChildNode : InNode->ChildNodes)
-			{
-				UActorComponent* ChildComponentTemplate = CastChecked<UActorComponent>(ChildNode->GetActualComponentTemplate(CastChecked<UBlueprintGeneratedClass>(LoadedBlueprint->GeneratedClass)));
-				if (ChildComponentTemplate)
-				{
-					OutParentsMap.Emplace(ChildComponentTemplate, GetObjectNameFromSCSNode(InNode->ComponentTemplate));
-				}
-
-				GatherParentComponentsInfo(ChildNode, OutParentsMap);
-			}
-		}
-	}
-}
-
 bool FDisplayClusterConfiguratorBlueprintEditor::SaveWithOpenFileDialog()
 {
 	const FString NDisplayFileDescription = LOCTEXT("NDisplayFileDescription", "nDisplay Config").ToString();
diff --git a/Engine/Plugins/Runtime/nDisplay/Source/DisplayClusterConfigurator/Private/DisplayClusterConfiguratorBlueprintEditor.h b/Engine/Plugins/Runtime/nDisplay/Source/DisplayClusterConfigurator/Private/DisplayClusterConfiguratorBlueprintEditor.h
--- a/Engine/Plugins/Runtime/nDisplay/Source/DisplayClusterConfigurator/Private/DisplayClusterConfiguratorBlueprintEditor.h
+++ b/Engine/Plugins/Runtime/nDisplay/Source/DisplayClusterConfigurator/Private/DisplayClusterConfiguratorBlueprintEditor.h
@@ -44,6 +44,7 @@ public:
 	//~ Begin FEditorUndoClient interface
 	virtual void PostUndo(bool bSuccess) override;
 	virtual void PostRedo(bool bSuccess) override;
+	virtual void RefreshEditors(ERefreshBlueprintEditorReason::Type Reason) override;
 	//~ End FEditorUndoClient interface
 
 	virtual void SelectObjects(TArray<UObject*>& InSelectedObjects, bool bFullRefresh = false);
@@ -208,12 +209,6 @@ private:
 
 	void BindCommands();
 
-	/** Returns normal component name. We get component name from a SCS component template which has _GEN_VARIABLE postfix. */
-	FString GetObjectNameFromSCSNode(const UObject* const Object) const;
-
-	/** Auxiliary recursive function that builds child-parentId components map */
-	void GatherParentComponentsInfo(const USCS_Node* const InNode, TMap<UActorComponent*, FString>& OutParentsMap) const;
-
 private:
 	TSharedPtr<FDisplayClusterConfiguratorViewOutputMapping> ViewOutputMapping;
 	TSharedPtr<FDisplayClusterConfiguratorViewCluster> ViewCluster;
diff --git a/Engine/Plugins/Runtime/nDisplay/Source/DisplayClusterConfigurator/Private/Views/Details/DisplayClusterConfiguratorDetailCustomization.cpp b/Engine/Plugins/Runtime/nDisplay/Source/DisplayClusterConfigurator/Private/Views/Details/DisplayClusterConfiguratorDetailCustomization.cpp
--- a/Engine/Plugins/Runtime/nDisplay/Source/DisplayClusterConfigurator/Private/Views/Details/DisplayClusterConfiguratorDetailCustomization.cpp
+++ b/Engine/Plugins/Runtime/nDisplay/Source/DisplayClusterConfigurator/Private/Views/Details/DisplayClusterConfiguratorDetailCustomization.cpp
@@ -566,6 +566,7 @@ void FDisplayClusterConfiguratorTypeCustomization::CustomizeHeader(TSharedRef<IP
 	if (OuterObjects.Num())
 	{
 		EditingObject = OuterObjects[0];
+		bMultipleObjectsSelected = OuterObjects.Num() > 1;
 	}
 }
 
@@ -779,7 +780,7 @@ void FDisplayClusterConfiguratorRenderSyncPolicyCustomization::AddRenderSyncPoli
 	ChildBuilder->AddCustomRow(TypeHandle->GetPropertyDisplayName())
 	.NameContent()
 	[
-		TypeHandle->CreatePropertyNameWidget()
+		TypeHandle->CreatePropertyNameWidget(FText::GetEmpty(), LOCTEXT("RenderSyncPolicyToolTip", "Specify your nDisplay Render Sync Policy"))
 	]
 	.ValueContent()
 	[
@@ -789,7 +790,6 @@ void FDisplayClusterConfiguratorRenderSyncPolicyCustomization::AddRenderSyncPoli
 		.OnSelectionChanged(this, &FDisplayClusterConfiguratorRenderSyncPolicyCustomization::OnRenderSyncPolicySelected)
 		.ContentPadding(2)
 		.MaxListHeight(200.0f)
-		.ToolTipText(LOCTEXT("RenderSyncPolicyToolTip", "Specify your nDisplay Render Sync Policy"))
 		.Content()
 		[
 			SNew(STextBlock)
@@ -1118,7 +1118,7 @@ void FDisplayClusterConfiguratorInputSyncPolicyCustomization::AddInputSyncPolicy
 	ChildBuilder->AddCustomRow(TypeHandle->GetPropertyDisplayName())
 	.NameContent()
 	[
-		TypeHandle->CreatePropertyNameWidget()
+		TypeHandle->CreatePropertyNameWidget(FText::GetEmpty(), LOCTEXT("InputSyncPolicyTooltip", "Specify your nDisplay Input Sync Policy"))
 	]
 	.ValueContent()
 	[
@@ -1128,7 +1128,6 @@ void FDisplayClusterConfiguratorInputSyncPolicyCustomization::AddInputSyncPolicy
 		.OnSelectionChanged(this, &FDisplayClusterConfiguratorInputSyncPolicyCustomization::OnInputSyncPolicySelected)
 		.ContentPadding(2)
 		.MaxListHeight(200.0f)
-		.ToolTipText(LOCTEXT("InputSyncPolicyTooltip", "Specify your nDisplay Input Sync Policy"))
 		.Content()
 		[
 			SNew(STextBlock)
diff --git a/Engine/Plugins/Runtime/nDisplay/Source/DisplayClusterConfigurator/Private/Views/Details/DisplayClusterConfiguratorDetailCustomization.h b/Engine/Plugins/Runtime/nDisplay/Source/DisplayClusterConfigurator/Private/Views/Details/DisplayClusterConfiguratorDetailCustomization.h
--- a/Engine/Plugins/Runtime/nDisplay/Source/DisplayClusterConfigurator/Private/Views/Details/DisplayClusterConfiguratorDetailCustomization.h
+++ b/Engine/Plugins/Runtime/nDisplay/Source/DisplayClusterConfigurator/Private/Views/Details/DisplayClusterConfiguratorDetailCustomization.h
@@ -212,6 +212,7 @@ protected:
 protected:
 	UObject* EditingObject;
 	TWeakPtr<IPropertyUtilities> PropertyUtilities;
+	bool bMultipleObjectsSelected = false;
 };
 
 /**
diff --git a/Engine/Plugins/Runtime/nDisplay/Source/DisplayClusterConfigurator/Private/Views/Details/DisplayClusterRootActorDetailsCustomization.cpp b/Engine/Plugins/Runtime/nDisplay/Source/DisplayClusterConfigurator/Private/Views/Details/DisplayClusterRootActorDetailsCustomization.cpp
--- a/Engine/Plugins/Runtime/nDisplay/Source/DisplayClusterConfigurator/Private/Views/Details/DisplayClusterRootActorDetailsCustomization.cpp
+++ b/Engine/Plugins/Runtime/nDisplay/Source/DisplayClusterConfigurator/Private/Views/Details/DisplayClusterRootActorDetailsCustomization.cpp
@@ -276,8 +276,6 @@ void FDisplayClusterRootActorDetailsCustomization::BuildLayout(IDetailLayoutBuil
 					ADD_GROUP_NESTED_PROPERTY_EDIT_CONDITION(NestedPropertyHelper, ADisplayClusterRootActor, CurrentConfigData->StageSettings.Lightcard.ShowOnlyList.RootActorComponentNames, ICVFXLightCardEnabledEditCondition)
 				}
 			END_GROUP();
-
-			ADD_ADVANCED_NESTED_PROPERTY_EDIT_CONDITION(NestedPropertyHelper, ADisplayClusterRootActor, CurrentConfigData->StageSettings.Lightcard.RenderSettings, ICVFXLightCardEnabledEditCondition)
 		END_CATEGORY();
 
 		// Add custom properties and lay out/order properties into their correct categories.
diff --git a/Engine/Plugins/Runtime/nDisplay/Source/DisplayClusterConfigurator/Private/Views/Details/Policies/DisplayClusterConfiguratorPolicyDetailCustomization.cpp b/Engine/Plugins/Runtime/nDisplay/Source/DisplayClusterConfigurator/Private/Views/Details/Policies/DisplayClusterConfiguratorPolicyDetailCustomization.cpp
--- a/Engine/Plugins/Runtime/nDisplay/Source/DisplayClusterConfigurator/Private/Views/Details/Policies/DisplayClusterConfiguratorPolicyDetailCustomization.cpp
+++ b/Engine/Plugins/Runtime/nDisplay/Source/DisplayClusterConfigurator/Private/Views/Details/Policies/DisplayClusterConfiguratorPolicyDetailCustomization.cpp
@@ -41,6 +41,12 @@
 
 void FDisplayClusterConfiguratorProjectionCustomization::CustomizeHeader(TSharedRef<IPropertyHandle> InPropertyHandle, FDetailWidgetRow& InHeaderRow, IPropertyTypeCustomizationUtils& CustomizationUtils)
 {
+	if (bMultipleObjectsSelected)
+	{
+		// Do not currently support multiple objects being edited.
+		return;
+	}
+	
 	FDisplayClusterConfiguratorPolymorphicEntityCustomization::CustomizeHeader(InPropertyHandle, InHeaderRow, CustomizationUtils);
 	CustomOption = MakeShared<FString>("Custom");
 
@@ -60,6 +66,12 @@ void FDisplayClusterConfiguratorProjectionCustomization::CustomizeHeader(TShared
 
 void FDisplayClusterConfiguratorProjectionCustomization::CustomizeChildren(TSharedRef<IPropertyHandle> InPropertyHandle, IDetailChildrenBuilder& InChildBuilder, IPropertyTypeCustomizationUtils& CustomizationUtils)
 {
+	if (bMultipleObjectsSelected)
+	{
+		// Do not currently support multiple objects being edited.
+		return;
+	}
+	
 	FDisplayClusterConfiguratorPolymorphicEntityCustomization::CustomizeChildren(InPropertyHandle, InChildBuilder, CustomizationUtils);
 
 	// Hide properties 
@@ -126,7 +138,7 @@ void FDisplayClusterConfiguratorProjectionCustomization::AddProjectionPolicyRow(
 	ChildBuilder->AddCustomRow(TypeHandle->GetPropertyDisplayName())
 	.NameContent()
 	[
-		TypeHandle->CreatePropertyNameWidget()
+		TypeHandle->CreatePropertyNameWidget(FText::GetEmpty(), LOCTEXT("ProjectionPolicyTypeTooltip", "Type of Projection Policy"))
 	]
 	.ValueContent()
 	[
@@ -358,13 +370,17 @@ void FDisplayClusterConfiguratorProjectionCustomization::BuildParametersForPolic
 
 void FDisplayClusterConfiguratorProjectionCustomization::CreateSimplePolicy(UDisplayClusterBlueprint* Blueprint)
 {
-	CustomPolicyParameters.Add(MakeShared<FPolicyParameterInfoComponentCombo>(
+	TSharedPtr<FPolicyParameterInfoComponentCombo> ScreenCombo = MakeShared<FPolicyParameterInfoComponentCombo>(
 		"Screen",
 		DisplayClusterProjectionStrings::cfg::simple::Screen,
 		Blueprint,
 		ConfigurationViewportPtr.Get(),
 		ParametersHandle,
-		TArray<TSubclassOf<UActorComponent>>{ UDisplayClusterScreenComponent::StaticClass() }));
+		TArray<TSubclassOf<UActorComponent>>{ UDisplayClusterScreenComponent::StaticClass() });
+
+	ScreenCombo->SetParameterTooltip(LOCTEXT("SimplePolicyScreenTooltip", "Target Screen or Display"));
+
+	CustomPolicyParameters.Add(ScreenCombo);
 }
 
 void FDisplayClusterConfiguratorProjectionCustomization::CreateCameraPolicy(UDisplayClusterBlueprint* Blueprint)
diff --git a/Engine/Plugins/Runtime/nDisplay/Source/DisplayClusterConfigurator/Private/Views/Details/Policies/DisplayClusterConfiguratorPolicyParameterCustomization.cpp b/Engine/Plugins/Runtime/nDisplay/Source/DisplayClusterConfigurator/Private/Views/Details/Policies/DisplayClusterConfiguratorPolicyParameterCustomization.cpp
--- a/Engine/Plugins/Runtime/nDisplay/Source/DisplayClusterConfigurator/Private/Views/Details/Policies/DisplayClusterConfiguratorPolicyParameterCustomization.cpp
+++ b/Engine/Plugins/Runtime/nDisplay/Source/DisplayClusterConfigurator/Private/Views/Details/Policies/DisplayClusterConfiguratorPolicyParameterCustomization.cpp
@@ -47,6 +47,7 @@ FPolicyParameterInfo::FPolicyParameterInfo(
 	const FString* InInitialValue)
 {
 	ParamDisplayName = InDisplayName;
+	ParamTooltip = FText::GetEmpty();
 	ParamKey = InKey;
 	BlueprintOwnerPtr = InBlueprint;
 	ConfigurationViewportPtr = InConfigurationViewport;
@@ -143,6 +144,7 @@ void FPolicyParameterInfoCombo::CreateCustomRowWidget(IDetailChildrenBuilder& In
 		SNew(STextBlock)
 		.Font(IDetailLayoutBuilder::GetDetailFont())
 		.Text(this, &FPolicyParameterInfo::GetParameterDisplayName)
+		.ToolTipText(this, &FPolicyParameterInfo::GetParameterTooltip)
 	]
 	.ValueContent()
 	[
diff --git a/Engine/Plugins/Runtime/nDisplay/Source/DisplayClusterConfigurator/Private/Views/Details/Policies/DisplayClusterConfiguratorPolicyParameterCustomization.h b/Engine/Plugins/Runtime/nDisplay/Source/DisplayClusterConfigurator/Private/Views/Details/Policies/DisplayClusterConfiguratorPolicyParameterCustomization.h
--- a/Engine/Plugins/Runtime/nDisplay/Source/DisplayClusterConfigurator/Private/Views/Details/Policies/DisplayClusterConfiguratorPolicyParameterCustomization.h
+++ b/Engine/Plugins/Runtime/nDisplay/Source/DisplayClusterConfigurator/Private/Views/Details/Policies/DisplayClusterConfiguratorPolicyParameterCustomization.h
@@ -44,6 +44,12 @@ public:
 	/** The friendly display name for this parameter. */
 	FText GetParameterDisplayName() const { return FText::FromString(ParamDisplayName); }
 
+	/** Gets the tooltip for this parameter. */
+	FText GetParameterTooltip() const { return ParamTooltip; }
+
+	/** Sets the tooltip for this parameter. */
+	void SetParameterTooltip(const FText& NewTooltip) { ParamTooltip = NewTooltip; }
+
 	/** The keu used for the parameter map. */
 	const FString& GetParameterKey() const { return ParamKey; }
 
@@ -89,6 +95,9 @@ private:
 	/** The display name only for the UI. */
 	FString ParamDisplayName;
 	
+	/** The tooltip to display in the UI */
+	FText ParamTooltip;
+
 	/** The proper title (key) of this parameter. */
 	FString ParamKey;
 };
@@ -192,6 +201,8 @@ public:
 		else
 		{
 			NumberValue = InDefaultValue;
+			GetOrAddCustomParameterValueText();
+			UpdateCustomParameterValueText(FString::SanitizeFloat(static_cast<float>(NumberValue)));
 		}
 	}
 
diff --git a/Engine/Plugins/Runtime/nDisplay/Source/DisplayClusterPostprocess/Private/DisplayClusterPostprocessLog.cpp b/Engine/Plugins/Runtime/nDisplay/Source/DisplayClusterPostprocess/Private/DisplayClusterPostprocessLog.cpp
--- a/Engine/Plugins/Runtime/nDisplay/Source/DisplayClusterPostprocess/Private/DisplayClusterPostprocessLog.cpp
+++ b/Engine/Plugins/Runtime/nDisplay/Source/DisplayClusterPostprocess/Private/DisplayClusterPostprocessLog.cpp
@@ -4,6 +4,5 @@
 
 
 DEFINE_LOG_CATEGORY(LogDisplayClusterPostprocess);
-DEFINE_LOG_CATEGORY(LogDisplayClusterPostprocessOutputRemap);
 DEFINE_LOG_CATEGORY(LogDisplayClusterPostprocessTextureShare);
 DEFINE_LOG_CATEGORY(LogDisplayClusterPostprocessD3D12CrossGPU);
diff --git a/Engine/Plugins/Runtime/nDisplay/Source/DisplayClusterPostprocess/Private/DisplayClusterPostprocessLog.h b/Engine/Plugins/Runtime/nDisplay/Source/DisplayClusterPostprocess/Private/DisplayClusterPostprocessLog.h
--- a/Engine/Plugins/Runtime/nDisplay/Source/DisplayClusterPostprocess/Private/DisplayClusterPostprocessLog.h
+++ b/Engine/Plugins/Runtime/nDisplay/Source/DisplayClusterPostprocess/Private/DisplayClusterPostprocessLog.h
@@ -5,6 +5,5 @@
 #include "CoreMinimal.h"
 
 DECLARE_LOG_CATEGORY_EXTERN(LogDisplayClusterPostprocess,              Log, All);
-DECLARE_LOG_CATEGORY_EXTERN(LogDisplayClusterPostprocessOutputRemap,   Log, All);
 DECLARE_LOG_CATEGORY_EXTERN(LogDisplayClusterPostprocessTextureShare,  Log, All);
 DECLARE_LOG_CATEGORY_EXTERN(LogDisplayClusterPostprocessD3D12CrossGPU, Log, All);
diff --git a/Engine/Plugins/Runtime/nDisplay/Source/DisplayClusterPostprocess/Private/DisplayClusterPostprocessModule.cpp b/Engine/Plugins/Runtime/nDisplay/Source/DisplayClusterPostprocess/Private/DisplayClusterPostprocessModule.cpp
--- a/Engine/Plugins/Runtime/nDisplay/Source/DisplayClusterPostprocess/Private/DisplayClusterPostprocessModule.cpp
+++ b/Engine/Plugins/Runtime/nDisplay/Source/DisplayClusterPostprocess/Private/DisplayClusterPostprocessModule.cpp
@@ -5,8 +5,6 @@
 #include "DisplayClusterPostprocessLog.h"
 #include "DisplayClusterPostprocessStrings.h"
 
-#include "PostProcess/OutputRemap/DisplayClusterPostprocessOutputRemapFactory.h"
-
 #if PLATFORM_WINDOWS
 #include "PostProcess/Windows/DisplayClusterPostprocessTextureShare.h"
 #include "PostProcess/Windows/DisplayClusterPostprocessDX12CrossGPU.h"
@@ -20,10 +18,6 @@ FDisplayClusterPostprocessModule::FDisplayClusterPostprocessModule()
 {
 	TSharedPtr<IDisplayClusterPostProcessFactory> Postprocess;
 
-	// Output Remap
-	Postprocess = MakeShared<FDisplayClusterPostprocessOutputRemapFactory>();
-	PostprocessAssets.Emplace(DisplayClusterPostprocessStrings::postprocess::OutputRemap, Postprocess);
-
 #if PLATFORM_WINDOWS
 	// Texture Share
 	//Postprocess = MakeShared<FDisplayClusterPostprocessTextureShare>();
diff --git a/Engine/Plugins/Runtime/nDisplay/Source/DisplayClusterPostprocess/Private/PostProcess/OutputRemap/DisplayClusterPostprocessOutputRemap.cpp b/Engine/Plugins/Runtime/nDisplay/Source/DisplayClusterPostprocess/Private/PostProcess/OutputRemap/DisplayClusterPostprocessOutputRemap.cpp
deleted file mode 100644
diff --git a/Engine/Plugins/Runtime/nDisplay/Source/DisplayClusterPostprocess/Private/PostProcess/OutputRemap/DisplayClusterPostprocessOutputRemap.h b/Engine/Plugins/Runtime/nDisplay/Source/DisplayClusterPostprocess/Private/PostProcess/OutputRemap/DisplayClusterPostprocessOutputRemap.h
deleted file mode 100644
diff --git a/Engine/Plugins/Runtime/nDisplay/Source/DisplayClusterPostprocess/Private/PostProcess/OutputRemap/DisplayClusterPostprocessOutputRemapFactory.cpp b/Engine/Plugins/Runtime/nDisplay/Source/DisplayClusterPostprocess/Private/PostProcess/OutputRemap/DisplayClusterPostprocessOutputRemapFactory.cpp
deleted file mode 100644
diff --git a/Engine/Plugins/Runtime/nDisplay/Source/DisplayClusterPostprocess/Private/PostProcess/OutputRemap/DisplayClusterPostprocessOutputRemapFactory.h b/Engine/Plugins/Runtime/nDisplay/Source/DisplayClusterPostprocess/Private/PostProcess/OutputRemap/DisplayClusterPostprocessOutputRemapFactory.h
deleted file mode 100644
diff --git a/Engine/Plugins/Runtime/nDisplay/Source/DisplayClusterPostprocess/Public/DisplayClusterPostprocessStrings.h b/Engine/Plugins/Runtime/nDisplay/Source/DisplayClusterPostprocess/Public/DisplayClusterPostprocessStrings.h
--- a/Engine/Plugins/Runtime/nDisplay/Source/DisplayClusterPostprocess/Public/DisplayClusterPostprocessStrings.h
+++ b/Engine/Plugins/Runtime/nDisplay/Source/DisplayClusterPostprocess/Public/DisplayClusterPostprocessStrings.h
@@ -7,12 +7,6 @@
 
 namespace DisplayClusterPostprocessStrings
 {
-	namespace output_remap
-	{
-		static constexpr auto File = TEXT("file"); // Load from external OBJ file
-		static constexpr auto Mesh = TEXT("mesh"); // Load from RootActorMeshComponentName
-	}
-
 	namespace texture_share
 	{
 		static constexpr auto ShareName = TEXT("share");
@@ -32,7 +26,6 @@ namespace DisplayClusterPostprocessStrings
 
 	namespace postprocess
 	{
-		static constexpr auto OutputRemap   = TEXT("OutputRemap");
 		static constexpr auto TextureShare  = TEXT("TextureShare");
 		static constexpr auto D3D12CrossGPU = TEXT("D3D12CrossGPU");
 	}
diff --git a/Engine/Plugins/Runtime/nDisplay/Source/DisplayClusterProjection/Private/Policy/Domeprojection/Windows/DX11/DisplayClusterProjectionDomeprojectionViewAdapterDX11.cpp b/Engine/Plugins/Runtime/nDisplay/Source/DisplayClusterProjection/Private/Policy/Domeprojection/Windows/DX11/DisplayClusterProjectionDomeprojectionViewAdapterDX11.cpp
--- a/Engine/Plugins/Runtime/nDisplay/Source/DisplayClusterProjection/Private/Policy/Domeprojection/Windows/DX11/DisplayClusterProjectionDomeprojectionViewAdapterDX11.cpp
+++ b/Engine/Plugins/Runtime/nDisplay/Source/DisplayClusterProjection/Private/Policy/Domeprojection/Windows/DX11/DisplayClusterProjectionDomeprojectionViewAdapterDX11.cpp
@@ -114,8 +114,7 @@ bool FDisplayClusterProjectionDomeprojectionViewAdapterDX11::GetProjectionMatrix
 bool FDisplayClusterProjectionDomeprojectionViewAdapterDX11::ImplApplyWarpBlend_RenderThread(FRHICommandListImmediate& RHICmdList, int InContextNum, FRHITexture2D* InputTextures, FRHITexture2D* OutputTextures)
 {
 	FD3D11DynamicRHI* d3d11RHI = static_cast<FD3D11DynamicRHI*>(GDynamicRHI);
-	FViewport* MainViewport = GEngine->GameViewport->Viewport;
-	if (d3d11RHI == nullptr || MainViewport == nullptr)
+	if (d3d11RHI == nullptr)
 	{
 		return false;
 	}
diff --git a/Engine/Plugins/Runtime/nDisplay/Source/DisplayClusterProjection/Private/Policy/EasyBlend/DisplayClusterProjectionEasyBlendPolicyBase.cpp b/Engine/Plugins/Runtime/nDisplay/Source/DisplayClusterProjection/Private/Policy/EasyBlend/DisplayClusterProjectionEasyBlendPolicyBase.cpp
--- a/Engine/Plugins/Runtime/nDisplay/Source/DisplayClusterProjection/Private/Policy/EasyBlend/DisplayClusterProjectionEasyBlendPolicyBase.cpp
+++ b/Engine/Plugins/Runtime/nDisplay/Source/DisplayClusterProjection/Private/Policy/EasyBlend/DisplayClusterProjectionEasyBlendPolicyBase.cpp
@@ -28,19 +28,24 @@ bool FDisplayClusterProjectionEasyBlendPolicyBase::HandleStartScene(class IDispl
 {
 	check(IsInGameThread());
 
-	if (bInitializeOnce)
+	if (!IsEasyBlendRenderingEnabled())
 	{
+		if (!bEasyBlendInitializeOnce)
+		{
+			UE_LOG(LogDisplayClusterProjectionEasyBlend, Error, TEXT("An error occurred during EasyBlend initialization : current UE render device not supported"));
+			bEasyBlendInitializeOnce = true;
+		}
+
 		return false;
 	}
 
-	bInitializeOnce = true;
-
-	if (!IsEasyBlendRenderingEnabled())
+	if (bInitializeOnce)
 	{
-		UE_LOG(LogDisplayClusterProjectionEasyBlend, Error, TEXT("An error occurred during EasyBlend initialization : current UE render device not supported"));
 		return false;
 	}
 
+	bInitializeOnce = true;
+
 	// The game side of the nDisplay has been initialized by the nDisplay Game Manager already
 	// so we can extend it by our projection related functionality/components/etc.
 
diff --git a/Engine/Plugins/Runtime/nDisplay/Source/DisplayClusterProjection/Private/Policy/EasyBlend/DisplayClusterProjectionEasyBlendPolicyBase.h b/Engine/Plugins/Runtime/nDisplay/Source/DisplayClusterProjection/Private/Policy/EasyBlend/DisplayClusterProjectionEasyBlendPolicyBase.h
--- a/Engine/Plugins/Runtime/nDisplay/Source/DisplayClusterProjection/Private/Policy/EasyBlend/DisplayClusterProjectionEasyBlendPolicyBase.h
+++ b/Engine/Plugins/Runtime/nDisplay/Source/DisplayClusterProjection/Private/Policy/EasyBlend/DisplayClusterProjectionEasyBlendPolicyBase.h
@@ -50,6 +50,7 @@ private:
 	FString OriginCompId;
 	float EasyBlendScale = 1.f;
 	bool bInitializeOnce = false;
+	bool bEasyBlendInitializeOnce = false;
 
 	// RHI depended view adapter (different RHI require different DLL/API etc.)
 	TUniquePtr<FDisplayClusterProjectionEasyBlendViewAdapterBase> ViewAdapter;
diff --git a/Engine/Source/Programs/TextureShare/TextureShareSDK/Private/TextureShareInterface.cpp b/Engine/Source/Programs/TextureShare/TextureShareSDK/Private/TextureShareInterface.cpp
--- a/Engine/Source/Programs/TextureShare/TextureShareSDK/Private/TextureShareInterface.cpp
+++ b/Engine/Source/Programs/TextureShare/TextureShareSDK/Private/TextureShareInterface.cpp
@@ -131,10 +131,10 @@ void FTextureShareInterface::SetSyncPolicySettings(const FTextureShareSyncPolicy
 {
 	ShareCoreAPI().SetSyncPolicySettings(ETextureShareProcess::Client, InSyncPolicySettings);
 }
-bool FTextureShareInterface::CreateTextureShare(const TCHAR* ShareName, ETextureShareProcess Process, FTextureShareSyncPolicy SyncMode, ETextureShareDevice DeviceType)
+bool FTextureShareInterface::CreateTextureShare(const TCHAR* ShareName, ETextureShareProcess Process, FTextureShareSyncPolicy SyncMode, ETextureShareDevice DeviceType, float SyncWaitTime)
 {
 	TSharedPtr<ITextureShareItem> ShareItem;
-	return ShareCoreAPI().CreateTextureShareItem(FString(ShareName), Process, SyncMode, DeviceType, ShareItem);
+	return ShareCoreAPI().CreateTextureShareItem(FString(ShareName), Process, SyncMode, DeviceType, ShareItem, SyncWaitTime);
 }
 
 bool FTextureShareInterface::ReleaseTextureShare(const TCHAR* ShareName)
diff --git a/Engine/Source/Programs/TextureShare/TextureShareSDK/Public/TextureShareInterface.h b/Engine/Source/Programs/TextureShare/TextureShareSDK/Public/TextureShareInterface.h
--- a/Engine/Source/Programs/TextureShare/TextureShareSDK/Public/TextureShareInterface.h
+++ b/Engine/Source/Programs/TextureShare/TextureShareSDK/Public/TextureShareInterface.h
@@ -21,7 +21,7 @@ public:
 	TEXTURE_SHARE_SDK_API static void SetSyncPolicySettings(const FTextureShareSyncPolicySettings& InSyncPolicySettings);
 
 	// Create shared resource object
-	TEXTURE_SHARE_SDK_API static bool CreateTextureShare(const TCHAR* ShareName, ETextureShareProcess Process, FTextureShareSyncPolicy SyncMode, ETextureShareDevice DeviceType);
+	TEXTURE_SHARE_SDK_API static bool CreateTextureShare(const TCHAR* ShareName, ETextureShareProcess Process, FTextureShareSyncPolicy SyncMode, ETextureShareDevice DeviceType, float SyncWaitTime = 0.03);
 	TEXTURE_SHARE_SDK_API static bool ReleaseTextureShare(const TCHAR* ShareName);
 
 	TEXTURE_SHARE_SDK_API static bool IsValid(const TCHAR* ShareName);
