Update DisplayClusterConfigurator with asset editor API changes so that it loads in UE5

diff --git a/Engine/Plugins/Runtime/nDisplay/Source/DisplayClusterConfigurator/Private/DisplayClusterConfiguratorBlueprintEditor.cpp b/Engine/Plugins/Runtime/nDisplay/Source/DisplayClusterConfigurator/Private/DisplayClusterConfiguratorBlueprintEditor.cpp
--- a/Engine/Plugins/Runtime/nDisplay/Source/DisplayClusterConfigurator/Private/DisplayClusterConfiguratorBlueprintEditor.cpp
+++ b/Engine/Plugins/Runtime/nDisplay/Source/DisplayClusterConfigurator/Private/DisplayClusterConfiguratorBlueprintEditor.cpp
@@ -45,7 +45,6 @@
 #include "ISCSEditorUICustomization.h"
 #include "SBlueprintEditorToolbar.h"
 #include "SubobjectEditorExtensionContext.h"
-#include "SubobjectEditorExtensionContext.h"
 #include "SKismetInspector.h"
 #include "SSCSEditor.h"
 #include "ToolMenu.h"
@@ -149,8 +148,6 @@ void FDisplayClusterConfiguratorBlueprintEditor::InitDisplayClusterBlueprintEdit
 	
 	RegisterMenus();
 	
-	TSharedRef<FApplicationMode> ConfigurationMode = MakeShared<FDisplayClusterConfiguratorEditorConfigurationMode>(Editor);
-	
 	{
 		const bool bCreateDefaultStandaloneMenu = true;
 		const bool bCreateDefaultToolbar = true;
@@ -166,6 +163,7 @@ void FDisplayClusterConfiguratorBlueprintEditor::InitDisplayClusterBlueprintEdit
 		);
 	}
 
+	TSharedRef<FApplicationMode> ConfigurationMode = MakeShared<FDisplayClusterConfiguratorEditorConfigurationMode>(Editor);
 	AddApplicationMode(ConfigurationMode->GetModeName(), ConfigurationMode);
 	
 	const TArray<UBlueprint*> Blueprints{ Blueprint };
@@ -1249,7 +1247,7 @@ void FDisplayClusterConfiguratorBlueprintEditor::CreateDCSCSEditors()
 
 	ViewportTabContent->Initialize([this](const FAssetEditorViewportConstructionArgs InArgs)
 	{
-		return SNew(SDisplayClusterConfiguratorSCSEditorViewport)
+		return SNew(SDisplayClusterConfiguratorSCSEditorViewport, InArgs)
 			.BlueprintEditor(SharedThis(this))
 			.OwningTab(ViewportTab);
 	}, ViewportTab.ToSharedRef(), LayoutId);
diff --git a/Engine/Plugins/Runtime/nDisplay/Source/DisplayClusterConfigurator/Private/Views/Viewport/DisplayClusterConfiguratorSCSEditorViewport.cpp b/Engine/Plugins/Runtime/nDisplay/Source/DisplayClusterConfigurator/Private/Views/Viewport/DisplayClusterConfiguratorSCSEditorViewport.cpp
--- a/Engine/Plugins/Runtime/nDisplay/Source/DisplayClusterConfigurator/Private/Views/Viewport/DisplayClusterConfiguratorSCSEditorViewport.cpp
+++ b/Engine/Plugins/Runtime/nDisplay/Source/DisplayClusterConfigurator/Private/Views/Viewport/DisplayClusterConfiguratorSCSEditorViewport.cpp
@@ -470,7 +470,7 @@ private:
 	TSharedPtr<FDisplayClusterConfiguratorSCSEditorViewportClient> ViewportClient;
 };
 
-void SDisplayClusterConfiguratorSCSEditorViewport::Construct(const FArguments& InArgs)
+void SDisplayClusterConfiguratorSCSEditorViewport::Construct(const FArguments& InArgs, const FAssetEditorViewportConstructionArgs& InViewportConstructionArgs)
 {
 	bIsActiveTimerRegistered = false;
 
@@ -478,7 +478,7 @@ void SDisplayClusterConfiguratorSCSEditorViewport::Construct(const FArguments& I
 	BlueprintEditorPtr = InArgs._BlueprintEditor;
 	OwnerTab = InArgs._OwningTab;
 	
-	SAssetEditorViewport::Construct(SAssetEditorViewport::FArguments(), FAssetEditorViewportConstructionArgs());
+	SAssetEditorViewport::Construct(SAssetEditorViewport::FArguments(), InViewportConstructionArgs);
 
 	// Restore last used feature level
 	if (ViewportClient.IsValid())
diff --git a/Engine/Plugins/Runtime/nDisplay/Source/DisplayClusterConfigurator/Private/Views/Viewport/DisplayClusterConfiguratorSCSEditorViewport.h b/Engine/Plugins/Runtime/nDisplay/Source/DisplayClusterConfigurator/Private/Views/Viewport/DisplayClusterConfiguratorSCSEditorViewport.h
--- a/Engine/Plugins/Runtime/nDisplay/Source/DisplayClusterConfigurator/Private/Views/Viewport/DisplayClusterConfiguratorSCSEditorViewport.h
+++ b/Engine/Plugins/Runtime/nDisplay/Source/DisplayClusterConfigurator/Private/Views/Viewport/DisplayClusterConfiguratorSCSEditorViewport.h
@@ -15,7 +15,7 @@ class SDisplayClusterConfiguratorSCSEditorViewport : public SAssetEditorViewport
 		SLATE_ARGUMENT(TSharedPtr<SDockTab>, OwningTab)
 	SLATE_END_ARGS()
 
-	void Construct(const FArguments& InArgs);
+	void Construct(const FArguments& InArgs, const FAssetEditorViewportConstructionArgs& InViewportConstructionArgs);
 
 	~SDisplayClusterConfiguratorSCSEditorViewport();
 	
