nDisplay: Fixed issue with styles for the Output Mapping toolbar that caused it to render incorrectly. Also fixed issue that prevented mouse interactions from working correctly on the Output Mapping graph editor.

diff --git a/Engine/Plugins/Runtime/nDisplay/Source/DisplayClusterConfigurator/Private/Views/OutputMapping/SDisplayClusterConfiguratorGraphEditor.cpp b/Engine/Plugins/Runtime/nDisplay/Source/DisplayClusterConfigurator/Private/Views/OutputMapping/SDisplayClusterConfiguratorGraphEditor.cpp
--- a/Engine/Plugins/Runtime/nDisplay/Source/DisplayClusterConfigurator/Private/Views/OutputMapping/SDisplayClusterConfiguratorGraphEditor.cpp
+++ b/Engine/Plugins/Runtime/nDisplay/Source/DisplayClusterConfigurator/Private/Views/OutputMapping/SDisplayClusterConfiguratorGraphEditor.cpp
@@ -20,6 +20,7 @@
 #include "Framework/Commands/GenericCommands.h"
 #include "Framework/MultiBox/MultiBoxBuilder.h"
 #include "GraphEditorActions.h"
+#include "SGraphPanel.h"
 #include "Misc/ConfigCacheIni.h"
 
 #include "Interfaces/Views/OutputMapping/IDisplayClusterConfiguratorViewOutputMapping.h"
@@ -29,6 +30,8 @@
 
 void SDisplayClusterConfiguratorGraphEditor::Tick(const FGeometry& AllottedGeometry, const double InCurrentTime, const float InDeltaTime)
 {
+	SortNodes();
+
 	SGraphEditor::Tick(AllottedGeometry, InCurrentTime, InDeltaTime);
 
 	// Every tick, give the graph nodes an opportunity to update their position and size, to allow procedurally positioned nodes a chance to reposition if need be.
@@ -840,6 +843,19 @@ void SDisplayClusterConfiguratorGraphEditor::SizeToChildNodes()
 	}
 }
 
+void SDisplayClusterConfiguratorGraphEditor::SortNodes()
+{
+	// Sort the nodes via their depth to ensure that overlapping nodes are properly rendered on top of each order and mouse interactions are properly handled
+	TSlotlessChildren<SNodePanel::SNode>* VisibleChildren = static_cast<TSlotlessChildren<SNodePanel::SNode>*>(GetGraphPanel()->GetChildren());
+
+	struct SNodeLessThanSort
+	{
+		FORCEINLINE bool operator()(const TSharedRef<SNodePanel::SNode>& A, const TSharedRef<SNodePanel::SNode>& B) const { return A.Get() < B.Get(); }
+	};
+
+	VisibleChildren->Sort(SNodeLessThanSort());
+}
+
 void SDisplayClusterConfiguratorGraphEditor::Construct(const FArguments& InArgs, const TSharedRef<FDisplayClusterConfiguratorBlueprintEditor>& InToolkit, const TSharedRef<FDisplayClusterConfiguratorViewOutputMapping>& InViewOutputMapping)
 {
 	ToolkitPtr = InToolkit;
diff --git a/Engine/Plugins/Runtime/nDisplay/Source/DisplayClusterConfigurator/Private/Views/OutputMapping/SDisplayClusterConfiguratorGraphEditor.h b/Engine/Plugins/Runtime/nDisplay/Source/DisplayClusterConfigurator/Private/Views/OutputMapping/SDisplayClusterConfiguratorGraphEditor.h
--- a/Engine/Plugins/Runtime/nDisplay/Source/DisplayClusterConfigurator/Private/Views/OutputMapping/SDisplayClusterConfiguratorGraphEditor.h
+++ b/Engine/Plugins/Runtime/nDisplay/Source/DisplayClusterConfigurator/Private/Views/OutputMapping/SDisplayClusterConfiguratorGraphEditor.h
@@ -110,6 +110,8 @@ private:
 	bool CanSizeToChildNodes() const;
 	void SizeToChildNodes();
 
+	void SortNodes();
+
 private:
 	TWeakPtr<FDisplayClusterConfiguratorBlueprintEditor> ToolkitPtr;
 
diff --git a/Engine/Plugins/Runtime/nDisplay/Source/DisplayClusterConfigurator/Private/Views/OutputMapping/Widgets/SDisplayClusterConfiguratorOutputMappingToolbar.cpp b/Engine/Plugins/Runtime/nDisplay/Source/DisplayClusterConfigurator/Private/Views/OutputMapping/Widgets/SDisplayClusterConfiguratorOutputMappingToolbar.cpp
--- a/Engine/Plugins/Runtime/nDisplay/Source/DisplayClusterConfigurator/Private/Views/OutputMapping/Widgets/SDisplayClusterConfiguratorOutputMappingToolbar.cpp
+++ b/Engine/Plugins/Runtime/nDisplay/Source/DisplayClusterConfigurator/Private/Views/OutputMapping/Widgets/SDisplayClusterConfiguratorOutputMappingToolbar.cpp
@@ -38,7 +38,7 @@ TSharedRef<SWidget> SDisplayClusterConfiguratorOutputMappingToolbar::MakeToolBar
 	const FDisplayClusterConfiguratorCommands& Commands = IDisplayClusterConfigurator::Get().GetCommands();
 
 	// Use a custom style
-	FName ToolBarStyle = "ViewportMenu";
+	FName ToolBarStyle = "EditorViewportToolBar";
 	ToolbarBuilder.SetStyle(&FEditorStyle::Get(), ToolBarStyle);
 	ToolbarBuilder.SetLabelVisibility(EVisibility::Collapsed);
 
@@ -53,39 +53,46 @@ TSharedRef<SWidget> SDisplayClusterConfiguratorOutputMappingToolbar::MakeToolBar
 				SNew(SEditorViewportToolbarMenu)
 				.ParentToolBar(SharedThis(this))
 				.Cursor(EMouseCursor::Default)
-				.Image("EditorViewportToolBar.MenuDropdown")
+				.Image("EditorViewportToolBar.OptionsDropdown")
 				.OnGetMenuContent(this, &SDisplayClusterConfiguratorOutputMappingToolbar::MakeAdvancedMenu)
 			]
 		);
 	}
 	ToolbarBuilder.EndSection();
 
-	ToolbarBuilder.BeginSection("View");
+	ToolbarBuilder.BeginSection("ClusterItems");
 	{
-		ToolbarBuilder.AddToolBarButton(Commands.ToggleWindowInfo,
-			NAME_None,
-			TAttribute<FText>(),
-			TAttribute<FText>(),
-			FSlateIcon(FDisplayClusterConfiguratorStyle::GetStyleSetName(), "DisplayClusterConfigurator.OutputMapping.ToggleWindowInfo"),
-			"ToggleWindowInfo"
+		ToolbarBuilder.BeginBlockGroup();
+		{
+			ToolbarBuilder.AddToolBarButton(Commands.ToggleWindowInfo,
+				NAME_None,
+				TAttribute<FText>(),
+				TAttribute<FText>(),
+				FSlateIcon(FDisplayClusterConfiguratorStyle::GetStyleSetName(), "DisplayClusterConfigurator.OutputMapping.ToggleWindowInfo"),
+				"ToggleWindowInfo"
 			);
 
-		ToolbarBuilder.AddToolBarButton(Commands.ToggleWindowCornerImage,
-			NAME_None,
-			TAttribute<FText>(),
-			TAttribute<FText>(),
-			FSlateIcon(FDisplayClusterConfiguratorStyle::GetStyleSetName(), "DisplayClusterConfigurator.OutputMapping.ToggleWindowCornerImage"),
-			"ToggleWindowCornerImage"
+			ToolbarBuilder.AddToolBarButton(Commands.ToggleWindowCornerImage,
+				NAME_None,
+				TAttribute<FText>(),
+				TAttribute<FText>(),
+				FSlateIcon(FDisplayClusterConfiguratorStyle::GetStyleSetName(), "DisplayClusterConfigurator.OutputMapping.ToggleWindowCornerImage"),
+				"ToggleWindowCornerImage"
 			);
-
-		ToolbarBuilder.AddToolBarButton(Commands.ToggleOutsideViewports,
-			NAME_None,
-			TAttribute<FText>(),
-			TAttribute<FText>(),
-			FSlateIcon(FDisplayClusterConfiguratorStyle::GetStyleSetName(), "DisplayClusterConfigurator.OutputMapping.ToggleOutsideViewports"),
-			"ToggleOutsideViewports"
+			ToolbarBuilder.AddToolBarButton(Commands.ToggleOutsideViewports,
+				NAME_None,
+				TAttribute<FText>(),
+				TAttribute<FText>(),
+				FSlateIcon(FDisplayClusterConfiguratorStyle::GetStyleSetName(), "DisplayClusterConfigurator.OutputMapping.ToggleOutsideViewports"),
+				"ToggleOutsideViewports"
 			);
+		}
+		ToolbarBuilder.EndBlockGroup();
+	}
+	ToolbarBuilder.EndSection();
 
+	ToolbarBuilder.BeginSection("View");
+	{
 		ToolbarBuilder.AddToolBarButton(Commands.ZoomToFit,
 			NAME_None,
 			TAttribute<FText>(),
