nDisplay: Update scs editor to subobject editor, adding support back in to add components. Fixed selecting ancillary components. Fixed double clicking on cameras.

diff --git a/Engine/Plugins/Runtime/nDisplay/Source/DisplayClusterConfigurator/Private/DisplayClusterConfiguratorBlueprintEditor.cpp b/Engine/Plugins/Runtime/nDisplay/Source/DisplayClusterConfigurator/Private/DisplayClusterConfiguratorBlueprintEditor.cpp
--- a/Engine/Plugins/Runtime/nDisplay/Source/DisplayClusterConfigurator/Private/DisplayClusterConfiguratorBlueprintEditor.cpp
+++ b/Engine/Plugins/Runtime/nDisplay/Source/DisplayClusterConfigurator/Private/DisplayClusterConfiguratorBlueprintEditor.cpp
@@ -32,7 +32,6 @@
 #include "Views/DisplayClusterConfiguratorToolbar.h"
 #include "Settings/DisplayClusterConfiguratorSettings.h"
 
-#include "Engine/SCS_Node.h"
 #include "Engine/SimpleConstructionScript.h"
 
 #include "Components/ActorComponent.h"
@@ -56,13 +55,17 @@
 #include "Tools/BaseAssetToolkit.h"
 #include "Engine/Selection.h"
 #include "Engine/TextureRenderTarget2D.h"
+#include "Framework/Commands/GenericCommands.h"
+#include "Framework/Notifications/NotificationManager.h"
 #include "Widgets/SOverlay.h"
 #include "Widgets/Docking/SDockTab.h"
+#include "Widgets/Notifications/SNotificationList.h"
 #include "HAL/ConsoleManager.h"
 #include "HAL/PlatformApplicationMisc.h"
 #include "Kismet2/BlueprintEditorUtils.h"
 #include "Kismet2/DebuggerCommands.h"
 #include "Subsystems/PanelExtensionSubsystem.h"
+#include "ScopedTransaction.h"
 
 
 #define LOCTEXT_NAMESPACE "DisplayClusterConfiguratorBlueprintEditor"
@@ -1054,8 +1057,7 @@ void FDisplayClusterConfiguratorBlueprintEditor::Compile()
 	FBlueprintEditor::Compile();
 }
 
-void FDisplayClusterConfiguratorBlueprintEditor::OnSelectionUpdated(
-	const TArray<TSharedPtr<FSCSEditorTreeNode>>& SelectedNodes)
+void FDisplayClusterConfiguratorBlueprintEditor::OnSelectionUpdated(const TArray<TSharedPtr<class FSubobjectEditorTreeNode>>& SelectedNodes)
 {
 	if (CurrentSelectionSource == ESelectionSource::Refresh)
 	{
@@ -1080,13 +1082,18 @@ void FDisplayClusterConfiguratorBlueprintEditor::OnSelectionUpdated(
 	// Update the selection visualization
 	if (ADisplayClusterRootActor* EditorActorInstance = Cast<ADisplayClusterRootActor>(Blueprint->SimpleConstructionScript->GetComponentEditorActorInstance()))
 	{
-		auto IsComponentSelected = [SelectedNodes, EditorActorInstance](UActorComponent* InComponent) -> bool
+		auto IsComponentSelected = [SelectedNodes, this](UActorComponent* InComponent) -> bool
 		{
-			for (const TSharedPtr<FSCSEditorTreeNode>& TreeNode : SelectedNodes)
+			for (const TSharedPtr<FSubobjectEditorTreeNode>& TreeNode : SelectedNodes)
 			{
-				if (TreeNode.IsValid() && TreeNode->FindComponentInstanceInActor(EditorActorInstance) == InComponent)
+				const FSubobjectData* NodeData = TreeNode ? TreeNode->GetDataSource() : nullptr;
+				if (NodeData)
 				{
-					return true;
+					const UActorComponent* ActorComponent = NodeData->GetObjectForBlueprint<UActorComponent>(GetBlueprintObj());
+					if (TreeNode.IsValid() && ActorComponent == InComponent)
+					{
+						return true;
+					}
 				}
 			}
 			
@@ -1137,31 +1144,32 @@ void FDisplayClusterConfiguratorBlueprintEditor::OnSelectionUpdated(
 		TArray<FString> SelectedComponentNames;
 		bool bShowComponents = true;
 		InspectorObjects.Empty(SelectedNodes.Num());
-		for (FSCSEditorTreeNodePtrType NodePtr : SelectedNodes)
+		for (FSubobjectEditorTreeNodePtrType NodePtr : SelectedNodes)
 		{
-			if (NodePtr.IsValid())
+			const FSubobjectData* NodeData = NodePtr ? NodePtr->GetDataSource() : nullptr;
+			if (NodeData)
 			{
-				if (NodePtr->IsActorNode())
+				if (const AActor* Actor = NodeData->GetObject<AActor>())
 				{
-					if (AActor* DefaultActor = NodePtr->GetEditableObjectForBlueprint<AActor>(GetBlueprintObj()))
+					if (const AActor* DefaultActor = NodeData->GetObjectForBlueprint<AActor>(GetBlueprintObj()))
 					{
-						InspectorObjects.Add(DefaultActor);
-
+						InspectorObjects.Add(const_cast<AActor*>(DefaultActor));
+					
 						FString Title;
 						DefaultActor->GetName(Title);
 						InspectorTitle = FText::FromString(Title);
 						bShowComponents = false;
-
+					
 						TryInvokingDetailsTab();
 					}
 				}
 				else
 				{
-					UActorComponent* EditableComponent = NodePtr->GetOrCreateEditableComponentTemplate(GetBlueprintObj());
+					const UActorComponent* EditableComponent = NodeData->GetObjectForBlueprint<UActorComponent>(GetBlueprintObj());
 					if (EditableComponent)
 					{
 						InspectorTitle = FText::FromString(NodePtr->GetDisplayString());
-						InspectorObjects.Add(EditableComponent);
+						InspectorObjects.Add(const_cast<UActorComponent*>(EditableComponent));
 
 						FString ComponentName = EditableComponent->GetName();
 						ComponentName.RemoveFromEnd(UActorComponent::ComponentTemplateNameSuffix);
@@ -1206,7 +1214,7 @@ void FDisplayClusterConfiguratorBlueprintEditor::OnSelectionUpdated(
 	}
 }
 
-void FDisplayClusterConfiguratorBlueprintEditor::OnComponentDoubleClicked(TSharedPtr<FSCSEditorTreeNode> Node)
+void FDisplayClusterConfiguratorBlueprintEditor::OnComponentDoubleClicked(TSharedPtr<class FSubobjectEditorTreeNode> Node)
 {
 	TSharedPtr<SDockTab> OwnerTab = Inspector->GetOwnerTab();
 	if (OwnerTab.IsValid())
@@ -1216,21 +1224,25 @@ void FDisplayClusterConfiguratorBlueprintEditor::OnComponentDoubleClicked(TShare
 		AActor* PreviewActor = GetPreviewActor();
 		if (ViewportTabContent.IsValid() && Node.IsValid() && PreviewActor)
 		{
-			if (USceneComponent* ComponentTemplate = Cast<USceneComponent>(Node->GetComponentTemplate()))
+			if (const USceneComponent* ComponentTemplate = Cast<USceneComponent>(Node->GetComponentTemplate()))
 			{
 				const bool bIsCamera = ComponentTemplate->IsA<UDisplayClusterCameraComponent>() || ComponentTemplate->IsA<UCameraComponent>();
 				if (bIsCamera)
 				{
-					if (USceneComponent* Component = Cast<USceneComponent>(Node->FindComponentInstanceInActor(PreviewActor)))
+					if (const FSubobjectData* NodeData = Node->GetDataSource())
 					{
-						TFunction<void(FName, TSharedPtr<IEditorViewportLayoutEntity>)> ViewportFunc =
-							[this, Component](FName Name, TSharedPtr<IEditorViewportLayoutEntity> Entity)
+						if (const USceneComponent* Component = Cast<USceneComponent>(NodeData->FindComponentInstanceInActor(PreviewActor)))
 						{
-							const TSharedRef<SDisplayClusterConfiguratorSCSEditorViewport> Viewport = StaticCastSharedRef<SDisplayClusterConfiguratorSCSEditorViewport>(Entity->AsWidget());
-							Viewport->GetDisplayClusterViewportClient()->SetCameraToComponent(const_cast<USceneComponent*>(Component));
-						};
+							TFunction<void(FName, TSharedPtr<IEditorViewportLayoutEntity>)> ViewportFunc =
+								[this, Component](FName Name, TSharedPtr<IEditorViewportLayoutEntity> Entity)
+								{
+									const TSharedRef<SDisplayClusterConfiguratorSCSEditorViewport> Viewport =
+										StaticCastSharedRef<SDisplayClusterConfiguratorSCSEditorViewport>(Entity->AsWidget());
+									Viewport->GetDisplayClusterViewportClient()->SetCameraToComponent(const_cast<USceneComponent*>(Component));
+								};
 
-						ViewportTabContent->PerformActionOnViewports(ViewportFunc);
+							ViewportTabContent->PerformActionOnViewports(ViewportFunc);
+						}
 					}
 				}
 			}
@@ -1268,18 +1280,20 @@ TSharedRef<SWidget> FDisplayClusterConfiguratorBlueprintEditor::CreateSCSEditorE
 	FWeakObjectPtr ExtensionContext)
 {
 	auto PerformComboAddClass = [=](TSubclassOf<UActorComponent> ComponentClass, EComponentCreateAction::Type ComponentCreateAction,
-		UObject* AssetOverride) -> UActorComponent*
+		UObject* AssetOverride) -> FSubobjectDataHandle
 	{
-		if (USubobjectEditorExtensionContext* SubobjectEditorExtensionContext = Cast<USubobjectEditorExtensionContext>(ExtensionContext.Get()))
+		FSubobjectDataHandle NewComponentHandle = FSubobjectDataHandle::InvalidHandle;
+		
+		if (const USubobjectEditorExtensionContext* SubobjectEditorExtensionContext = Cast<USubobjectEditorExtensionContext>(ExtensionContext.Get()))
 		{
 			if (const TSharedPtr<SSubobjectEditor> SubobjectEditor = SubobjectEditorExtensionContext->GetSubobjectEditor().Pin())
 			{
 				UClass* NewClass = ComponentClass;
-
-				UActorComponent* NewComponent = nullptr;
-
+				
 				if (NewClass != nullptr)
 				{
+					TUniquePtr<FScopedTransaction> AddTransaction = MakeUnique<FScopedTransaction>(LOCTEXT("AddComponent", "Add Component"));
+					
 					FEditorDelegates::LoadSelectedAssetsIfNeeded.Broadcast();
 					USelection* Selection = GEditor->GetSelectedObjects();
 
@@ -1289,13 +1303,13 @@ TSharedRef<SWidget> FDisplayClusterConfiguratorBlueprintEditor::CreateSCSEditorE
 					// Otherwise the DisplayCluster class name is used, not the custom nDisplay display name.
 					if (AssetOverride == nullptr && NewClass->IsClassGroupName(TEXT("DisplayCluster")))
 					{
-						const FString DisplayName =  FDisplayClusterConfiguratorUtils::FormatNDisplayComponentName(NewClass);
+						const FString DisplayName = FDisplayClusterConfiguratorUtils::FormatNDisplayComponentName(NewClass);
 						AssetOverride = NewObject<UActorComponent>(GetTransientPackage(), NewClass, *DisplayName);
 					}
 					
 					USubobjectDataSubsystem* SubobjectSystem = USubobjectDataSubsystem::Get();
 					FAddNewSubobjectParams AddSubobjectParams;
-					FText FailReason;
+					AddSubobjectParams.BlueprintContext = LoadedBlueprint.Get();
 					AddSubobjectParams.NewClass = NewClass;
 										
 					// This adds components according to the type selected in the drop down. If the user
@@ -1303,24 +1317,30 @@ TSharedRef<SWidget> FDisplayClusterConfiguratorBlueprintEditor::CreateSCSEditorE
 					// else we go down the previous route of adding components by type.
 					//
 					// Furthermore don't try to match up assets for USceneComponent it will match lots of things and doesn't have any nice behavior for asset adds 
+
+					FText OutFailReason;
+					
 					if (Selection->Num() > 0 && !AssetOverride && NewClass != USceneComponent::StaticClass())
 					{
 						for (FSelectionIterator ObjectIter(*Selection); ObjectIter; ++ObjectIter)
 						{
 							UObject* Object = *ObjectIter;
-							TArray< TSubclassOf<UActorComponent> > ComponentClasses = FComponentAssetBrokerage::GetComponentsForAsset(Object);
+							TArray<TSubclassOf<UActorComponent>> ComponentClasses = FComponentAssetBrokerage::GetComponentsForAsset(Object);
 
 							// if the selected asset supports the selected component type then go ahead and add it
 							for (int32 ComponentIndex = 0; ComponentIndex < ComponentClasses.Num(); ComponentIndex++)
 							{
 								if (ComponentClasses[ComponentIndex]->IsChildOf(NewClass))
 								{
-									AddSubobjectParams.AssetOverride = Object;
-									AddSubobjectParams.ParentHandle = SubobjectSystem->FindHandleForObject(SubobjectEditor->GetSceneRootNode()->GetDataHandle(), Object);
-									SubobjectSystem->AddNewSubobject(AddSubobjectParams, FailReason);
-									
-									bAddedComponent = true;
-									break;
+									const FSubobjectDataHandle ParentHandle = SubobjectSystem->FindHandleForObject(SubobjectEditor->GetSceneRootNode()->GetDataHandle(), Object);
+									if(ParentHandle.IsValid())
+									{
+										AddSubobjectParams.AssetOverride = Object;
+										AddSubobjectParams.ParentHandle = ParentHandle;
+										NewComponentHandle = SubobjectSystem->AddNewSubobject(AddSubobjectParams, OutFailReason);
+										bAddedComponent = true;
+										break;
+									}
 								}
 							}
 						}
@@ -1328,20 +1348,63 @@ TSharedRef<SWidget> FDisplayClusterConfiguratorBlueprintEditor::CreateSCSEditorE
 
 					if (!bAddedComponent)
 					{
-						// As the SCS splits up the scene and actor components, can now add directly
-						AddSubobjectParams.NewClass = NewClass;
-						AddSubobjectParams.AssetOverride = AssetOverride;
-						AddSubobjectParams.ParentHandle = SubobjectSystem->FindHandleForObject(SubobjectEditor->GetSceneRootNode()->GetDataHandle(), AssetOverride);
-						SubobjectSystem->AddNewSubobject(AddSubobjectParams, FailReason);
+						// Attach this component to the override asset first, but if none is given then use the actor context			
+						FSubobjectDataHandle ParentHandle = SubobjectSystem->FindHandleForObject(SubobjectEditor->GetSceneRootNode()->GetDataHandle(), AssetOverride);
+			
+						if(!ParentHandle.IsValid())
+						{
+							// If we have something selected, then we should attach it to that
+							TArray<FSubobjectEditorTreeNodePtrType> SelectedTreeNodes = SubobjectEditor->GetSelectedNodes();
+							if(SelectedTreeNodes.Num() > 0)
+							{
+								ParentHandle = SelectedTreeNodes[0]->GetDataHandle();
+							}
+							// Otherwise fall back to the root node
+							else
+							{
+								const TArray<FSubobjectEditorTreeNodePtrType>& RootNodes = SubobjectEditor->GetRootNodes();
+								ParentHandle = RootNodes.Num() > 0
+												? RootNodes[0]->GetDataHandle()
+												: FSubobjectDataHandle::InvalidHandle;					
+							}
+						}
+			
+						if(ParentHandle.IsValid())
+						{
+							AddSubobjectParams.AssetOverride = AssetOverride;
+							AddSubobjectParams.ParentHandle = ParentHandle;
+							NewComponentHandle = SubobjectSystem->AddNewSubobject(AddSubobjectParams, OutFailReason);
+						}
 					}
 					
-					SubobjectEditor->UpdateTree();
-				}
+					if(!NewComponentHandle.IsValid())
+					{
+						if (OutFailReason.IsEmpty())
+						{
+							OutFailReason = LOCTEXT("AddComponentFailed_Generic", "Failed to add component!");
+						}
+						FNotificationInfo Info(OutFailReason);
+						Info.Image = FEditorStyle::GetBrush(TEXT("Icons.Error"));
+						Info.bFireAndForget = true;
+						Info.bUseSuccessFailIcons = false;
+						Info.ExpireDuration = 5.0f;
 
-				return NewComponent;
+						FSlateNotificationManager::Get().AddNotification(Info);
+					}
+					else
+					{
+						SubobjectEditor->UpdateTree();
+						// Set focus to the newly created subobject
+						if(const FSubobjectEditorTreeNodePtrType NewNode = SubobjectEditor->FindSlateNodeForHandle(NewComponentHandle))
+						{
+							SubobjectEditor->GetDragDropTree()->SetSelection(NewNode);
+							SubobjectEditor->GetCommandList()->ExecuteAction(FGenericCommands::Get().Rename.ToSharedRef());
+						}
+					}
+				}
 			}
 		}
-		return nullptr;
+		return NewComponentHandle;
 	};
 	
 	return SNew(SHorizontalBox)
@@ -1351,7 +1414,7 @@ TSharedRef<SWidget> FDisplayClusterConfiguratorBlueprintEditor::CreateSCSEditorE
 		.HAlign(HAlign_Left)
 		[
 			SNew(SDisplayClusterConfiguratorComponentClassCombo)
-			.OnComponentClassSelected(FComponentClassSelected::CreateLambda(PerformComboAddClass))
+			.OnSubobjectClassSelected(FSubobjectClassSelected::CreateLambda(PerformComboAddClass))
 			.Visibility_Lambda([this, ExtensionContext]
 			{
 				if(USubobjectEditorExtensionContext* SubobjectEditorExtensionContext = Cast<USubobjectEditorExtensionContext>(ExtensionContext.Get()))
diff --git a/Engine/Plugins/Runtime/nDisplay/Source/DisplayClusterConfigurator/Private/DisplayClusterConfiguratorBlueprintEditor.h b/Engine/Plugins/Runtime/nDisplay/Source/DisplayClusterConfigurator/Private/DisplayClusterConfiguratorBlueprintEditor.h
--- a/Engine/Plugins/Runtime/nDisplay/Source/DisplayClusterConfigurator/Private/DisplayClusterConfiguratorBlueprintEditor.h
+++ b/Engine/Plugins/Runtime/nDisplay/Source/DisplayClusterConfigurator/Private/DisplayClusterConfiguratorBlueprintEditor.h
@@ -20,7 +20,7 @@ class FDisplayClusterConfiguratorViewCluster;
 class FDisplayClusterConfiguratorViewScene;
 class FDisplayClusterConfiguratorToolbar;
 class SDisplayClusterConfiguratorSCSEditorViewport;
-class USCS_Node;
+class FSubobjectEditorTreeNode;
 
 /**
  * nDisplay editor UI (should call functions on the subsystem or UNDisplayAssetEditor)
@@ -136,10 +136,10 @@ protected:
 
 	// SSCS Implementation
 	/** Delegate invoked when the selection is changed in the SCS editor widget */
-	virtual void OnSelectionUpdated(const TArray<TSharedPtr<class FSCSEditorTreeNode>>& SelectedNodes) override;
+	virtual void OnSelectionUpdated(const TArray<TSharedPtr<class FSubobjectEditorTreeNode>>& SelectedNodes) override;
 
 	/** Delegate invoked when an item is double clicked in the SCS editor widget */
-	virtual void OnComponentDoubleClicked(TSharedPtr<class FSCSEditorTreeNode> Node) override;
+	virtual void OnComponentDoubleClicked(TSharedPtr<class FSubobjectEditorTreeNode> Node) override;
 
 	friend struct FDisplayClusterSCSViewportSummoner;
 	void CreateDCSCSEditors();
diff --git a/Engine/Plugins/Runtime/nDisplay/Source/DisplayClusterConfigurator/Private/Views/SCSEditor/SDisplayClusterConfiguratorComponentCombo.cpp b/Engine/Plugins/Runtime/nDisplay/Source/DisplayClusterConfigurator/Private/Views/SCSEditor/SDisplayClusterConfiguratorComponentCombo.cpp
--- a/Engine/Plugins/Runtime/nDisplay/Source/DisplayClusterConfigurator/Private/Views/SCSEditor/SDisplayClusterConfiguratorComponentCombo.cpp
+++ b/Engine/Plugins/Runtime/nDisplay/Source/DisplayClusterConfigurator/Private/Views/SCSEditor/SDisplayClusterConfiguratorComponentCombo.cpp
@@ -1,9 +1,6 @@
 // Copyright Epic Games, Inc. All Rights Reserved.
 
-#include "Views/SCSEditor/SDisplayClusterConfiguratorComponentCombo.h"
-
-#include "DisplayClusterConfiguratorBlueprintEditor.h"
-#include "DisplayClusterRootActor.h"
+#include "SDisplayClusterConfiguratorComponentCombo.h"
 
 #include "Components/DisplayClusterCameraComponent.h"
 #include "Components/DisplayClusterICVFXCameraComponent.h"
@@ -15,32 +12,31 @@
 #include "Widgets/SToolTip.h"
 #include "Widgets/Views/SListView.h"
 #include "Widgets/Input/SComboBox.h"
-#include "Widgets/Input/SComboButton.h"
 #include "Widgets/Input/SSearchBox.h"
+#include "Widgets/Layout/SSeparator.h"
 
-#include "Misc/TextFilterExpressionEvaluator.h"
+#include "EditorStyleSet.h"
 #include "Components/SceneComponent.h"
 #include "Engine/Blueprint.h"
 #include "Engine/Selection.h"
-
 #include "Editor.h"
-#include "EditorStyleSet.h"
 #include "Styling/SlateIconFinder.h"
 #include "ComponentAssetBroker.h"
 #include "ComponentTypeRegistry.h"
 #include "EditorClassUtils.h"
+
 #include "SListViewSelectorDropdownMenu.h"
-#include "SComponentClassCombo.h"
+#include "Misc/TextFilterExpressionEvaluator.h"
+#include "SEditorHeaderButton.h"
 
 #define LOCTEXT_NAMESPACE "DisplayClusterConfiguratorComponentClassCombo"
 
 void SDisplayClusterConfiguratorComponentClassCombo::Construct(const FArguments& InArgs)
 {
 	PrevSelectedIndex = INDEX_NONE;
-	OnComponentClassSelected = InArgs._OnComponentClassSelected;
-	
+	OnSubobjectClassSelected = InArgs._OnSubobjectClassSelected;
 	TextFilter = MakeShared<FTextFilterExpressionEvaluator>(ETextFilterExpressionEvaluatorMode::BasicString);
-	
+
 	FComponentTypeRegistry::Get().SubscribeToComponentList(ComponentClassListPtr).AddRaw(this, &SDisplayClusterConfiguratorComponentClassCombo::UpdateComponentClassList);
 
 	UpdateComponentClassList();
@@ -56,76 +52,50 @@ void SDisplayClusterConfiguratorComponentClassCombo::Construct(const FArguments&
 		.OnTextChanged( this, &SDisplayClusterConfiguratorComponentClassCombo::OnSearchBoxTextChanged )
 		.OnTextCommitted( this, &SDisplayClusterConfiguratorComponentClassCombo::OnSearchBoxTextCommitted );
 
-	// Create the Construct arguments for the parent class (SComboButton)
-	SComboButton::FArguments Args;
-	Args.ButtonContent()
+	ChildSlot
 	[
-		SNew(SHorizontalBox)
-		+SHorizontalBox::Slot()
-		.VAlign(VAlign_Center)
-		.AutoWidth()
-		.Padding(1.f,1.f)
-		[
-			SNew(STextBlock)
-			.TextStyle(FEditorStyle::Get(), "ContentBrowser.TopBar.Font")
-			.Font(FEditorStyle::Get().GetFontStyle("FontAwesome.10"))
-			.Text(FText::FromString(FString(TEXT("\xf067"))) /*fa-plus*/)
-		]
-		+ SHorizontalBox::Slot()
-		.VAlign(VAlign_Center)
-		.Padding(1.f)
+		SAssignNew(AddNewButton, SEditorHeaderButton)
+		.Icon(FAppStyle::Get().GetBrush("Icons.Plus"))
+		.Text(LOCTEXT("AddComponent", "Add Component"))
+		.OnComboBoxOpened(this, &SDisplayClusterConfiguratorComponentClassCombo::ClearSelection)
+		.MenuContent()
 		[
-			SNew(STextBlock)
-			.Text(LOCTEXT("AddnDisplayComponentButtonLabel", "Add Component"))
-			.TextStyle(FEditorStyle::Get(), "ContentBrowser.TopBar.Font")
-			.Visibility(InArgs._IncludeText.Get() ? EVisibility::Visible : EVisibility::Collapsed)
-		]
-	]
-	.MenuContent()
-	[
-
-		SNew(SListViewSelectorDropdownMenu<FComponentClassComboEntryPtr>, SearchBox, ComponentClassListView)
-		[
-			SNew(SBorder)
-			.BorderImage(FEditorStyle::GetBrush("Menu.Background"))
-			.Padding(2)
+			SNew(SListViewSelectorDropdownMenu<FComponentClassComboEntryPtr>, SearchBox, ComponentClassListView)
 			[
-				SNew(SBox)
-				.WidthOverride(250)
-				[				
-					SNew(SVerticalBox)
-					+SVerticalBox::Slot()
-					.Padding(1.f)
-					.AutoHeight()
-					[
-						SearchBox.ToSharedRef()
-					]
-					+SVerticalBox::Slot()
-					.MaxHeight(400)
-					[
-						ComponentClassListView.ToSharedRef()
+				SNew(SBorder)
+				.BorderImage(FEditorStyle::GetBrush("Menu.Background"))
+				.Padding(2)
+				[
+					SNew(SBox)
+					.WidthOverride(250)
+					[				
+						SNew(SVerticalBox)
+						+SVerticalBox::Slot()
+						.Padding(1.f)
+						.AutoHeight()
+						[
+							SearchBox.ToSharedRef()
+						]
+						+SVerticalBox::Slot()
+						.MaxHeight(400)
+						[
+							ComponentClassListView.ToSharedRef()
+						]
 					]
 				]
 			]
 		]
-	]
-	.IsFocusable(true)
-	.ContentPadding(FMargin(5, 0))
-	.ComboButtonStyle(FEditorStyle::Get(), "ToolbarComboButton")
-	.ButtonStyle(FEditorStyle::Get(), "FlatButton.Success")
-	.ForegroundColor(FLinearColor::White)
-	.OnComboBoxOpened(this, &SDisplayClusterConfiguratorComponentClassCombo::ClearSelection);
+	];
 
-	SComboButton::Construct(Args);
+	
 
 	ComponentClassListView->EnableToolTipForceField( true );
-	// The base class can automatically handle setting focus to a specified control when the combo button is opened
-	SetMenuContentWidgetToFocus( SearchBox );
+	// The button can automatically handle setting focus to a specified control when the combo button is opened
+	AddNewButton->SetMenuContentWidgetToFocus( SearchBox );
 }
 
 SDisplayClusterConfiguratorComponentClassCombo::~SDisplayClusterConfiguratorComponentClassCombo()
 {
-	OnComponentClassSelected.Unbind();
 	FComponentTypeRegistry::Get().GetOnComponentTypeListChanged().RemoveAll(this);
 	FilteredComponentClassList.Empty();
 	ComponentClassList.Empty();
@@ -231,6 +201,7 @@ void SDisplayClusterConfiguratorComponentClassCombo::OnSearchBoxTextCommitted(co
 	}
 }
 
+// @todo: use FKismetEditorUtilities if the main SCS editor combo also moves their version of this function.
 static UClass* GetAuthoritativeBlueprintClass(UBlueprint const* const Blueprint)
 {
 	UClass* BpClass = (Blueprint->SkeletonGeneratedClass != nullptr) ? Blueprint->SkeletonGeneratedClass :
@@ -256,20 +227,12 @@ void SDisplayClusterConfiguratorComponentClassCombo::OnAddComponentSelectionChan
 		// We don't want the item to remain selected
 		ClearSelection();
 
-		// The filtered list won't always clear & regenerate. Seems to be a UE bug.
-		{
-			const FText EmptyText;
-			TextFilter->SetFilterText(EmptyText);
-
-			GenerateFilteredComponentList();
-		}
-		
 		if ( InItem->IsClass() )
 		{
 			// Neither do we want the combo dropdown staying open once the user has clicked on a valid option
-			SetIsOpen(false, false);
+			AddNewButton->SetIsMenuOpen(false, false);
 
-			if( OnComponentClassSelected.IsBound() )
+			if( OnSubobjectClassSelected.IsBound() )
 			{
 				UClass* ComponentClass = InItem->GetComponentClass();
 				if (ComponentClass == nullptr)
@@ -279,11 +242,11 @@ void SDisplayClusterConfiguratorComponentClassCombo::OnAddComponentSelectionChan
 					UBlueprint* LoadedObject = LoadObject<UBlueprint>(nullptr, *InItem->GetComponentPath(), nullptr, LoadFlags, nullptr);
 					ComponentClass = GetAuthoritativeBlueprintClass(LoadedObject);
 				}
-
-				UActorComponent* NewActorComponent = OnComponentClassSelected.Execute(ComponentClass, InItem->GetComponentCreateAction(), InItem->GetAssetOverride());
-				if(NewActorComponent)
+				
+				FSubobjectDataHandle NewActorCompHandle = OnSubobjectClassSelected.Execute(ComponentClass, InItem->GetComponentCreateAction(), InItem->GetAssetOverride());
+				if(NewActorCompHandle.IsValid())
 				{
-					InItem->GetOnComponentCreated().ExecuteIfBound(NewActorComponent);
+					InItem->GetOnSubobjectCreated().ExecuteIfBound(NewActorCompHandle);
 				}
 			}
 		}
@@ -347,13 +310,9 @@ TSharedRef<ITableRow> SDisplayClusterConfiguratorComponentClassCombo::GenerateAd
 				.Style(&FEditorStyle::Get().GetWidgetStyle<FTableRowStyle>("TableView.NoHoverTableRow"))
 				.ShowSelection(false)
 			[
-				SNew(SBox)
-				.Padding(1.f)
-				[
-					SNew(SBorder)
-					.Padding(FEditorStyle::GetMargin(TEXT("Menu.Separator.Padding")))
-					.BorderImage(FEditorStyle::GetBrush(TEXT("Menu.Separator")))
-				]
+				SNew(SSeparator)
+				.SeparatorImage(FEditorStyle::Get().GetBrush("Menu.Separator"))
+				.Thickness(1.0f)
 			];
 	}
 	else
@@ -403,6 +362,73 @@ void SDisplayClusterConfiguratorComponentClassCombo::UpdateComponentClassList()
 	GenerateFilteredComponentList();
 }
 
+FText SDisplayClusterConfiguratorComponentClassCombo::GetFriendlyComponentName(FComponentClassComboEntryPtr Entry) const
+{
+	// Get a user friendly string from the component name
+	FString FriendlyComponentName;
+
+	if( Entry->GetComponentCreateAction() == EComponentCreateAction::CreateNewCPPClass )
+	{
+		FriendlyComponentName = LOCTEXT("NewCPPComponentFriendlyName", "New C++ Component...").ToString();
+	}
+	else if (Entry->GetComponentCreateAction() == EComponentCreateAction::CreateNewBlueprintClass )
+	{
+		FriendlyComponentName = LOCTEXT("NewBlueprintComponentFriendlyName", "New Blueprint Script Component...").ToString();
+	}
+	else
+	{
+		FriendlyComponentName = SComponentClassCombo::GetSanitizedComponentName(Entry);
+
+		// Don't try to match up assets for USceneComponent it will match lots of things and doesn't have any nice behavior for asset adds 
+		if (Entry->GetComponentClass() != USceneComponent::StaticClass() && Entry->GetComponentNameOverride().IsEmpty())
+		{
+			// Search the selected assets and look for any that can be used as a source asset for this type of component
+			// If there is one we append the asset name to the component name, if there are many we append "Multiple Assets"
+			FString AssetName;
+			UObject* PreviousMatchingAsset = nullptr;
+
+			FEditorDelegates::LoadSelectedAssetsIfNeeded.Broadcast();
+			USelection* Selection = GEditor->GetSelectedObjects();
+			for(FSelectionIterator ObjectIter(*Selection); ObjectIter; ++ObjectIter)
+			{
+				UObject* Object = *ObjectIter;
+				check(Object);
+				UClass* Class = Object->GetClass();
+
+				TArray<TSubclassOf<UActorComponent> > ComponentClasses = FComponentAssetBrokerage::GetComponentsForAsset(Object);
+				for(int32 ComponentIndex = 0; ComponentIndex < ComponentClasses.Num(); ComponentIndex++)
+				{
+					if(ComponentClasses[ComponentIndex]->IsChildOf(Entry->GetComponentClass()))
+					{
+						if(AssetName.IsEmpty())
+						{
+							// If there is no previous asset then we just accept the name
+							AssetName = Object->GetName();
+							PreviousMatchingAsset = Object;
+						}
+						else
+						{
+							// if there is a previous asset then check that we didn't just find multiple appropriate components
+							// in a single asset - if the asset differs then we don't display the name, just "Multiple Assets"
+							if(PreviousMatchingAsset != Object)
+							{
+								AssetName = LOCTEXT("MultipleAssetsForComponentAnnotation", "Multiple Assets").ToString();
+								PreviousMatchingAsset = Object;
+							}
+						}
+					}
+				}
+			}
+
+			if(!AssetName.IsEmpty())
+			{
+				FriendlyComponentName += FString(" (") + AssetName + FString(")");
+			}
+		}
+	}
+	return FText::FromString(FriendlyComponentName);
+}
+
 void SDisplayClusterConfiguratorComponentClassCombo::GenerateComponentClassList()
 {
 	ComponentClassList.Reset();
@@ -457,62 +483,20 @@ void SDisplayClusterConfiguratorComponentClassCombo::GenerateComponentClassList(
 	}
 }
 
-FText SDisplayClusterConfiguratorComponentClassCombo::GetFriendlyComponentName(FComponentClassComboEntryPtr Entry) const
+TSharedRef<SToolTip> SDisplayClusterConfiguratorComponentClassCombo::GetComponentToolTip(FComponentClassComboEntryPtr Entry) const
 {
-	// Get a user friendly string from the component name
-	FString FriendlyComponentName = SComponentClassCombo::GetSanitizedComponentName(Entry);
-
-	// Don't try to match up assets for USceneComponent it will match lots of things and doesn't have any nice behavior for asset adds 
-	if (Entry->GetComponentClass() != USceneComponent::StaticClass() && Entry->GetComponentNameOverride().IsEmpty())
+	// Special handling for the "New..." options
+	if (Entry->GetComponentCreateAction() == EComponentCreateAction::CreateNewCPPClass)
 	{
-		// Search the selected assets and look for any that can be used as a source asset for this type of component
-		// If there is one we append the asset name to the component name, if there are many we append "Multiple Assets"
-		FString AssetName;
-		UObject* PreviousMatchingAsset = nullptr;
-
-		FEditorDelegates::LoadSelectedAssetsIfNeeded.Broadcast();
-		USelection* Selection = GEditor->GetSelectedObjects();
-		for(FSelectionIterator ObjectIter(*Selection); ObjectIter; ++ObjectIter)
-		{
-			UObject* Object = *ObjectIter;
-			check(Object);
-			UClass* Class = Object->GetClass();
-
-			TArray<TSubclassOf<UActorComponent> > ComponentClasses = FComponentAssetBrokerage::GetComponentsForAsset(Object);
-			for(int32 ComponentIndex = 0; ComponentIndex < ComponentClasses.Num(); ComponentIndex++)
-			{
-				if(ComponentClasses[ComponentIndex]->IsChildOf(Entry->GetComponentClass()))
-				{
-					if(AssetName.IsEmpty())
-					{
-						// If there is no previous asset then we just accept the name
-						AssetName = Object->GetName();
-						PreviousMatchingAsset = Object;
-					}
-					else
-					{
-						// if there is a previous asset then check that we didn't just find multiple appropriate components
-						// in a single asset - if the asset differs then we don't display the name, just "Multiple Assets"
-						if(PreviousMatchingAsset != Object)
-						{
-							AssetName = LOCTEXT("MultipleAssetsForComponentAnnotation", "Multiple Assets").ToString();
-							PreviousMatchingAsset = Object;
-						}
-					}
-				}
-			}
-		}
-
-		if(!AssetName.IsEmpty())
-		{
-			FriendlyComponentName += FString(" (") + AssetName + FString(")");
-		}
+		return SNew(SToolTip)
+			.Text(LOCTEXT("NewCPPComponentToolTip", "Create a custom actor component using C++"));
 	}
-	return FText::FromString(FriendlyComponentName);
-}
-
-TSharedRef<SToolTip> SDisplayClusterConfiguratorComponentClassCombo::GetComponentToolTip(FComponentClassComboEntryPtr Entry) const
-{
+	else if (Entry->GetComponentCreateAction() == EComponentCreateAction::CreateNewBlueprintClass)
+	{
+		return SNew(SToolTip)
+			.Text(LOCTEXT("NewBlueprintComponentToolTip", "Create a custom actor component using Blueprints"));
+	}
+	
 	// Handle components which have a currently loaded class
 	if (const UClass* ComponentClass = Entry->GetComponentClass())
 	{
diff --git a/Engine/Plugins/Runtime/nDisplay/Source/DisplayClusterConfigurator/Private/Views/SCSEditor/SDisplayClusterConfiguratorComponentCombo.h b/Engine/Plugins/Runtime/nDisplay/Source/DisplayClusterConfigurator/Private/Views/SCSEditor/SDisplayClusterConfiguratorComponentCombo.h
--- a/Engine/Plugins/Runtime/nDisplay/Source/DisplayClusterConfigurator/Private/Views/SCSEditor/SDisplayClusterConfiguratorComponentCombo.h
+++ b/Engine/Plugins/Runtime/nDisplay/Source/DisplayClusterConfigurator/Private/Views/SCSEditor/SDisplayClusterConfiguratorComponentCombo.h
@@ -2,30 +2,21 @@
 
 #pragma once
 
-#include "CoreMinimal.h"
 #include "SComponentClassCombo.h"
-#include "SlateFwd.h"
-#include "Widgets/DeclarativeSyntaxSupport.h"
-#include "Widgets/Input/SComboButton.h"
-#include "Widgets/Views/STableViewBase.h"
-#include "Widgets/Views/STableRow.h"
 
-class UDisplayClusterBlueprint;
-class FTextFilterExpressionEvaluator;
-class SSearchBox;
+class FComponentClassComboEntry;
 class SToolTip;
+class FTextFilterExpressionEvaluator;
 
-DECLARE_DELEGATE_OneParam(FTemplateSelected, UDisplayClusterBlueprint* /* Template */);
-
-class SDisplayClusterConfiguratorComponentClassCombo : public SComboButton
+class SDisplayClusterConfiguratorComponentClassCombo : public SCompoundWidget
 {
 public:
-	SLATE_BEGIN_ARGS(SDisplayClusterConfiguratorComponentClassCombo)
+	SLATE_BEGIN_ARGS( SDisplayClusterConfiguratorComponentClassCombo )
 		: _IncludeText(true)
 	{}
 
 		SLATE_ATTRIBUTE(bool, IncludeText)
-		SLATE_EVENT( FComponentClassSelected, OnComponentClassSelected )
+		SLATE_EVENT( FSubobjectClassSelected, OnSubobjectClassSelected )
 
 	SLATE_END_ARGS()
 
@@ -62,19 +53,18 @@ public:
 protected:
 	/** Generate the source component list to use for nDisplay. */
 	void GenerateComponentClassList();
-private:
-
+	
 	FText GetFriendlyComponentName(FComponentClassComboEntryPtr Entry) const;
 
 	TSharedRef<SToolTip> GetComponentToolTip(FComponentClassComboEntryPtr Entry) const;
+	
+	FSubobjectClassSelected OnSubobjectClassSelected;
 
-	FComponentClassSelected OnComponentClassSelected;
-
+	TArray<FComponentClassComboEntryPtr>* ComponentClassListPtr;
+	
 	/** List of component class names used by combo box */
 	TArray<FComponentClassComboEntryPtr> ComponentClassList;
 
-	TArray<FComponentClassComboEntryPtr>* ComponentClassListPtr = nullptr;
-	
 	/** List of component class names, filtered by the current search string */
 	TArray<FComponentClassComboEntryPtr> FilteredComponentClassList;
 
@@ -84,6 +74,9 @@ private:
 	/** The search box control - part of the combo drop down */
 	TSharedPtr<SSearchBox> SearchBox;
 
+	/** The Add combo button. */
+	TSharedPtr<class SEditorHeaderButton> AddNewButton;
+
 	/** The component list control - part of the combo drop down */
 	TSharedPtr< SListView<FComponentClassComboEntryPtr> > ComponentClassListView;
 
diff --git a/Engine/Plugins/Runtime/nDisplay/Source/DisplayClusterConfigurator/Private/Views/Viewport/DisplayClusterConfiguratorSCSEditorViewportClient.cpp b/Engine/Plugins/Runtime/nDisplay/Source/DisplayClusterConfigurator/Private/Views/Viewport/DisplayClusterConfiguratorSCSEditorViewportClient.cpp
--- a/Engine/Plugins/Runtime/nDisplay/Source/DisplayClusterConfigurator/Private/Views/Viewport/DisplayClusterConfiguratorSCSEditorViewportClient.cpp
+++ b/Engine/Plugins/Runtime/nDisplay/Source/DisplayClusterConfigurator/Private/Views/Viewport/DisplayClusterConfiguratorSCSEditorViewportClient.cpp
@@ -769,18 +769,21 @@ FVector FDisplayClusterConfiguratorSCSEditorViewportClient::GetWidgetLocation()
 		if (SelectedNodes.Num() > 0 && SelectedNodes.Last().IsValid())
 		{
 			// Use the last selected item for the widget location
-			const USceneComponent* SceneComp = Cast<USceneComponent>(SelectedNodes.Last()->GetDataSource()->FindComponentInstanceInActor(PreviewActor));
-			if (SceneComp)
+			if (const FSubobjectData* NodeData = SelectedNodes.Last()->GetDataSource())
 			{
-				TSharedPtr<ISCSEditorCustomization> Customization = BlueprintEditorPtr.Pin()->CustomizeSubobjectEditor(SceneComp);
-				FVector CustomLocation;
-				if (Customization.IsValid() && Customization->HandleGetWidgetLocation(SceneComp, CustomLocation))
-				{
-					Location = CustomLocation;
-				}
-				else
+				const USceneComponent* SceneComp = Cast<USceneComponent>(NodeData->FindComponentInstanceInActor(PreviewActor));
+				if (SceneComp)
 				{
-					Location = SceneComp->GetComponentLocation();
+					TSharedPtr<ISCSEditorCustomization> Customization = BlueprintEditorPtr.Pin()->CustomizeSubobjectEditor(SceneComp);
+					FVector CustomLocation;
+					if (Customization.IsValid() && Customization->HandleGetWidgetLocation(SceneComp, CustomLocation))
+					{
+						Location = CustomLocation;
+					}
+					else
+					{
+						Location = SceneComp->GetComponentLocation();
+					}
 				}
 			}
 		}
