nDisplay: Fix nDisplay OCIO Profile customization crashing when used as a variable outside of the root actor. In this case default customization is used.

diff --git a/Engine/Plugins/Runtime/nDisplay/Source/DisplayClusterConfigurator/Private/Views/Details/DisplayClusterConfiguratorDetailCustomization.cpp b/Engine/Plugins/Runtime/nDisplay/Source/DisplayClusterConfigurator/Private/Views/Details/DisplayClusterConfiguratorDetailCustomization.cpp
--- a/Engine/Plugins/Runtime/nDisplay/Source/DisplayClusterConfigurator/Private/Views/Details/DisplayClusterConfiguratorDetailCustomization.cpp
+++ b/Engine/Plugins/Runtime/nDisplay/Source/DisplayClusterConfigurator/Private/Views/Details/DisplayClusterConfiguratorDetailCustomization.cpp
@@ -1424,8 +1424,26 @@ void FDisplayClusterConfiguratorOCIOProfileCustomization::CustomizeHeader(TShare
 {
 	FDisplayClusterConfiguratorTypeCustomization::CustomizeHeader(PropertyHandle, HeaderRow, CustomizationUtils);
 
+	ADisplayClusterRootActor* RootActor = FindRootActor();
+	FDisplayClusterConfiguratorBlueprintEditor* BPEditor = FDisplayClusterConfiguratorUtils::GetBlueprintEditorForObject(EditingObject);
+
+	if (RootActor == nullptr && BPEditor == nullptr)
+	{
+		HeaderRow.NameContent()
+		[
+			PropertyHandle->CreatePropertyNameWidget()
+		];
+		HeaderRow.ValueContent()
+		[
+			PropertyHandle->CreatePropertyValueWidget()
+		];
+
+		bIsDefaultDetailsDisplay = true;
+		return;
+	}
+	
 	Mode = FDisplayClusterConfiguratorNodeSelection::GetOperationModeFromProperty(PropertyHandle->GetProperty()->GetOwnerProperty());
-	NodeSelection = MakeShared<FDisplayClusterConfiguratorNodeSelection>(Mode, FindRootActor(), FDisplayClusterConfiguratorUtils::GetBlueprintEditorForObject(EditingObject));
+	NodeSelection = MakeShared<FDisplayClusterConfiguratorNodeSelection>(Mode, RootActor, BPEditor);
 	
 	FText ElementTooltip = FText::GetEmpty();
 
@@ -1449,6 +1467,19 @@ void FDisplayClusterConfiguratorOCIOProfileCustomization::CustomizeHeader(TShare
 void FDisplayClusterConfiguratorOCIOProfileCustomization::CustomizeChildren(TSharedRef<IPropertyHandle> PropertyHandle,
 	IDetailChildrenBuilder& ChildBuilder, IPropertyTypeCustomizationUtils& CustomizationUtils)
 {
+	if (bIsDefaultDetailsDisplay)
+	{
+		uint32 NumChildren;
+		PropertyHandle->GetNumChildren(NumChildren);
+
+		for (uint32 ChildIndex = 0; ChildIndex < NumChildren; ++ChildIndex)
+		{
+			const TSharedRef<IPropertyHandle> ChildHandle = PropertyHandle->GetChildHandle(ChildIndex).ToSharedRef();
+			ChildBuilder.AddProperty(ChildHandle);
+		}
+		return;
+	}
+	
 	GET_CHILD_PROPERTY_HANDLE(PropertyHandle, EnableOCIOHandle, FDisplayClusterConfigurationOCIOProfile, bIsEnabled);
 
 	GET_CHILD_PROPERTY_HANDLE(PropertyHandle, OCIOConfigurationHandle, FDisplayClusterConfigurationOCIOProfile, OCIOConfiguration);
diff --git a/Engine/Plugins/Runtime/nDisplay/Source/DisplayClusterConfigurator/Private/Views/Details/DisplayClusterConfiguratorDetailCustomization.h b/Engine/Plugins/Runtime/nDisplay/Source/DisplayClusterConfigurator/Private/Views/Details/DisplayClusterConfiguratorDetailCustomization.h
--- a/Engine/Plugins/Runtime/nDisplay/Source/DisplayClusterConfigurator/Private/Views/Details/DisplayClusterConfiguratorDetailCustomization.h
+++ b/Engine/Plugins/Runtime/nDisplay/Source/DisplayClusterConfigurator/Private/Views/Details/DisplayClusterConfiguratorDetailCustomization.h
@@ -458,6 +458,8 @@ protected:
 private:
 	TSharedPtr<FDisplayClusterConfiguratorNodeSelection> NodeSelection;
 	FDisplayClusterConfiguratorNodeSelection::EOperationMode Mode = FDisplayClusterConfiguratorNodeSelection::EOperationMode::Viewports;
+
+	bool bIsDefaultDetailsDisplay = false;
 };
 
 
