nDisplay: Fix viewing and modifying projection policies with multiple viewports selected. Fix individual parameter undo not working.

diff --git a/Engine/Plugins/Runtime/nDisplay/Source/DisplayClusterConfigurator/Private/Views/Details/DisplayClusterConfiguratorDetailCustomization.cpp b/Engine/Plugins/Runtime/nDisplay/Source/DisplayClusterConfigurator/Private/Views/Details/DisplayClusterConfiguratorDetailCustomization.cpp
--- a/Engine/Plugins/Runtime/nDisplay/Source/DisplayClusterConfigurator/Private/Views/Details/DisplayClusterConfiguratorDetailCustomization.cpp
+++ b/Engine/Plugins/Runtime/nDisplay/Source/DisplayClusterConfigurator/Private/Views/Details/DisplayClusterConfiguratorDetailCustomization.cpp
@@ -561,18 +561,17 @@ void FDisplayClusterConfiguratorTypeCustomization::CustomizeHeader(TSharedRef<IP
 {
 	PropertyUtilities = CustomizationUtils.GetPropertyUtilities();
 	
-	TArray<UObject*> OuterObjects;
-	InPropertyHandle->GetOuterObjects(OuterObjects);
-	if (OuterObjects.Num())
+	EditingObjects = PropertyUtilities.Pin()->GetSelectedObjects();
+	
+	if (EditingObjects.Num())
 	{
-		EditingObject = OuterObjects[0];
-		bMultipleObjectsSelected = OuterObjects.Num() > 1;
+		EditingObject = EditingObjects[0];
 	}
 }
 
 void FDisplayClusterConfiguratorTypeCustomization::RefreshBlueprint()
 {
-	if (FDisplayClusterConfiguratorBlueprintEditor* BlueprintEditor = FDisplayClusterConfiguratorUtils::GetBlueprintEditorForObject(EditingObject))
+	if (FDisplayClusterConfiguratorBlueprintEditor* BlueprintEditor = FDisplayClusterConfiguratorUtils::GetBlueprintEditorForObject(EditingObject.Get()))
 	{
 		BlueprintEditor->RefreshDisplayClusterPreviewActor();
 	}
@@ -580,7 +579,7 @@ void FDisplayClusterConfiguratorTypeCustomization::RefreshBlueprint()
 
 void FDisplayClusterConfiguratorTypeCustomization::ModifyBlueprint()
 {
-	if (UDisplayClusterBlueprint* Blueprint = FDisplayClusterConfiguratorUtils::FindBlueprintFromObject(EditingObject))
+	if (UDisplayClusterBlueprint* Blueprint = FDisplayClusterConfiguratorUtils::FindBlueprintFromObject(EditingObject.Get()))
 	{
 		FDisplayClusterConfiguratorUtils::MarkDisplayClusterBlueprintAsModified(Blueprint, false);
 	}
@@ -1425,7 +1424,7 @@ void FDisplayClusterConfiguratorOCIOProfileCustomization::CustomizeHeader(TShare
 	FDisplayClusterConfiguratorTypeCustomization::CustomizeHeader(PropertyHandle, HeaderRow, CustomizationUtils);
 
 	ADisplayClusterRootActor* RootActor = FindRootActor();
-	FDisplayClusterConfiguratorBlueprintEditor* BPEditor = FDisplayClusterConfiguratorUtils::GetBlueprintEditorForObject(EditingObject);
+	FDisplayClusterConfiguratorBlueprintEditor* BPEditor = FDisplayClusterConfiguratorUtils::GetBlueprintEditorForObject(EditingObject.Get());
 
 	if (RootActor == nullptr && BPEditor == nullptr)
 	{
@@ -1526,7 +1525,7 @@ void FDisplayClusterConfiguratorPerViewportColorGradingCustomization::CustomizeH
 	FDetailWidgetRow& HeaderRow, IPropertyTypeCustomizationUtils& CustomizationUtils)
 {
 	FDisplayClusterConfiguratorTypeCustomization::CustomizeHeader(PropertyHandle, HeaderRow, CustomizationUtils);
-	NodeSelection = MakeShared<FDisplayClusterConfiguratorNodeSelection>(FDisplayClusterConfiguratorNodeSelection::EOperationMode::Viewports, FindRootActor(), FDisplayClusterConfiguratorUtils::GetBlueprintEditorForObject(EditingObject));
+	NodeSelection = MakeShared<FDisplayClusterConfiguratorNodeSelection>(FDisplayClusterConfiguratorNodeSelection::EOperationMode::Viewports, FindRootActor(), FDisplayClusterConfiguratorUtils::GetBlueprintEditorForObject(EditingObject.Get()));
 
 	FText ElementTooltip = FText::GetEmpty();
 
@@ -1580,7 +1579,7 @@ void FDisplayClusterConfiguratorPerNodeColorGradingCustomization::CustomizeHeade
 	FDetailWidgetRow& HeaderRow, IPropertyTypeCustomizationUtils& CustomizationUtils)
 {
 	FDisplayClusterConfiguratorTypeCustomization::CustomizeHeader(PropertyHandle, HeaderRow, CustomizationUtils);
-	NodeSelection = MakeShared<FDisplayClusterConfiguratorNodeSelection>(FDisplayClusterConfiguratorNodeSelection::EOperationMode::ClusterNodes, FindRootActor(), FDisplayClusterConfiguratorUtils::GetBlueprintEditorForObject(EditingObject));
+	NodeSelection = MakeShared<FDisplayClusterConfiguratorNodeSelection>(FDisplayClusterConfiguratorNodeSelection::EOperationMode::ClusterNodes, FindRootActor(), FDisplayClusterConfiguratorUtils::GetBlueprintEditorForObject(EditingObject.Get()));
 
 	HeaderRow.NameContent()
 		[
diff --git a/Engine/Plugins/Runtime/nDisplay/Source/DisplayClusterConfigurator/Private/Views/Details/DisplayClusterConfiguratorDetailCustomization.h b/Engine/Plugins/Runtime/nDisplay/Source/DisplayClusterConfigurator/Private/Views/Details/DisplayClusterConfiguratorDetailCustomization.h
--- a/Engine/Plugins/Runtime/nDisplay/Source/DisplayClusterConfigurator/Private/Views/Details/DisplayClusterConfiguratorDetailCustomization.h
+++ b/Engine/Plugins/Runtime/nDisplay/Source/DisplayClusterConfigurator/Private/Views/Details/DisplayClusterConfiguratorDetailCustomization.h
@@ -210,9 +210,9 @@ protected:
 	ADisplayClusterRootActor* FindRootActor() const;
 
 protected:
-	UObject* EditingObject;
+	TWeakObjectPtr<UObject> EditingObject;
+	TArray<TWeakObjectPtr<UObject>> EditingObjects;
 	TWeakPtr<IPropertyUtilities> PropertyUtilities;
-	bool bMultipleObjectsSelected = false;
 };
 
 /**
diff --git a/Engine/Plugins/Runtime/nDisplay/Source/DisplayClusterConfigurator/Private/Views/Details/Policies/DisplayClusterConfiguratorPolicyDetailCustomization.cpp b/Engine/Plugins/Runtime/nDisplay/Source/DisplayClusterConfigurator/Private/Views/Details/Policies/DisplayClusterConfiguratorPolicyDetailCustomization.cpp
--- a/Engine/Plugins/Runtime/nDisplay/Source/DisplayClusterConfigurator/Private/Views/Details/Policies/DisplayClusterConfiguratorPolicyDetailCustomization.cpp
+++ b/Engine/Plugins/Runtime/nDisplay/Source/DisplayClusterConfigurator/Private/Views/Details/Policies/DisplayClusterConfiguratorPolicyDetailCustomization.cpp
@@ -41,22 +41,23 @@
 
 void FDisplayClusterConfiguratorProjectionCustomization::CustomizeHeader(TSharedRef<IPropertyHandle> InPropertyHandle, FDetailWidgetRow& InHeaderRow, IPropertyTypeCustomizationUtils& CustomizationUtils)
 {
-	if (bMultipleObjectsSelected)
-	{
-		// Do not currently support multiple objects being edited.
-		return;
-	}
-	
 	FDisplayClusterConfiguratorPolymorphicEntityCustomization::CustomizeHeader(InPropertyHandle, InHeaderRow, CustomizationUtils);
 	CustomOption = MakeShared<FString>("Custom");
 
 	// Get the Editing object
 	ConfigurationViewportPtr = CastChecked<UDisplayClusterConfigurationViewport>(EditingObject);
 
+	for (const TWeakObjectPtr<UObject>& Object : EditingObjects)
+	{
+		UDisplayClusterConfigurationViewport* Viewport = CastChecked<UDisplayClusterConfigurationViewport>(Object.Get());
+		ConfigurationViewports.Add(Viewport);
+	}
+	
 	// Store what's currently selected.
 	CurrentSelectedPolicy = ConfigurationViewportPtr->ProjectionPolicy.Type;
-	
-	bIsCustomPolicy = IsCustomTypeInConfig();
+
+	const bool bRequireCustomPolicy = true;
+	bIsCustomPolicy = IsCustomTypeInConfig() && IsPolicyIdenticalAcrossEditedObjects(bRequireCustomPolicy);
 	if (bIsCustomPolicy)
 	{
 		// Load default config
@@ -66,12 +67,6 @@ void FDisplayClusterConfiguratorProjectionCustomization::CustomizeHeader(TShared
 
 void FDisplayClusterConfiguratorProjectionCustomization::CustomizeChildren(TSharedRef<IPropertyHandle> InPropertyHandle, IDetailChildrenBuilder& InChildBuilder, IPropertyTypeCustomizationUtils& CustomizationUtils)
 {
-	if (bMultipleObjectsSelected)
-	{
-		// Do not currently support multiple objects being edited.
-		return;
-	}
-	
 	FDisplayClusterConfiguratorPolymorphicEntityCustomization::CustomizeChildren(InPropertyHandle, InChildBuilder, CustomizationUtils);
 
 	// Hide properties 
@@ -99,7 +94,7 @@ TSharedRef<SWidget> FDisplayClusterConfiguratorProjectionCustomization::MakeProj
 
 EVisibility FDisplayClusterConfiguratorProjectionCustomization::GetCustomRowsVisibility() const
 {
-	return bIsCustomPolicy ? EVisibility::Visible :  EVisibility::Collapsed;
+	return bIsCustomPolicy && IsPolicyIdenticalAcrossEditedObjects() ? EVisibility::Visible :  EVisibility::Collapsed;
 }
 
 void FDisplayClusterConfiguratorProjectionCustomization::ResetProjectionPolicyOptions()
@@ -192,7 +187,7 @@ void FDisplayClusterConfiguratorProjectionCustomization::OnProjectionPolicySelec
 		ConfigurationViewport->Modify();
 		ModifyBlueprint();
 		
-		if (SelectedPolicy.Equals(*CustomOption.Get()))
+		if (SelectedPolicy.Equals(*CustomOption.Get()) && IsPolicyIdenticalAcrossEditedObjects())
 		{
 			bIsCustomPolicy = true;
 			CustomPolicy = ConfigurationViewport->ProjectionPolicy.Type;
@@ -225,7 +220,12 @@ void FDisplayClusterConfiguratorProjectionCustomization::OnProjectionPolicySelec
 
 FText FDisplayClusterConfiguratorProjectionCustomization::GetSelectedProjectionPolicyText() const
 {
-	return FText::FromString(GetCurrentPolicy());
+	if (IsPolicyIdenticalAcrossEditedObjects())
+	{
+		return FText::FromString(GetCurrentPolicy());
+	}
+
+	return LOCTEXT("MultipleValues", "Multiple Values");
 }
 
 FText FDisplayClusterConfiguratorProjectionCustomization::GetCustomPolicyText() const
@@ -296,11 +296,32 @@ void FDisplayClusterConfiguratorProjectionCustomization::OnTextCommittedInCustom
 	*/
 }
 
+bool FDisplayClusterConfiguratorProjectionCustomization::IsPolicyIdenticalAcrossEditedObjects(bool bRequireCustomPolicy) const
+{
+	for (const TWeakObjectPtr<UDisplayClusterConfigurationViewport>& Viewport : ConfigurationViewports)
+	{
+		if (Viewport.IsValid() &&
+			Viewport->ProjectionPolicy.Type != CurrentSelectedPolicy ||
+			bRequireCustomPolicy && !Viewport->ProjectionPolicy.bIsCustom ||
+			!bRequireCustomPolicy && Viewport->ProjectionPolicy.bIsCustom != bIsCustomPolicy)
+		{
+			return false;
+		}
+	}
+
+	return true;
+}
+
 void FDisplayClusterConfiguratorProjectionCustomization::BuildParametersForPolicy(const FString& Policy)
 {
 	CustomPolicyParameters.Reset();
 
-	UDisplayClusterBlueprint* Blueprint = FDisplayClusterConfiguratorUtils::FindBlueprintFromObject(EditingObject);
+	if (!IsPolicyIdenticalAcrossEditedObjects())
+	{
+		return;
+	}
+	
+	UDisplayClusterBlueprint* Blueprint = FDisplayClusterConfiguratorUtils::FindBlueprintFromObject(EditingObject.Get());
 	check(Blueprint);
 
 	const FString PolicyLower = Policy.ToLower();
@@ -374,8 +395,7 @@ void FDisplayClusterConfiguratorProjectionCustomization::CreateSimplePolicy(UDis
 		"Screen",
 		DisplayClusterProjectionStrings::cfg::simple::Screen,
 		Blueprint,
-		ConfigurationViewportPtr.Get(),
-		ParametersHandle,
+		ConfigurationViewports,
 		TArray<TSubclassOf<UActorComponent>>{ UDisplayClusterScreenComponent::StaticClass() });
 
 	ScreenCombo->SetParameterTooltip(LOCTEXT("SimplePolicyScreenTooltip", "Target Screen or Display"));
@@ -389,16 +409,14 @@ void FDisplayClusterConfiguratorProjectionCustomization::CreateCameraPolicy(UDis
 		"Camera",
 		DisplayClusterProjectionStrings::cfg::camera::Component,
 		Blueprint,
-		ConfigurationViewportPtr.Get(),
-		ParametersHandle,
+		ConfigurationViewports,
 		TArray<TSubclassOf<UActorComponent>>{ UCameraComponent::StaticClass() } ));
 
 	CustomPolicyParameters.Add(MakeShared<FPolicyParameterInfoBool>(
 		"Use nDisplay Renderer",
 		DisplayClusterProjectionStrings::cfg::camera::Native,
 		Blueprint,
-		ConfigurationViewportPtr.Get(),
-		ParametersHandle,
+		ConfigurationViewports,
 		true));
 }
 
@@ -408,8 +426,7 @@ void FDisplayClusterConfiguratorProjectionCustomization::CreateMeshPolicy(UDispl
 		"Mesh",
 		DisplayClusterProjectionStrings::cfg::mesh::Component,
 		Blueprint,
-		ConfigurationViewportPtr.Get(),
-		ParametersHandle,
+		ConfigurationViewports,
 		TArray<TSubclassOf<UActorComponent>>{ UStaticMeshComponent::StaticClass() }));
 }
 
@@ -419,24 +436,21 @@ void FDisplayClusterConfiguratorProjectionCustomization::CreateDomePolicy(UDispl
 		"File",
 		DisplayClusterProjectionStrings::cfg::domeprojection::File,
 		Blueprint,
-		ConfigurationViewportPtr.Get(),
-		ParametersHandle,
+		ConfigurationViewports,
 		TArray<FString>{"xml"}));
 
 	CustomPolicyParameters.Add(MakeShared<FPolicyParameterInfoComponentCombo>(
 		"Origin",
 		DisplayClusterProjectionStrings::cfg::domeprojection::Origin,
 		Blueprint,
-		ConfigurationViewportPtr.Get(),
-		ParametersHandle,
+		ConfigurationViewports,
 		TArray<TSubclassOf<UActorComponent>>{ USceneComponent::StaticClass() }));
 
 	CustomPolicyParameters.Add(MakeShared<FPolicyParameterInfoNumber<int32>>(
 		"Channel",
 		DisplayClusterProjectionStrings::cfg::domeprojection::Channel,
 		Blueprint,
-		ConfigurationViewportPtr.Get(),
-		ParametersHandle));
+		ConfigurationViewports));
 }
 
 void FDisplayClusterConfiguratorProjectionCustomization::CreateVIOSOPolicy(UDisplayClusterBlueprint* Blueprint)
@@ -445,24 +459,21 @@ void FDisplayClusterConfiguratorProjectionCustomization::CreateVIOSOPolicy(UDisp
 		"File",
 		DisplayClusterProjectionStrings::cfg::VIOSO::File,
 		Blueprint,
-		ConfigurationViewportPtr.Get(),
-		ParametersHandle,
+		ConfigurationViewports,
 		TArray<FString>{"vwf"}));
 
 	CustomPolicyParameters.Add(MakeShared<FPolicyParameterInfoComponentCombo>(
 		"Origin",
 		DisplayClusterProjectionStrings::cfg::VIOSO::Origin,
 		Blueprint,
-		ConfigurationViewportPtr.Get(),
-		ParametersHandle,
+		ConfigurationViewports,
 		TArray<TSubclassOf<UActorComponent>>{ USceneComponent::StaticClass() }));
 
 	CustomPolicyParameters.Add(MakeShared<FPolicyParameterInfo4x4Matrix>(
 		"Matrix",
 		DisplayClusterProjectionStrings::cfg::VIOSO::BaseMatrix,
 		Blueprint,
-		ConfigurationViewportPtr.Get(),
-		ParametersHandle));
+		ConfigurationViewports));
 }
 
 void FDisplayClusterConfiguratorProjectionCustomization::CreateEasyBlendPolicy(UDisplayClusterBlueprint* Blueprint)
@@ -471,24 +482,21 @@ void FDisplayClusterConfiguratorProjectionCustomization::CreateEasyBlendPolicy(U
 		"File",
 		DisplayClusterProjectionStrings::cfg::easyblend::File,
 		Blueprint,
-		ConfigurationViewportPtr.Get(),
-		ParametersHandle,
+		ConfigurationViewports,
 		TArray<FString>{"pol*", "ol*"}));
 
 	CustomPolicyParameters.Add(MakeShared<FPolicyParameterInfoComponentCombo>(
 		"Origin",
 		DisplayClusterProjectionStrings::cfg::easyblend::Origin,
 		Blueprint,
-		ConfigurationViewportPtr.Get(),
-		ParametersHandle,
+		ConfigurationViewports,
 		TArray<TSubclassOf<UActorComponent>>{ USceneComponent::StaticClass() }));
 
 	CustomPolicyParameters.Add(MakeShared<FPolicyParameterInfoNumber<float>>(
 		"Scale",
 		DisplayClusterProjectionStrings::cfg::easyblend::Scale,
 		Blueprint,
-		ConfigurationViewportPtr.Get(),
-		ParametersHandle));
+		ConfigurationViewports));
 }
 
 void FDisplayClusterConfiguratorProjectionCustomization::CreateManualPolicy(UDisplayClusterBlueprint* Blueprint)
@@ -515,8 +523,7 @@ void FDisplayClusterConfiguratorProjectionCustomization::CreateManualPolicy(UDis
 		"Rendering",
 		RenderingKey,
 		Blueprint,
-		ConfigurationViewportPtr.Get(),
-		ParametersHandle,
+		ConfigurationViewports,
 		//TArray<FString>{RenderingMono, RenderingStereo, RenderingMonoStereo}, temporarily disabled MonoStereo, not supported  implementation from projection policy side
 		TArray<FString>{RenderingMono, RenderingStereo},
 		& RenderingMono,
@@ -528,8 +535,7 @@ void FDisplayClusterConfiguratorProjectionCustomization::CreateManualPolicy(UDis
 		"Frustum",
 		FrustumKey,
 		Blueprint,
-		ConfigurationViewportPtr.Get(),
-		ParametersHandle,
+		ConfigurationViewports,
 		TArray<FString>{FrustumMatrix, FrustumAngles},
 		& FrustumMatrix,
 		bSort);
@@ -544,8 +550,7 @@ void FDisplayClusterConfiguratorProjectionCustomization::CreateManualPolicy(UDis
 			"Rotation",
 			DisplayClusterProjectionStrings::cfg::manual::Rotation,
 			Blueprint,
-			ConfigurationViewportPtr.Get(),
-			ParametersHandle));
+			ConfigurationViewports));
 	}
 	
 	/*
@@ -569,8 +574,7 @@ void FDisplayClusterConfiguratorProjectionCustomization::CreateManualPolicy(UDis
 				"Matrix",
 				DisplayClusterProjectionStrings::cfg::manual::Matrix,
 				Blueprint,
-				ConfigurationViewportPtr.Get(),
-				ParametersHandle);
+				ConfigurationViewports);
 			CustomPolicyParameters.Add(MatrixPolicy);
 		}
 	}
@@ -593,8 +597,7 @@ void FDisplayClusterConfiguratorProjectionCustomization::CreateManualPolicy(UDis
 				"MatrixLeft",
 				DisplayClusterProjectionStrings::cfg::manual::MatrixLeft,
 				Blueprint,
-				ConfigurationViewportPtr.Get(),
-				ParametersHandle);
+				ConfigurationViewports);
 			CustomPolicyParameters.Add(MatrixLeftPolicy);
 
 			const TSharedPtr<FPolicyParameterInfo4x4Matrix> MatrixRightPolicy =
@@ -602,8 +605,7 @@ void FDisplayClusterConfiguratorProjectionCustomization::CreateManualPolicy(UDis
 					"MatrixRight",
 					DisplayClusterProjectionStrings::cfg::manual::MatrixRight,
 					Blueprint,
-					ConfigurationViewportPtr.Get(),
-					ParametersHandle);
+					ConfigurationViewports);
 			CustomPolicyParameters.Add(MatrixRightPolicy);
 		}
 	}
@@ -627,8 +629,7 @@ void FDisplayClusterConfiguratorProjectionCustomization::CreateManualPolicy(UDis
 				"Frustum",
 				DisplayClusterProjectionStrings::cfg::manual::Frustum,
 				Blueprint,
-				ConfigurationViewportPtr.Get(),
-				ParametersHandle));
+				ConfigurationViewports));
 		}
 
 		auto IsLeftRightFrustumVisible = [RenderingCombo, RenderingStereo, RenderingMonoStereo, FrustumCombo, FrustumAngles]() -> bool
@@ -646,15 +647,13 @@ void FDisplayClusterConfiguratorProjectionCustomization::CreateManualPolicy(UDis
 				"FrustumLeft",
 				DisplayClusterProjectionStrings::cfg::manual::FrustumLeft,
 				Blueprint,
-				ConfigurationViewportPtr.Get(),
-				ParametersHandle));
+				ConfigurationViewports));
 
 			CustomPolicyParameters.Add(MakeShared<FPolicyParameterInfoFrustumAngle>(
 				"FrustumRight",
 				DisplayClusterProjectionStrings::cfg::manual::FrustumRight,
 				Blueprint,
-				ConfigurationViewportPtr.Get(),
-				ParametersHandle));
+				ConfigurationViewports));
 		}
 	}
 }
@@ -678,14 +677,13 @@ void FDisplayClusterConfiguratorProjectionCustomization::CreateMPCDIPolicy(UDisp
 		"MPCDI Type",
 		MPCDITypeKey,
 		Blueprint,
-		ConfigurationViewportPtr.Get(),
-		ParametersHandle,
+		ConfigurationViewports,
 		TArray<FString>{TypeMPCDI, TypePFM},
 		&TypeMPCDI,
 		bSort);
 	MPCDICombo->SetOnSelectedDelegate(FPolicyParameterInfoCombo::FOnItemSelected::CreateLambda(RefreshPolicy));
 	CustomPolicyParameters.Add(MPCDICombo);
-	
+
 	const FString Setting = MPCDICombo->GetOrAddCustomParameterValueText().ToString();
 	if (Setting == TypeMPCDI)
 	{
@@ -693,23 +691,20 @@ void FDisplayClusterConfiguratorProjectionCustomization::CreateMPCDIPolicy(UDisp
 			"File",
 			DisplayClusterProjectionStrings::cfg::mpcdi::File,
 			Blueprint,
-			ConfigurationViewportPtr.Get(),
-			ParametersHandle,
+			ConfigurationViewports,
 			TArray<FString>{"mpcdi"}));
 
 		CustomPolicyParameters.Add(MakeShared<FPolicyParameterInfoText>(
 			"Buffer",
 			DisplayClusterProjectionStrings::cfg::mpcdi::Buffer,
 			Blueprint,
-			ConfigurationViewportPtr.Get(),
-			ParametersHandle));
+			ConfigurationViewports));
 
 		CustomPolicyParameters.Add(MakeShared<FPolicyParameterInfoText>(
 			"Region",
 			DisplayClusterProjectionStrings::cfg::mpcdi::Region,
 			Blueprint,
-			ConfigurationViewportPtr.Get(),
-			ParametersHandle));
+			ConfigurationViewports));
 	}
 	else if (Setting == TypePFM)
 	{
@@ -717,64 +712,56 @@ void FDisplayClusterConfiguratorProjectionCustomization::CreateMPCDIPolicy(UDisp
 			"File",
 			DisplayClusterProjectionStrings::cfg::mpcdi::FilePFM,
 			Blueprint,
-			ConfigurationViewportPtr.Get(),
-			ParametersHandle,
+			ConfigurationViewports,
 			TArray<FString>{"pfm"}));
 
 		CustomPolicyParameters.Add(MakeShared<FPolicyParameterInfoFile>(
 			"Alpha Mask",
 			DisplayClusterProjectionStrings::cfg::mpcdi::FileAlpha,
 			Blueprint,
-			ConfigurationViewportPtr.Get(),
-			ParametersHandle,
+			ConfigurationViewports,
 			TArray<FString>{"png"}));
 		
 		CustomPolicyParameters.Add(MakeShared<FPolicyParameterInfoNumber<float>>(
 			"Alpha Gamma",
 			DisplayClusterProjectionStrings::cfg::mpcdi::AlphaGamma,
 			Blueprint,
-			ConfigurationViewportPtr.Get(),
-			ParametersHandle,
+			ConfigurationViewports,
 			1.f));
 
 		CustomPolicyParameters.Add(MakeShared<FPolicyParameterInfoFile>(
 			"Beta Mask",
 			DisplayClusterProjectionStrings::cfg::mpcdi::FileBeta,
 			Blueprint,
-			ConfigurationViewportPtr.Get(),
-			ParametersHandle,
+			ConfigurationViewports,
 			TArray<FString>{"png"}));
 
 		CustomPolicyParameters.Add(MakeShared<FPolicyParameterInfoNumber<float>>(
 			"Scale",
 			DisplayClusterProjectionStrings::cfg::mpcdi::WorldScale,
 			Blueprint,
-			ConfigurationViewportPtr.Get(), 
-			ParametersHandle,
+			ConfigurationViewports,
 			1.f));
 		
 		CustomPolicyParameters.Add(MakeShared<FPolicyParameterInfoBool>(
 			"Use Unreal Axis",
 			DisplayClusterProjectionStrings::cfg::mpcdi::UseUnrealAxis,
 			Blueprint,
-			ConfigurationViewportPtr.Get(),
-			ParametersHandle));
+			ConfigurationViewports));
 	}
 	
 	TSharedPtr<FPolicyParameterInfoCombo> Origin = MakeShared<FPolicyParameterInfoComponentCombo>(
 		"Origin",
 		DisplayClusterProjectionStrings::cfg::mpcdi::Origin,
 		Blueprint,
-		ConfigurationViewportPtr.Get(),
-		ParametersHandle,
+		ConfigurationViewports,
 		TArray<TSubclassOf<UActorComponent>>{ USceneComponent::StaticClass() });
 
 	CustomPolicyParameters.Add(MakeShared<FPolicyParameterInfoBool>(
 		"Enable Preview",
 		DisplayClusterProjectionStrings::cfg::mpcdi::EnablePreview,
 		Blueprint,
-		ConfigurationViewportPtr.Get(),
-		ParametersHandle));
+		ConfigurationViewports));
 }
 
 #undef LOCTEXT_NAMESPACE
diff --git a/Engine/Plugins/Runtime/nDisplay/Source/DisplayClusterConfigurator/Private/Views/Details/Policies/DisplayClusterConfiguratorPolicyDetailCustomization.h b/Engine/Plugins/Runtime/nDisplay/Source/DisplayClusterConfigurator/Private/Views/Details/Policies/DisplayClusterConfiguratorPolicyDetailCustomization.h
--- a/Engine/Plugins/Runtime/nDisplay/Source/DisplayClusterConfigurator/Private/Views/Details/Policies/DisplayClusterConfiguratorPolicyDetailCustomization.h
+++ b/Engine/Plugins/Runtime/nDisplay/Source/DisplayClusterConfigurator/Private/Views/Details/Policies/DisplayClusterConfiguratorPolicyDetailCustomization.h
@@ -41,7 +41,7 @@ protected:
 	const FString& GetCurrentPolicy() const;
 	bool IsCustomTypeInConfig() const;
 	void OnTextCommittedInCustomPolicyText(const FText& InValue, ETextCommit::Type CommitType);
-
+	bool IsPolicyIdenticalAcrossEditedObjects(bool bRequireCustomPolicy = false) const;
 private:
 	TSharedPtr<FString>	CustomOption;
 	TArray< TSharedPtr<FString> > ProjectionPolicyOptions;
@@ -72,5 +72,6 @@ private:
 	// End Custom Parameters Selection
 	
 private:
+	TArray<TWeakObjectPtr<UDisplayClusterConfigurationViewport>> ConfigurationViewports;
 	TWeakObjectPtr<UDisplayClusterConfigurationViewport> ConfigurationViewportPtr;
 };
\ No newline at end of file
diff --git a/Engine/Plugins/Runtime/nDisplay/Source/DisplayClusterConfigurator/Private/Views/Details/Policies/DisplayClusterConfiguratorPolicyParameterCustomization.cpp b/Engine/Plugins/Runtime/nDisplay/Source/DisplayClusterConfigurator/Private/Views/Details/Policies/DisplayClusterConfiguratorPolicyParameterCustomization.cpp
--- a/Engine/Plugins/Runtime/nDisplay/Source/DisplayClusterConfigurator/Private/Views/Details/Policies/DisplayClusterConfiguratorPolicyParameterCustomization.cpp
+++ b/Engine/Plugins/Runtime/nDisplay/Source/DisplayClusterConfigurator/Private/Views/Details/Policies/DisplayClusterConfiguratorPolicyParameterCustomization.cpp
@@ -24,6 +24,7 @@
 #include "PropertyHandle.h"
 #include "Kismet2/BlueprintEditorUtils.h"
 #include "Components/StaticMeshComponent.h"
+#include "ScopedTransaction.h"
 
 #include "Widgets/Input/SEditableTextBox.h"
 #include "Widgets/Input/SNumericEntryBox.h"
@@ -42,16 +43,15 @@ FPolicyParameterInfo::FPolicyParameterInfo(
 	const FString& InDisplayName,
 	const FString& InKey,
 	UDisplayClusterBlueprint* InBlueprint,
-	UDisplayClusterConfigurationViewport* InConfigurationViewport,
-	const TSharedPtr<IPropertyHandle>& InParametersHandle,
+	const TArray<TWeakObjectPtr<UDisplayClusterConfigurationViewport>>& InConfigurationViewports,
 	const FString* InInitialValue)
 {
 	ParamDisplayName = InDisplayName;
 	ParamTooltip = FText::GetEmpty();
 	ParamKey = InKey;
 	BlueprintOwnerPtr = InBlueprint;
-	ConfigurationViewportPtr = InConfigurationViewport;
-	ParametersHandle = InParametersHandle;
+	ConfigurationViewports = InConfigurationViewports;
+	
 	BlueprintEditorPtrCached = FDisplayClusterConfiguratorUtils::GetBlueprintEditorForObject(InBlueprint);
 
 	if (InInitialValue)
@@ -67,48 +67,125 @@ void FPolicyParameterInfo::SetParameterVisibilityDelegate(FParameterVisible InDe
 
 FText FPolicyParameterInfo::GetOrAddCustomParameterValueText() const
 {
-	UDisplayClusterConfigurationViewport* ConfigurationViewport = ConfigurationViewportPtr.Get();
-	check(ConfigurationViewport != nullptr);
-
-	FString* ParameterValue = ConfigurationViewport->ProjectionPolicy.Parameters.Find(GetParameterKey());
-	if (ParameterValue == nullptr)
+	if (!DoParametersMatchForAllViewports())
 	{
-		UpdateCustomParameterValueText(InitialValue.IsValid() ? *InitialValue : TEXT(""), false);
+		return LOCTEXT("MultipleValues", "Multiple Values");
+	}
+	
+	FString* ParameterValue = nullptr;
+	for (const TWeakObjectPtr<UDisplayClusterConfigurationViewport>& Viewport : ConfigurationViewports)
+	{
+		UDisplayClusterConfigurationViewport* ConfigurationViewport = Viewport.Get();
+		check(ConfigurationViewport != nullptr);
+
 		ParameterValue = ConfigurationViewport->ProjectionPolicy.Parameters.Find(GetParameterKey());
+		if (ParameterValue == nullptr)
+		{
+			UpdateCustomParameterValueText(InitialValue.IsValid() ? *InitialValue : TEXT(""), false);
+			ParameterValue = ConfigurationViewport->ProjectionPolicy.Parameters.Find(GetParameterKey());
+		}
+		check(ParameterValue)
+
+		// Only need values from one viewport since they match.
+		break;
 	}
-	check(ParameterValue)
 	
 	return FText::FromString(*ParameterValue);
 }
 
 bool FPolicyParameterInfo::IsParameterAlreadyAdded() const
 {
-	UDisplayClusterConfigurationViewport* ConfigurationViewport = ConfigurationViewportPtr.Get();
-	check(ConfigurationViewport != nullptr);
+	for (const TWeakObjectPtr<UDisplayClusterConfigurationViewport>& Viewport : ConfigurationViewports)
+	{
+		UDisplayClusterConfigurationViewport* ConfigurationViewport = Viewport.Get();
+		check(ConfigurationViewport != nullptr);
+
+		if (ConfigurationViewport->ProjectionPolicy.Parameters.Contains(GetParameterKey()))
+		{
+			return true;
+		}
+	}
 
-	return ConfigurationViewport->ProjectionPolicy.Parameters.Contains(GetParameterKey());
+	return false;
 }
 
-void FPolicyParameterInfo::UpdateCustomParameterValueText(const FString& NewValue, bool bNotify) const
+bool FPolicyParameterInfo::DoParametersMatchForAllViewports() const
 {
-	UDisplayClusterConfigurationViewport* ConfigurationViewport = ConfigurationViewportPtr.Get();
-	check(ConfigurationViewport != nullptr);
-	ConfigurationViewport->Modify();
+	FString LastParameterValue;
+	bool bFirstParamFound = false;
+	for (const TWeakObjectPtr<UDisplayClusterConfigurationViewport>& Viewport : ConfigurationViewports)
+	{
+		if (!Viewport.IsValid())
+		{
+			continue;
+		}
+
+		if (FString* ParameterValue = Viewport->ProjectionPolicy.Parameters.Find(GetParameterKey()))
+		{
+			if (bFirstParamFound && LastParameterValue != *ParameterValue)
+			{
+				return false;
+			}
 
-	FStructProperty* StructProperty = FindFProperty<FStructProperty>(ConfigurationViewport->GetClass(), GET_MEMBER_NAME_CHECKED(UDisplayClusterConfigurationViewport, ProjectionPolicy));
-	check(StructProperty);
+			LastParameterValue = *ParameterValue;
+		}
 
-	uint8* MapContainer = StructProperty->ContainerPtrToValuePtr<uint8>(ConfigurationViewport);
-	DisplayClusterConfiguratorPropertyUtils::AddKeyValueToMap(MapContainer, ParametersHandle, GetParameterKey(), NewValue);
+		bFirstParamFound = true;
+	}
 
-	if (bNotify)
+	return true;
+}
+
+void FPolicyParameterInfo::UpdateCustomParameterValueText(const FString& NewValue, bool bNotify) const
+{
+	bool bHasChanged = false;
+	FScopedTransaction Transaction(LOCTEXT("UpdateParameterText", "Update Parameter"));
+	for (const TWeakObjectPtr<UDisplayClusterConfigurationViewport>& Viewport : ConfigurationViewports)
 	{
-		if (FDisplayClusterConfiguratorBlueprintEditor* BlueprintEditor = FDisplayClusterConfiguratorUtils::GetBlueprintEditorForObject(BlueprintOwnerPtr.Get()))
+		UDisplayClusterConfigurationViewport* ConfigurationViewport = Viewport.Get();
+		check(ConfigurationViewport != nullptr);
+
+		const FString& Key = GetParameterKey();
+		
+		if (const FString *OldValue = ConfigurationViewport->ProjectionPolicy.Parameters.Find(Key))
+		{
+			if (*OldValue == NewValue)
+			{
+				continue;
+			}
+		}
+
+		bHasChanged = true;
+		ConfigurationViewport->Modify();
+
+		FStructProperty* StructProperty = FindFProperty<FStructProperty>(ConfigurationViewport->GetClass(), GET_MEMBER_NAME_CHECKED(UDisplayClusterConfigurationViewport, ProjectionPolicy));
+		check(StructProperty);
+
+		const TSharedPtr<ISinglePropertyView> ProjectionPolicyView = DisplayClusterConfiguratorPropertyUtils::GetPropertyView(
+			ConfigurationViewport, GET_MEMBER_NAME_CHECKED(UDisplayClusterConfigurationViewport, ProjectionPolicy));
+
+		check(ProjectionPolicyView.IsValid());
+
+		TSharedPtr<IPropertyHandle> PropertyHandle = ProjectionPolicyView->GetPropertyHandle()->GetChildHandle(GET_MEMBER_NAME_CHECKED(FDisplayClusterConfigurationPolymorphicEntity, Parameters));
+		check(PropertyHandle.IsValid());
+		
+		uint8* MapContainer = StructProperty->ContainerPtrToValuePtr<uint8>(ConfigurationViewport);
+		DisplayClusterConfiguratorPropertyUtils::AddKeyValueToMap(MapContainer, PropertyHandle, Key, NewValue);
+
+		if (bNotify)
 		{
-			BlueprintEditor->ClusterChanged(true);
-			BlueprintEditor->RefreshDisplayClusterPreviewActor();
+			if (FDisplayClusterConfiguratorBlueprintEditor* BlueprintEditor = FDisplayClusterConfiguratorUtils::GetBlueprintEditorForObject(BlueprintOwnerPtr.Get()))
+			{
+				BlueprintEditor->ClusterChanged(true);
+				BlueprintEditor->RefreshDisplayClusterPreviewActor();
+			}
 		}
 	}
+
+	if (!bHasChanged)
+	{
+		Transaction.Cancel();
+	}
 }
 
 EVisibility FPolicyParameterInfo::IsParameterVisible() const
@@ -117,9 +194,8 @@ EVisibility FPolicyParameterInfo::IsParameterVisible() const
 }
 
 FPolicyParameterInfoCombo::FPolicyParameterInfoCombo(const FString& InDisplayName, const FString& InKey,
-	UDisplayClusterBlueprint* InBlueprint, UDisplayClusterConfigurationViewport* InConfigurationViewport,
-	const TSharedPtr<IPropertyHandle>& InParametersHandle,
-	const TArray<FString>& InValues, const FString* InInitialItem, bool bSort) : FPolicyParameterInfo(InDisplayName, InKey, InBlueprint, InConfigurationViewport, InParametersHandle, InInitialItem)
+													UDisplayClusterBlueprint* InBlueprint,const TArray<TWeakObjectPtr<UDisplayClusterConfigurationViewport>>& InConfigurationViewports,
+													const TArray<FString>& InValues, const FString* InInitialItem, bool bSort) : FPolicyParameterInfo(InDisplayName, InKey, InBlueprint, InConfigurationViewports, InInitialItem)
 {
 	for (const FString& Value : InValues)
 	{
@@ -185,9 +261,8 @@ FPolicyParameterInfoComponentCombo::FPolicyParameterInfoComponentCombo(
 	const FString& InDisplayName,
 	const FString& InKey,
 	UDisplayClusterBlueprint* InBlueprint,
-	UDisplayClusterConfigurationViewport* InConfigurationViewport,
-	const TSharedPtr<IPropertyHandle>& InParametersHandle,
-	const TArray<TSubclassOf<UActorComponent>>& InComponentClasses) : FPolicyParameterInfoCombo(InDisplayName, InKey, InBlueprint, InConfigurationViewport, InParametersHandle, {}, nullptr)
+	const TArray<TWeakObjectPtr<UDisplayClusterConfigurationViewport>>& InConfigurationViewports,
+	const TArray<TSubclassOf<UActorComponent>>& InComponentClasses) : FPolicyParameterInfoCombo(InDisplayName, InKey, InBlueprint, InConfigurationViewports, {}, nullptr)
 {
 	ComponentTypes = InComponentClasses;
 	if (BlueprintEditorPtrCached)
@@ -229,9 +304,8 @@ void FPolicyParameterInfoComponentCombo::CreateParameterValues(ADisplayClusterRo
 
 FPolicyParameterInfoText::FPolicyParameterInfoText(const FString& InDisplayName, const FString& InKey,
 	UDisplayClusterBlueprint* InBlueprint,
-	UDisplayClusterConfigurationViewport* InConfigurationViewport,
-	const TSharedPtr<IPropertyHandle>& InParametersHandle) :
-	FPolicyParameterInfo(InDisplayName, InKey, InBlueprint, InConfigurationViewport, InParametersHandle)
+	const TArray<TWeakObjectPtr<UDisplayClusterConfigurationViewport>>& InConfigurationViewports) :
+	FPolicyParameterInfo(InDisplayName, InKey, InBlueprint, InConfigurationViewports)
 {
 }
 
@@ -257,9 +331,8 @@ void FPolicyParameterInfoText::CreateCustomRowWidget(IDetailChildrenBuilder& InD
 }
 
 FPolicyParameterInfoBool::FPolicyParameterInfoBool(const FString& InDisplayName, const FString& InKey,
-	UDisplayClusterBlueprint* InBlueprint, UDisplayClusterConfigurationViewport* InConfigurationViewport,
-	const TSharedPtr<IPropertyHandle>& InParametersHandle, bool bInvertValue) :
-	FPolicyParameterInfo(InDisplayName, InKey, InBlueprint, InConfigurationViewport, InParametersHandle),
+	UDisplayClusterBlueprint* InBlueprint, const TArray<TWeakObjectPtr<UDisplayClusterConfigurationViewport>>& InConfigurationViewports, bool bInvertValue) :
+	FPolicyParameterInfo(InDisplayName, InKey, InBlueprint, InConfigurationViewports),
 	bInvertValue(bInvertValue)
 {
 }
@@ -287,6 +360,11 @@ void FPolicyParameterInfoBool::CreateCustomRowWidget(IDetailChildrenBuilder& InD
 
 ECheckBoxState FPolicyParameterInfoBool::IsChecked() const
 {
+	if (!DoParametersMatchForAllViewports())
+	{
+		return ECheckBoxState::Undetermined;	
+	}
+	
 	const FString StrValue = GetOrAddCustomParameterValueText().ToString().ToLower();
 	const bool bValue = DisplayClusterTypesConverter::template FromString<bool>(StrValue);
 	const bool bIsChecked = bInvertValue ? !bValue : bValue;
@@ -395,7 +473,7 @@ FReply FPolicyParameterInfoFile::OnChangePathClicked()
 }
 
 
-TSharedRef<SWidget> FPolicyParameterInfoFloatReference::MakeFloatInputWidget(TSharedRef<float>& ProxyValue, const FText& Label,
+TSharedRef<SWidget> FPolicyParameterInfoFloatReference::MakeFloatInputWidget(TSharedRef<TOptional<float>>& ProxyValue, const FText& Label,
                                                                      bool bRotationInDegrees,
                                                                      const FLinearColor& LabelColor,
                                                                      const FLinearColor& LabelBackgroundColor)
@@ -417,7 +495,7 @@ TSharedRef<SWidget> FPolicyParameterInfoFloatReference::MakeFloatInputWidget(TSh
 		];
 }
 
-void FPolicyParameterInfoFloatReference::OnValueCommitted(float NewValue, ETextCommit::Type CommitType, TSharedRef<float> Value)
+void FPolicyParameterInfoFloatReference::OnValueCommitted(float NewValue, ETextCommit::Type CommitType, TSharedRef<TOptional<float>> Value)
 {
 	*Value = NewValue;
 	FormatTextAndUpdateParameter();
@@ -427,18 +505,17 @@ const FString FPolicyParameterInfoMatrix::BaseMatrixString = DisplayClusterTypes
 
 FPolicyParameterInfoMatrix::FPolicyParameterInfoMatrix(const FString& InDisplayName, const FString& InKey,
                                                        UDisplayClusterBlueprint* InBlueprint,
-                                                       UDisplayClusterConfigurationViewport* InConfigurationViewport,
-														const TSharedPtr<IPropertyHandle>& InParametersHandle):
-	FPolicyParameterInfoFloatReference(InDisplayName, InKey, InBlueprint, InConfigurationViewport, InParametersHandle, &BaseMatrixString),
-	CachedTranslationX(MakeShared<float>()),
-	CachedTranslationY(MakeShared<float>()),
-	CachedTranslationZ(MakeShared<float>()),
-	CachedRotationYaw(MakeShared<float>()),
-	CachedRotationPitch(MakeShared<float>()),
-	CachedRotationRoll(MakeShared<float>()),
-	CachedScaleX(MakeShared<float>()),
-	CachedScaleY(MakeShared<float>()),
-	CachedScaleZ(MakeShared<float>())
+                                                       const TArray<TWeakObjectPtr<UDisplayClusterConfigurationViewport>>& InConfigurationViewports):
+	FPolicyParameterInfoFloatReference(InDisplayName, InKey, InBlueprint, InConfigurationViewports, &BaseMatrixString),
+	CachedTranslationX(MakeShared<TOptional<float>>()),
+	CachedTranslationY(MakeShared<TOptional<float>>()),
+	CachedTranslationZ(MakeShared<TOptional<float>>()),
+	CachedRotationYaw(MakeShared<TOptional<float>>()),
+	CachedRotationPitch(MakeShared<TOptional<float>>()),
+	CachedRotationRoll(MakeShared<TOptional<float>>()),
+	CachedScaleX(MakeShared<TOptional<float>>()),
+	CachedScaleY(MakeShared<TOptional<float>>()),
+	CachedScaleZ(MakeShared<TOptional<float>>())
 {
 	const FText TextValue = GetOrAddCustomParameterValueText();
 	const FMatrix Matrix = DisplayClusterTypesConverter::template FromString<FMatrix>(TextValue.ToString());
@@ -458,6 +535,8 @@ FPolicyParameterInfoMatrix::FPolicyParameterInfoMatrix(const FString& InDisplayN
 	*CachedScaleX = Scale.X;
 	*CachedScaleY = Scale.Y;
 	*CachedScaleZ = Scale.Z;
+
+	// TODO: Multiple values.
 }
 
 
@@ -581,9 +660,22 @@ void FPolicyParameterInfoMatrix::CustomizeScale(FDetailWidgetRow& InDetailWidget
 
 void FPolicyParameterInfoMatrix::FormatTextAndUpdateParameter()
 {
-	const FVector Translation(*CachedTranslationX, *CachedTranslationY, *CachedTranslationZ);
-	const FRotator Rotation(*CachedRotationPitch, *CachedRotationYaw, *CachedRotationRoll);
-	const FVector Scale(*CachedScaleX, *CachedScaleY, *CachedScaleZ);
+	const bool bAllValuesSet =
+		CachedTranslationX->IsSet() && CachedTranslationY->IsSet() && CachedTranslationZ->IsSet() &&
+		CachedRotationPitch->IsSet() && CachedRotationYaw->IsSet() && CachedRotationRoll->IsSet() &&
+		CachedScaleX->IsSet() && CachedScaleY->IsSet() && CachedScaleZ->IsSet();
+
+	if (!bAllValuesSet)
+	{
+		// We can't call Get() on optional values that aren't set.
+		// This check isn't ideal but we need the complete matrix
+		// to process it back to a string.
+		return;
+	}
+	
+	const FVector Translation(*CachedTranslationX.Get(), *CachedTranslationY.Get(), *CachedTranslationZ.Get());
+	const FRotator Rotation(*CachedRotationPitch.Get(), *CachedRotationYaw.Get(), *CachedRotationRoll.Get());
+	const FVector Scale(*CachedScaleX.Get(), *CachedScaleY.Get(), *CachedScaleZ.Get());
 	
 	const FMatrix Matrix = FScaleRotationTranslationMatrix(Scale, Rotation, Translation);
 	
@@ -594,35 +686,71 @@ void FPolicyParameterInfoMatrix::FormatTextAndUpdateParameter()
 const FString FPolicyParameterInfo4x4Matrix::BaseMatrixString = DisplayClusterTypesConverter::template ToString<FMatrix>(FMatrix(FPlane(1.0f, 0.0f, 0.0f, 0.0f), FPlane(0.0f, 1.0f, 0.0f, 0.0f), FPlane(0.0f, 0.0f, 0.0f, 1.0f), FPlane(0.0f, 0.0f, 1.0f, 0.0f)));
 
 FPolicyParameterInfo4x4Matrix::FPolicyParameterInfo4x4Matrix(const FString& InDisplayName, const FString& InKey,
-	UDisplayClusterBlueprint* InBlueprint, UDisplayClusterConfigurationViewport* InConfigurationViewport,
-	const TSharedPtr<IPropertyHandle>& InParametersHandle) :
-	FPolicyParameterInfoFloatReference(InDisplayName, InKey, InBlueprint, InConfigurationViewport, InParametersHandle, &BaseMatrixString),
-	A(MakeShared<float>()), B(MakeShared<float>()), C(MakeShared<float>()), D(MakeShared<float>()), E(MakeShared<float>())
-	, F(MakeShared<float>()), G(MakeShared<float>()), H(MakeShared<float>()), I(MakeShared<float>()), J(MakeShared<float>())
-	, K(MakeShared<float>()), L(MakeShared<float>()), M(MakeShared<float>()), N(MakeShared<float>()), O(MakeShared<float>())
-	, P(MakeShared<float>())
+	UDisplayClusterBlueprint* InBlueprint, const TArray<TWeakObjectPtr<UDisplayClusterConfigurationViewport>>& InConfigurationViewports) :
+	FPolicyParameterInfoFloatReference(InDisplayName, InKey, InBlueprint, InConfigurationViewports, &BaseMatrixString),
+	A(MakeShared<TOptional<float>>()), B(MakeShared<TOptional<float>>()), C(MakeShared<TOptional<float>>()), D(MakeShared<TOptional<float>>()), E(MakeShared<TOptional<float>>())
+	, F(MakeShared<TOptional<float>>()), G(MakeShared<TOptional<float>>()), H(MakeShared<TOptional<float>>()), I(MakeShared<TOptional<float>>()), J(MakeShared<TOptional<float>>())
+	, K(MakeShared<TOptional<float>>()), L(MakeShared<TOptional<float>>()), M(MakeShared<TOptional<float>>()), N(MakeShared<TOptional<float>>()), O(MakeShared<TOptional<float>>())
+	, P(MakeShared<TOptional<float>>())
 {
-	const FText TextValue = GetOrAddCustomParameterValueText();
-	const FMatrix Matrix = DisplayClusterTypesConverter::template FromString<FMatrix>(TextValue.ToString());
-	*A = Matrix.M[0][0];
-	*B = Matrix.M[0][1];
-	*C = Matrix.M[0][2];
-	*D = Matrix.M[0][3];
-
-	*E = Matrix.M[1][0];
-	*F = Matrix.M[1][1];
-	*G = Matrix.M[1][2];
-	*H = Matrix.M[1][3];
-
-	*I = Matrix.M[2][0];
-	*J = Matrix.M[2][1];
-	*K = Matrix.M[2][2];
-	*L = Matrix.M[2][3];
-
-	*M = Matrix.M[3][0];
-	*N = Matrix.M[3][1];
-	*O = Matrix.M[3][2];
-	*P = Matrix.M[3][3];
+	if (!DoParametersMatchForAllViewports())
+	{
+		for (const TWeakObjectPtr<UDisplayClusterConfigurationViewport>& Viewport : ConfigurationViewports)
+		{
+			UDisplayClusterConfigurationViewport* ConfigurationViewport = Viewport.Get();
+			check(ConfigurationViewport != nullptr);
+
+			FString* ParameterValue = ConfigurationViewport->ProjectionPolicy.Parameters.Find(GetParameterKey());
+			check(ParameterValue)
+
+			const FText TextValue = FText::FromString(*ParameterValue);
+			const FMatrix Matrix = DisplayClusterTypesConverter::template FromString<FMatrix>(TextValue.ToString());
+
+			ResetOrSetCachedValue(A, Matrix.M[0][0]);
+			ResetOrSetCachedValue(B, Matrix.M[0][1]);
+			ResetOrSetCachedValue(C, Matrix.M[0][2]);
+			ResetOrSetCachedValue(D, Matrix.M[0][3]);
+
+			ResetOrSetCachedValue(E, Matrix.M[1][0]);
+			ResetOrSetCachedValue(F, Matrix.M[1][1]);
+			ResetOrSetCachedValue(G, Matrix.M[1][2]);
+			ResetOrSetCachedValue(H, Matrix.M[1][3]);
+
+			ResetOrSetCachedValue(I, Matrix.M[2][0]);
+			ResetOrSetCachedValue(J, Matrix.M[2][1]);
+			ResetOrSetCachedValue(K, Matrix.M[2][2]);
+			ResetOrSetCachedValue(L, Matrix.M[2][3]);
+
+			ResetOrSetCachedValue(M, Matrix.M[3][0]);
+			ResetOrSetCachedValue(N, Matrix.M[3][1]);
+			ResetOrSetCachedValue(O, Matrix.M[3][2]);
+			ResetOrSetCachedValue(P, Matrix.M[3][3]);
+		}
+	}
+	else
+	{
+		const FText TextValue = GetOrAddCustomParameterValueText();
+		const FMatrix Matrix = DisplayClusterTypesConverter::template FromString<FMatrix>(TextValue.ToString());
+		*A = Matrix.M[0][0];
+		*B = Matrix.M[0][1];
+		*C = Matrix.M[0][2];
+		*D = Matrix.M[0][3];
+
+		*E = Matrix.M[1][0];
+		*F = Matrix.M[1][1];
+		*G = Matrix.M[1][2];
+		*H = Matrix.M[1][3];
+
+		*I = Matrix.M[2][0];
+		*J = Matrix.M[2][1];
+		*K = Matrix.M[2][2];
+		*L = Matrix.M[2][3];
+
+		*M = Matrix.M[3][0];
+		*N = Matrix.M[3][1];
+		*O = Matrix.M[3][2];
+		*P = Matrix.M[3][3];
+	}
 }
 
 void FPolicyParameterInfo4x4Matrix::CreateCustomRowWidget(IDetailChildrenBuilder& InDetailWidgetRow)
@@ -645,8 +773,8 @@ void FPolicyParameterInfo4x4Matrix::CreateCustomRowWidget(IDetailChildrenBuilder
 	Group.ToggleExpansion(true);
 }
 
-void FPolicyParameterInfo4x4Matrix::CustomizeRow(const FText& InHeaderText, TSharedRef<float>& InX, TSharedRef<float>& InY,
-		TSharedRef<float>& InZ, TSharedRef<float>& InW, FDetailWidgetRow& InDetailWidgetRow)
+void FPolicyParameterInfo4x4Matrix::CustomizeRow(const FText& InHeaderText, TSharedRef<TOptional<float>>& InX, TSharedRef<TOptional<float>>& InY,
+		TSharedRef<TOptional<float>>& InZ, TSharedRef<TOptional<float>>& InW, FDetailWidgetRow& InDetailWidgetRow)
 {
 	InDetailWidgetRow
 	.Visibility(MakeAttributeRaw(this, &FPolicyParameterInfo4x4Matrix::IsParameterVisible))
@@ -686,10 +814,24 @@ void FPolicyParameterInfo4x4Matrix::CustomizeRow(const FText& InHeaderText, TSha
 
 void FPolicyParameterInfo4x4Matrix::FormatTextAndUpdateParameter()
 {
-	const FPlane PlaneX(*A, *B, *C, *D);
-	const FPlane PlaneY(*E, *F, *G, *H);
-	const FPlane PlaneZ(*I, *J, *K, *L);
-	const FPlane PlaneW(*M, *N, *O, *P);
+	const bool bAllValuesSet =
+		A->IsSet() && B->IsSet() && C->IsSet() && D->IsSet() &&
+		E->IsSet() && F->IsSet() && G->IsSet() && H->IsSet() &&
+		I->IsSet() && J->IsSet() && K->IsSet() && L->IsSet() &&
+		M->IsSet() && N->IsSet() && O->IsSet() && P->IsSet();
+
+	if (!bAllValuesSet)
+	{
+		// We can't call Get() on optional values that aren't set.
+		// This check isn't ideal but we need the complete matrix
+		// to process it back to a string.
+		return;
+	}
+	
+	const FPlane PlaneX(*A.Get(), *B.Get(), *C.Get(), *D.Get());
+	const FPlane PlaneY(*E.Get(), *F.Get(), *G.Get(), *H.Get());
+	const FPlane PlaneZ(*I.Get(), *J.Get(), *K.Get(), *L.Get());
+	const FPlane PlaneW(*M.Get(), *N.Get(), *O.Get(), *P.Get());
 	
 	const FMatrix Matrix(PlaneX, PlaneY, PlaneZ, PlaneW);
 
@@ -701,19 +843,40 @@ const FString FPolicyParameterInfoRotator::BaseRotatorString = DisplayClusterTyp
 
 FPolicyParameterInfoRotator::FPolicyParameterInfoRotator(const FString& InDisplayName, const FString& InKey,
                                                          UDisplayClusterBlueprint* InBlueprint,
-                                                         UDisplayClusterConfigurationViewport* InConfigurationViewport,
-                                                         const TSharedPtr<IPropertyHandle>& InParametersHandle) :
-	FPolicyParameterInfoFloatReference(InDisplayName, InKey, InBlueprint, InConfigurationViewport, InParametersHandle, &BaseRotatorString),
-	CachedRotationYaw(MakeShared<float>()),
-	CachedRotationPitch(MakeShared<float>()),
-	CachedRotationRoll(MakeShared<float>())
+                                                         const TArray<TWeakObjectPtr<UDisplayClusterConfigurationViewport>>& InConfigurationViewports) :
+	FPolicyParameterInfoFloatReference(InDisplayName, InKey, InBlueprint, InConfigurationViewports, &BaseRotatorString),
+	CachedRotationYaw(MakeShared<TOptional<float>>()),
+	CachedRotationPitch(MakeShared<TOptional<float>>()),
+	CachedRotationRoll(MakeShared<TOptional<float>>())
 {
-	const FText TextValue = GetOrAddCustomParameterValueText();
-	const FRotator Rotation = DisplayClusterTypesConverter::template FromString<FRotator>(TextValue.ToString());
 
-	*CachedRotationYaw = Rotation.Yaw;
-	*CachedRotationPitch = Rotation.Pitch;
-	*CachedRotationRoll = Rotation.Roll;
+	if (!DoParametersMatchForAllViewports())
+	{
+		for (const TWeakObjectPtr<UDisplayClusterConfigurationViewport>& Viewport : ConfigurationViewports)
+		{
+			UDisplayClusterConfigurationViewport* ConfigurationViewport = Viewport.Get();
+			check(ConfigurationViewport != nullptr);
+
+			FString* ParameterValue = ConfigurationViewport->ProjectionPolicy.Parameters.Find(GetParameterKey());
+			check(ParameterValue)
+
+			const FText TextValue = FText::FromString(*ParameterValue);
+			const FRotator Rotation = DisplayClusterTypesConverter::template FromString<FRotator>(TextValue.ToString());
+
+			ResetOrSetCachedValue(CachedRotationYaw, Rotation.Yaw);
+			ResetOrSetCachedValue(CachedRotationPitch, Rotation.Pitch);
+			ResetOrSetCachedValue(CachedRotationRoll, Rotation.Roll);
+		}
+	}
+	else
+	{
+		const FText TextValue = GetOrAddCustomParameterValueText();
+		const FRotator Rotation = DisplayClusterTypesConverter::template FromString<FRotator>(TextValue.ToString());
+
+		*CachedRotationYaw = Rotation.Yaw;
+		*CachedRotationPitch = Rotation.Pitch;
+		*CachedRotationRoll = Rotation.Roll;
+	}
 }
 
 void FPolicyParameterInfoRotator::CreateCustomRowWidget(IDetailChildrenBuilder& InDetailWidgetRow)
@@ -751,7 +914,13 @@ void FPolicyParameterInfoRotator::CreateCustomRowWidget(IDetailChildrenBuilder&
 
 void FPolicyParameterInfoRotator::FormatTextAndUpdateParameter()
 {
-	const FRotator Rotation(*CachedRotationPitch, *CachedRotationYaw, *CachedRotationRoll);
+	const bool bAllValuesSet = CachedRotationPitch->IsSet() && CachedRotationYaw->IsSet() && CachedRotationRoll->IsSet();
+	if (!bAllValuesSet)
+	{
+		return;
+	}
+	
+	const FRotator Rotation(*CachedRotationPitch.Get(), *CachedRotationYaw.Get(), *CachedRotationRoll.Get());
 
 	const FString RotationString = DisplayClusterTypesConverter::template ToString(Rotation);
 	UpdateCustomParameterValueText(RotationString);
@@ -760,38 +929,76 @@ void FPolicyParameterInfoRotator::FormatTextAndUpdateParameter()
 const FString FPolicyParameterInfoFrustumAngle::BaseFrustumPlanesString = FString("l=-30, r=30, t=30, b=-30");
 
 FPolicyParameterInfoFrustumAngle::FPolicyParameterInfoFrustumAngle(const FString& InDisplayName, const FString& InKey,
-	UDisplayClusterBlueprint* InBlueprint, UDisplayClusterConfigurationViewport* InConfigurationViewport,
-	const TSharedPtr<IPropertyHandle>& InParametersHandle) :
-	FPolicyParameterInfoFloatReference(InDisplayName, InKey, InBlueprint, InConfigurationViewport, InParametersHandle, &BaseFrustumPlanesString),
-	CachedAngleL(MakeShared<float>()),
-	CachedAngleR(MakeShared<float>()),
-	CachedAngleT(MakeShared<float>()),
-	CachedAngleB(MakeShared<float>())
-{
-	const FString TextValue = GetOrAddCustomParameterValueText().ToString();
-
-	float Left;
-	if (DisplayClusterHelpers::str::ExtractValue(TextValue, FString(DisplayClusterProjectionStrings::cfg::manual::AngleL), Left))
+	UDisplayClusterBlueprint* InBlueprint, const TArray<TWeakObjectPtr<UDisplayClusterConfigurationViewport>>& InConfigurationViewports) :
+	FPolicyParameterInfoFloatReference(InDisplayName, InKey, InBlueprint, InConfigurationViewports, &BaseFrustumPlanesString),
+	CachedAngleL(MakeShared<TOptional<float>>()),
+	CachedAngleR(MakeShared<TOptional<float>>()),
+	CachedAngleT(MakeShared<TOptional<float>>()),
+	CachedAngleB(MakeShared<TOptional<float>>())
+{
+	if (!DoParametersMatchForAllViewports())
 	{
-		*CachedAngleL = Left;
-	}
+		for (const TWeakObjectPtr<UDisplayClusterConfigurationViewport>& Viewport : ConfigurationViewports)
+		{
+			UDisplayClusterConfigurationViewport* ConfigurationViewport = Viewport.Get();
+			check(ConfigurationViewport != nullptr);
 
-	float Right;
-	if (DisplayClusterHelpers::str::ExtractValue(TextValue, FString(DisplayClusterProjectionStrings::cfg::manual::AngleR), Right))
-	{
-		*CachedAngleR = Right;
+			if (const FString* ParameterValue = ConfigurationViewport->ProjectionPolicy.Parameters.Find(GetParameterKey()))
+			{
+				const FString TextValue = *ParameterValue;
+				float Left;
+				if (DisplayClusterHelpers::str::ExtractValue(TextValue, FString(DisplayClusterProjectionStrings::cfg::manual::AngleL), Left))
+				{
+					ResetOrSetCachedValue(CachedAngleL, Left);
+				}
+
+				float Right;
+				if (DisplayClusterHelpers::str::ExtractValue(TextValue, FString(DisplayClusterProjectionStrings::cfg::manual::AngleR), Right))
+				{
+					ResetOrSetCachedValue(CachedAngleR, Right);
+				}
+
+				float Top;
+				if (DisplayClusterHelpers::str::ExtractValue(TextValue, FString(DisplayClusterProjectionStrings::cfg::manual::AngleT), Top))
+				{
+					ResetOrSetCachedValue(CachedAngleT, Top);
+				}
+
+				float Bottom;
+				if (DisplayClusterHelpers::str::ExtractValue(TextValue, FString(DisplayClusterProjectionStrings::cfg::manual::AngleB), Bottom))
+				{
+					ResetOrSetCachedValue(CachedAngleB, Bottom);
+				}
+			}
+		}
 	}
-
-	float Top;
-	if (DisplayClusterHelpers::str::ExtractValue(TextValue, FString(DisplayClusterProjectionStrings::cfg::manual::AngleT), Top))
+	else
 	{
-		*CachedAngleT = Top;
-	}
+		const FString TextValue = GetOrAddCustomParameterValueText().ToString();
 
-	float Bottom;
-	if (DisplayClusterHelpers::str::ExtractValue(TextValue, FString(DisplayClusterProjectionStrings::cfg::manual::AngleB), Bottom))
-	{
-		*CachedAngleB = Bottom;
+		float Left;
+		if (DisplayClusterHelpers::str::ExtractValue(TextValue, FString(DisplayClusterProjectionStrings::cfg::manual::AngleL), Left))
+		{
+			*CachedAngleL = Left;
+		}
+
+		float Right;
+		if (DisplayClusterHelpers::str::ExtractValue(TextValue, FString(DisplayClusterProjectionStrings::cfg::manual::AngleR), Right))
+		{
+			*CachedAngleR = Right;
+		}
+
+		float Top;
+		if (DisplayClusterHelpers::str::ExtractValue(TextValue, FString(DisplayClusterProjectionStrings::cfg::manual::AngleT), Top))
+		{
+			*CachedAngleT = Top;
+		}
+
+		float Bottom;
+		if (DisplayClusterHelpers::str::ExtractValue(TextValue, FString(DisplayClusterProjectionStrings::cfg::manual::AngleB), Bottom))
+		{
+			*CachedAngleB = Bottom;
+		}
 	}
 }
 
@@ -834,10 +1041,16 @@ void FPolicyParameterInfoFrustumAngle::CreateCustomRowWidget(IDetailChildrenBuil
 
 void FPolicyParameterInfoFrustumAngle::FormatTextAndUpdateParameter()
 {
-	const FString AngleLStr = FString::SanitizeFloat(*CachedAngleL);
-	const FString AngleRStr = FString::SanitizeFloat(*CachedAngleR);
-	const FString AngleTStr = FString::SanitizeFloat(*CachedAngleT);
-	const FString AngleBStr = FString::SanitizeFloat(*CachedAngleB);
+	const bool bAllValuesSet = CachedAngleL->IsSet() && CachedAngleR->IsSet() && CachedAngleT->IsSet()  && CachedAngleB->IsSet();
+	if (!bAllValuesSet)
+	{
+		return;
+	}
+	
+	const FString AngleLStr = FString::SanitizeFloat(*CachedAngleL.Get());
+	const FString AngleRStr = FString::SanitizeFloat(*CachedAngleR.Get());
+	const FString AngleTStr = FString::SanitizeFloat(*CachedAngleT.Get());
+	const FString AngleBStr = FString::SanitizeFloat(*CachedAngleB.Get());
 	
 	const FString AngleString = FString::Printf(TEXT("l=%s, r=%s, t=%s, b=%s"), *AngleLStr, *AngleRStr, *AngleTStr, *AngleBStr);
 
diff --git a/Engine/Plugins/Runtime/nDisplay/Source/DisplayClusterConfigurator/Private/Views/Details/Policies/DisplayClusterConfiguratorPolicyParameterCustomization.h b/Engine/Plugins/Runtime/nDisplay/Source/DisplayClusterConfigurator/Private/Views/Details/Policies/DisplayClusterConfiguratorPolicyParameterCustomization.h
--- a/Engine/Plugins/Runtime/nDisplay/Source/DisplayClusterConfigurator/Private/Views/Details/Policies/DisplayClusterConfiguratorPolicyParameterCustomization.h
+++ b/Engine/Plugins/Runtime/nDisplay/Source/DisplayClusterConfigurator/Private/Views/Details/Policies/DisplayClusterConfiguratorPolicyParameterCustomization.h
@@ -35,8 +35,7 @@ public:
 		const FString& InDisplayName,
 		const FString& InKey,
 		UDisplayClusterBlueprint* InBlueprint,
-		UDisplayClusterConfigurationViewport* InConfigurationViewport,
-		const TSharedPtr<IPropertyHandle>& InParametersHandle,
+		const TArray<TWeakObjectPtr<UDisplayClusterConfigurationViewport>>& InConfigurationViewports,
 		const FString* InInitialValue = nullptr);
 
 	virtual ~FPolicyParameterInfo() {}
@@ -65,6 +64,9 @@ protected:
 	
 	/** Checks the data model if the parameter exists. */
 	bool IsParameterAlreadyAdded() const;
+
+	/** Compares all selected viewport values. */
+	bool DoParametersMatchForAllViewports() const;
 	
 	/** Update the parameter value in the viewport data object. */
 	void UpdateCustomParameterValueText(const FString& NewValue, bool bNotify = true) const;
@@ -77,7 +79,7 @@ protected:
 	TWeakObjectPtr<UDisplayClusterBlueprint> BlueprintOwnerPtr;
 
 	/** Viewport owning policy. */
-	TWeakObjectPtr<UDisplayClusterConfigurationViewport> ConfigurationViewportPtr;
+	TArray<TWeakObjectPtr<UDisplayClusterConfigurationViewport>> ConfigurationViewports;
 
 	/** BP editor found from BlueprintOwner. */
 	FDisplayClusterConfiguratorBlueprintEditor* BlueprintEditorPtrCached;
@@ -88,9 +90,6 @@ protected:
 	/** Used to determine parameter visibility. */
 	FParameterVisible OnParameterVisibilityCheck;
 
-	/** Handle to the owning property being set. */
-	TSharedPtr<IPropertyHandle> ParametersHandle;
-
 private:
 	/** The display name only for the UI. */
 	FString ParamDisplayName;
@@ -115,8 +114,7 @@ public:
 		const FString& InDisplayName,
 		const FString& InKey,
 		UDisplayClusterBlueprint* InBlueprint,
-		UDisplayClusterConfigurationViewport* InConfigurationViewport,
-		const TSharedPtr<IPropertyHandle>& InParametersHandle,
+		const TArray<TWeakObjectPtr<UDisplayClusterConfigurationViewport>>& InConfigurationViewports,
 		const TArray<FString>& InValues,
 		const FString* InInitialItem = nullptr,
 		bool bSort = true);
@@ -160,8 +158,7 @@ public:
 		const FString& InDisplayName,
 		const FString& InKey,
 		UDisplayClusterBlueprint* InBlueprint,
-		UDisplayClusterConfigurationViewport* InConfigurationViewport,
-		const TSharedPtr<IPropertyHandle>& InParametersHandle,
+		const TArray<TWeakObjectPtr<UDisplayClusterConfigurationViewport>>& InConfigurationViewports,
 		const TArray<TSubclassOf<UActorComponent>>& InComponentClasses);
 
 	/** Populate CustomParameterOptions based on this parameter. */
@@ -186,13 +183,17 @@ public:
 		const FString& InDisplayName,
 		const FString& InKey,
 		UDisplayClusterBlueprint* InBlueprint,
-		UDisplayClusterConfigurationViewport* InConfigurationViewport,
-		const TSharedPtr<IPropertyHandle>& InParametersHandle,
-		T InDefaultValue = 0, TOptional<T> InMinValue = TOptional<T>(), TOptional<T> InMaxValue = TOptional<T>()) : FPolicyParameterInfo(InDisplayName, InKey, InBlueprint, InConfigurationViewport, InParametersHandle)
+		const TArray<TWeakObjectPtr<UDisplayClusterConfigurationViewport>>& InConfigurationViewports,
+		T InDefaultValue = 0, TOptional<T> InMinValue = TOptional<T>(), TOptional<T> InMaxValue = TOptional<T>()) : FPolicyParameterInfo(InDisplayName, InKey, InBlueprint, InConfigurationViewports)
 	{
 		MinValue = InMinValue;
 		MaxValue = InMaxValue;
 
+		if (!DoParametersMatchForAllViewports())
+		{
+			bIsMultipleValues = true;
+		}
+		
 		if (IsParameterAlreadyAdded())
 		{
 			const FText TextValue = GetOrAddCustomParameterValueText();
@@ -218,27 +219,40 @@ public:
 			]
 			.ValueContent()
 			[
-				SNew(SSpinBox<T>)
+				SNew(SNumericEntryBox<T>)
 				.Font(IDetailLayoutBuilder::GetDetailFont())
+				.Value(this, &FPolicyParameterInfoNumber::OnGetValue)
+				.AllowSpin(true)
 				.MinValue(MinValue)
 				.MaxValue(MaxValue)
-				.Value_Lambda([this]()
-				{
-					return NumberValue;
-				})
+				.MinSliderValue(MinValue)
+				.MaxSliderValue(MaxValue)
+				.UndeterminedString(NSLOCTEXT("PropertyEditor", "MultipleValues", "Multiple Values"))
 				.OnValueChanged_Lambda([this](T InValue)
 				{
 					NumberValue = InValue;
+					bIsMultipleValues = false;
 					UpdateCustomParameterValueText(FString::SanitizeFloat(static_cast<float>(NumberValue)));
 				})
 			];
 	}
 	// ~FPolicyParameterInfo
+	
+	TOptional<T> OnGetValue() const
+	{
+		TOptional<T> Result(NumberValue);
+		if (bIsMultipleValues)
+		{
+			Result.Reset();
+		}
+		return Result;
+	}
 
 private:
 	TOptional<T> MinValue;
 	TOptional<T> MaxValue;
 	T NumberValue;
+	bool bIsMultipleValues = false;
 };
 
 /**
@@ -251,8 +265,7 @@ public:
 		const FString& InDisplayName,
 		const FString& InKey,
 		UDisplayClusterBlueprint* InBlueprint,
-		UDisplayClusterConfigurationViewport* InConfigurationViewport,
-		const TSharedPtr<IPropertyHandle>& InParametersHandle);
+		const TArray<TWeakObjectPtr<UDisplayClusterConfigurationViewport>>& InConfigurationViewports);
 
 	// FPolicyParameterInfo
 	virtual void CreateCustomRowWidget(IDetailChildrenBuilder& InDetailWidgetRow) override;
@@ -269,8 +282,7 @@ public:
 		const FString& InDisplayName,
 		const FString& InKey,
 		UDisplayClusterBlueprint* InBlueprint,
-		UDisplayClusterConfigurationViewport* InConfigurationViewport,
-		const TSharedPtr<IPropertyHandle>& InParametersHandle,
+		const TArray<TWeakObjectPtr<UDisplayClusterConfigurationViewport>>& InConfigurationViewports,
 		bool bInvertValue = false
 		);
 
@@ -294,9 +306,8 @@ public:
 		const FString& InDisplayName,
 		const FString& InKey,
 		UDisplayClusterBlueprint* InBlueprint,
-		UDisplayClusterConfigurationViewport* InConfigurationViewport,
-		const TSharedPtr<IPropertyHandle>& InParametersHandle,
-		const TArray<FString>& InFileExtensions) : FPolicyParameterInfo(InDisplayName, InKey, InBlueprint, InConfigurationViewport, InParametersHandle)
+		const TArray<TWeakObjectPtr<UDisplayClusterConfigurationViewport>>& InConfigurationViewports,
+		const TArray<FString>& InFileExtensions) : FPolicyParameterInfo(InDisplayName, InKey, InBlueprint, InConfigurationViewports)
 	{
 		FileExtensions = InFileExtensions;
 	}
@@ -324,25 +335,36 @@ public:
 		const FString& InDisplayName,
 		const FString& InKey,
 		UDisplayClusterBlueprint* InBlueprint,
-		UDisplayClusterConfigurationViewport* InConfigurationViewport,
-		const TSharedPtr<IPropertyHandle>& InParametersHandle,
+		const TArray<TWeakObjectPtr<UDisplayClusterConfigurationViewport>>& InConfigurationViewports,
 		const FString* InInitialValue = nullptr) :
-	FPolicyParameterInfo(InDisplayName, InKey, InBlueprint, InConfigurationViewport, InParametersHandle, InInitialValue)
+	FPolicyParameterInfo(InDisplayName, InKey, InBlueprint, InConfigurationViewports, InInitialValue)
 	{
 	}
 
 protected:
 	virtual void FormatTextAndUpdateParameter() = 0;
 	
-	TSharedRef<SWidget> MakeFloatInputWidget(TSharedRef<float>& ProxyValue, const FText& Label, bool bRotationInDegrees,
+	TSharedRef<SWidget> MakeFloatInputWidget(TSharedRef<TOptional<float>>& ProxyValue, const FText& Label, bool bRotationInDegrees,
 	                                         const FLinearColor& LabelColor, const FLinearColor& LabelBackgroundColor);
 
-	TOptional<float> OnGetValue(TSharedRef<float> Value) const
+	TOptional<float> OnGetValue(TSharedRef<TOptional<float>> Value) const
 	{
 		return Value.Get();
 	}
 
-	void OnValueCommitted(float NewValue, ETextCommit::Type CommitType, TSharedRef<float> Value);
+	void OnValueCommitted(float NewValue, ETextCommit::Type CommitType, TSharedRef<TOptional<float>> Value);
+
+	static void ResetOrSetCachedValue(TSharedRef<TOptional<float>>& InCachedValue, float InCompareValue)
+	{
+		if (InCachedValue->IsSet() && InCachedValue.Get() != InCompareValue)
+		{
+			InCachedValue->Reset();
+		}
+		else
+		{
+			*InCachedValue = InCompareValue;
+		}
+	};
 };
 
 /**
@@ -355,8 +377,7 @@ public:
 		const FString& InDisplayName,
 		const FString& InKey,
 		UDisplayClusterBlueprint* InBlueprint,
-		UDisplayClusterConfigurationViewport* InConfigurationViewport,
-		const TSharedPtr<IPropertyHandle>& InParametersHandle);
+		const TArray<TWeakObjectPtr<UDisplayClusterConfigurationViewport>>& InConfigurationViewports);
 
 	// FPolicyParameterInfo
 	virtual void CreateCustomRowWidget(IDetailChildrenBuilder& InDetailWidgetRow) override;
@@ -369,17 +390,17 @@ private:
 	virtual void FormatTextAndUpdateParameter() override;
 
 private:
-	mutable TSharedRef<float> CachedTranslationX;
-	mutable TSharedRef<float> CachedTranslationY;
-	mutable TSharedRef<float> CachedTranslationZ;
+	mutable TSharedRef<TOptional<float>> CachedTranslationX;
+	mutable TSharedRef<TOptional<float>> CachedTranslationY;
+	mutable TSharedRef<TOptional<float>> CachedTranslationZ;
 
-	mutable TSharedRef<float> CachedRotationYaw;
-	mutable TSharedRef<float> CachedRotationPitch;
-	mutable TSharedRef<float> CachedRotationRoll;
+	mutable TSharedRef<TOptional<float>> CachedRotationYaw;
+	mutable TSharedRef<TOptional<float>> CachedRotationPitch;
+	mutable TSharedRef<TOptional<float>> CachedRotationRoll;
 
-	mutable TSharedRef<float> CachedScaleX;
-	mutable TSharedRef<float> CachedScaleY;
-	mutable TSharedRef<float> CachedScaleZ;
+	mutable TSharedRef<TOptional<float>> CachedScaleX;
+	mutable TSharedRef<TOptional<float>> CachedScaleY;
+	mutable TSharedRef<TOptional<float>> CachedScaleZ;
 
 	static const FString BaseMatrixString;
 };
@@ -394,38 +415,37 @@ public:
 		const FString& InDisplayName,
 		const FString& InKey,
 		UDisplayClusterBlueprint* InBlueprint,
-		UDisplayClusterConfigurationViewport* InConfigurationViewport,
-		const TSharedPtr<IPropertyHandle>& InParametersHandle);
+		const TArray<TWeakObjectPtr<UDisplayClusterConfigurationViewport>>& InConfigurationViewports);
 
 	// FPolicyParameterInfo
 	virtual void CreateCustomRowWidget(IDetailChildrenBuilder& InDetailWidgetRow) override;
 	// ~FPolicyParameterInfo
 
 private:
-	void CustomizeRow(const FText& InHeaderText, TSharedRef<float>& InX, TSharedRef<float>& InY,
-		TSharedRef<float>& InZ, TSharedRef<float>& InW, FDetailWidgetRow& InDetailWidgetRow);
+	void CustomizeRow(const FText& InHeaderText, TSharedRef<TOptional<float>>& InX, TSharedRef<TOptional<float>>& InY,
+		TSharedRef<TOptional<float>>& InZ, TSharedRef<TOptional<float>>& InW, FDetailWidgetRow& InDetailWidgetRow);
 	virtual void FormatTextAndUpdateParameter() override;
 
 private:
-	mutable TSharedRef<float> A;
-	mutable TSharedRef<float> B;
-	mutable TSharedRef<float> C;
-	mutable TSharedRef<float> D;
-
-	mutable TSharedRef<float> E;
-	mutable TSharedRef<float> F;
-	mutable TSharedRef<float> G;
-	mutable TSharedRef<float> H;
-
-	mutable TSharedRef<float> I;
-	mutable TSharedRef<float> J;
-	mutable TSharedRef<float> K;
-	mutable TSharedRef<float> L;
-
-	mutable TSharedRef<float> M;
-	mutable TSharedRef<float> N;
-	mutable TSharedRef<float> O;
-	mutable TSharedRef<float> P;
+	mutable TSharedRef<TOptional<float>> A;
+	mutable TSharedRef<TOptional<float>> B;
+	mutable TSharedRef<TOptional<float>> C;
+	mutable TSharedRef<TOptional<float>> D;
+
+	mutable TSharedRef<TOptional<float>> E;
+	mutable TSharedRef<TOptional<float>> F;
+	mutable TSharedRef<TOptional<float>> G;
+	mutable TSharedRef<TOptional<float>> H;
+
+	mutable TSharedRef<TOptional<float>> I;
+	mutable TSharedRef<TOptional<float>> J;
+	mutable TSharedRef<TOptional<float>> K;
+	mutable TSharedRef<TOptional<float>> L;
+
+	mutable TSharedRef<TOptional<float>> M;
+	mutable TSharedRef<TOptional<float>> N;
+	mutable TSharedRef<TOptional<float>> O;
+	mutable TSharedRef<TOptional<float>> P;
 
 	static const FString BaseMatrixString;
 };
@@ -440,8 +460,7 @@ public:
 		const FString& InDisplayName,
 		const FString& InKey,
 		UDisplayClusterBlueprint* InBlueprint,
-		UDisplayClusterConfigurationViewport* InConfigurationViewport,
-		const TSharedPtr<IPropertyHandle>& InParametersHandle);
+		const TArray<TWeakObjectPtr<UDisplayClusterConfigurationViewport>>& InConfigurationViewports);
 
 	// FPolicyParameterInfo
 	virtual void CreateCustomRowWidget(IDetailChildrenBuilder& InDetailWidgetRow) override;
@@ -451,9 +470,9 @@ private:
 	virtual void FormatTextAndUpdateParameter() override;
 
 private:
-	mutable TSharedRef<float> CachedRotationYaw;
-	mutable TSharedRef<float> CachedRotationPitch;
-	mutable TSharedRef<float> CachedRotationRoll;
+	mutable TSharedRef<TOptional<float>> CachedRotationYaw;
+	mutable TSharedRef<TOptional<float>> CachedRotationPitch;
+	mutable TSharedRef<TOptional<float>> CachedRotationRoll;
 
 	static const FString BaseRotatorString;
 };
@@ -468,8 +487,7 @@ public:
 		const FString& InDisplayName,
 		const FString& InKey,
 		UDisplayClusterBlueprint* InBlueprint,
-		UDisplayClusterConfigurationViewport* InConfigurationViewport,
-		const TSharedPtr<IPropertyHandle>& InParametersHandle);
+		const TArray<TWeakObjectPtr<UDisplayClusterConfigurationViewport>>& InConfigurationViewports);
 
 	// FPolicyParameterInfo
 	virtual void CreateCustomRowWidget(IDetailChildrenBuilder& InDetailWidgetRow) override;
@@ -479,10 +497,10 @@ private:
 	virtual void FormatTextAndUpdateParameter() override;
 
 private:
-	mutable TSharedRef<float> CachedAngleL;
-	mutable TSharedRef<float> CachedAngleR;
-	mutable TSharedRef<float> CachedAngleT;
-	mutable TSharedRef<float> CachedAngleB;
+	mutable TSharedRef<TOptional<float>> CachedAngleL;
+	mutable TSharedRef<TOptional<float>> CachedAngleR;
+	mutable TSharedRef<TOptional<float>> CachedAngleT;
+	mutable TSharedRef<TOptional<float>> CachedAngleB;
 
 	static const FString BaseFrustumPlanesString;
 };
\ No newline at end of file
