nDisplay: Fix Linux build. Horde Issue 95631

diff --git a/Engine/Plugins/Runtime/nDisplay/Source/DisplayClusterConfigurator/Private/Views/Details/Policies/DisplayClusterConfiguratorPolicyDetailCustomization.cpp b/Engine/Plugins/Runtime/nDisplay/Source/DisplayClusterConfigurator/Private/Views/Details/Policies/DisplayClusterConfiguratorPolicyDetailCustomization.cpp
--- a/Engine/Plugins/Runtime/nDisplay/Source/DisplayClusterConfigurator/Private/Views/Details/Policies/DisplayClusterConfiguratorPolicyDetailCustomization.cpp
+++ b/Engine/Plugins/Runtime/nDisplay/Source/DisplayClusterConfigurator/Private/Views/Details/Policies/DisplayClusterConfiguratorPolicyDetailCustomization.cpp
@@ -301,9 +301,9 @@ bool FDisplayClusterConfiguratorProjectionCustomization::IsPolicyIdenticalAcross
 	for (const TWeakObjectPtr<UDisplayClusterConfigurationViewport>& Viewport : ConfigurationViewports)
 	{
 		if (Viewport.IsValid() &&
-			Viewport->ProjectionPolicy.Type != CurrentSelectedPolicy ||
-			bRequireCustomPolicy && !Viewport->ProjectionPolicy.bIsCustom ||
-			!bRequireCustomPolicy && Viewport->ProjectionPolicy.bIsCustom != bIsCustomPolicy)
+			(Viewport->ProjectionPolicy.Type != CurrentSelectedPolicy ||
+			(bRequireCustomPolicy && !Viewport->ProjectionPolicy.bIsCustom) ||
+			(!bRequireCustomPolicy && Viewport->ProjectionPolicy.bIsCustom != bIsCustomPolicy)))
 		{
 			return false;
 		}
