nDisplay: Fix static analysis, horde issue 96235.

diff --git a/Engine/Plugins/Runtime/nDisplay/Source/DisplayClusterConfigurator/Private/Views/Details/Policies/DisplayClusterConfiguratorPolicyParameterCustomization.cpp b/Engine/Plugins/Runtime/nDisplay/Source/DisplayClusterConfigurator/Private/Views/Details/Policies/DisplayClusterConfiguratorPolicyParameterCustomization.cpp
--- a/Engine/Plugins/Runtime/nDisplay/Source/DisplayClusterConfigurator/Private/Views/Details/Policies/DisplayClusterConfiguratorPolicyParameterCustomization.cpp
+++ b/Engine/Plugins/Runtime/nDisplay/Source/DisplayClusterConfigurator/Private/Views/Details/Policies/DisplayClusterConfiguratorPolicyParameterCustomization.cpp
@@ -84,12 +84,15 @@ FText FPolicyParameterInfo::GetOrAddCustomParameterValueText() const
 			UpdateCustomParameterValueText(InitialValue.IsValid() ? *InitialValue : TEXT(""), false);
 			ParameterValue = ConfigurationViewport->ProjectionPolicy.Parameters.Find(GetParameterKey());
 		}
-		check(ParameterValue)
 
 		// Only need values from one viewport since they match.
 		break;
 	}
-	
+
+	if (ParameterValue == nullptr)
+	{
+		return FText::GetEmpty();
+	}
 	return FText::FromString(*ParameterValue);
 }
 
