Fix crash when exposing an ndisplay root property

diff --git a/Engine/Plugins/Runtime/nDisplay/Source/DisplayClusterConfigurator/Private/Views/Details/Policies/DisplayClusterConfiguratorPolicyDetailCustomization.cpp b/Engine/Plugins/Runtime/nDisplay/Source/DisplayClusterConfigurator/Private/Views/Details/Policies/DisplayClusterConfiguratorPolicyDetailCustomization.cpp
--- a/Engine/Plugins/Runtime/nDisplay/Source/DisplayClusterConfigurator/Private/Views/Details/Policies/DisplayClusterConfiguratorPolicyDetailCustomization.cpp
+++ b/Engine/Plugins/Runtime/nDisplay/Source/DisplayClusterConfigurator/Private/Views/Details/Policies/DisplayClusterConfiguratorPolicyDetailCustomization.cpp
@@ -45,6 +45,13 @@ void FDisplayClusterConfiguratorProjectionCustomization::CustomizeHeader(TShared
 	CustomOption = MakeShared<FString>("Custom");
 
 	// Get the Editing object
+	if (!EditingObject->IsA<UDisplayClusterConfigurationViewport>())
+	{
+		// The editing object should only be invalid in the case where the customized row was created in a different context than the config editor,
+		// ie. when using a property row generator like in the Remote Control Preset.
+		return;
+	}
+
 	ConfigurationViewportPtr = CastChecked<UDisplayClusterConfigurationViewport>(EditingObject);
 
 	for (const TWeakObjectPtr<UObject>& Object : EditingObjects)
@@ -101,25 +108,25 @@ void FDisplayClusterConfiguratorProjectionCustomization::ResetProjectionPolicyOp
 {
 	ProjectionPolicyOptions.Reset();
 
-	UDisplayClusterConfigurationViewport* ConfigurationViewport = ConfigurationViewportPtr.Get();
-	check(ConfigurationViewport != nullptr);
-
-	for (const FString& ProjectionPolicy : UDisplayClusterConfigurationData::ProjectionPolicies)
+	if (UDisplayClusterConfigurationViewport* ConfigurationViewport = ConfigurationViewportPtr.Get())
 	{
-		ProjectionPolicyOptions.Add(MakeShared<FString>(ProjectionPolicy));
-	}
+		for (const FString& ProjectionPolicy : UDisplayClusterConfigurationData::ProjectionPolicies)
+		{
+			ProjectionPolicyOptions.Add(MakeShared<FString>(ProjectionPolicy));
+		}
 
-	// Add Custom option
-	if (!bIsCustomPolicy)
-	{
-		ProjectionPolicyOptions.Add(CustomOption);
-	}
+		// Add Custom option
+		if (!bIsCustomPolicy)
+		{
+			ProjectionPolicyOptions.Add(CustomOption);
+		}
 
-	if (ProjectionPolicyComboBox.IsValid())
-	{
-		// Refreshes the available options now that the shared array has been updated.
+		if (ProjectionPolicyComboBox.IsValid())
+		{
+			// Refreshes the available options now that the shared array has been updated.
 		
-		ProjectionPolicyComboBox->ResetOptionsSource();
+			ProjectionPolicyComboBox->ResetOptionsSource();
+		}
 	}
 }
 
@@ -181,36 +188,36 @@ void FDisplayClusterConfiguratorProjectionCustomization::OnProjectionPolicySelec
 	{
 		FString SelectedPolicy = *InPolicy.Get();
 
-		UDisplayClusterConfigurationViewport* ConfigurationViewport = ConfigurationViewportPtr.Get();
-		check(ConfigurationViewport != nullptr);
-
-		ConfigurationViewport->Modify();
-		ModifyBlueprint();
-		
-		if (SelectedPolicy.Equals(*CustomOption.Get()) && IsPolicyIdenticalAcrossEditedObjects())
-		{
-			bIsCustomPolicy = true;
-			CustomPolicy = ConfigurationViewport->ProjectionPolicy.Type;
-			IsCustomHandle->SetValue(true);
-		}
-		else
+		if (UDisplayClusterConfigurationViewport* ConfigurationViewport = ConfigurationViewportPtr.Get())
 		{
-			bIsCustomPolicy = false;
-			IsCustomHandle->SetValue(false);
+			ConfigurationViewport->Modify();
+			ModifyBlueprint();
+		
+			if (SelectedPolicy.Equals(*CustomOption.Get()) && IsPolicyIdenticalAcrossEditedObjects())
+			{
+				bIsCustomPolicy = true;
+				CustomPolicy = ConfigurationViewport->ProjectionPolicy.Type;
+				IsCustomHandle->SetValue(true);
+			}
+			else
+			{
+				bIsCustomPolicy = false;
+				IsCustomHandle->SetValue(false);
 			
-			TypeHandle->SetValue(SelectedPolicy);
+				TypeHandle->SetValue(SelectedPolicy);
 
-			if (CurrentSelectedPolicy.ToLower() != SelectedPolicy.ToLower())
-			{
-				// Reset when going from custom to another policy.
-				ensure(ParametersHandle->AsMap()->Empty() == FPropertyAccess::Result::Success);
+				if (CurrentSelectedPolicy.ToLower() != SelectedPolicy.ToLower())
+				{
+					// Reset when going from custom to another policy.
+					ensure(ParametersHandle->AsMap()->Empty() == FPropertyAccess::Result::Success);
+				}
 			}
-		}
 
-		CurrentSelectedPolicy = SelectedPolicy;
+			CurrentSelectedPolicy = SelectedPolicy;
 
-		RefreshBlueprint();
-		PropertyUtilities.Pin()->ForceRefresh();
+			RefreshBlueprint();
+			PropertyUtilities.Pin()->ForceRefresh();
+		}
 	}
 	else
 	{
@@ -240,27 +247,30 @@ const FString& FDisplayClusterConfiguratorProjectionCustomization::GetCurrentPol
 		return *CustomOption.Get();
 	}
 
-	UDisplayClusterConfigurationViewport* ConfigurationViewport = ConfigurationViewportPtr.Get();
-	check(ConfigurationViewport != nullptr);
+	if (UDisplayClusterConfigurationViewport* ConfigurationViewport = ConfigurationViewportPtr.Get())
+	{
+		return ConfigurationViewport->ProjectionPolicy.Type;
+	}
 
-	return ConfigurationViewport->ProjectionPolicy.Type;
+	static FString Empty;
+	return Empty;
 }
 
 bool FDisplayClusterConfiguratorProjectionCustomization::IsCustomTypeInConfig() const
 {
-	UDisplayClusterConfigurationViewport* ConfigurationViewport = ConfigurationViewportPtr.Get();
-	check(ConfigurationViewport != nullptr);
-
-	if (ConfigurationViewport->ProjectionPolicy.bIsCustom)
+	if (UDisplayClusterConfigurationViewport* ConfigurationViewport = ConfigurationViewportPtr.Get())
 	{
-		return true;
-	}
+		if (ConfigurationViewport->ProjectionPolicy.bIsCustom)
+		{
+			return true;
+		}
 	
-	for (const FString& ProjectionPolicy : UDisplayClusterConfigurationData::ProjectionPolicies)
-	{
-		if (ConfigurationViewport->ProjectionPolicy.Type.ToLower().Equals(ProjectionPolicy.ToLower()))
+		for (const FString& ProjectionPolicy : UDisplayClusterConfigurationData::ProjectionPolicies)
 		{
-			return false;
+			if (ConfigurationViewport->ProjectionPolicy.Type.ToLower().Equals(ProjectionPolicy.ToLower()))
+			{
+				return false;
+			}
 		}
 	}
 
