nDisplay: Update combo box style to new standard.

diff --git a/Engine/Plugins/Runtime/nDisplay/Source/DisplayClusterConfigurator/Private/Views/Details/Widgets/SDisplayClusterConfigurationSearchableComboBox.h b/Engine/Plugins/Runtime/nDisplay/Source/DisplayClusterConfigurator/Private/Views/Details/Widgets/SDisplayClusterConfigurationSearchableComboBox.h
--- a/Engine/Plugins/Runtime/nDisplay/Source/DisplayClusterConfigurator/Private/Views/Details/Widgets/SDisplayClusterConfigurationSearchableComboBox.h
+++ b/Engine/Plugins/Runtime/nDisplay/Source/DisplayClusterConfigurator/Private/Views/Details/Widgets/SDisplayClusterConfigurationSearchableComboBox.h
@@ -16,11 +16,11 @@ public:
 
 	SLATE_BEGIN_ARGS(SDisplayClusterConfigurationSearchableComboBox)
 		: _Content()
-		, _ComboBoxStyle(&FCoreStyle::Get().GetWidgetStyle< FComboBoxStyle >("ComboBox"))
+		, _ComboBoxStyle(&FAppStyle::Get().GetWidgetStyle<FComboBoxStyle>("ComboBox"))
 		, _ButtonStyle(nullptr)
-		, _ItemStyle(&FCoreStyle::Get().GetWidgetStyle< FTableRowStyle >("TableView.Row"))
-		, _ContentPadding(FMargin(4.0, 2.0))
-		, _ForegroundColor(FCoreStyle::Get().GetSlateColor("InvertedForeground"))
+		, _ItemStyle(&FAppStyle::Get().GetWidgetStyle<FTableRowStyle>("ComboBox.Row"))
+		, _ContentPadding(_ComboBoxStyle->ContentPadding)
+		, _ForegroundColor(FSlateColor::UseStyle())
 		, _OptionsSource()
 		, _OnSelectionChanged()
 		, _OnGenerateWidget()
