nDisplay: Fix multiple values sometimes incorrectly displaying for policies with one viewport selected.

diff --git a/Engine/Plugins/Runtime/nDisplay/Source/DisplayClusterConfigurator/Private/Views/Details/Policies/DisplayClusterConfiguratorPolicyDetailCustomization.cpp b/Engine/Plugins/Runtime/nDisplay/Source/DisplayClusterConfigurator/Private/Views/Details/Policies/DisplayClusterConfiguratorPolicyDetailCustomization.cpp
--- a/Engine/Plugins/Runtime/nDisplay/Source/DisplayClusterConfigurator/Private/Views/Details/Policies/DisplayClusterConfiguratorPolicyDetailCustomization.cpp
+++ b/Engine/Plugins/Runtime/nDisplay/Source/DisplayClusterConfigurator/Private/Views/Details/Policies/DisplayClusterConfiguratorPolicyDetailCustomization.cpp
@@ -308,6 +308,10 @@ void FDisplayClusterConfiguratorProjectionCustomization::OnTextCommittedInCustom
 
 bool FDisplayClusterConfiguratorProjectionCustomization::IsPolicyIdenticalAcrossEditedObjects(bool bRequireCustomPolicy) const
 {
+	if (ConfigurationViewports.Num() <= 1)
+	{
+		return true;
+	}
 	for (const TWeakObjectPtr<UDisplayClusterConfigurationViewport>& Viewport : ConfigurationViewports)
 	{
 		if (Viewport.IsValid() &&
