nDisplay: Fixed button styles in the Configuration Editor, and adjusted the nDisplay New Asset dialog to match UE5's style.

diff --git a/Engine/Plugins/Runtime/nDisplay/Source/DisplayClusterConfigurator/Private/ClusterConfiguration/SDisplayClusterConfiguratorNewClusterItemDialog.cpp b/Engine/Plugins/Runtime/nDisplay/Source/DisplayClusterConfigurator/Private/ClusterConfiguration/SDisplayClusterConfiguratorNewClusterItemDialog.cpp
--- a/Engine/Plugins/Runtime/nDisplay/Source/DisplayClusterConfigurator/Private/ClusterConfiguration/SDisplayClusterConfiguratorNewClusterItemDialog.cpp
+++ b/Engine/Plugins/Runtime/nDisplay/Source/DisplayClusterConfigurator/Private/ClusterConfiguration/SDisplayClusterConfiguratorNewClusterItemDialog.cpp
@@ -12,6 +12,8 @@
 #include "Widgets/Layout/SBox.h"
 #include "Widgets/Layout/SScrollBox.h"
 #include "Widgets/Layout/SSpacer.h"
+#include "Widgets/Layout/SSeparator.h"
+#include "Widgets/Layout/SUniformGridPanel.h"
 
 #define LOCTEXT_NAMESPACE "SDisplayClusterConfiguratorNewClusterItemDialog"
 
@@ -230,9 +232,9 @@ void SDisplayClusterConfiguratorNewClusterItemDialog::Construct(const FArguments
 						.HAlign(HAlign_Fill)
 						.Padding(0.0f, 4.0f, 0.0f, 2.0f)
 						[
-							SNew(SBorder)
-							.Padding(FEditorStyle::GetMargin(TEXT("Menu.Separator.Padding")))
-							.BorderImage(FEditorStyle::GetBrush(TEXT("Menu.Separator")))
+							SNew(SSeparator)
+							.SeparatorImage(FAppStyle::Get().GetBrush("Menu.Separator"))
+							.Thickness(1.0f)
 						]
 
 						+SVerticalBox::Slot()
@@ -259,34 +261,27 @@ void SDisplayClusterConfiguratorNewClusterItemDialog::Construct(const FArguments
 			.AutoHeight()
 			.Padding(2)
 			[
-				SNew(SHorizontalBox)
-				
-				+SHorizontalBox::Slot()
-				.AutoWidth()
-				.Padding(2)
+				SNew(SUniformGridPanel)
+				.MinDesiredSlotHeight(FAppStyle::Get().GetFloat("StandardDialog.MinDesiredSlotHeight"))
+				.MinDesiredSlotWidth(FAppStyle::Get().GetFloat("StandardDialog.MinDesiredSlotWidth"))
+				.SlotPadding(FAppStyle::Get().GetMargin("StandardDialog.SlotPadding"))
+
+				+SUniformGridPanel::Slot(0, 0)
 				[
-					SNew(SBox)
-					.WidthOverride(64)
-					[
-						SNew(SButton)
-						.HAlign(HAlign_Center)
-						.Text(LOCTEXT("AddItemLabel", "Add"))
-						.OnClicked(this, &SDisplayClusterConfiguratorNewClusterItemDialog::OnAddButtonClicked)
-					]
+					SNew(SButton)
+					.ContentPadding(FAppStyle::Get().GetMargin("StandardDialog.ContentPadding"))
+					.HAlign(HAlign_Center)
+					.Text(LOCTEXT("AddItemLabel", "Add"))
+					.OnClicked(this, &SDisplayClusterConfiguratorNewClusterItemDialog::OnAddButtonClicked)
 				]
 
-				+SHorizontalBox::Slot()
-				.AutoWidth()
-				.Padding(2)
+				+SUniformGridPanel::Slot(1, 0)
 				[
-					SNew(SBox)
-					.WidthOverride(64)
-					[
-						SNew(SButton)
-						.HAlign(HAlign_Center)
-						.Text(LOCTEXT("CancelLabel", "Cancel"))
-						.OnClicked(this, &SDisplayClusterConfiguratorNewClusterItemDialog::OnCancelButtonClicked)
-					]
+					SNew(SButton)
+					.ContentPadding(FAppStyle::Get().GetMargin("StandardDialog.ContentPadding"))
+					.HAlign(HAlign_Center)
+					.Text(LOCTEXT("CancelLabel", "Cancel"))
+					.OnClicked(this, &SDisplayClusterConfiguratorNewClusterItemDialog::OnCancelButtonClicked)
 				]
 			]
 		]
diff --git a/Engine/Plugins/Runtime/nDisplay/Source/DisplayClusterConfigurator/Private/DisplayClusterConfiguratorStyle.cpp b/Engine/Plugins/Runtime/nDisplay/Source/DisplayClusterConfigurator/Private/DisplayClusterConfiguratorStyle.cpp
--- a/Engine/Plugins/Runtime/nDisplay/Source/DisplayClusterConfigurator/Private/DisplayClusterConfiguratorStyle.cpp
+++ b/Engine/Plugins/Runtime/nDisplay/Source/DisplayClusterConfigurator/Private/DisplayClusterConfiguratorStyle.cpp
@@ -59,6 +59,8 @@ const FSlateBrush* FDisplayClusterConfiguratorStyle::GetBrush(FName PropertyName
 #define IMAGE_BRUSH( RelativePath, ... ) FSlateImageBrush( Style->RootToContentDir( RelativePath, TEXT(".png") ), __VA_ARGS__ )
 #define BOX_BRUSH( RelativePath, ... ) FSlateBoxBrush( Style->RootToContentDir( RelativePath, TEXT(".png") ), __VA_ARGS__ )
 #define BORDER_BRUSH( RelativePath, ... ) FSlateBorderBrush( Style->RootToContentDir( RelativePath, TEXT(".png") ), __VA_ARGS__ )
+#define EDITOR_BOX_BRUSH( RelativePath, ... ) FSlateBoxBrush( FPaths::EngineContentDir() / "Editor/Slate" / RelativePath + TEXT(".png"), __VA_ARGS__ )
+#define DEFAULT_FONT(...) FCoreStyle::GetDefaultFontStyle(__VA_ARGS__)
 
 const FVector2D Icon128x128(128.0f, 128.0f);
 const FVector2D Icon64x64(64.0f, 64.0f);
@@ -252,12 +254,23 @@ TSharedRef< FSlateStyleSet > FDisplayClusterConfiguratorStyle::Create()
 		}
 	}
 
+	// New Asset Dialog
+	{
+		const FTextBlockStyle DefaultTextStyle = FAppStyle::Get().GetWidgetStyle<FTextBlockStyle>("NormalText");
+
+		Style->Set("DisplayClusterConfigurator.NewAssetDialog.OptionText", FTextBlockStyle(DefaultTextStyle).SetFont(DEFAULT_FONT("Bold", 11)));
+		Style->Set("DisplayClusterConfigurator.NewAssetDialog.SubBorder", new EDITOR_BOX_BRUSH("Common/GroupBorderLight", FMargin(4.0f / 16.0f)));
+		Style->Set("DisplayClusterConfigurator.NewAssetDialog.ActiveOptionBorderColor", FLinearColor(FColor(96, 96, 96)));
+	}
+
 	return Style;
 }
 
 #undef IMAGE_BRUSH
 #undef BOX_BRUSH
 #undef BORDER_BRUSH
+#undef EDITOR_BOX_BRUSH
+#undef DEFAULT_FONT
 
 void FDisplayClusterConfiguratorStyle::ReloadTextures()
 {
diff --git a/Engine/Plugins/Runtime/nDisplay/Source/DisplayClusterConfigurator/Private/Views/NewAsset/SDisplayClusterConfiguratorNewAssetDialog.cpp b/Engine/Plugins/Runtime/nDisplay/Source/DisplayClusterConfigurator/Private/Views/NewAsset/SDisplayClusterConfiguratorNewAssetDialog.cpp
--- a/Engine/Plugins/Runtime/nDisplay/Source/DisplayClusterConfigurator/Private/Views/NewAsset/SDisplayClusterConfiguratorNewAssetDialog.cpp
+++ b/Engine/Plugins/Runtime/nDisplay/Source/DisplayClusterConfigurator/Private/Views/NewAsset/SDisplayClusterConfiguratorNewAssetDialog.cpp
@@ -2,6 +2,7 @@
 
 #include "Views/NewAsset/SDisplayClusterConfiguratorNewAssetDialog.h"
 #include "Settings/DisplayClusterConfiguratorSettings.h"
+#include "DisplayClusterConfiguratorStyle.h"
 
 #include "AssetData.h"
 
@@ -59,6 +60,7 @@ void SDisplayClusterConfiguratorNewAssetDialog::Construct(const FArguments& InAr
 			.AutoHeight()
 			[
 				SNew(SBorder)
+				.BorderImage(FDisplayClusterConfiguratorStyle::Get().GetBrush("DisplayClusterConfigurator.NewAssetDialog.SubBorder"))
 				.BorderBackgroundColor(this, &SDisplayClusterConfiguratorNewAssetDialog::GetOptionBorderColor, OptionIndex)
 				[
 						SNew(SCheckBox)
@@ -68,24 +70,31 @@ void SDisplayClusterConfiguratorNewAssetDialog::Construct(const FArguments& InAr
 						.OnCheckStateChanged(this, &SDisplayClusterConfiguratorNewAssetDialog::OptionCheckBoxStateChanged, OptionIndex)
 						.Content()
 						[
-							SNew(SVerticalBox)
-							+ SVerticalBox::Slot()
-							.AutoHeight()
-							.Padding(5, 2)
-							[		
-								SNew(STextBlock)
-								.ColorAndOpacity(this, &SDisplayClusterConfiguratorNewAssetDialog::GetOptionTextColor, OptionIndex)
-								.Text(Option.OptionText)
-								.AutoWrapText(true)
-							]
-							+ SVerticalBox::Slot()
-							.AutoHeight()
-							.Padding(5, 2, 5, 7)
+							// this border catches the double click before the checkbox can
+							SNew(SBorder)
+							.BorderImage(FEditorStyle::Get().GetBrush("NoBorder"))
+							.OnMouseDoubleClick(this, &SDisplayClusterConfiguratorNewAssetDialog::OnOptionDoubleClicked, OptionIndex)
 							[
-								SNew(STextBlock)
-								.ColorAndOpacity(this, &SDisplayClusterConfiguratorNewAssetDialog::GetOptionTextColor, OptionIndex)
-								.Text(Option.OptionDescription)
-								.AutoWrapText(true)
+								SNew(SVerticalBox)
+								+ SVerticalBox::Slot()
+								.AutoHeight()
+								.Padding(5, 2)
+								[		
+									SNew(STextBlock)
+									.TextStyle(FDisplayClusterConfiguratorStyle::Get(), "DisplayClusterConfigurator.NewAssetDialog.OptionText")
+									.ColorAndOpacity(this, &SDisplayClusterConfiguratorNewAssetDialog::GetOptionTextColor, OptionIndex)
+									.Text(Option.OptionText)
+									.AutoWrapText(true)
+								]
+								+ SVerticalBox::Slot()
+								.AutoHeight()
+								.Padding(5, 2, 5, 7)
+								[
+									SNew(STextBlock)
+									.ColorAndOpacity(this, &SDisplayClusterConfiguratorNewAssetDialog::GetOptionTextColor, OptionIndex)
+									.Text(Option.OptionDescription)
+									.AutoWrapText(true)
+								]
 							]
 						]
 					
@@ -97,26 +106,24 @@ void SDisplayClusterConfiguratorNewAssetDialog::Construct(const FArguments& InAr
 	SWindow::Construct(SWindow::FArguments()
 		.Title(FText::Format(LOCTEXT("NewConfigurationDialogTitle", "Pick a starting point for your {0}"), AssetTypeDisplayName))
 		.SizingRule(ESizingRule::UserSized)
-		.ClientSize(FVector2D(400.f, 400.f))
+		.ClientSize(FVector2D(480.f, 600.f))
 		.SupportsMaximize(false)
 		.SupportsMinimize(false)
 		[
-			SNew(SWizard)
+			SAssignNew(Wizard, SWizard)
 			.OnCanceled(this, &SDisplayClusterConfiguratorNewAssetDialog::OnCancelButtonClicked)
 			.OnFinished(this, &SDisplayClusterConfiguratorNewAssetDialog::OnOkButtonClicked)
 			.CanFinish(this, &SDisplayClusterConfiguratorNewAssetDialog::IsOkButtonEnabled)
 			.ShowPageList(false)
-			.ButtonStyle(FEditorStyle::Get(), "FlatButton.Default")
-			.CancelButtonStyle(FEditorStyle::Get(), "FlatButton.Default")
-			.FinishButtonStyle(FEditorStyle::Get(), "FlatButton.Success")
-			.ButtonTextStyle(FEditorStyle::Get(), "FlatButton.DefaultTextStyle")
+
 			+SWizard::Page()
 			.CanShow(true)
 			.OnEnter(this, &SDisplayClusterConfiguratorNewAssetDialog::ResetStage)
 			[
 				RootBox
 			]
-			+ SWizard::Page()
+
+			+SWizard::Page()
 			.CanShow(this, &SDisplayClusterConfiguratorNewAssetDialog::HasAssetPage)
 			.OnEnter(this, &SDisplayClusterConfiguratorNewAssetDialog::GetAssetPicker)
 			[
@@ -181,7 +188,7 @@ void SDisplayClusterConfiguratorNewAssetDialog::OnWindowClosed(const TSharedRef<
 FSlateColor SDisplayClusterConfiguratorNewAssetDialog::GetOptionBorderColor(int32 OptionIndex) const
 {
 	return SelectedOptionIndex == OptionIndex
-		? FSlateColor(FLinearColor::Blue)
+		? FDisplayClusterConfiguratorStyle::Get().GetColor("DisplayClusterConfigurator.NewAssetDialog.ActiveOptionBorderColor")
 		: FSlateColor(FLinearColor::Transparent);
 }
 
@@ -209,6 +216,27 @@ void SDisplayClusterConfiguratorNewAssetDialog::OptionCheckBoxStateChanged(EChec
 	}
 }
 
+FReply SDisplayClusterConfiguratorNewAssetDialog::OnOptionDoubleClicked(const FGeometry& Geometry, const FPointerEvent& PointerEvent, int32 OptionIndex)
+{
+	SelectedOptionIndex = OptionIndex;
+	if (Wizard->CanShowPage(Wizard->GetCurrentPageIndex() + 1))
+	{
+		Wizard->AdvanceToPage(Wizard->GetCurrentPageIndex() + 1);
+		return FReply::Handled();
+	}
+	// if we can't advance, we attempt to finish instead
+	else
+	{
+		if (IsOkButtonEnabled())
+		{
+			OnOkButtonClicked();
+			return FReply::Handled();
+		}
+	}
+
+	return FReply::Unhandled();
+}
+
 FText SDisplayClusterConfiguratorNewAssetDialog::GetAssetPickersLabelText() const
 {
 	if (SelectedOptionIndex < 0 || SelectedOptionIndex >= Options.Num())
diff --git a/Engine/Plugins/Runtime/nDisplay/Source/DisplayClusterConfigurator/Private/Views/NewAsset/SDisplayClusterConfiguratorNewAssetDialog.h b/Engine/Plugins/Runtime/nDisplay/Source/DisplayClusterConfigurator/Private/Views/NewAsset/SDisplayClusterConfiguratorNewAssetDialog.h
--- a/Engine/Plugins/Runtime/nDisplay/Source/DisplayClusterConfigurator/Private/Views/NewAsset/SDisplayClusterConfiguratorNewAssetDialog.h
+++ b/Engine/Plugins/Runtime/nDisplay/Source/DisplayClusterConfigurator/Private/Views/NewAsset/SDisplayClusterConfiguratorNewAssetDialog.h
@@ -11,6 +11,7 @@
 #include "AssetRegistry/AssetData.h"
 
 class FDisplayClusterConfiguratorBlueprintEditor;
+class SWizard;
 
 class SDisplayClusterConfiguratorNewAssetDialog : public SWindow
 {
@@ -70,6 +71,8 @@ private:
 	void OptionCheckBoxStateChanged(ECheckBoxState InCheckBoxState, int32 OptionIndex);
 	FSlateColor GetOptionTextColor(int32 OptionIndex) const;
 
+	FReply OnOptionDoubleClicked(const FGeometry& Geometry, const FPointerEvent& PointerEvent, int32 OptionIndex);
+
 	FText GetAssetPickersLabelText() const;
 	bool IsOkButtonEnabled() const;
 	void OnOkButtonClicked();
@@ -77,6 +80,7 @@ private:
 	bool HasAssetPage() const;
 
 private:
+	TSharedPtr<SWizard> Wizard;
 	TArray<FDisplayClusterConfiguratorNewAssetDialogOption> Options;
 	int32 SelectedOptionIndex;
 	bool bUserConfirmedSelection;
diff --git a/Engine/Plugins/Runtime/nDisplay/Source/DisplayClusterConfigurator/Private/Views/TreeViews/SDisplayClusterConfiguratorViewTree.cpp b/Engine/Plugins/Runtime/nDisplay/Source/DisplayClusterConfigurator/Private/Views/TreeViews/SDisplayClusterConfiguratorViewTree.cpp
--- a/Engine/Plugins/Runtime/nDisplay/Source/DisplayClusterConfigurator/Private/Views/TreeViews/SDisplayClusterConfiguratorViewTree.cpp
+++ b/Engine/Plugins/Runtime/nDisplay/Source/DisplayClusterConfigurator/Private/Views/TreeViews/SDisplayClusterConfiguratorViewTree.cpp
@@ -22,6 +22,7 @@
 #include "Widgets/Images/SImage.h"
 #include "Widgets/Input/SSearchBox.h"
 #include "Widgets/Input/SComboButton.h"
+#include "SEditorHeaderButton.h"
 
 #define LOCTEXT_NAMESPACE "SDisplayClusterConfiguratorViewTree"
 
@@ -88,35 +89,12 @@ void SDisplayClusterConfiguratorViewTree::Construct(const FArguments& InArgs,
 						.VAlign(VAlign_Center)
 						[
 					
-							SNew(SComboButton)
-							.ContentPadding(FMargin(5, 0))
-							.ComboButtonStyle(FEditorStyle::Get(), "ToolbarComboButton")
-							.ButtonStyle(FEditorStyle::Get(), "FlatButton.Success")
-							.ForegroundColor(FLinearColor::White)
+							SNew(SEditorHeaderButton)
+							.Text(LOCTEXT("AddNewItem", "Add New"))
+							.Icon(FAppStyle::Get().GetBrush("Icons.Plus"))
+							.ToolTipText(LOCTEXT("AddNewClusterItemTooltip", "Adds a new Cluster Item to the Cluster"))
 							.Visibility(this, &SDisplayClusterConfiguratorViewTree::GetAddNewComboButtonVisibility)
 							.OnGetMenuContent(this, &SDisplayClusterConfiguratorViewTree::CreateAddNewMenuContent)
-							.ButtonContent()
-							[
-								SNew(SHorizontalBox)
-								+SHorizontalBox::Slot()
-								.VAlign(VAlign_Center)
-								.AutoWidth()
-								.Padding(1.f,1.f)
-								[
-									SNew(STextBlock)
-									.TextStyle(FEditorStyle::Get(), "ContentBrowser.TopBar.Font")
-									.Font(FEditorStyle::Get().GetFontStyle("FontAwesome.10"))
-									.Text(FText::FromString(FString(TEXT("\xf067"))) /*fa-plus*/)
-								]
-								+ SHorizontalBox::Slot()
-								.VAlign(VAlign_Center)
-								.Padding(1.f)
-								[
-									SNew(STextBlock)
-									.Text(LOCTEXT("AddNewItem", "Add New"))
-									.TextStyle(FEditorStyle::Get(), "ContentBrowser.TopBar.Font")
-								]
-							]
 						]
 
 						+SHorizontalBox::Slot()
