nDisplay: Fixed custom widgets in various dropdown menus in the Configuration Editor so that they are propoerly aligned with the regular menu entries instead of being flush with the menu edge.

diff --git a/Engine/Plugins/Runtime/nDisplay/Source/DisplayClusterConfigurator/Private/Views/OutputMapping/Widgets/SDisplayClusterConfiguratorOutputMappingToolbar.cpp b/Engine/Plugins/Runtime/nDisplay/Source/DisplayClusterConfigurator/Private/Views/OutputMapping/Widgets/SDisplayClusterConfiguratorOutputMappingToolbar.cpp
--- a/Engine/Plugins/Runtime/nDisplay/Source/DisplayClusterConfigurator/Private/Views/OutputMapping/Widgets/SDisplayClusterConfiguratorOutputMappingToolbar.cpp
+++ b/Engine/Plugins/Runtime/nDisplay/Source/DisplayClusterConfigurator/Private/Views/OutputMapping/Widgets/SDisplayClusterConfiguratorOutputMappingToolbar.cpp
@@ -175,30 +175,29 @@ TSharedRef<SWidget> SDisplayClusterConfiguratorOutputMappingToolbar::MakeSnappin
 		MenuBuilder.AddMenuEntry(IDisplayClusterConfigurator::Get().GetCommands().ToggleAdjacentEdgeSnapping);
 
 		MenuBuilder.AddWidget(
-			SNew(SHorizontalBox)
-			.IsEnabled(this, &SDisplayClusterConfiguratorOutputMappingToolbar::IsAdjacentEdgeSnappingEnabled)
-
-			+ SHorizontalBox::Slot()
-			.AutoWidth()
-			.Padding(FMargin(0.f, 0.f, 5.f, 0.f))
-			[
-				SNew(STextBlock)
-				.Text(LOCTEXT("AlignmentSettings_AdjacentEdgePadding", "Adjacent Edge Padding"))
-			]
-
-			+ SHorizontalBox::Slot()
-			.VAlign(VAlign_Bottom)
-			.FillWidth(1.f)
+			SNew(SBox)
+			.HAlign(HAlign_Right)
 			[
-				SNew(SNumericEntryBox<int>)
-				.Value(this, &SDisplayClusterConfiguratorOutputMappingToolbar::GetAdjacentEdgePadding)
-				.OnValueChanged(this, &SDisplayClusterConfiguratorOutputMappingToolbar::SetAdjacentEdgePadding)
-				.MinValue(0)
-				.MaxValue(INT_MAX)
-				.MaxSliderValue(100)
-				.AllowSpin(true)
+				SNew(SBox)
+				.Padding(FMargin(8.0f, 0.0f, 0.0f, 0.0f))
+				.WidthOverride(100.0f)
+				[
+					SNew(SBorder)
+					.BorderImage(FAppStyle::Get().GetBrush("Menu.WidgetBorder"))
+					.Padding(FMargin(1.0f))
+					[
+						SNew(SNumericEntryBox<int>)
+						.Value(this, &SDisplayClusterConfiguratorOutputMappingToolbar::GetAdjacentEdgePadding)
+						.OnValueChanged(this, &SDisplayClusterConfiguratorOutputMappingToolbar::SetAdjacentEdgePadding)
+						.IsEnabled(this, &SDisplayClusterConfiguratorOutputMappingToolbar::IsAdjacentEdgeSnappingEnabled)
+						.MinValue(0)
+						.MaxValue(INT_MAX)
+						.MaxSliderValue(100)
+						.AllowSpin(true)
+					]
+				]
 			],
-			FText::GetEmpty()
+			LOCTEXT("AlignmentSettings_AdjacentEdgePadding", "Adjacent Edge Padding")
 		);
 	}
 	MenuBuilder.EndSection();
@@ -212,30 +211,29 @@ TSharedRef<SWidget> SDisplayClusterConfiguratorOutputMappingToolbar::MakeSnappin
 	MenuBuilder.BeginSection(TEXT("GeneralSnapping"));
 	{
 		MenuBuilder.AddWidget(
-			SNew(SHorizontalBox)
-			.IsEnabled(this, &SDisplayClusterConfiguratorOutputMappingToolbar::IsSnappingEnabled)
-			
-			+ SHorizontalBox::Slot()
-			.AutoWidth()
-			.Padding(FMargin(0.f, 0.f, 5.f, 0.f))
-			[
-				SNew(STextBlock)
-				.Text(LOCTEXT("AlignmentSettings_SnapProximity", "Snap Proximity"))
-			]
-
-			+ SHorizontalBox::Slot()
-			.VAlign(VAlign_Bottom)
-			.FillWidth(1.f)
+			SNew(SBox)
+			.HAlign(HAlign_Right)
 			[
-				SNew(SNumericEntryBox<int>)
-				.Value(this, &SDisplayClusterConfiguratorOutputMappingToolbar::GetSnapProximity)
-				.OnValueChanged(this, &SDisplayClusterConfiguratorOutputMappingToolbar::SetSnapProximity)
-				.MinValue(0)
-				.MaxValue(INT_MAX)
-				.MaxSliderValue(100)
-				.AllowSpin(true)
+				SNew(SBox)
+				.Padding(FMargin(8.0f, 0.0f, 0.0f, 0.0f))
+				.WidthOverride(100.0f)
+				[
+					SNew(SBorder)
+					.BorderImage(FAppStyle::Get().GetBrush("Menu.WidgetBorder"))
+					.Padding(FMargin(1.0f))
+					[
+						SNew(SNumericEntryBox<int>)
+						.Value(this, &SDisplayClusterConfiguratorOutputMappingToolbar::GetSnapProximity)
+						.OnValueChanged(this, &SDisplayClusterConfiguratorOutputMappingToolbar::SetSnapProximity)
+						.IsEnabled(this, &SDisplayClusterConfiguratorOutputMappingToolbar::IsSnappingEnabled)
+						.MinValue(0)
+						.MaxValue(INT_MAX)
+						.MaxSliderValue(100)
+						.AllowSpin(true)
+					]
+				]
 			],
-			FText::GetEmpty()
+			LOCTEXT("AlignmentSettings_SnapProximity", "Snap Proximity")
 		);
 	}
 	MenuBuilder.EndSection();
@@ -378,58 +376,56 @@ void SDisplayClusterConfiguratorOutputMappingToolbar::MakeHostArrangementTypeSub
 	MenuBuilder.BeginSection(TEXT("HostArrangementSettings"));
 	{
 		MenuBuilder.AddWidget(
-			SNew(SHorizontalBox)
-			.IsEnabled(this, &SDisplayClusterConfiguratorOutputMappingToolbar::IsHostArrangementTypeChecked, EHostArrangementType::Wrap)
-
-			+ SHorizontalBox::Slot()
-			.AutoWidth()
-			.Padding(FMargin(0.f, 0.f, 5.f, 0.f))
-			[
-				SNew(STextBlock)
-				.Text(LOCTEXT("HostArrangementSettings_WrapThreshold", "Wrapping Threshold"))
-			]
-
-			+ SHorizontalBox::Slot()
-			.VAlign(VAlign_Bottom)
-			.FillWidth(1.f)
+			SNew(SBox)
+			.HAlign(HAlign_Right)
 			[
-				SNew(SNumericEntryBox<int>)
-				.Value(this, &SDisplayClusterConfiguratorOutputMappingToolbar::GetHostWrapThreshold)
-				.OnValueChanged(this, &SDisplayClusterConfiguratorOutputMappingToolbar::SetHostWrapThreshold)
-				.MinValue(0)
-				.MaxValue(INT_MAX)
-				.MaxSliderValue(10000)
-				.AllowSpin(true)
+				SNew(SBox)
+				.Padding(FMargin(8.0f, 0.0f, 0.0f, 0.0f))
+				.WidthOverride(100.0f)
+				[
+					SNew(SBorder)
+					.BorderImage(FAppStyle::Get().GetBrush("Menu.WidgetBorder"))
+					.Padding(FMargin(1.0f))
+					[
+						SNew(SNumericEntryBox<int>)
+						.Value(this, &SDisplayClusterConfiguratorOutputMappingToolbar::GetHostWrapThreshold)
+						.OnValueChanged(this, &SDisplayClusterConfiguratorOutputMappingToolbar::SetHostWrapThreshold)
+						.IsEnabled(this, &SDisplayClusterConfiguratorOutputMappingToolbar::IsHostArrangementTypeChecked, EHostArrangementType::Wrap)
+						.MinValue(0)
+						.MaxValue(INT_MAX)
+						.MaxSliderValue(10000)
+						.AllowSpin(true)
+					]
+				]
 			],
-			FText::GetEmpty()
+			LOCTEXT("HostArrangementSettings_WrapThreshold", "Wrapping Threshold")
 		);
 
 		MenuBuilder.AddWidget(
-			SNew(SHorizontalBox)
-			.IsEnabled(this, &SDisplayClusterConfiguratorOutputMappingToolbar::IsHostArrangementTypeChecked, EHostArrangementType::Grid)
-			
-			+ SHorizontalBox::Slot()
-			.AutoWidth()
-			.Padding(FMargin(0.f, 0.f, 5.f, 0.f))
-			[
-				SNew(STextBlock)
-				.Text(LOCTEXT("HostArrangementSettings_GridSize", "Grid Size"))
-			]
-
-			+ SHorizontalBox::Slot()
-			.VAlign(VAlign_Bottom)
-			.FillWidth(1.f)
+			SNew(SBox)
+			.HAlign(HAlign_Right)
 			[
-				SNew(SNumericEntryBox<int>)
-				.Value(this, &SDisplayClusterConfiguratorOutputMappingToolbar::GetHostGridSize)
-				.OnValueChanged(this, &SDisplayClusterConfiguratorOutputMappingToolbar::SetHostGridSize)
-				.MinValue(1)
-				.MaxValue(10)
-				.MinSliderValue(1)
-				.MaxSliderValue(5)
-				.AllowSpin(true)
+				SNew(SBox)
+				.Padding(FMargin(8.0f, 0.0f, 0.0f, 0.0f))
+				.WidthOverride(100.0f)
+				[
+					SNew(SBorder)
+					.BorderImage(FAppStyle::Get().GetBrush("Menu.WidgetBorder"))
+					.Padding(FMargin(1.0f))
+					[
+						SNew(SNumericEntryBox<int>)
+						.Value(this, &SDisplayClusterConfiguratorOutputMappingToolbar::GetHostGridSize)
+						.OnValueChanged(this, &SDisplayClusterConfiguratorOutputMappingToolbar::SetHostGridSize)
+						.IsEnabled(this, &SDisplayClusterConfiguratorOutputMappingToolbar::IsHostArrangementTypeChecked, EHostArrangementType::Grid)
+						.MinValue(1)
+						.MaxValue(10)
+						.MinSliderValue(1)
+						.MaxSliderValue(5)
+						.AllowSpin(true)
+					]
+				]
 			],
-			FText::GetEmpty()
+			LOCTEXT("HostArrangementSettings_GridSize", "Grid Size")
 		);
 	}
 	MenuBuilder.EndSection();
diff --git a/Engine/Plugins/Runtime/nDisplay/Source/DisplayClusterConfigurator/Private/Views/Viewport/DisplayClusterConfiguratorSCSEditorViewport.cpp b/Engine/Plugins/Runtime/nDisplay/Source/DisplayClusterConfigurator/Private/Views/Viewport/DisplayClusterConfiguratorSCSEditorViewport.cpp
--- a/Engine/Plugins/Runtime/nDisplay/Source/DisplayClusterConfigurator/Private/Views/Viewport/DisplayClusterConfiguratorSCSEditorViewport.cpp
+++ b/Engine/Plugins/Runtime/nDisplay/Source/DisplayClusterConfigurator/Private/Views/Viewport/DisplayClusterConfiguratorSCSEditorViewport.cpp
@@ -267,43 +267,37 @@ public:
 			ViewportsMenuBuilder.AddMenuEntry(FDisplayClusterConfiguratorCommands::Get().Show3DViewportNames);
 
 			ViewportsMenuBuilder.AddWidget(
-				SNew(SHorizontalBox)
-
-				+ SHorizontalBox::Slot()
-				.AutoWidth()
-				.VAlign(VAlign_Center)
-				.Padding(FMargin(0.f, 0.f, 5.f, 0.f))
-				[
-					SNew(STextBlock)
-					.Text(LOCTEXT("PreviewResolution_Label", "Preview Resolution"))
-				]
-
-				+ SHorizontalBox::Slot()
-				.VAlign(VAlign_Bottom)
-				.AutoWidth()
+				SNew(SBox)
+				.HAlign(HAlign_Right)
 				[
 					SNew(SBox)
-					.MinDesiredWidth(64)
+					.Padding(FMargin(8.0f, 0.0f, 0.0f, 0.0f))
+					.WidthOverride(100.0f)
 					[
-						SNew(SNumericEntryBox<float>)
-						.Value(ViewportClient.Get(), &FDisplayClusterConfiguratorSCSEditorViewportClient::GetPreviewResolutionScale)
-						.OnValueChanged(ViewportClient.Get(), &FDisplayClusterConfiguratorSCSEditorViewportClient::SetPreviewResolutionScale)
-						.OnValueCommitted_Lambda([this](const float InValue, ETextCommit::Type)
-						{
-							if (ViewportClient.IsValid())
+						SNew(SBorder)
+						.BorderImage(FAppStyle::Get().GetBrush("Menu.WidgetBorder"))
+						.Padding(FMargin(1.0f))
+						[
+							SNew(SNumericEntryBox<float>)
+							.Value(ViewportClient.Get(), &FDisplayClusterConfiguratorSCSEditorViewportClient::GetPreviewResolutionScale)
+							.OnValueChanged(ViewportClient.Get(), &FDisplayClusterConfiguratorSCSEditorViewportClient::SetPreviewResolutionScale)
+							.OnValueCommitted_Lambda([this](const float InValue, ETextCommit::Type)
 							{
-								ViewportClient->Invalidate();
-							}
-						})
-
-						.MinValue(0.05f)
-						.MaxValue(1.f)
-						.MinSliderValue(0.05f)
-						.MaxSliderValue(1.f)
-						.AllowSpin(true)
+								if (ViewportClient.IsValid())
+								{
+									ViewportClient->Invalidate();
+								}
+							})
+
+							.MinValue(0.05f)
+							.MaxValue(1.f)
+							.MinSliderValue(0.05f)
+							.MaxSliderValue(1.f)
+							.AllowSpin(true)
+						]
 					]
 				],
-				FText::GetEmpty()
+				LOCTEXT("PreviewResolution_Label", "Preview Resolution")
 		);
 		}
 		ViewportsMenuBuilder.EndSection();
@@ -313,35 +307,33 @@ public:
 			ViewportsMenuBuilder.AddMenuEntry(FDisplayClusterConfiguratorCommands::Get().ToggleShowXformGizmos);
 
 			ViewportsMenuBuilder.AddWidget(
-				SNew(SHorizontalBox)
-
-				+ SHorizontalBox::Slot()
-				.AutoWidth()
-				.VAlign(VAlign_Center)
-				.Padding(FMargin(0.f, 0.f, 5.f, 0.f))
-				[
-					SNew(STextBlock)
-					.Text(LOCTEXT("XformGizmoScale_Label", "Xform Gizmo Scale"))
-				]
-
-				+ SHorizontalBox::Slot()
-				.VAlign(VAlign_Bottom)
-				.AutoWidth()
+				SNew(SBox)
+				.HAlign(HAlign_Right)
 				[
 					SNew(SBox)
-					.MinDesiredWidth(64)
+					.Padding(FMargin(8.0f, 0.0f, 0.0f, 0.0f))
+					.WidthOverride(100.0f)
 					[
-						SNew(SNumericEntryBox<float>)
-						.Value(ViewportClient.Get(), &FDisplayClusterConfiguratorSCSEditorViewportClient::GetXformGizmoScale)
-						.OnValueChanged(ViewportClient.Get(), &FDisplayClusterConfiguratorSCSEditorViewportClient::SetXformGizmoScale)
-						.MinValue(0)
-						.MaxValue(FLT_MAX)
-						.MinSliderValue(0)
-						.MaxSliderValue(2)
-						.AllowSpin(true)
+						SNew(SBorder)
+						.BorderImage(FAppStyle::Get().GetBrush("Menu.WidgetBorder"))
+						.Padding(FMargin(1.0f))
+						[
+							SNew(SBox)
+							.MinDesiredWidth(64)
+							[
+								SNew(SNumericEntryBox<float>)
+								.Value(ViewportClient.Get(), &FDisplayClusterConfiguratorSCSEditorViewportClient::GetXformGizmoScale)
+								.OnValueChanged(ViewportClient.Get(), &FDisplayClusterConfiguratorSCSEditorViewportClient::SetXformGizmoScale)
+								.MinValue(0)
+								.MaxValue(FLT_MAX)
+								.MinSliderValue(0)
+								.MaxSliderValue(2)
+								.AllowSpin(true)
+							]
+						]
 					]
 				],
-				FText::GetEmpty()
+				LOCTEXT("XformGizmoScale_Label", "Xform Gizmo Scale")
 			);
 		}
 		ViewportsMenuBuilder.EndSection();
