Additional files for update of Release-Engine-Test from Release-Engine-Staging @ 16264272

diff --git a/Engine/Plugins/Runtime/nDisplay/Source/DisplayCluster/Private/Render/Viewport/DisplayClusterViewport.cpp b/Engine/Plugins/Runtime/nDisplay/Source/DisplayCluster/Private/Render/Viewport/DisplayClusterViewport.cpp
--- a/Engine/Plugins/Runtime/nDisplay/Source/DisplayCluster/Private/Render/Viewport/DisplayClusterViewport.cpp
+++ b/Engine/Plugins/Runtime/nDisplay/Source/DisplayCluster/Private/Render/Viewport/DisplayClusterViewport.cpp
@@ -180,6 +180,31 @@ inline void AdjustRect(FIntRect& InOutRect, const float multX, const float multY
 	InOutRect.Max.Y *= multY;
 }
 
+FIntRect FDisplayClusterViewport::GetValidRect(const FIntRect& InRect, const TCHAR* DbgSourceName)
+{
+	// The target always needs be within GMaxTextureDimensions, larger dimensions are not supported by the engine
+	static const int32 MaxTextureSize = 1 << (GMaxTextureMipCount - 1);
+
+	float RectScale = 1;
+
+	// Make sure the rect doesn't exceed the maximum resolution, and preserve its aspect ratio if it needs to be clamped
+	int RectMaxSize = InRect.Max.GetMax();
+	if (RectMaxSize > MaxTextureSize)
+	{
+		RectScale = float(MaxTextureSize) / RectMaxSize;
+		UE_LOG(LogDisplayClusterViewport, Error, TEXT("The viewport '%s' rect '%s' size %dx%d clamped: max texture dimensions is %d"), *GetId(), (DbgSourceName==nullptr) ? TEXT("none") : DbgSourceName, InRect.Max.X, InRect.Max.Y, MaxTextureSize);
+	}
+
+	FIntRect OutRect = InRect;
+	OutRect.Min.X = FMath::Min(OutRect.Min.X, MaxTextureSize);
+	OutRect.Min.Y = FMath::Min(OutRect.Min.Y, MaxTextureSize);
+
+	OutRect.Max.X = FMath::Clamp(int(OutRect.Max.X * RectScale), OutRect.Min.X, MaxTextureSize);
+	OutRect.Max.Y = FMath::Clamp(int(OutRect.Max.Y * RectScale), OutRect.Min.Y, MaxTextureSize);
+
+	return OutRect;
+}
+
 bool FDisplayClusterViewport::UpdateFrameContexts(const uint32 InViewPassNum, const FDisplayClusterRenderFrameSettings& InFrameSettings)
 {
 	check(IsInGameThread());
@@ -215,10 +240,8 @@ bool FDisplayClusterViewport::UpdateFrameContexts(const uint32 InViewPassNum, co
 		}
 	}
 
-	float RTTSizeMult = RenderSettings.RenderTargetRatio * InFrameSettings.ClusterRenderTargetRatioMult;
-
 	uint32 FrameTargetsAmmount = 2;
-	FIntRect     FrameTargetRect = RenderSettings.Rect;
+	FIntRect     DesiredFrameTargetRect = RenderSettings.Rect;
 	{
 		switch (InFrameSettings.RenderMode)
 		{
@@ -226,20 +249,19 @@ bool FDisplayClusterViewport::UpdateFrameContexts(const uint32 InViewPassNum, co
 			FrameTargetsAmmount = 1;
 			break;
 		case EDisplayClusterRenderFrameMode::SideBySide:
-			AdjustRect(FrameTargetRect, 0.5f, 1.f);
+			AdjustRect(DesiredFrameTargetRect, 0.5f, 1.f);
 			break;
 		case EDisplayClusterRenderFrameMode::TopBottom:
-			AdjustRect(FrameTargetRect, 1.f, 0.5f);
+			AdjustRect(DesiredFrameTargetRect, 1.f, 0.5f);
 			break;
 		case EDisplayClusterRenderFrameMode::PreviewMono:
 		{
-			// Apply preview scale
-			RTTSizeMult *= InFrameSettings.PreviewRenderTargetRatioMult;
-			float MultXY = InFrameSettings.PreviewRenderTargetRatioMult;
-			AdjustRect(FrameTargetRect, MultXY, MultXY);
+			// Preview downscale in range 0..1
+			float MultXY = FMath::Clamp(InFrameSettings.PreviewRenderTargetRatioMult, 0.f, 1.f);
+			AdjustRect(DesiredFrameTargetRect, MultXY, MultXY);
 
 			// Align each frame to zero
-			FrameTargetRect = FIntRect(FIntPoint(0, 0), FrameTargetRect.Size());
+			DesiredFrameTargetRect = FIntRect(FIntPoint(0, 0), DesiredFrameTargetRect.Size());
 
 			// Mono
 			FrameTargetsAmmount = 1;
@@ -250,55 +272,43 @@ bool FDisplayClusterViewport::UpdateFrameContexts(const uint32 InViewPassNum, co
 		}
 	}
 
-	static int MaximumSupportedResolution = 1 << (GMaxTextureMipCount - 1);
+	// Special case mono->stereo
+	uint32 ViewportContextAmmount = RenderSettings.bForceMono ? 1 : FrameTargetsAmmount;
 
 	// Make sure the frame target rect doesn't exceed the maximum resolution, and preserve its aspect ratio if it needs to be clamped
-	if (FrameTargetRect.Max.X > MaximumSupportedResolution)
-	{
-		FrameTargetRect.Max.Y = FrameTargetRect.Max.Y * MaximumSupportedResolution / (float)FrameTargetRect.Max.X;
-		FrameTargetRect.Max.X = MaximumSupportedResolution;
-	}
+	FIntRect FrameTargetRect = GetValidRect(DesiredFrameTargetRect, TEXT("Context Frame"));
 
-	if (FrameTargetRect.Max.Y > MaximumSupportedResolution)
+	// Exclude zero-size viewports from render
+	if (FrameTargetRect.Size().GetMin() <= 0)
 	{
-		FrameTargetRect.Max.X = FrameTargetRect.Max.X * MaximumSupportedResolution / (float)FrameTargetRect.Max.Y;
-		FrameTargetRect.Max.Y = MaximumSupportedResolution;
+		UE_LOG(LogDisplayClusterViewport, Error, TEXT("The viewport '%s' FrameTarget rect has zero size %dx%d: Disabled"), *GetId(), FrameTargetRect.Size().X, FrameTargetRect.Size().Y);
+		return false;
 	}
 
-	// Special case mono->stereo
-	uint32 ViewportContextAmmount = RenderSettings.bForceMono ? 1 : FrameTargetsAmmount;
 
+	FIntPoint DesiredContextSize = FrameTargetRect.Size();
 
-	FIntPoint ContextSize = FrameTargetRect.Size();
+	// Cluster mult downscale in range 0..1
+	float ClusterRTTMult = FMath::Clamp(InFrameSettings.ClusterRenderTargetRatioMult, 0.f, 1.f);
 
+	// Scale context for rendering
+	float ViewportContextSizeMult = FMath::Max(RenderSettings.RenderTargetRatio * ClusterRTTMult, 0.f);
+	DesiredContextSize.X *= ViewportContextSizeMult;
+	DesiredContextSize.Y *= ViewportContextSizeMult;
 
-	// Get valid viewport size for render:
-	{
-		int ViewportResolution = ContextSize.GetMax() * RTTSizeMult;
+	FIntRect RenderTargetRect = GetValidRect(FIntRect(FIntPoint(0, 0), DesiredContextSize), TEXT("Context RenderTarget"));
 
-		// In UE, the maximum texture resolution is computed as:
-		if (ViewportResolution > MaximumSupportedResolution)
-		{
-			RTTSizeMult *= float(ViewportResolution) / MaximumSupportedResolution;
-	}
-
-		ContextSize.X = FMath::Min(int(RTTSizeMult * ContextSize.X), MaximumSupportedResolution);
-		ContextSize.Y = FMath::Min(int(RTTSizeMult * ContextSize.Y), MaximumSupportedResolution);
-	}
-
-	FIntRect RenderTargetRect(FIntPoint(0, 0), ContextSize);
-
-	// Test vs zero size rects
-	if (FrameTargetRect.Width() <= 0 || FrameTargetRect.Height() <= 0 || RenderTargetRect.Width() <= 0 || RenderTargetRect.Height() <= 0)
+	// Exclude zero-size viewports from render
+	if (RenderTargetRect.Size().GetMin()<=0)
 	{
+		UE_LOG(LogDisplayClusterViewport, Error, TEXT("The viewport '%s' RenderTarget rect has zero size %dx%d: Disabled"), *GetId(), RenderTargetRect.Size().X, RenderTargetRect.Size().Y);
 		return false;
 	}
 
-	// Support overscan rendering feature
-	OverscanRendering.Update(RenderTargetRect);
+	FIntPoint ContextSize = RenderTargetRect.Size();
 
-	check(FrameTargetRect.Max.X <= MaximumSupportedResolution);
-	check(FrameTargetRect.Max.Y <= MaximumSupportedResolution);
+	// Support overscan rendering feature
+	OverscanRendering.Update(*this, RenderTargetRect);
 
 	// Support override feature
 	bool bDisableRender = PostRenderSettings.Override.IsEnabled();
diff --git a/Engine/Plugins/Runtime/nDisplay/Source/DisplayCluster/Private/Render/Viewport/DisplayClusterViewportManager.cpp b/Engine/Plugins/Runtime/nDisplay/Source/DisplayCluster/Private/Render/Viewport/DisplayClusterViewportManager.cpp
--- a/Engine/Plugins/Runtime/nDisplay/Source/DisplayCluster/Private/Render/Viewport/DisplayClusterViewportManager.cpp
+++ b/Engine/Plugins/Runtime/nDisplay/Source/DisplayCluster/Private/Render/Viewport/DisplayClusterViewportManager.cpp
@@ -275,9 +275,9 @@ void FDisplayClusterViewportManager::ConfigureViewFamily(const FDisplayClusterRe
 	}
 }
 
-void FDisplayClusterViewportManager::RenderFrame(const bool bWarpBlendEnabled, FRHITexture2D* FrameOutputRTT)
+void FDisplayClusterViewportManager::RenderFrame(const bool bWarpBlendEnabled, FViewport* InViewport)
 {
-	ViewportManagerProxy->ImplRenderFrame(bWarpBlendEnabled, FrameOutputRTT);
+	ViewportManagerProxy->ImplRenderFrame(bWarpBlendEnabled, InViewport);
 }
 
 bool FDisplayClusterViewportManager::CreateViewport(const FString& InViewportId, const class UDisplayClusterConfigurationViewport* ConfigurationViewport)
diff --git a/Engine/Plugins/Runtime/nDisplay/Source/DisplayCluster/Private/Render/Viewport/DisplayClusterViewportManagerProxy.cpp b/Engine/Plugins/Runtime/nDisplay/Source/DisplayCluster/Private/Render/Viewport/DisplayClusterViewportManagerProxy.cpp
--- a/Engine/Plugins/Runtime/nDisplay/Source/DisplayCluster/Private/Render/Viewport/DisplayClusterViewportManagerProxy.cpp
+++ b/Engine/Plugins/Runtime/nDisplay/Source/DisplayCluster/Private/Render/Viewport/DisplayClusterViewportManagerProxy.cpp
@@ -136,10 +136,10 @@ void FDisplayClusterViewportManagerProxy::ImplUpdateViewports(const TArray<FDisp
 
 DECLARE_GPU_STAT_NAMED(nDisplay_ViewportManager_RenderFrame, TEXT("nDisplay ViewportManager::RenderFrame"));
 
-void FDisplayClusterViewportManagerProxy::ImplRenderFrame(const bool bWarpBlendEnabled, FRHITexture2D* FrameOutputRTT)
+void FDisplayClusterViewportManagerProxy::ImplRenderFrame(const bool bWarpBlendEnabled, FViewport* InViewport)
 {
 	ENQUEUE_RENDER_COMMAND(DeleteDisplayClusterViewportProxy)(
-		[ViewportManagerProxy = this, bWarpBlendEnabled, FrameOutputRTT](FRHICommandListImmediate& RHICmdList)
+		[ViewportManagerProxy = this, bWarpBlendEnabled, InViewport](FRHICommandListImmediate& RHICmdList)
 	{
 		SCOPED_GPU_STAT(RHICmdList, nDisplay_ViewportManager_RenderFrame);
 		SCOPED_DRAW_EVENT(RHICmdList, nDisplay_ViewportManager_RenderFrame);
@@ -158,22 +158,25 @@ void FDisplayClusterViewportManagerProxy::ImplRenderFrame(const bool bWarpBlendE
 		// Update the frame resources: post-processing, warping, and finally resolving everything to the frame resource
 		ViewportManagerProxy->UpdateFrameResources_RenderThread(RHICmdList, bWarpBlendEnabled);
 
-		if (FrameOutputRTT)
+		if (InViewport)
 		{
-			// For quadbuf stereo copy only left eye, right copy from OutputFrameTarget
-			//@todo Copy QuadBuf_LeftEye/(mono,sbs,tp) to separate rtt, before UI and debug rendering
-			//@todo QuadBuf_LeftEye copied latter, before present
-			switch(ViewportManagerProxy->RenderFrameSettings.RenderMode)
-			{ 
-			case EDisplayClusterRenderFrameMode::SideBySide:
-			case EDisplayClusterRenderFrameMode::TopBottom:
-				ViewportManagerProxy->ResolveFrameTargetToBackBuffer_RenderThread(RHICmdList, 1, 0, FrameOutputRTT, FrameOutputRTT->GetSizeXY());
-				break;
-			default:
-				break;
-			}
+			if (FRHITexture2D* FrameOutputRTT = InViewport->GetRenderTargetTexture())
+			{
+				// For quadbuf stereo copy only left eye, right copy from OutputFrameTarget
+				//@todo Copy QuadBuf_LeftEye/(mono,sbs,tp) to separate rtt, before UI and debug rendering
+				//@todo QuadBuf_LeftEye copied latter, before present
+				switch (ViewportManagerProxy->RenderFrameSettings.RenderMode)
+				{
+				case EDisplayClusterRenderFrameMode::SideBySide:
+				case EDisplayClusterRenderFrameMode::TopBottom:
+					ViewportManagerProxy->ResolveFrameTargetToBackBuffer_RenderThread(RHICmdList, 1, 0, FrameOutputRTT, FrameOutputRTT->GetSizeXY());
+					break;
+				default:
+					break;
+				}
 
-			ViewportManagerProxy->ResolveFrameTargetToBackBuffer_RenderThread(RHICmdList, 0, 0, FrameOutputRTT, FrameOutputRTT->GetSizeXY());
+				ViewportManagerProxy->ResolveFrameTargetToBackBuffer_RenderThread(RHICmdList, 0, 0, FrameOutputRTT, FrameOutputRTT->GetSizeXY());
+			}
 		}
 	});
 }
diff --git a/Engine/Source/Runtime/Engine/Private/Components/SkeletalMeshComponent.cpp b/Engine/Source/Runtime/Engine/Private/Components/SkeletalMeshComponent.cpp
--- a/Engine/Source/Runtime/Engine/Private/Components/SkeletalMeshComponent.cpp
+++ b/Engine/Source/Runtime/Engine/Private/Components/SkeletalMeshComponent.cpp
@@ -4063,7 +4063,11 @@ void USkeletalMeshComponent::FinalizeBoneTransform()
 
 	ConditionallyDispatchQueuedAnimEvents();
 
-	OnBoneTransformsFinalized.Broadcast();
+	PRAGMA_DISABLE_DEPRECATION_WARNINGS
+	OnBoneTransformsFinalized.Broadcast();  // Deprecated in 4.27
+	PRAGMA_ENABLE_DEPRECATION_WARNINGS
+
+	OnBoneTransformsFinalizedMC.Broadcast();
 
 	TRACE_SKELETAL_MESH_COMPONENT(this);
 }
@@ -4138,6 +4142,15 @@ void USkeletalMeshComponent::UnregisterOnTeleportDelegate(const FDelegateHandle&
 	OnSkelMeshPhysicsTeleported.Remove(DelegateHandle);
 }
 
+FDelegateHandle USkeletalMeshComponent::RegisterOnBoneTransformsFinalizedDelegate(const FOnBoneTransformsFinalizedMultiCast::FDelegate& Delegate)
+{
+	return OnBoneTransformsFinalizedMC.Add(Delegate);
+}
+
+void USkeletalMeshComponent::UnregisterOnBoneTransformsFinalizedDelegate(const FDelegateHandle& DelegateHandle)
+{
+	OnBoneTransformsFinalizedMC.Remove(DelegateHandle);
+}
 
 bool USkeletalMeshComponent::MoveComponentImpl(const FVector& Delta, const FQuat& NewRotation, bool bSweep, FHitResult* OutHit /*= nullptr*/, EMoveComponentFlags MoveFlags /*= MOVECOMP_NoFlags*/, ETeleportType Teleport /*= ETeleportType::None*/)
 {
diff --git a/Engine/Source/Runtime/Engine/Private/GameplayStatics.cpp b/Engine/Source/Runtime/Engine/Private/GameplayStatics.cpp
--- a/Engine/Source/Runtime/Engine/Private/GameplayStatics.cpp
+++ b/Engine/Source/Runtime/Engine/Private/GameplayStatics.cpp
@@ -1463,7 +1463,7 @@ UAudioComponent* UGameplayStatics::SpawnSound2D(const UObject* WorldContextObjec
 	return AudioComponent;
 }
 
-void UGameplayStatics::PlaySoundAtLocation(const UObject* WorldContextObject, class USoundBase* Sound, FVector Location, FRotator Rotation, float VolumeMultiplier, float PitchMultiplier, float StartTime, class USoundAttenuation* AttenuationSettings, class USoundConcurrency* ConcurrencySettings, AActor* OwningActor)
+void UGameplayStatics::PlaySoundAtLocation(const UObject* WorldContextObject, class USoundBase* Sound, FVector Location, FRotator Rotation, float VolumeMultiplier, float PitchMultiplier, float StartTime, class USoundAttenuation* AttenuationSettings, class USoundConcurrency* ConcurrencySettings, AActor* OwningActor, UInitialActiveSoundParams* InitialParams)
 {
 	if (!Sound || !GEngine || !GEngine->UseSound())
 	{
@@ -1478,7 +1478,8 @@ void UGameplayStatics::PlaySoundAtLocation(const UObject* WorldContextObject, cl
 
 	if (FAudioDeviceHandle AudioDevice = ThisWorld->GetAudioDevice())
 	{
-		AudioDevice->PlaySoundAtLocation(Sound, ThisWorld, VolumeMultiplier, PitchMultiplier, StartTime, Location, Rotation, AttenuationSettings, ConcurrencySettings, nullptr, OwningActor);
+		const TArray<FAudioComponentParam>* Params = InitialParams ? &InitialParams->AudioParams : nullptr;
+		AudioDevice->PlaySoundAtLocation(Sound, ThisWorld, VolumeMultiplier, PitchMultiplier, StartTime, Location, Rotation, AttenuationSettings, ConcurrencySettings, Params, OwningActor);
 	}
 }
 
