nDisplay: fixed build errors

diff --git a/Engine/Plugins/Runtime/nDisplay/Source/DisplayCluster/Public/Blueprints/DisplayClusterBlueprintContainers.h b/Engine/Plugins/Runtime/nDisplay/Source/DisplayCluster/Public/Blueprints/DisplayClusterBlueprintContainers.h
--- a/Engine/Plugins/Runtime/nDisplay/Source/DisplayCluster/Public/Blueprints/DisplayClusterBlueprintContainers.h
+++ b/Engine/Plugins/Runtime/nDisplay/Source/DisplayCluster/Public/Blueprints/DisplayClusterBlueprintContainers.h
@@ -21,27 +21,27 @@ public:
 
 	// Location on a backbuffer.
 	UPROPERTY(EditAnywhere, Category = "NDisplay")
-	FIntPoint RectLocation;
+	FIntPoint RectLocation = FIntPoint::ZeroValue;
 
 	// Size on a backbuffer.
 	UPROPERTY(EditAnywhere, Category = "NDisplay")
-	FIntPoint RectSize;
+	FIntPoint RectSize = FIntPoint::ZeroValue;
 
 	// Camera view location
 	UPROPERTY(EditAnywhere, Category = "NDisplay")
-	FVector  ViewLocation;
+	FVector  ViewLocation = FVector::ZeroVector;
 
 	// Camera view rotation
 	UPROPERTY(EditAnywhere, Category = "NDisplay")
-	FRotator ViewRotation;
+	FRotator ViewRotation = FRotator::ZeroRotator;
 
 	// Camera projection Matrix
 	UPROPERTY(EditAnywhere, Category = "NDisplay")
-	FMatrix ProjectionMatrix;
+	FMatrix ProjectionMatrix = FMatrix::Identity;
 
 	// Rendering status for this viewport (Overlay and other configuration rules can disable rendering for this viewport.)
 	UPROPERTY(EditAnywhere, Category = "NDisplay")
-	bool bIsRendering;
+	bool bIsRendering = true;
 };
 
 USTRUCT()
@@ -56,11 +56,11 @@ public:
 
 	// Location on a backbuffer.
 	UPROPERTY(EditAnywhere, Category = "NDisplay")
-	FIntPoint RectLocation;
+	FIntPoint RectLocation = FIntPoint::ZeroValue;
 
 	// Size on a backbuffer.
 	UPROPERTY(EditAnywhere, Category = "NDisplay")
-	FIntPoint RectSize;
+	FIntPoint RectSize = FIntPoint::ZeroValue;
 
 	// Camera view location
 	UPROPERTY(EditAnywhere, Category = "NDisplay")
@@ -76,5 +76,5 @@ public:
 
 	// Rendering status for this viewport (Overlay and other configuration rules can disable rendering for this viewport.)
 	UPROPERTY(EditAnywhere, Category = "NDisplay")
-	bool bIsRendering;
+	bool bIsRendering = true;
 };
diff --git a/Engine/Plugins/Runtime/nDisplay/Source/DisplayCluster/Public/Cluster/DisplayClusterClusterEvent.h b/Engine/Plugins/Runtime/nDisplay/Source/DisplayCluster/Public/Cluster/DisplayClusterClusterEvent.h
--- a/Engine/Plugins/Runtime/nDisplay/Source/DisplayCluster/Public/Cluster/DisplayClusterClusterEvent.h
+++ b/Engine/Plugins/Runtime/nDisplay/Source/DisplayCluster/Public/Cluster/DisplayClusterClusterEvent.h
@@ -26,11 +26,11 @@ public:
 public:
 	// Is nDisplay internal event (should never be true for end users)
 	UPROPERTY(EditAnywhere, BlueprintReadWrite, meta = (DisplayName = "Is Sytem Event. 'True' is reserved for nDisplay internals."), Category = "NDisplay")
-	bool bIsSystemEvent;
+	bool bIsSystemEvent = false;
 
 	// Should older events with the same Name/Type/Category (for JSON) or ID (for binary) be discarded if a new one received
 	UPROPERTY(EditAnywhere, BlueprintReadWrite, Category = "NDisplay")
-	bool bShouldDiscardOnRepeat;
+	bool bShouldDiscardOnRepeat = true;
 };
 
 
@@ -83,7 +83,7 @@ struct FDisplayClusterClusterEventBinary
 public:
 	// Event ID (used for discarding outdated events)
 	UPROPERTY(EditAnywhere, BlueprintReadWrite, Category = "NDisplay")
-	int32 EventId;
+	int32 EventId = -1;
 
 	// Binary event data
 	UPROPERTY(EditAnywhere, BlueprintReadWrite, Category = "NDisplay")
diff --git a/Engine/Plugins/Runtime/nDisplay/Source/DisplayClusterConfiguration/Private/Formats/JSON426/DisplayClusterConfigurationJsonTypes_426.h b/Engine/Plugins/Runtime/nDisplay/Source/DisplayClusterConfiguration/Private/Formats/JSON426/DisplayClusterConfigurationJsonTypes_426.h
--- a/Engine/Plugins/Runtime/nDisplay/Source/DisplayClusterConfiguration/Private/Formats/JSON426/DisplayClusterConfigurationJsonTypes_426.h
+++ b/Engine/Plugins/Runtime/nDisplay/Source/DisplayClusterConfiguration/Private/Formats/JSON426/DisplayClusterConfigurationJsonTypes_426.h
@@ -361,16 +361,16 @@ public:
 	FString Camera;
 
 	UPROPERTY()
-	float BufferRatio;
+	float BufferRatio = 1.0f;
 
 	UPROPERTY()
-	int GPUIndex;
+	int32 GPUIndex = 0;
 
 	UPROPERTY()
-	bool AllowCrossGPUTransfer;
+	bool AllowCrossGPUTransfer = false;
 
 	UPROPERTY()
-	bool IsShared;
+	bool IsShared = false;
 
 	UPROPERTY()
 	FDisplayClusterConfigurationJsonRectangle_426 Region;
@@ -389,10 +389,10 @@ public:
 	FString Host;
 
 	UPROPERTY()
-	bool Sound;
+	bool Sound = true;
 
 	UPROPERTY()
-	bool FullScreen;
+	bool FullScreen = false;
 
 	UPROPERTY()
 	FDisplayClusterConfigurationJsonRectangle_426 Window;
@@ -430,13 +430,13 @@ struct FDisplayClusterConfigurationJsonDiagnostics_426
 
 public:
 	UPROPERTY()
-	bool SimulateLag;
+	bool SimulateLag = false;
 
 	UPROPERTY()
-	float MinLagTime;
+	float MinLagTime = 0.01f;
 
 	UPROPERTY()
-	float MaxLagTime;
+	float MaxLagTime = 0.3f;
 };
 
 
diff --git a/Engine/Plugins/Runtime/nDisplay/Source/DisplayClusterConfiguration/Private/Formats/JSON427/DisplayClusterConfigurationJsonTypes_427.h b/Engine/Plugins/Runtime/nDisplay/Source/DisplayClusterConfiguration/Private/Formats/JSON427/DisplayClusterConfigurationJsonTypes_427.h
--- a/Engine/Plugins/Runtime/nDisplay/Source/DisplayClusterConfiguration/Private/Formats/JSON427/DisplayClusterConfigurationJsonTypes_427.h
+++ b/Engine/Plugins/Runtime/nDisplay/Source/DisplayClusterConfiguration/Private/Formats/JSON427/DisplayClusterConfigurationJsonTypes_427.h
@@ -386,16 +386,16 @@ public:
 	FString Camera;
 
 	UPROPERTY()
-	float BufferRatio;
+	float BufferRatio = 1.0f;
 
 	UPROPERTY()
-	int GPUIndex;
+	int32 GPUIndex = 0;
 
 	UPROPERTY()
-	bool AllowCrossGPUTransfer;
+	bool AllowCrossGPUTransfer = false;
 
 	UPROPERTY()
-	bool IsShared;
+	bool IsShared = false;
 
 	UPROPERTY()
 	FDisplayClusterConfigurationJsonRectangle_427 Region;
@@ -433,10 +433,10 @@ public:
 	FString Host;
 
 	UPROPERTY()
-	bool Sound;
+	bool Sound = true;
 
 	UPROPERTY()
-	bool FullScreen;
+	bool FullScreen = false;
 
 	UPROPERTY()
 	FDisplayClusterConfigurationJsonRectangle_427 Window;
@@ -477,13 +477,13 @@ struct FDisplayClusterConfigurationJsonDiagnostics_427
 
 public:
 	UPROPERTY()
-	bool SimulateLag;
+	bool SimulateLag = false;
 
 	UPROPERTY()
-	float MinLagTime;
+	float MinLagTime = 0.01f;
 
 	UPROPERTY()
-	float MaxLagTime;
+	float MaxLagTime = 0.3f;
 };
 
 
diff --git a/Engine/Plugins/Runtime/nDisplay/Source/DisplayClusterConfiguration/Public/DisplayClusterConfigurationTypes.h b/Engine/Plugins/Runtime/nDisplay/Source/DisplayClusterConfiguration/Public/DisplayClusterConfigurationTypes.h
--- a/Engine/Plugins/Runtime/nDisplay/Source/DisplayClusterConfiguration/Public/DisplayClusterConfigurationTypes.h
+++ b/Engine/Plugins/Runtime/nDisplay/Source/DisplayClusterConfiguration/Public/DisplayClusterConfigurationTypes.h
@@ -442,13 +442,13 @@ struct DISPLAYCLUSTERCONFIGURATION_API FDisplayClusterConfigurationDiagnostics
 
 public:
 	UPROPERTY(EditAnywhere, Category = NDisplay)
-	bool bSimulateLag;
+	bool bSimulateLag = false;
 
 	UPROPERTY(EditAnywhere, Category = NDisplay)
-	float MinLagTime;
+	float MinLagTime = 0.01f;
 
 	UPROPERTY(EditAnywhere, Category = NDisplay)
-	float MaxLagTime;
+	float MaxLagTime = 0.3f;
 };
 
 struct FDisplayClusterConfigurationDataMetaInfo
diff --git a/Engine/Plugins/Runtime/nDisplay/Source/DisplayClusterConfiguration/Public/DisplayClusterConfigurationTypes_Base.h b/Engine/Plugins/Runtime/nDisplay/Source/DisplayClusterConfiguration/Public/DisplayClusterConfigurationTypes_Base.h
--- a/Engine/Plugins/Runtime/nDisplay/Source/DisplayClusterConfiguration/Public/DisplayClusterConfigurationTypes_Base.h
+++ b/Engine/Plugins/Runtime/nDisplay/Source/DisplayClusterConfiguration/Public/DisplayClusterConfigurationTypes_Base.h
@@ -20,11 +20,11 @@ struct DISPLAYCLUSTERCONFIGURATION_API FTextureCropOrigin
 public:
 	// Replace texture origin X location, in pixels
 	UPROPERTY(EditAnywhere, BlueprintReadWrite, Category = NDisplay, meta = (DisplayName = "X"))
-	int32 X;
+	int32 X = 0;
 
 	// Replace texture origin Y position, in pixels
 	UPROPERTY(EditAnywhere, BlueprintReadWrite, Category = NDisplay, meta = (DisplayName = "Y"))
-	int32 Y;
+	int32 Y = 0;
 };
 
 /**
@@ -38,11 +38,11 @@ struct DISPLAYCLUSTERCONFIGURATION_API FTextureCropSize
 public:
 	// Replace texture crop width, in pixels
 	UPROPERTY(EditAnywhere, BlueprintReadWrite, Category = NDisplay, meta = (DisplayName = "W"))
-	int32 W;
+	int32 W = 0;
 
 	// Replace texture crop height, in pixels
 	UPROPERTY(EditAnywhere, BlueprintReadWrite, Category = NDisplay, meta = (DisplayName = "H"))
-	int32 H;
+	int32 H = 0;
 };
 
 /**
diff --git a/Engine/Plugins/Runtime/nDisplay/Source/DisplayClusterConfiguration/Public/DisplayClusterConfigurationTypes_Postprocess.h b/Engine/Plugins/Runtime/nDisplay/Source/DisplayClusterConfiguration/Public/DisplayClusterConfigurationTypes_Postprocess.h
--- a/Engine/Plugins/Runtime/nDisplay/Source/DisplayClusterConfiguration/Public/DisplayClusterConfigurationTypes_Postprocess.h
+++ b/Engine/Plugins/Runtime/nDisplay/Source/DisplayClusterConfiguration/Public/DisplayClusterConfigurationTypes_Postprocess.h
@@ -55,18 +55,26 @@ public:
 USTRUCT(Blueprintable)
 struct DISPLAYCLUSTERCONFIGURATION_API FDisplayClusterConfigurationViewport_ColorGradingSettings
 {
-	GENERATED_BODY()
+	GENERATED_BODY();
+
+	FDisplayClusterConfigurationViewport_ColorGradingSettings()
+		: bOverride_Saturation(false)
+		, bOverride_Contrast(false)
+		, bOverride_Gamma(false)
+		, bOverride_Gain(false)
+		, bOverride_Offset(false)
+	{ }
 
 	UPROPERTY(EditAnywhere, BlueprintReadWrite, Category = "Overrides", meta = (PinHiddenByDefault, InlineEditConditionToggle))
-	uint8 bOverride_Saturation:1;
+	uint8 bOverride_Saturation : 1;
 	UPROPERTY(EditAnywhere, BlueprintReadWrite, Category = "Overrides", meta = (PinHiddenByDefault, InlineEditConditionToggle))
-	uint8 bOverride_Contrast:1;
+	uint8 bOverride_Contrast : 1;
 	UPROPERTY(EditAnywhere, BlueprintReadWrite, Category = "Overrides", meta = (PinHiddenByDefault, InlineEditConditionToggle))
-	uint8 bOverride_Gamma:1;
+	uint8 bOverride_Gamma : 1;
 	UPROPERTY(EditAnywhere, BlueprintReadWrite, Category = "Overrides", meta = (PinHiddenByDefault, InlineEditConditionToggle))
-	uint8 bOverride_Gain:1;
+	uint8 bOverride_Gain : 1;
 	UPROPERTY(EditAnywhere, BlueprintReadWrite, Category = "Overrides", meta = (PinHiddenByDefault, InlineEditConditionToggle))
-	uint8 bOverride_Offset:1;
+	uint8 bOverride_Offset : 1;
 
 	// Saturation
 	UPROPERTY(EditAnywhere, BlueprintReadWrite, Category = "Color Grading", meta = (UIMin = "0.0", UIMax = "2.0", Delta = "0.01", EditCondition = "bOverride_Saturation", ColorGradingMode = "saturation"))
diff --git a/Engine/Plugins/Runtime/nDisplay/Source/DisplayClusterShaders/Public/Blueprints/MPCDIGeometryData.h b/Engine/Plugins/Runtime/nDisplay/Source/DisplayClusterShaders/Public/Blueprints/MPCDIGeometryData.h
--- a/Engine/Plugins/Runtime/nDisplay/Source/DisplayClusterShaders/Public/Blueprints/MPCDIGeometryData.h
+++ b/Engine/Plugins/Runtime/nDisplay/Source/DisplayClusterShaders/Public/Blueprints/MPCDIGeometryData.h
@@ -13,10 +13,10 @@ struct FMPCDIGeometryImportData
 
 public:
 	UPROPERTY(EditAnywhere, BlueprintReadWrite, Category = "MPCDI")
-	int Width;
+	int32 Width = 0;
 
 	UPROPERTY(EditAnywhere, BlueprintReadWrite, Category = "MPCDI")
-	int Height;
+	int32 Height = 0;
 
 	UPROPERTY(EditAnywhere, BlueprintReadWrite, Category = "MPCDI")
 	TArray<FVector>   Vertices;
