nDisplay: fixed vertex declaration to adopt to LWC

diff --git a/Engine/Plugins/Runtime/nDisplay/Source/DisplayCluster/Private/Render/Containers/DisplayClusterRender_MeshComponentProxy.cpp b/Engine/Plugins/Runtime/nDisplay/Source/DisplayCluster/Private/Render/Containers/DisplayClusterRender_MeshComponentProxy.cpp
--- a/Engine/Plugins/Runtime/nDisplay/Source/DisplayCluster/Private/Render/Containers/DisplayClusterRender_MeshComponentProxy.cpp
+++ b/Engine/Plugins/Runtime/nDisplay/Source/DisplayCluster/Private/Render/Containers/DisplayClusterRender_MeshComponentProxy.cpp
@@ -71,39 +71,44 @@ void FDisplayClusterRender_MeshComponentProxy::UpdateRHI_RenderThread(FRHIComman
 
 		EBufferUsageFlags Usage = BUF_ShaderResource | BUF_Static;
 
-	// Create Vertex buffer RHI:
-	{
-			FRHIResourceCreateInfo CreateInfo(TEXT("DisplayClusterRender_MeshComponentProxy_VertexBuffer"));
-				;
-		size_t VertexDataSize = sizeof(FDisplayClusterMeshVertex) * NumVertices;
-		if (VertexDataSize == 0)
+		// Create Vertex buffer RHI:
 		{
-			UE_LOG(LogDisplayClusterRender, Warning, TEXT("MeshComponent has a vertex size of 0, please make sure a mesh is assigned."))
-			return;
-		}
+			size_t VertexDataSize = sizeof(FDisplayClusterMeshVertexType) * NumVertices;
+			if (VertexDataSize == 0)
+			{
+				UE_LOG(LogDisplayClusterRender, Warning, TEXT("MeshComponent has a vertex size of 0, please make sure a mesh is assigned."))
+				return;
+			}
 		
-		VertexBufferRHI = RHICreateVertexBuffer(VertexDataSize, Usage, CreateInfo);
-		FDisplayClusterMeshVertex* DestVertexData = reinterpret_cast<FDisplayClusterMeshVertex*>(RHILockBuffer(VertexBufferRHI, 0, VertexDataSize, RLM_WriteOnly));
-		if (DestVertexData)
-		{
-				FPlatformMemory::Memcpy(DestVertexData, InMeshData->GetVertexData().GetData(), VertexDataSize);
-			RHIUnlockBuffer(VertexBufferRHI);
+			FRHIResourceCreateInfo CreateInfo(TEXT("DisplayClusterRender_MeshComponentProxy_VertexBuffer"));
+			VertexBufferRHI = RHICreateVertexBuffer(VertexDataSize, Usage, CreateInfo);
+
+			FDisplayClusterMeshVertexType* DestVertexData = reinterpret_cast<FDisplayClusterMeshVertexType*>(RHILockBuffer(VertexBufferRHI, 0, VertexDataSize, RLM_WriteOnly));
+			if (DestVertexData)
+			{
+				const FDisplayClusterMeshVertex* SrcVertexData = InMeshData->GetVertexData().GetData();
+				for (uint32 i = 0; i < NumVertices; i++)
+				{
+					DestVertexData[i] = SrcVertexData[i];
+				}
+			
+				RHIUnlockBuffer(VertexBufferRHI);
+			}
 		}
-	}
 
-	// Create Index buffer RHI:
-	{
-			FRHIResourceCreateInfo CreateInfo(TEXT("DisplayClusterRender_MeshComponentProxy_IndexBuffer"));
+		// Create Index buffer RHI:
+		{
 			size_t IndexDataSize = sizeof(uint32) * InMeshData->GetIndexData().Num();
 
-		IndexBufferRHI = RHICreateIndexBuffer(sizeof(uint32), IndexDataSize, Usage, CreateInfo);
+			FRHIResourceCreateInfo CreateInfo(TEXT("DisplayClusterRender_MeshComponentProxy_IndexBuffer"));
+			IndexBufferRHI = RHICreateIndexBuffer(sizeof(uint32), IndexDataSize, Usage, CreateInfo);
 
-		uint32* DestIndexData = reinterpret_cast<uint32*>(RHILockBuffer(IndexBufferRHI, 0, IndexDataSize, RLM_WriteOnly));
-		if(DestIndexData)
-		{
+			uint32* DestIndexData = reinterpret_cast<uint32*>(RHILockBuffer(IndexBufferRHI, 0, IndexDataSize, RLM_WriteOnly));
+			if(DestIndexData)
+			{
 				FPlatformMemory::Memcpy(DestIndexData, InMeshData->GetIndexData().GetData(), IndexDataSize);
-			RHIUnlockBuffer(IndexBufferRHI);
+				RHIUnlockBuffer(IndexBufferRHI);
+			}
 		}
 	}
 }
-}
diff --git a/Engine/Plugins/Runtime/nDisplay/Source/DisplayCluster/Private/Render/Containers/DisplayClusterRender_MeshResources.cpp b/Engine/Plugins/Runtime/nDisplay/Source/DisplayCluster/Private/Render/Containers/DisplayClusterRender_MeshResources.cpp
--- a/Engine/Plugins/Runtime/nDisplay/Source/DisplayCluster/Private/Render/Containers/DisplayClusterRender_MeshResources.cpp
+++ b/Engine/Plugins/Runtime/nDisplay/Source/DisplayCluster/Private/Render/Containers/DisplayClusterRender_MeshResources.cpp
@@ -8,11 +8,11 @@ void FDisplayClusterMeshVertexDeclaration::InitRHI()
 {
 	FVertexDeclarationElementList Elements;
 
-	uint32 Stride = sizeof(FDisplayClusterMeshVertex);
+	uint32 Stride = sizeof(FDisplayClusterMeshVertexType);
 
-	Elements.Add(FVertexElement(0, STRUCT_OFFSET(FDisplayClusterMeshVertex, Position), VET_Float4, 0, Stride));
-	Elements.Add(FVertexElement(0, STRUCT_OFFSET(FDisplayClusterMeshVertex, UV), VET_Float2, 1, Stride));
-	Elements.Add(FVertexElement(0, STRUCT_OFFSET(FDisplayClusterMeshVertex, UV_Chromakey), VET_Float2, 2, Stride));
+	Elements.Add(FVertexElement(0, STRUCT_OFFSET(FDisplayClusterMeshVertexType, Position), VET_Float4, 0, Stride));
+	Elements.Add(FVertexElement(0, STRUCT_OFFSET(FDisplayClusterMeshVertexType, UV), VET_Float2, 1, Stride));
+	Elements.Add(FVertexElement(0, STRUCT_OFFSET(FDisplayClusterMeshVertexType, UV_Chromakey), VET_Float2, 2, Stride));
 
 	VertexDeclarationRHI = PipelineStateCache::GetOrCreateVertexDeclaration(Elements);
 }
diff --git a/Engine/Plugins/Runtime/nDisplay/Source/DisplayCluster/Private/Render/Containers/DisplayClusterRender_MeshResources.h b/Engine/Plugins/Runtime/nDisplay/Source/DisplayCluster/Private/Render/Containers/DisplayClusterRender_MeshResources.h
--- a/Engine/Plugins/Runtime/nDisplay/Source/DisplayCluster/Private/Render/Containers/DisplayClusterRender_MeshResources.h
+++ b/Engine/Plugins/Runtime/nDisplay/Source/DisplayCluster/Private/Render/Containers/DisplayClusterRender_MeshResources.h
@@ -7,6 +7,22 @@
 #include "RenderResource.h"
 #include "Render/Containers/DisplayClusterRender_MeshComponentTypes.h"
 
+
+/** The vertex data used to filter a texture. */
+struct FDisplayClusterMeshVertexType
+{
+	FVector4f Position;
+	FVector2D UV;
+	FVector2D UV_Chromakey;
+
+	void operator=(const FDisplayClusterMeshVertex& In)
+	{
+		Position = In.Position;
+		UV = In.UV;
+		UV_Chromakey = In.UV_Chromakey;
+	}
+};
+
 /** The filter vertex declaration resource type. */
 class FDisplayClusterMeshVertexDeclaration : public FRenderResource
 {
