nDisplay: fixed missing header for non unity build

diff --git a/Engine/Plugins/Runtime/nDisplay/Source/DisplayCluster/Private/Render/Containers/DisplayClusterRender_MeshComponent.cpp b/Engine/Plugins/Runtime/nDisplay/Source/DisplayCluster/Private/Render/Containers/DisplayClusterRender_MeshComponent.cpp
--- a/Engine/Plugins/Runtime/nDisplay/Source/DisplayCluster/Private/Render/Containers/DisplayClusterRender_MeshComponent.cpp
+++ b/Engine/Plugins/Runtime/nDisplay/Source/DisplayCluster/Private/Render/Containers/DisplayClusterRender_MeshComponent.cpp
@@ -4,6 +4,8 @@
 #include "Render/Containers/DisplayClusterRender_MeshComponentProxy.h"
 #include "Render/Containers/DisplayClusterRender_MeshComponentProxyData.h"
 
+#include "Engine/StaticMesh.h"
+
 #include "Misc/DisplayClusterLog.h"
 
 static void ImplUpdateMeshComponentProxyData(FDisplayClusterRender_MeshComponentProxy* MeshComponentProxy, FDisplayClusterRender_MeshComponentProxyData* NewProxyData)
