Do not reset the nodes ready.  Basically there is only one chance to
synchronize nDisplay nodes which is on the first startup.

diff --git a/Engine/Plugins/Runtime/nDisplay/Source/DisplayClusterMessageInterceptor/Private/DisplayClusterMessageInterceptionModule.cpp b/Engine/Plugins/Runtime/nDisplay/Source/DisplayClusterMessageInterceptor/Private/DisplayClusterMessageInterceptionModule.cpp
--- a/Engine/Plugins/Runtime/nDisplay/Source/DisplayClusterMessageInterceptor/Private/DisplayClusterMessageInterceptionModule.cpp
+++ b/Engine/Plugins/Runtime/nDisplay/Source/DisplayClusterMessageInterceptor/Private/DisplayClusterMessageInterceptionModule.cpp
@@ -370,7 +370,8 @@ private:
 	{
 		UE_LOG(LogDisplayClusterInterception, Display, TEXT("Temporarily disabling multi-user workspace sync."));
 		bCanFinalizeWorkspace = false;
-		NodesReady.Empty();
+		// Note, we intentially do not reset NodesReady here because of the timing of joining the MU session thereby causing the node
+		// to forever not receive updates.
 	}
 
 	bool bWasEverDisconnected = false;
