nDisplay: Fixed issues with property categorization that was causing incorrect grouping in the DisplayClusterRootActor details panel. Refactored details customization a bit to prepare for future changes to the display of root actor properties.

diff --git a/Engine/Plugins/Runtime/nDisplay/Source/DisplayCluster/Public/DisplayClusterRootActor.h b/Engine/Plugins/Runtime/nDisplay/Source/DisplayCluster/Public/DisplayClusterRootActor.h
--- a/Engine/Plugins/Runtime/nDisplay/Source/DisplayCluster/Public/DisplayClusterRootActor.h
+++ b/Engine/Plugins/Runtime/nDisplay/Source/DisplayCluster/Public/DisplayClusterRootActor.h
@@ -163,10 +163,10 @@ public:
 	UDisplayClusterSceneComponent* GetComponentById(const FString& ComponentId) const;
 
 public:
-	UPROPERTY(EditAnywhere, Category = "NDisplay|Settings", NoClear)
+	UPROPERTY(EditAnywhere, Category = "NDisplay", NoClear)
 	UDisplayClusterConfigurationICVFX_StageSettings* StageSettings;
 
-	UPROPERTY(EditAnywhere, Category = "NDisplay|Settings", NoClear)
+	UPROPERTY(EditAnywhere, Category = "NDisplay", NoClear)
 	UDisplayClusterConfigurationRenderFrame* RenderFrameSettings;
 
 	IDisplayClusterViewportManager* GetViewportManager() const
@@ -195,7 +195,7 @@ private:
 	/**
 	 * If set from the DisplayCluster BP Compiler it will be loaded from the class default subobjects in run-time.
 	 */
-	UPROPERTY(VisibleAnywhere, Instanced, Category = "Config Data")
+	UPROPERTY(VisibleAnywhere, Instanced, Category = "NDisplay")
 	UDisplayClusterConfigurationData* CurrentConfigData;
 
 	/**
diff --git a/Engine/Plugins/Runtime/nDisplay/Source/DisplayClusterConfiguration/Public/DisplayClusterConfigurationStrings.h b/Engine/Plugins/Runtime/nDisplay/Source/DisplayClusterConfiguration/Public/DisplayClusterConfigurationStrings.h
--- a/Engine/Plugins/Runtime/nDisplay/Source/DisplayClusterConfiguration/Public/DisplayClusterConfigurationStrings.h
+++ b/Engine/Plugins/Runtime/nDisplay/Source/DisplayClusterConfiguration/Public/DisplayClusterConfigurationStrings.h
@@ -14,6 +14,15 @@ namespace DisplayClusterConfigurationStrings
 		}
 	}
 
+	// Property Categories
+	namespace categories
+	{
+		static constexpr auto DefaultCategory = TEXT("NDisplay");
+		static constexpr auto ClusterCategory = TEXT("NDisplay Cluster");
+		static constexpr auto ConfigurationCategory = TEXT("NDisplay Configuration");
+		static constexpr auto PreviewCategory = TEXT("NDisplay Preview (Editor only)");
+	}
+
 	// Command line arguments
 	namespace args
 	{
diff --git a/Engine/Plugins/Runtime/nDisplay/Source/DisplayClusterConfigurator/Private/Views/Details/DisplayClusterConfiguratorDetailCustomization.cpp b/Engine/Plugins/Runtime/nDisplay/Source/DisplayClusterConfigurator/Private/Views/Details/DisplayClusterConfiguratorDetailCustomization.cpp
--- a/Engine/Plugins/Runtime/nDisplay/Source/DisplayClusterConfigurator/Private/Views/Details/DisplayClusterConfiguratorDetailCustomization.cpp
+++ b/Engine/Plugins/Runtime/nDisplay/Source/DisplayClusterConfigurator/Private/Views/Details/DisplayClusterConfiguratorDetailCustomization.cpp
@@ -30,19 +30,37 @@
 #include "Widgets/Text/STextBlock.h"
 #include "Widgets/SBoxPanel.h"
 
+#define BEGIN_CATEGORY(CategoryName) { \
+	IDetailCategoryBuilder& CurrentCategory = InLayoutBuilder.EditCategory(CategoryName);
+
+#define END_CATEGORY() }
+
+#define ADD_PROPERTY(ClassName, PropertyName) { \
+	TSharedRef<IPropertyHandle> PropertyHandle = InLayoutBuilder.GetProperty(GET_MEMBER_NAME_CHECKED(ClassName, PropertyName)); \
+	check(PropertyHandle->IsValidHandle()); \
+	CurrentCategory.AddProperty(PropertyHandle); \
+}
+
+#define ADD_EXPANDED_PROPERTY(ClassName, PropertyName) { \
+	TSharedRef<IPropertyHandle> PropertyHandle = InLayoutBuilder.GetProperty(GET_MEMBER_NAME_CHECKED(ClassName, PropertyName)); \
+	check(PropertyHandle->IsValidHandle()); \
+	CurrentCategory.AddProperty(PropertyHandle).ShouldAutoExpand(true); \
+}
+
+#define ADD_CUSTOM_PROPERTY(FilterText) CurrentCategory.AddCustomRow(FilterText)
+
+#define REPLACE_PROPERTY_WITH_CUSTOM(ClassName, PropertyName, Widget) { \
+	TSharedRef<IPropertyHandle> PropertyHandle = InLayoutBuilder.GetProperty(GET_MEMBER_NAME_CHECKED(ClassName, PropertyName)); \
+	check(PropertyHandle->IsValidHandle()); \
+	InLayoutBuilder.HideProperty(PropertyHandle); \
+	CurrentCategory.AddCustomRow(PropertyHandle->GetPropertyDisplayName()).NameContent()[PropertyHandle->CreatePropertyNameWidget()].ValueContent()[Widget]; \
+}
 
 #define LOCTEXT_NAMESPACE "FDisplayClusterConfiguratorDetailCustomization"
 
 //////////////////////////////////////////////////////////////////////////////////////////////
 // Base UCLASS Detail Customization
 //////////////////////////////////////////////////////////////////////////////////////////////
-
-FDisplayClusterConfiguratorDetailCustomization::FDisplayClusterConfiguratorDetailCustomization()
-	: ToolkitPtr(nullptr)
-	, LayoutBuilder(nullptr)
-	, NDisplayCategory(nullptr)
-{}
-
 void FDisplayClusterConfiguratorDetailCustomization::CustomizeDetails(IDetailLayoutBuilder& InLayoutBuilder)
 {
 	UObject* ObjectBeingEdited = nullptr;
@@ -72,20 +90,14 @@ void FDisplayClusterConfiguratorDetailCustomization::CustomizeDetails(IDetailLay
 
 	check(RootActorPtr.IsValid() || ToolkitPtr.IsValid());
 	
-	LayoutBuilder = &InLayoutBuilder;
-	NDisplayCategory = &LayoutBuilder->EditCategory("NDisplay", FText::GetEmpty());
-	NDisplayCategory->InitiallyCollapsed(false);
-
-	// Hide properties that should only be visible on the instance. VisibleInstanceOnly doesn't work with these properties.
-	TArray<TSharedRef<IPropertyHandle>> CategoryProperties;
-	NDisplayCategory->GetDefaultProperties(CategoryProperties);
-
-	for (TSharedRef<IPropertyHandle>& PropertyHandle : CategoryProperties)
+	// Iterate over all of the properties in the object being edited to find properties marked with specific custom metadata tags and hide those properties if necessary
+	for (TFieldIterator<FProperty> It(ObjectBeingEdited->GetClass()); It; ++It)
 	{
-		FProperty* Property = PropertyHandle->GetProperty();
+		FProperty* Property = *It;
+		TSharedRef<IPropertyHandle> PropertyHandle = InLayoutBuilder.GetProperty(Property->GetFName());
 
-		const bool bShouldHide = PropertyHandle->HasMetaData("nDisplayHidden") || 
-			(IsRunningForBlueprintEditor() ? 
+		const bool bShouldHide = PropertyHandle->HasMetaData("nDisplayHidden") ||
+			(IsRunningForBlueprintEditor() ?
 				PropertyHandle->HasMetaData("nDisplayInstanceOnly") || Property && Property->HasAnyPropertyFlags(CPF_DisableEditOnTemplate) :
 				false);
 
@@ -154,17 +166,11 @@ void FDisplayClusterConfiguratorDataDetailCustomization::CustomizeDetails(IDetai
 {
 	Super::CustomizeDetails(InLayoutBuilder);
 
-	TSharedRef<IPropertyHandle> InfoHandle = InLayoutBuilder.GetProperty(GET_MEMBER_NAME_CHECKED(UDisplayClusterConfigurationData, Info));
-	check(InfoHandle->IsValidHandle());
-	NDisplayCategory->AddProperty(InfoHandle).ShouldAutoExpand(true);
-
-	TSharedRef<IPropertyHandle> DiagnosticsHandle = InLayoutBuilder.GetProperty(GET_MEMBER_NAME_CHECKED(UDisplayClusterConfigurationData, Diagnostics));
-	check(DiagnosticsHandle->IsValidHandle());
-	NDisplayCategory->AddProperty(DiagnosticsHandle).ShouldAutoExpand(true);
-
-	TSharedPtr<IPropertyHandle> CustomParamsHandle = InLayoutBuilder.GetProperty(GET_MEMBER_NAME_CHECKED(UDisplayClusterConfigurationData, CustomParameters));
-	check(CustomParamsHandle->IsValidHandle());
-	NDisplayCategory->AddProperty(CustomParamsHandle).ShouldAutoExpand(false);
+	BEGIN_CATEGORY(DisplayClusterConfigurationStrings::categories::ConfigurationCategory)
+		ADD_EXPANDED_PROPERTY(UDisplayClusterConfigurationData, Info);
+		ADD_EXPANDED_PROPERTY(UDisplayClusterConfigurationData, Diagnostics);
+		ADD_PROPERTY(UDisplayClusterConfigurationData, CustomParameters);
+	END_CATEGORY()
 }
 
 //////////////////////////////////////////////////////////////////////////////////////////////
@@ -174,48 +180,41 @@ void FDisplayClusterConfiguratorClusterDetailCustomization::CustomizeDetails(IDe
 {
 	Super::CustomizeDetails(InLayoutBuilder);
 
-	NDisplayCategory->InitiallyCollapsed(false);
-
-	TSharedPtr<IPropertyHandle> MasterNodeHandle = InLayoutBuilder.GetProperty(GET_MEMBER_NAME_CHECKED(UDisplayClusterConfigurationCluster, MasterNode), UDisplayClusterConfigurationCluster::StaticClass());
-	check(MasterNodeHandle->IsValidHandle());
-
-	TSharedPtr<IPropertyHandle> SyncHandle = InLayoutBuilder.GetProperty(GET_MEMBER_NAME_CHECKED(UDisplayClusterConfigurationCluster, Sync), UDisplayClusterConfigurationCluster::StaticClass());
-	check(SyncHandle->IsValidHandle());
-
-	TSharedPtr<IPropertyHandle> NetworkHandle = InLayoutBuilder.GetProperty(GET_MEMBER_NAME_CHECKED(UDisplayClusterConfigurationCluster, Network), UDisplayClusterConfigurationCluster::StaticClass());
-	check(NetworkHandle->IsValidHandle());
-
-	NDisplayCategory->AddProperty(MasterNodeHandle).ShouldAutoExpand(true);
-	NDisplayCategory->AddProperty(SyncHandle).ShouldAutoExpand(true);
-	NDisplayCategory->AddProperty(NetworkHandle).ShouldAutoExpand(true);
-
+	// Store the Nodes property handle for use later
 	ClusterNodesHandle = InLayoutBuilder.GetProperty(GET_MEMBER_NAME_CHECKED(UDisplayClusterConfigurationCluster, Nodes));
 	check(ClusterNodesHandle->IsValidHandle());
 
-	if (!IsRunningForBlueprintEditor())
-	{
-		// Add a reset cluster nodes default button. Normal reset to defaults on the map won't display because of EditFixedSize flag.
-		// This is also handy since UE will stop propagating CDO container structure changes when an instanced value is modified.
-		
-		NDisplayCategory->AddCustomRow(ClusterNodesHandle->GetDefaultCategoryText())
-		.NameContent()
-		[
-			SNew(SHorizontalBox)
-			+ SHorizontalBox::Slot()
-			.FillWidth(1.0f)
-			.HAlign(HAlign_Fill)
-			[
-				SNew(SButton)
-				.Text(LOCTEXT("ResetClusterNodesButton_Label", "Reset Cluster Nodes"))
-				.ToolTipText(LOCTEXT("ResetClusterNodesButton_Tooltip", "Reset all cluster nodes to class defaults."))
-				.OnClicked(FOnClicked::CreateLambda([this]()
-				{
-					ClusterNodesHandle->ResetToDefault();
-					return FReply::Handled();
-				}))
-			]
-		];
-	}
+	BEGIN_CATEGORY(DisplayClusterConfigurationStrings::categories::ClusterCategory)
+		ADD_EXPANDED_PROPERTY(UDisplayClusterConfigurationCluster, MasterNode);
+		ADD_EXPANDED_PROPERTY(UDisplayClusterConfigurationCluster, Sync);
+		ADD_EXPANDED_PROPERTY(UDisplayClusterConfigurationCluster, Network);
+		ADD_PROPERTY(UDisplayClusterConfigurationCluster, bUseOverallClusterPostProcess);
+		ADD_PROPERTY(UDisplayClusterConfigurationCluster, OverallClusterPostProcessSettings);
+
+		if (!IsRunningForBlueprintEditor())
+		{
+			ADD_CUSTOM_PROPERTY(LOCTEXT("ResetClusterNodesButton_Label", "Reset Cluster Nodes"))
+				.NameContent()
+				[
+					SNew(SHorizontalBox)
+					+ SHorizontalBox::Slot()
+					.FillWidth(1.0f)
+					.HAlign(HAlign_Fill)
+					[
+						SNew(SButton)
+						.Text(LOCTEXT("ResetClusterNodesButton_Label", "Reset Cluster Nodes"))
+						.ToolTipText(LOCTEXT("ResetClusterNodesButton_Tooltip", "Reset all cluster nodes to class defaults."))
+						.OnClicked(FOnClicked::CreateLambda([this]()
+						{
+							ClusterNodesHandle->ResetToDefault();
+							return FReply::Handled();
+						}))
+					]
+				];
+		}
+
+		ADD_PROPERTY(UDisplayClusterConfigurationCluster, Nodes);
+	END_CATEGORY()
 }
 
 //////////////////////////////////////////////////////////////////////////////////////////////
@@ -241,13 +240,14 @@ void FDisplayClusterConfiguratorViewportDetailCustomization::CustomizeDetails(ID
 	}
 	check(ConfigurationViewportPtr != nullptr);
 
-	// Hide properties
-	CameraHandle = InLayoutBuilder.GetProperty(GET_MEMBER_NAME_CHECKED(UDisplayClusterConfigurationViewport, Camera), UDisplayClusterConfigurationViewport::StaticClass());
-	check(CameraHandle->IsValidHandle());
-	InLayoutBuilder.HideProperty(CameraHandle);
+	CameraHandle = InLayoutBuilder.GetProperty(GET_MEMBER_NAME_CHECKED(UDisplayClusterConfigurationViewport, Camera));
+	check(CameraHandle.IsValid());
 
 	ResetCameraOptions();
-	AddCameraRow();
+
+	BEGIN_CATEGORY(CameraHandle->GetDefaultCategoryName())
+		REPLACE_PROPERTY_WITH_CUSTOM(UDisplayClusterConfigurationViewport, Camera, CreateCustomCameraWidget());
+	END_CATEGORY()
 }
 
 void FDisplayClusterConfiguratorViewportDetailCustomization::ResetCameraOptions()
@@ -281,21 +281,14 @@ void FDisplayClusterConfiguratorViewportDetailCustomization::ResetCameraOptions(
 	}
 }
 
-void FDisplayClusterConfiguratorViewportDetailCustomization::AddCameraRow()
+TSharedRef<SWidget> FDisplayClusterConfiguratorViewportDetailCustomization::CreateCustomCameraWidget()
 {
 	if (CameraComboBox.IsValid())
 	{
-		return;
+		return CameraComboBox.ToSharedRef();
 	}
-	
-	NDisplayCategory->AddCustomRow(CameraHandle->GetPropertyDisplayName())
-	.NameContent()
-	[
-		CameraHandle->CreatePropertyNameWidget()
-	]
-	.ValueContent()
-	[
-		SAssignNew(CameraComboBox, SDisplayClusterConfigurationSearchableComboBox)
+
+	return SAssignNew(CameraComboBox, SDisplayClusterConfigurationSearchableComboBox)
 		.OptionsSource(&CameraOptions)
 		.OnGenerateWidget(this, &FDisplayClusterConfiguratorViewportDetailCustomization::MakeCameraOptionComboWidget)
 		.OnSelectionChanged(this, &FDisplayClusterConfiguratorViewportDetailCustomization::OnCameraSelected)
@@ -305,8 +298,7 @@ void FDisplayClusterConfiguratorViewportDetailCustomization::AddCameraRow()
 		[
 			SNew(STextBlock)
 			.Text(this, &FDisplayClusterConfiguratorViewportDetailCustomization::GetSelectedCameraText)
-		]
-	];
+		];
 }
 
 TSharedRef<SWidget> FDisplayClusterConfiguratorViewportDetailCustomization::MakeCameraOptionComboWidget(TSharedPtr<FString> InItem)
diff --git a/Engine/Plugins/Runtime/nDisplay/Source/DisplayClusterConfigurator/Private/Views/Details/DisplayClusterConfiguratorDetailCustomization.h b/Engine/Plugins/Runtime/nDisplay/Source/DisplayClusterConfigurator/Private/Views/Details/DisplayClusterConfiguratorDetailCustomization.h
--- a/Engine/Plugins/Runtime/nDisplay/Source/DisplayClusterConfigurator/Private/Views/Details/DisplayClusterConfiguratorDetailCustomization.h
+++ b/Engine/Plugins/Runtime/nDisplay/Source/DisplayClusterConfigurator/Private/Views/Details/DisplayClusterConfiguratorDetailCustomization.h
@@ -42,8 +42,6 @@ class FDisplayClusterConfiguratorDetailCustomization
 	: public IDetailCustomization
 {
 public:
-	FDisplayClusterConfiguratorDetailCustomization();
-
 	/** IDetailCustomization interface */
 	virtual void CustomizeDetails(IDetailLayoutBuilder& InLayoutBuilder) override;
 	/** End IDetailCustomization interface */
@@ -65,12 +63,8 @@ protected:
 	bool IsRunningForBlueprintEditor() const { return ToolkitPtr.IsValid(); }
 
 protected:
-	TWeakPtr<FDisplayClusterConfiguratorBlueprintEditor> ToolkitPtr;
+	TWeakPtr<FDisplayClusterConfiguratorBlueprintEditor> ToolkitPtr = nullptr;
 	TWeakObjectPtr<ADisplayClusterRootActor> RootActorPtr;
-	
-	IDetailLayoutBuilder* LayoutBuilder;
-
-	IDetailCategoryBuilder* NDisplayCategory;
 };
 
 /**
@@ -117,7 +111,7 @@ public:
 private:
 	void ResetCameraOptions();
 
-	void AddCameraRow();
+	TSharedRef<SWidget> CreateCustomCameraWidget();
 
 	TSharedRef<SWidget> MakeCameraOptionComboWidget(TSharedPtr<FString> InItem);
 
diff --git a/Engine/Plugins/Runtime/nDisplay/Source/DisplayClusterEditor/Private/DetailsCustomization/DisplayClusterRootActorDetailsCustomization.cpp b/Engine/Plugins/Runtime/nDisplay/Source/DisplayClusterEditor/Private/DetailsCustomization/DisplayClusterRootActorDetailsCustomization.cpp
--- a/Engine/Plugins/Runtime/nDisplay/Source/DisplayClusterEditor/Private/DetailsCustomization/DisplayClusterRootActorDetailsCustomization.cpp
+++ b/Engine/Plugins/Runtime/nDisplay/Source/DisplayClusterEditor/Private/DetailsCustomization/DisplayClusterRootActorDetailsCustomization.cpp
@@ -15,6 +15,33 @@
 #include "DisplayClusterConfigurationTypes.h"
 #include "DisplayClusterConfigurationTypes_Viewport.h"
 
+// TODO: Macros duplicated from DisplayClusterConfiguratorDetailCustomization.h. Eventually, we will want to move the RootActor details customization into the same module
+// as the other details customizations to unify this and remove these duplicate macros.
+#define BEGIN_CATEGORY(CategoryName) { \
+	IDetailCategoryBuilder& CurrentCategory = InLayoutBuilder.EditCategory(CategoryName);
+
+#define END_CATEGORY() }
+
+#define ADD_PROPERTY(ClassName, PropertyName) { \
+	TSharedRef<IPropertyHandle> PropertyHandle = InLayoutBuilder.GetProperty(GET_MEMBER_NAME_CHECKED(ClassName, PropertyName)); \
+	check(PropertyHandle->IsValidHandle()); \
+	CurrentCategory.AddProperty(PropertyHandle); \
+}
+
+#define ADD_EXPANDED_PROPERTY(ClassName, PropertyName) { \
+	TSharedRef<IPropertyHandle> PropertyHandle = InLayoutBuilder.GetProperty(GET_MEMBER_NAME_CHECKED(ClassName, PropertyName)); \
+	check(PropertyHandle->IsValidHandle()); \
+	CurrentCategory.AddProperty(PropertyHandle).ShouldAutoExpand(true); \
+}
+
+#define ADD_CUSTOM_PROPERTY(FilterText) CurrentCategory.AddCustomRow(FilterText)
+
+#define REPLACE_PROPERTY_WITH_CUSTOM(ClassName, PropertyName, Widget) { \
+	TSharedRef<IPropertyHandle> PropertyHandle = InLayoutBuilder.GetProperty(GET_MEMBER_NAME_CHECKED(ClassName, PropertyName)); \
+	check(PropertyHandle->IsValidHandle()); \
+	InLayoutBuilder.HideProperty(PropertyHandle); \
+	CurrentCategory.AddCustomRow(PropertyHandle->GetPropertyDisplayName()).NameContent()[PropertyHandle->CreatePropertyNameWidget()].ValueContent()[Widget]; \
+}
 
 #define LOCTEXT_NAMESPACE "DisplayClusterRootActorDetailsCustomization"
 
@@ -26,8 +53,6 @@ TSharedRef<IDetailCustomization> FDisplayClusterRootActorDetailsCustomization::M
 
 void FDisplayClusterRootActorDetailsCustomization::CustomizeDetails(IDetailLayoutBuilder& InLayoutBuilder)
 {
-	LayoutBuilder = &InLayoutBuilder;
-
 	// Hide the following categories, we don't really need them
 	InLayoutBuilder.HideCategory(TEXT("Rendering"));
 	InLayoutBuilder.HideCategory(TEXT("Replication"));
@@ -51,80 +76,63 @@ void FDisplayClusterRootActorDetailsCustomization::CustomizeDetails(IDetailLayou
 		return;
 	}
 
-	// Store preview category
-	CategoryPreview = &InLayoutBuilder.EditCategory("NDisplay Preview (Editor only)");
-	check(CategoryPreview);
+	// Store the NodeId property handle to use later
+	PropertyNodeId = InLayoutBuilder.GetProperty(GET_MEMBER_NAME_CHECKED(ADisplayClusterRootActor, PreviewNodeId), ADisplayClusterRootActor::StaticClass());
+	check(PropertyNodeId->IsValidHandle());
+
+	// Lay out all of the root actor properties into their correct orders and categories. Because we are adding custom properties,
+	// most of the actor's properties will need to be laid out manually here even if they aren't being customized.
+	BuildLayout(InLayoutBuilder);
 
-	// Finally, do the customization
-	BuildLayout();
+	// Update the selected item in the NodeId combo box to match the current value on the root actor
+	UpdateNodeIdSelection();
 }
 
-void FDisplayClusterRootActorDetailsCustomization::BuildLayout()
+void FDisplayClusterRootActorDetailsCustomization::BuildLayout(IDetailLayoutBuilder& InLayoutBuilder)
 {
-	// Setup preivew cluster node ID combobox
-	AddNodeIdRow();
+	// Force a particular order that the property categories show up in the details panel
+	InLayoutBuilder.EditCategory(DisplayClusterConfigurationStrings::categories::DefaultCategory);
+	InLayoutBuilder.EditCategory(DisplayClusterConfigurationStrings::categories::ConfigurationCategory);
+	InLayoutBuilder.EditCategory(DisplayClusterConfigurationStrings::categories::ClusterCategory);
+	InLayoutBuilder.EditCategory(DisplayClusterConfigurationStrings::categories::PreviewCategory);
+
+	// Add custom properties and lay out/order properties into their correct categories.
+	BEGIN_CATEGORY(DisplayClusterConfigurationStrings::categories::PreviewCategory)
+		if (RebuildNodeIdOptionsList())
+		{
+			REPLACE_PROPERTY_WITH_CUSTOM(ADisplayClusterRootActor, PreviewNodeId, CreateCustomNodeIdWidget());
+		}
+	END_CATEGORY();
 }
 
 
 //////////////////////////////////////////////////////////////////////////////////////////////
 // Preview node ID
 //////////////////////////////////////////////////////////////////////////////////////////////
-void FDisplayClusterRootActorDetailsCustomization::AddNodeIdRow()
+TSharedRef<SWidget> FDisplayClusterRootActorDetailsCustomization::CreateCustomNodeIdWidget()
 {
-	// Hide PreviewNodeId property, it will be replaced with a combobox
-	PropertyNodeId = LayoutBuilder->GetProperty(GET_MEMBER_NAME_CHECKED(ADisplayClusterRootActor, PreviewNodeId), ADisplayClusterRootActor::StaticClass());
-	check(PropertyNodeId->IsValidHandle());
-	LayoutBuilder->HideProperty(PropertyNodeId);
-
-	if (!RebuildNodeIdOptionsList())
-	{
-		return;
-	}
-
-	// Create GUI representation
+	if (NodeIdComboBox.IsValid())
 	{
-		CategoryPreview->AddCustomRow(PropertyNodeId->GetPropertyDisplayName())
-			.NameContent()
-			[
-				PropertyNodeId->CreatePropertyNameWidget()
-			]
-			.ValueContent()
-			[
-				SAssignNew(NodeIdComboBox, SSearchableComboBox)
-				.OptionsSource(&NodeIdOptions)
-				.OnGenerateWidget(this, &FDisplayClusterRootActorDetailsCustomization::CreateComboWidget)
-				.OnSelectionChanged(this, &FDisplayClusterRootActorDetailsCustomization::OnNodeIdSelected)
-				.ContentPadding(2)
-				.Content()
-				[
-					SNew(STextBlock)
-					.Text(this, &FDisplayClusterRootActorDetailsCustomization::GetSelectedNodeIdText)
-				]
-			];
+		return NodeIdComboBox.ToSharedRef();
 	}
 
-	const FString& CurrentPreviewNode = EditedObject->PreviewNodeId;
-	TSharedPtr<FString>* FoundItem = NodeIdOptions.FindByPredicate([CurrentPreviewNode](const TSharedPtr<FString>& Item)
-	{
-		return Item->Equals(CurrentPreviewNode, ESearchCase::IgnoreCase);
-	});
-
-	if (FoundItem)
-	{
-		NodeIdComboBox->SetSelectedItem(*FoundItem);
-	}
-	else
-	{
-		// Set combobox selected item (options list is not empty here)
-		NodeIdComboBox->SetSelectedItem(NodeIdOptionAll);
-	}
+	return SAssignNew(NodeIdComboBox, SSearchableComboBox)
+		.OptionsSource(&NodeIdOptions)
+		.OnGenerateWidget(this, &FDisplayClusterRootActorDetailsCustomization::CreateComboWidget)
+		.OnSelectionChanged(this, &FDisplayClusterRootActorDetailsCustomization::OnNodeIdSelected)
+		.ContentPadding(2)
+		.Content()
+		[
+			SNew(STextBlock)
+			.Text(this, &FDisplayClusterRootActorDetailsCustomization::GetSelectedNodeIdText)
+		];
 }
 
 bool FDisplayClusterRootActorDetailsCustomization::RebuildNodeIdOptionsList()
 {
 	// Get current configuration data
 	const UDisplayClusterConfigurationData* ConfigurationData = EditedObject->GetConfigData();
-	if (!ConfigurationData || !CategoryPreview)
+	if (!ConfigurationData)
 	{
 		return false;
 	}
@@ -156,6 +164,25 @@ bool FDisplayClusterRootActorDetailsCustomization::RebuildNodeIdOptionsList()
 	return NodeIdOptions.Num() >= 3;
 }
 
+void FDisplayClusterRootActorDetailsCustomization::UpdateNodeIdSelection()
+{
+	const FString& CurrentPreviewNode = EditedObject->PreviewNodeId;
+	TSharedPtr<FString>* FoundItem = NodeIdOptions.FindByPredicate([CurrentPreviewNode](const TSharedPtr<FString>& Item)
+	{
+		return Item->Equals(CurrentPreviewNode, ESearchCase::IgnoreCase);
+	});
+
+	if (FoundItem)
+	{
+		NodeIdComboBox->SetSelectedItem(*FoundItem);
+	}
+	else
+	{
+		// Set combobox selected item (options list is not empty here)
+		NodeIdComboBox->SetSelectedItem(NodeIdOptionAll);
+	}
+}
+
 void FDisplayClusterRootActorDetailsCustomization::OnNodeIdSelected(TSharedPtr<FString> PreviewNodeId, ESelectInfo::Type SelectInfo)
 {
 	const FString NewValue = (PreviewNodeId.IsValid() ? *PreviewNodeId : *NodeIdOptionNone);
@@ -181,7 +208,7 @@ bool FDisplayClusterRootActorDetailsCustomization::RebuildDefaultCameraOptionsLi
 {
 	// Get current configuration data
 	const UDisplayClusterConfigurationData* ConfigurationData = EditedObject->GetConfigData();
-	if (!ConfigurationData || !CategoryPreview)
+	if (!ConfigurationData)
 	{
 		return false;
 	}
diff --git a/Engine/Plugins/Runtime/nDisplay/Source/DisplayClusterEditor/Private/DetailsCustomization/DisplayClusterRootActorDetailsCustomization.h b/Engine/Plugins/Runtime/nDisplay/Source/DisplayClusterEditor/Private/DetailsCustomization/DisplayClusterRootActorDetailsCustomization.h
--- a/Engine/Plugins/Runtime/nDisplay/Source/DisplayClusterEditor/Private/DetailsCustomization/DisplayClusterRootActorDetailsCustomization.h
+++ b/Engine/Plugins/Runtime/nDisplay/Source/DisplayClusterEditor/Private/DetailsCustomization/DisplayClusterRootActorDetailsCustomization.h
@@ -22,18 +22,20 @@ public:
 
 protected:
 	// IDetailCustomization interface
-	virtual void CustomizeDetails(class IDetailLayoutBuilder& LayoutBuilder) override;
+	virtual void CustomizeDetails(IDetailLayoutBuilder& InLayoutBuilder) override;
 	// End IDetailCustomization interface
 
 protected:
-	void BuildLayout();
+	void BuildLayout(IDetailLayoutBuilder& InLayoutBuilder);
 
 protected:
 	//////////////////////////////////////////////////////////////////////////////////////////////
 	// Node ID
 	//////////////////////////////////////////////////////////////////////////////////////////////
-	void AddNodeIdRow();
+	TSharedRef<SWidget> CreateCustomNodeIdWidget();
 	bool RebuildNodeIdOptionsList();
+	void UpdateNodeIdSelection();
+
 	void OnNodeIdSelected(TSharedPtr<FString> NodeId, ESelectInfo::Type SelectInfo);
 	FText GetSelectedNodeIdText() const;
 
@@ -72,9 +74,4 @@ protected:
 protected:
 	// ADisplayClusterRootActor on which we're acting
 	TWeakObjectPtr<ADisplayClusterRootActor> EditedObject;
-
-	// Keep a reference to force refresh the layout
-	IDetailLayoutBuilder* LayoutBuilder = nullptr;
-	// The layout category we'll be dealing with
-	IDetailCategoryBuilder* CategoryPreview;
 };
