nDisplay: Fixed compile warning introduced with recent CL

diff --git a/Engine/Plugins/Runtime/nDisplay/Source/DisplayClusterConfigurator/Private/Views/Details/DisplayClusterConfiguratorDetailCustomization.cpp b/Engine/Plugins/Runtime/nDisplay/Source/DisplayClusterConfigurator/Private/Views/Details/DisplayClusterConfiguratorDetailCustomization.cpp
--- a/Engine/Plugins/Runtime/nDisplay/Source/DisplayClusterConfigurator/Private/Views/Details/DisplayClusterConfiguratorDetailCustomization.cpp
+++ b/Engine/Plugins/Runtime/nDisplay/Source/DisplayClusterConfigurator/Private/Views/Details/DisplayClusterConfiguratorDetailCustomization.cpp
@@ -93,17 +93,19 @@ void FDisplayClusterConfiguratorDetailCustomization::CustomizeDetails(IDetailLay
 	// Iterate over all of the properties in the object being edited to find properties marked with specific custom metadata tags and hide those properties if necessary
 	for (TFieldIterator<FProperty> It(ObjectBeingEdited->GetClass()); It; ++It)
 	{
-		FProperty* Property = *It;
-		TSharedRef<IPropertyHandle> PropertyHandle = InLayoutBuilder.GetProperty(Property->GetFName());
+		if (FProperty* Property = *It)
+		{
+			TSharedRef<IPropertyHandle> PropertyHandle = InLayoutBuilder.GetProperty(Property->GetFName());
 
-		const bool bShouldHide = PropertyHandle->HasMetaData("nDisplayHidden") ||
-			(IsRunningForBlueprintEditor() ?
-				PropertyHandle->HasMetaData("nDisplayInstanceOnly") || Property && Property->HasAnyPropertyFlags(CPF_DisableEditOnTemplate) :
-				false);
+			const bool bShouldHide = PropertyHandle->HasMetaData("nDisplayHidden") ||
+				(IsRunningForBlueprintEditor() ?
+					PropertyHandle->HasMetaData("nDisplayInstanceOnly") || Property && Property->HasAnyPropertyFlags(CPF_DisableEditOnTemplate) :
+					false);
 
-		if (bShouldHide)
-		{
-			PropertyHandle->MarkHiddenByCustomization();
+			if (bShouldHide)
+			{
+				PropertyHandle->MarkHiddenByCustomization();
+			}
 		}
 	}
 }
