nDisplay: Added custom PostProcessSettings blending to EndFinalPostProcessSettings instead of Override to have full control over the custom blend
Refactored/combined macros
Fixed bug with per-viewport PPS exclusion flag not working correctly if bIsEnabled was false

diff --git a/Engine/Plugins/Runtime/nDisplay/Source/DisplayCluster/Private/Render/Device/DisplayClusterDeviceBase.cpp b/Engine/Plugins/Runtime/nDisplay/Source/DisplayCluster/Private/Render/Device/DisplayClusterDeviceBase.cpp
--- a/Engine/Plugins/Runtime/nDisplay/Source/DisplayCluster/Private/Render/Device/DisplayClusterDeviceBase.cpp
+++ b/Engine/Plugins/Runtime/nDisplay/Source/DisplayCluster/Private/Render/Device/DisplayClusterDeviceBase.cpp
@@ -38,6 +38,8 @@
 #include "Render/Viewport/IDisplayClusterViewport.h"
 #include "Render/Viewport/IDisplayClusterViewportProxy.h"
 
+#include "Render/Viewport/Configuration/DisplayClusterViewportConfigurationHelpers.h"
+
 #include <utility>
 
 
@@ -600,8 +602,26 @@ void FDisplayClusterDeviceBase::EndFinalPostprocessSettings(struct FPostProcessS
 		IDisplayClusterViewport* pViewport = ViewportManagerPtr->FindViewport(StereoPassType);
 		if (pViewport)
 		{
-			pViewport->GetViewport_CustomPostProcessSettings().DoPostProcess(IDisplayClusterViewport_CustomPostProcessSettings::ERenderPass::Final, FinalPostProcessingSettings);
+			if (FinalPostProcessingSettings != nullptr)
+			{
+				// Get the final overall cluster + per-viewport PPS from nDisplay
+				FPostProcessSettings RequestedFinalPPS;
+				pViewport->GetViewport_CustomPostProcessSettings().DoPostProcess(IDisplayClusterViewport_CustomPostProcessSettings::ERenderPass::Final, &RequestedFinalPPS);
+
+				FDisplayClusterConfigurationViewport_PerViewportSettings InPPSnDisplay;
+				DisplayClusterViewportConfigurationHelpers::CopyPPSStructConditional(&InPPSnDisplay, &RequestedFinalPPS);
+
+				// Get the passed-in cumulative PPS from the game/viewport (includes all PPVs affecting this viewport)
+				FDisplayClusterConfigurationViewport_PerViewportSettings InPPSCumulative;
+				DisplayClusterViewportConfigurationHelpers::CopyPPSStruct(&InPPSCumulative, FinalPostProcessingSettings);
+
+				// Blend both together with our custom math instead of the default PPS blending
+				DisplayClusterViewportConfigurationHelpers::BlendPostProcessSettings(*FinalPostProcessingSettings, InPPSCumulative, InPPSnDisplay);
+			}
+			else
+			{
+				pViewport->GetViewport_CustomPostProcessSettings().DoPostProcess(IDisplayClusterViewport_CustomPostProcessSettings::ERenderPass::Final, FinalPostProcessingSettings);
+			}
 		}
 	}
 }
-
diff --git a/Engine/Plugins/Runtime/nDisplay/Source/DisplayCluster/Private/Render/Viewport/Configuration/DisplayClusterViewportConfigurationHelpers.h b/Engine/Plugins/Runtime/nDisplay/Source/DisplayCluster/Private/Render/Viewport/Configuration/DisplayClusterViewportConfigurationHelpers.h
--- a/Engine/Plugins/Runtime/nDisplay/Source/DisplayCluster/Private/Render/Viewport/Configuration/DisplayClusterViewportConfigurationHelpers.h
+++ b/Engine/Plugins/Runtime/nDisplay/Source/DisplayCluster/Private/Render/Viewport/Configuration/DisplayClusterViewportConfigurationHelpers.h
@@ -81,67 +81,131 @@ public:
 		}
 	}
 
-	#define PP_CONDITIONAL_BLEND_WITH_OFFSET(OP, NAME, OFFSET) \
-		OutputPP.bOverride_##NAME = ClusterPPSettings.bOverride_##NAME || ViewportPPSettings.bOverride_##NAME; \
-		if (ClusterPPSettings.bOverride_##NAME && ViewportPPSettings.bOverride_##NAME) \
-		{ \
-			OutputPP.##NAME = ClusterPPSettings.##NAME ##OP ViewportPPSettings.##NAME ##OP ##OFFSET; \
-		} \
-		else if (ClusterPPSettings.bOverride_##NAME) \
-		{ \
-			OutputPP.##NAME = ClusterPPSettings.##NAME; \
-		} \
-		else if (ViewportPPSettings.bOverride_##NAME) \
-		{ \
-			OutputPP.##NAME = ViewportPPSettings.##NAME; \
-		} \
-
 	// Note that skipped parameters in macro definitions will just evaluate to nothing
 	// This is intentional to get around the inconsistent naming in the color grading fields in FPostProcessSettings
-	#define PP_CONDITIONAL_BLEND_COLOR(OP, OUTGROUP, INGROUP, NAME) \
-		OutputPP.bOverride_Color##NAME##OUTGROUP = ClusterPPSettings.##INGROUP.bOverride_##NAME || ViewportPPSettings.##INGROUP.bOverride_##NAME; \
-		if (ClusterPPSettings.##INGROUP.bOverride_##NAME && ViewportPPSettings.##INGROUP.bOverride_##NAME) \
+	#define PP_CONDITIONAL_BLEND(BLENDOP, COLOR, OUTGROUP, INGROUP, NAME, OFFSETOP, OFFSETVALUE) \
+		OutputPP.bOverride_##COLOR##NAME##OUTGROUP = ClusterPPSettings.##INGROUP bOverride_##NAME || ViewportPPSettings.##INGROUP bOverride_##NAME; \
+		if (ClusterPPSettings.##INGROUP bOverride_##NAME && ViewportPPSettings.##INGROUP bOverride_##NAME) \
 		{ \
-			OutputPP.Color##NAME##OUTGROUP = ClusterPPSettings.##INGROUP.##NAME ##OP ViewportPPSettings.##INGROUP.##NAME; \
+			OutputPP.##COLOR##NAME##OUTGROUP = ClusterPPSettings.##INGROUP##NAME ##BLENDOP ViewportPPSettings.##INGROUP##NAME ##OFFSETOP ##OFFSETVALUE; \
 		} \
-		else if (ClusterPPSettings.##INGROUP.bOverride_##NAME) \
+		else if (ClusterPPSettings.##INGROUP bOverride_##NAME) \
 		{ \
-			OutputPP.Color##NAME##OUTGROUP = ClusterPPSettings.##INGROUP.##NAME; \
+			OutputPP.##COLOR##NAME##OUTGROUP = ClusterPPSettings.##INGROUP##NAME; \
 		} \
-		else if (ViewportPPSettings.##INGROUP.bOverride_##NAME) \
+		else if (ViewportPPSettings.##INGROUP bOverride_##NAME) \
 		{ \
-			OutputPP.Color##NAME##OUTGROUP = ViewportPPSettings.##INGROUP.##NAME; \
+			OutputPP.##COLOR##NAME##OUTGROUP = ViewportPPSettings.##INGROUP##NAME; \
 		} \
 
 	static void BlendPostProcessSettings(FPostProcessSettings& OutputPP, const FDisplayClusterConfigurationViewport_PerViewportSettings& ClusterPPSettings, const FDisplayClusterConfigurationViewport_PerViewportSettings& ViewportPPSettings)
 	{
-		PP_CONDITIONAL_BLEND_WITH_OFFSET(+, WhiteTemp, -6500.0f);
-		PP_CONDITIONAL_BLEND_WITH_OFFSET(+, WhiteTint, 0.0f);
-		PP_CONDITIONAL_BLEND_WITH_OFFSET(+, AutoExposureBias, 0.0f);
-
-		PP_CONDITIONAL_BLEND_COLOR(*, , Global, Saturation);
-		PP_CONDITIONAL_BLEND_COLOR(*, , Global, Contrast);
-		PP_CONDITIONAL_BLEND_COLOR(*, , Global, Gamma);
-		PP_CONDITIONAL_BLEND_COLOR(*, , Global, Gain);
-		PP_CONDITIONAL_BLEND_COLOR(+, , Global, Offset);
-
-		PP_CONDITIONAL_BLEND_COLOR(*, Shadows, Shadows, Saturation);
-		PP_CONDITIONAL_BLEND_COLOR(*, Shadows, Shadows, Contrast);
-		PP_CONDITIONAL_BLEND_COLOR(*, Shadows, Shadows, Gamma);
-		PP_CONDITIONAL_BLEND_COLOR(*, Shadows, Shadows, Gain);
-		PP_CONDITIONAL_BLEND_COLOR(+, Shadows, Shadows, Offset);
-
-		PP_CONDITIONAL_BLEND_COLOR(*, Midtones, Midtones, Saturation);
-		PP_CONDITIONAL_BLEND_COLOR(*, Midtones, Midtones, Contrast);
-		PP_CONDITIONAL_BLEND_COLOR(*, Midtones, Midtones, Gamma);
-		PP_CONDITIONAL_BLEND_COLOR(*, Midtones, Midtones, Gain);
-		PP_CONDITIONAL_BLEND_COLOR(+, Midtones, Midtones, Offset);
-
-		PP_CONDITIONAL_BLEND_COLOR(*, Highlights, Highlights, Saturation);
-		PP_CONDITIONAL_BLEND_COLOR(*, Highlights, Highlights, Contrast);
-		PP_CONDITIONAL_BLEND_COLOR(*, Highlights, Highlights, Gamma);
-		PP_CONDITIONAL_BLEND_COLOR(*, Highlights, Highlights, Gain);
-		PP_CONDITIONAL_BLEND_COLOR(+, Highlights, Highlights, Offset);
+		PP_CONDITIONAL_BLEND(+, , , , WhiteTemp, +, -6500.0f);
+		PP_CONDITIONAL_BLEND(+, , , , WhiteTint, , );
+		PP_CONDITIONAL_BLEND(+, , , , AutoExposureBias, , );
+
+		PP_CONDITIONAL_BLEND(*, Color, , Global., Saturation, , );
+		PP_CONDITIONAL_BLEND(*, Color, , Global., Contrast, , );
+		PP_CONDITIONAL_BLEND(*, Color, , Global., Gamma, , );
+		PP_CONDITIONAL_BLEND(*, Color, , Global., Gain, , );
+		PP_CONDITIONAL_BLEND(+, Color, , Global., Offset, , );
+
+		PP_CONDITIONAL_BLEND(*, Color, Shadows, Shadows., Saturation, , );
+		PP_CONDITIONAL_BLEND(*, Color, Shadows, Shadows., Contrast, , );
+		PP_CONDITIONAL_BLEND(*, Color, Shadows, Shadows., Gamma, , );
+		PP_CONDITIONAL_BLEND(*, Color, Shadows, Shadows., Gain, , );
+		PP_CONDITIONAL_BLEND(+, Color, Shadows, Shadows., Offset, , );
+
+		PP_CONDITIONAL_BLEND(*, Color, Midtones, Midtones., Saturation, , );
+		PP_CONDITIONAL_BLEND(*, Color, Midtones, Midtones., Contrast, , );
+		PP_CONDITIONAL_BLEND(*, Color, Midtones, Midtones., Gamma, , );
+		PP_CONDITIONAL_BLEND(*, Color, Midtones, Midtones., Gain, , );
+		PP_CONDITIONAL_BLEND(+, Color, Midtones, Midtones., Offset, , );
+
+		PP_CONDITIONAL_BLEND(*, Color, Highlights, Highlights., Saturation, , );
+		PP_CONDITIONAL_BLEND(*, Color, Highlights, Highlights., Contrast, , );
+		PP_CONDITIONAL_BLEND(*, Color, Highlights, Highlights., Gamma, , );
+		PP_CONDITIONAL_BLEND(*, Color, Highlights, Highlights., Gain, , );
+		PP_CONDITIONAL_BLEND(+, Color, Highlights, Highlights., Offset, , );
+	}
+
+	#define PP_CONDITIONAL_COPY(COLOR, OUTGROUP, INGROUP, NAME) \
+		if (InPPS->bOverride_##COLOR##NAME##INGROUP) \
+		{ \
+			OutViewportPPSettings->##OUTGROUP##NAME = InPPS->##COLOR##NAME##INGROUP; \
+			OutViewportPPSettings->##OUTGROUP bOverride_##NAME = true; \
+		} \
+
+	static void CopyPPSStructConditional(FDisplayClusterConfigurationViewport_PerViewportSettings* OutViewportPPSettings, FPostProcessSettings* InPPS)
+	{
+		if ((OutViewportPPSettings != nullptr) && (InPPS != nullptr))
+		{
+			PP_CONDITIONAL_COPY( , , , WhiteTemp);
+			PP_CONDITIONAL_COPY( , , , WhiteTint);
+			PP_CONDITIONAL_COPY( , , , AutoExposureBias);
+
+			PP_CONDITIONAL_COPY(Color, Global., , Saturation);
+			PP_CONDITIONAL_COPY(Color, Global., , Contrast);
+			PP_CONDITIONAL_COPY(Color, Global., , Gamma);
+			PP_CONDITIONAL_COPY(Color, Global., , Gain);
+			PP_CONDITIONAL_COPY(Color, Global., , Offset);
+
+			PP_CONDITIONAL_COPY(Color, Shadows., Shadows, Saturation);
+			PP_CONDITIONAL_COPY(Color, Shadows., Shadows, Contrast);
+			PP_CONDITIONAL_COPY(Color, Shadows., Shadows, Gamma);
+			PP_CONDITIONAL_COPY(Color, Shadows., Shadows, Gain);
+			PP_CONDITIONAL_COPY(Color, Shadows., Shadows, Offset);
+
+			PP_CONDITIONAL_COPY(Color, Midtones., Midtones, Saturation);
+			PP_CONDITIONAL_COPY(Color, Midtones., Midtones, Contrast);
+			PP_CONDITIONAL_COPY(Color, Midtones., Midtones, Gamma);
+			PP_CONDITIONAL_COPY(Color, Midtones., Midtones, Gain);
+			PP_CONDITIONAL_COPY(Color, Midtones., Midtones, Offset);
+
+			PP_CONDITIONAL_COPY(Color, Highlights., Highlights, Saturation);
+			PP_CONDITIONAL_COPY(Color, Highlights., Highlights, Contrast);
+			PP_CONDITIONAL_COPY(Color, Highlights., Highlights, Gamma);
+			PP_CONDITIONAL_COPY(Color, Highlights., Highlights, Gain);
+			PP_CONDITIONAL_COPY(Color, Highlights., Highlights, Offset);
+		}
+	}
+
+	#define PP_COPY(COLOR, OUTGROUP, INGROUP, NAME) \
+		OutViewportPPSettings->##OUTGROUP##NAME = InPPS->##COLOR##NAME##INGROUP; \
+		OutViewportPPSettings->##OUTGROUP bOverride_##NAME = true; \
+
+	static void CopyPPSStruct(FDisplayClusterConfigurationViewport_PerViewportSettings* OutViewportPPSettings, FPostProcessSettings* InPPS)
+	{
+		if ((OutViewportPPSettings != nullptr) && (InPPS != nullptr))
+		{
+			PP_COPY(, , , WhiteTemp);
+			PP_COPY(, , , WhiteTint);
+			PP_COPY(, , , AutoExposureBias);
+
+			PP_COPY(Color, Global., , Saturation);
+			PP_COPY(Color, Global., , Contrast);
+			PP_COPY(Color, Global., , Gamma);
+			PP_COPY(Color, Global., , Gain);
+			PP_COPY(Color, Global., , Offset);
+
+			PP_COPY(Color, Shadows., Shadows, Saturation);
+			PP_COPY(Color, Shadows., Shadows, Contrast);
+			PP_COPY(Color, Shadows., Shadows, Gamma);
+			PP_COPY(Color, Shadows., Shadows, Gain);
+			PP_COPY(Color, Shadows., Shadows, Offset);
+
+			PP_COPY(Color, Midtones., Midtones, Saturation);
+			PP_COPY(Color, Midtones., Midtones, Contrast);
+			PP_COPY(Color, Midtones., Midtones, Gamma);
+			PP_COPY(Color, Midtones., Midtones, Gain);
+			PP_COPY(Color, Midtones., Midtones, Offset);
+
+			PP_COPY(Color, Highlights., Highlights, Saturation);
+			PP_COPY(Color, Highlights., Highlights, Contrast);
+			PP_COPY(Color, Highlights., Highlights, Gamma);
+			PP_COPY(Color, Highlights., Highlights, Gain);
+			PP_COPY(Color, Highlights., Highlights, Offset);
+		}
 	}
 
 	static void UpdateViewportPostProcessSettings(FDisplayClusterViewport& DstViewport, const FDisplayClusterConfigurationViewport_PostProcessSettings& InPostProcessSettings)
@@ -179,11 +243,20 @@ public:
 						}
 						else
 						{
-							// This viewport doesn't use per-viewport settings, so just use the overall cluster settings and blend weight
-							BlendPostProcessSettings(CustomPPS.PostProcessSettings, ClusterConfig->OverallClusterPostProcessSettings, EmptyPPSettings);
+							if (InPostProcessSettings.bExcludeFromOverallClusterPostProcess)
+							{
+								// This viewport doesn't use per-viewport settings and is also excluded from the overall cluster, so do nothing
+								CustomPPS.BlendWeight = 0.0f;
+							}
+							else
+							{
+								// This viewport doesn't use per-viewport settings, so just use the overall cluster settings and blend weight
+								BlendPostProcessSettings(CustomPPS.PostProcessSettings, ClusterConfig->OverallClusterPostProcessSettings, EmptyPPSettings);
+							}
 						}
 
-						ImplUpdateViewportSetting_CustomPostprocess(DstViewport, CustomPPS, IDisplayClusterViewport_CustomPostProcessSettings::ERenderPass::Override);
+						// Use ERenderPass::Final here because we do a custom blend in FDisplayClusterDeviceBase::EndFinalPostprocessSettings
+						ImplUpdateViewportSetting_CustomPostprocess(DstViewport, CustomPPS, IDisplayClusterViewport_CustomPostProcessSettings::ERenderPass::Final);
 
 						return;
 					}
@@ -202,7 +275,9 @@ public:
 			const FDisplayClusterConfigurationViewport_PerViewportSettings EmptyPPSettings;
 
 			BlendPostProcessSettings(CustomPPS.PostProcessSettings, EmptyPPSettings, InPostProcessSettings.ViewportSettings);
-			ImplUpdateViewportSetting_CustomPostprocess(DstViewport, CustomPPS, IDisplayClusterViewport_CustomPostProcessSettings::ERenderPass::Override);
+
+			// Use ERenderPass::Final here because we do a custom blend in FDisplayClusterDeviceBase::EndFinalPostprocessSettings
+			ImplUpdateViewportSetting_CustomPostprocess(DstViewport, CustomPPS, IDisplayClusterViewport_CustomPostProcessSettings::ERenderPass::Final);
 		}
 	}
 
