nDisplay: camera projection policy implementation

diff --git a/Engine/Plugins/Runtime/nDisplay/Source/DisplayClusterProjection/Private/Policy/Camera/DisplayClusterProjectionCameraPolicy.cpp b/Engine/Plugins/Runtime/nDisplay/Source/DisplayClusterProjection/Private/Policy/Camera/DisplayClusterProjectionCameraPolicy.cpp
--- a/Engine/Plugins/Runtime/nDisplay/Source/DisplayClusterProjection/Private/Policy/Camera/DisplayClusterProjectionCameraPolicy.cpp
+++ b/Engine/Plugins/Runtime/nDisplay/Source/DisplayClusterProjection/Private/Policy/Camera/DisplayClusterProjectionCameraPolicy.cpp
@@ -4,6 +4,7 @@
 
 #include "DisplayClusterProjectionLog.h"
 #include "DisplayClusterProjectionStrings.h"
+#include "DisplayClusterRootActor.h"
 
 #include "Game/IDisplayClusterGameManager.h"
 
@@ -132,11 +133,45 @@ bool FDisplayClusterProjectionCameraPolicy::GetProjectionMatrix(IDisplayClusterV
 bool FDisplayClusterProjectionCameraPolicy::GetSettingsFromConfig(class IDisplayClusterViewport* InViewport, UCameraComponent*& OutCamera, FDisplayClusterProjectionCameraPolicySettings& OutCameraSettings)
 {
 	check(InViewport);
-
-	//@todo: add camera settings reading from config!
 	
+	ADisplayClusterRootActor* const RootActor = InViewport->GetOwner().GetRootActor();
+	if (!RootActor)
+	{
+		UE_LOG(LogDisplayClusterProjectionCamera, Error, TEXT("Couldn't get a DisplayClusterRootActor root object"));
+		return false;
+	}
 
-	return true;
+	FString CameraComponentId;
+	// Get assigned camera ID
+	if (!DisplayClusterHelpers::map::template ExtractValue(GetParameters(), DisplayClusterProjectionStrings::cfg::camera::Component, CameraComponentId))
+	{
+
+#if WITH_EDITOR
+		if (CameraComponentId.IsEmpty())
+		{
+			return false;
+		}
+#endif
+
+		UE_LOG(LogDisplayClusterProjectionCamera, Error, TEXT("No camera component ID '%s' specified for projection policy '%s'"), *CameraComponentId, *GetId());
+		return false;
+	}
+
+	// Get camera component
+	TArray<UCameraComponent*> CameraComps;
+	RootActor->GetComponents<UCameraComponent>(CameraComps);
+	for (UCameraComponent* Comp : CameraComps)
+	{
+		if (Comp->GetName() == CameraComponentId)
+		{
+			OutCamera = Comp;
+			return true;
+		}
+	}
+
+	UE_LOG(LogDisplayClusterProjectionCamera, Error, TEXT("Camera component ID '%s' not found for projection policy '%s'"), *CameraComponentId, *GetId());
+
+	return false;
 }
 
 //////////////////////////////////////////////////////////////////////////////////////////////
