nDisplay: Update camera policy selection. Remove camera selection from viewport. Fix rename variable not propagating changes.

diff --git a/Engine/Plugins/Runtime/nDisplay/Source/DisplayCluster/Public/Components/DisplayClusterICVFX_CineCameraComponent.h b/Engine/Plugins/Runtime/nDisplay/Source/DisplayCluster/Public/Components/DisplayClusterICVFX_CineCameraComponent.h
--- a/Engine/Plugins/Runtime/nDisplay/Source/DisplayCluster/Public/Components/DisplayClusterICVFX_CineCameraComponent.h
+++ b/Engine/Plugins/Runtime/nDisplay/Source/DisplayCluster/Public/Components/DisplayClusterICVFX_CineCameraComponent.h
@@ -13,7 +13,7 @@
 /**
  * ICVFX camera with configuration
  */
-UCLASS(ClassGroup = (DisplayCluster), meta = (DisplayName="NDisplay ICVFX_CineCamera"))
+UCLASS(ClassGroup = (DisplayCluster), meta = (DisplayName="ICVFX_CineCamera"))
 class DISPLAYCLUSTER_API UDisplayClusterICVFX_CineCameraComponent 
 	: public UCineCameraComponent
 {
diff --git a/Engine/Plugins/Runtime/nDisplay/Source/DisplayCluster/Public/Components/DisplayClusterICVFX_RefCineCameraComponent.h b/Engine/Plugins/Runtime/nDisplay/Source/DisplayCluster/Public/Components/DisplayClusterICVFX_RefCineCameraComponent.h
--- a/Engine/Plugins/Runtime/nDisplay/Source/DisplayCluster/Public/Components/DisplayClusterICVFX_RefCineCameraComponent.h
+++ b/Engine/Plugins/Runtime/nDisplay/Source/DisplayCluster/Public/Components/DisplayClusterICVFX_RefCineCameraComponent.h
@@ -13,7 +13,7 @@
 /**
  * ICVFX camera with configuration
  */
-UCLASS(ClassGroup = (DisplayCluster), meta = (DisplayName="NDisplay ICVFX_RefCineCamera"))
+UCLASS(ClassGroup = (DisplayCluster), meta = (DisplayName="ICVFX_RefCineCamera"))
 class DISPLAYCLUSTER_API UDisplayClusterICVFX_RefCineCameraComponent 
 	: public USceneComponent
 {
diff --git a/Engine/Plugins/Runtime/nDisplay/Source/DisplayClusterConfiguration/Public/DisplayClusterConfigurationTypes_Viewport.h b/Engine/Plugins/Runtime/nDisplay/Source/DisplayClusterConfiguration/Public/DisplayClusterConfigurationTypes_Viewport.h
--- a/Engine/Plugins/Runtime/nDisplay/Source/DisplayClusterConfiguration/Public/DisplayClusterConfigurationTypes_Viewport.h
+++ b/Engine/Plugins/Runtime/nDisplay/Source/DisplayClusterConfiguration/Public/DisplayClusterConfigurationTypes_Viewport.h
@@ -282,7 +282,7 @@ public:
 	UPROPERTY(EditAnywhere, Category = "NDisplay Viewport")
 	bool bAllowRendering = true;
 
-	UPROPERTY(EditAnywhere, Category = "NDisplay Viewport")
+	UPROPERTY()
 	FString Camera;
 
 	UPROPERTY(EditDefaultsOnly, Category = "NDisplay Viewport")
diff --git a/Engine/Plugins/Runtime/nDisplay/Source/DisplayClusterConfigurator/Private/DisplayClusterConfiguratorBlueprintEditor.cpp b/Engine/Plugins/Runtime/nDisplay/Source/DisplayClusterConfigurator/Private/DisplayClusterConfiguratorBlueprintEditor.cpp
--- a/Engine/Plugins/Runtime/nDisplay/Source/DisplayClusterConfigurator/Private/DisplayClusterConfiguratorBlueprintEditor.cpp
+++ b/Engine/Plugins/Runtime/nDisplay/Source/DisplayClusterConfigurator/Private/DisplayClusterConfiguratorBlueprintEditor.cpp
@@ -19,6 +19,7 @@
 #include "Components/DisplayClusterScreenComponent.h"
 
 #include "ClusterConfiguration/DisplayClusterConfiguratorClusterUtils.h"
+#include "DisplayClusterConfiguratorPropertyUtils.h"
 #include "DisplayClusterConfiguratorVersionUtils.h"
 #include "Views/General/DisplayClusterConfiguratorViewGeneral.h"
 #include "Views/OutputMapping/DisplayClusterConfiguratorViewOutputMapping.h"
@@ -663,11 +664,26 @@ void FDisplayClusterConfiguratorBlueprintEditor::OnRenameVariable(UBlueprint* Bl
 				// of the parameters in the custom case
 				if (!Viewport->ProjectionPolicy.bIsCustom)
 				{
-					for (TPair<FString, FString>& PolicyParameters : Viewport->ProjectionPolicy.Parameters)
+					TMap<FString, FString> PolicyCopy = Viewport->ProjectionPolicy.Parameters;
+
+					const TSharedPtr<ISinglePropertyView> ProjectPolicyView =
+						DisplayClusterConfiguratorPropertyUtils::GetPropertyView(
+							Viewport, GET_MEMBER_NAME_CHECKED(UDisplayClusterConfigurationViewport, ProjectionPolicy));
+					check(ProjectPolicyView);
+					
+					const TSharedPtr<IPropertyHandle> ParametersHandle = ProjectPolicyView->GetPropertyHandle()->GetChildHandle(GET_MEMBER_NAME_CHECKED(FDisplayClusterConfigurationProjection, Parameters));
+					check(ParametersHandle);
+
+					FStructProperty* StructProperty = FindFProperty<FStructProperty>(Viewport->GetClass(), GET_MEMBER_NAME_CHECKED(UDisplayClusterConfigurationViewport, ProjectionPolicy));
+					check(StructProperty);
+					uint8* MapContainer = StructProperty->ContainerPtrToValuePtr<uint8>(Viewport);
+					
+					for (TPair<FString, FString>& PolicyParameters : PolicyCopy)
 					{
 						if (PolicyParameters.Value == OldVariableName.ToString())
 						{
-							PolicyParameters.Value = NewVariableName.ToString();
+							DisplayClusterConfiguratorPropertyUtils::RemoveKeyFromMap(MapContainer, ParametersHandle, PolicyParameters.Key);
+							DisplayClusterConfiguratorPropertyUtils::AddKeyValueToMap(MapContainer, ParametersHandle, PolicyParameters.Key, NewVariableName.ToString());
 						}
 					}
 				}
diff --git a/Engine/Plugins/Runtime/nDisplay/Source/DisplayClusterConfigurator/Private/DisplayClusterConfiguratorCompiler.cpp b/Engine/Plugins/Runtime/nDisplay/Source/DisplayClusterConfigurator/Private/DisplayClusterConfiguratorCompiler.cpp
--- a/Engine/Plugins/Runtime/nDisplay/Source/DisplayClusterConfigurator/Private/DisplayClusterConfiguratorCompiler.cpp
+++ b/Engine/Plugins/Runtime/nDisplay/Source/DisplayClusterConfigurator/Private/DisplayClusterConfiguratorCompiler.cpp
@@ -163,10 +163,6 @@ void FDisplayClusterConfiguratorKismetCompilerContext::ValidateConfiguration()
 				{
 					MessageLog.Warning(*LOCTEXT("NoPolicyError", "No projection policy assigned to viewport @@.").ToString(), Viewport.Value);
 				}
-				if (bCameraFound && Viewport.Value->Camera.IsEmpty())
-				{
-					MessageLog.Note(*LOCTEXT("NoCameraForViewportNote", "Using default camera for viewport @@ because none is assigned.").ToString(), Viewport.Value);
-				}
 			}
 		}
 	}
diff --git a/Engine/Plugins/Runtime/nDisplay/Source/DisplayClusterConfigurator/Private/Views/Details/DisplayClusterConfiguratorDetailCustomization.cpp b/Engine/Plugins/Runtime/nDisplay/Source/DisplayClusterConfigurator/Private/Views/Details/DisplayClusterConfiguratorDetailCustomization.cpp
--- a/Engine/Plugins/Runtime/nDisplay/Source/DisplayClusterConfigurator/Private/Views/Details/DisplayClusterConfiguratorDetailCustomization.cpp
+++ b/Engine/Plugins/Runtime/nDisplay/Source/DisplayClusterConfigurator/Private/Views/Details/DisplayClusterConfiguratorDetailCustomization.cpp
@@ -225,124 +225,6 @@ void FDisplayClusterConfiguratorClusterDetailCustomization::CustomizeDetails(IDe
 void FDisplayClusterConfiguratorViewportDetailCustomization::CustomizeDetails(IDetailLayoutBuilder& InLayoutBuilder)
 {
 	Super::CustomizeDetails(InLayoutBuilder);
-
-	ConfigurationViewportPtr = nullptr;
-	NoneOption = MakeShared<FString>("None");
-
-	// Set config data pointer
-	UDisplayClusterConfigurationData* ConfigurationData = GetConfigData();
-	check(ConfigurationData != nullptr);
-	ConfigurationDataPtr = ConfigurationData;
-
-	// Get the Editing object
-	const TArray<TWeakObjectPtr<UObject>>& SelectedObjects = InLayoutBuilder.GetSelectedObjects();
-	if (SelectedObjects.Num())
-	{
-		ConfigurationViewportPtr = Cast<UDisplayClusterConfigurationViewport>(SelectedObjects[0]);
-	}
-	check(ConfigurationViewportPtr != nullptr);
-
-	CameraHandle = InLayoutBuilder.GetProperty(GET_MEMBER_NAME_CHECKED(UDisplayClusterConfigurationViewport, Camera));
-	check(CameraHandle.IsValid());
-
-	ResetCameraOptions();
-
-	BEGIN_CATEGORY(CameraHandle->GetDefaultCategoryName())
-		REPLACE_PROPERTY_WITH_CUSTOM(UDisplayClusterConfigurationViewport, Camera, CreateCustomCameraWidget());
-	END_CATEGORY()
-}
-
-void FDisplayClusterConfiguratorViewportDetailCustomization::ResetCameraOptions()
-{
-	CameraOptions.Reset();
-
-	UDisplayClusterConfigurationViewport* ConfigurationViewport = ConfigurationViewportPtr.Get();
-	check(ConfigurationViewport != nullptr);
-	
-	AActor* RootActor = GetRootActor();
-	
-	TArray<UActorComponent*> ActorComponents;
-	RootActor->GetComponents(UDisplayClusterCameraComponent::StaticClass(), ActorComponents);
-	for (UActorComponent* ActorComponent : ActorComponents)
-	{
-		const FString ComponentName = ActorComponent->GetName();
-		CameraOptions.Add(MakeShared<FString>(ComponentName));
-	}
-
-	// Component order not guaranteed, sort for consistency.
-	CameraOptions.Sort([](const TSharedPtr<FString>& A, const TSharedPtr<FString>& B)
-	{
-		// Default sort isn't compatible with TSharedPtr<FString>.
-		return *A < *B;
-	});
-	
-	// Add None option
-	if (!ConfigurationViewport->Camera.IsEmpty())
-	{
-		CameraOptions.Add(NoneOption);
-	}
-}
-
-TSharedRef<SWidget> FDisplayClusterConfiguratorViewportDetailCustomization::CreateCustomCameraWidget()
-{
-	if (CameraComboBox.IsValid())
-	{
-		return CameraComboBox.ToSharedRef();
-	}
-
-	return SAssignNew(CameraComboBox, SDisplayClusterConfigurationSearchableComboBox)
-		.OptionsSource(&CameraOptions)
-		.OnGenerateWidget(this, &FDisplayClusterConfiguratorViewportDetailCustomization::MakeCameraOptionComboWidget)
-		.OnSelectionChanged(this, &FDisplayClusterConfiguratorViewportDetailCustomization::OnCameraSelected)
-		.ContentPadding(2)
-		.MaxListHeight(200.0f)
-		.Content()
-		[
-			SNew(STextBlock)
-			.Text(this, &FDisplayClusterConfiguratorViewportDetailCustomization::GetSelectedCameraText)
-		];
-}
-
-TSharedRef<SWidget> FDisplayClusterConfiguratorViewportDetailCustomization::MakeCameraOptionComboWidget(TSharedPtr<FString> InItem)
-{
-	return SNew(STextBlock).Text(FText::FromString(*InItem));
-}
-
-void FDisplayClusterConfiguratorViewportDetailCustomization::OnCameraSelected(TSharedPtr<FString> InCamera, ESelectInfo::Type SelectInfo)
-{
-	if (InCamera.IsValid())
-	{
-		UDisplayClusterConfigurationViewport* ConfigurationViewport = ConfigurationViewportPtr.Get();
-		check(ConfigurationViewport != nullptr);
-
-		// Handle empty case
-		if (InCamera->Equals(*NoneOption.Get()))
-		{
-			CameraHandle->SetValue(TEXT(""));
-
-		}
-		else
-		{
-			CameraHandle->SetValue(*InCamera.Get());
-		}
-
-		// Reset available options
-		ResetCameraOptions();
-		CameraComboBox->ResetOptionsSource(&CameraOptions);
-
-		CameraComboBox->SetIsOpen(false);
-	}
-}
-
-FText FDisplayClusterConfiguratorViewportDetailCustomization::GetSelectedCameraText() const
-{
-	FString SelectedOption = ConfigurationViewportPtr.Get()->Camera;
-	if (SelectedOption.IsEmpty())
-	{
-		SelectedOption = *NoneOption.Get();
-	}
-
-	return FText::FromString(SelectedOption);
 }
 
 //////////////////////////////////////////////////////////////////////////////////////////////
diff --git a/Engine/Plugins/Runtime/nDisplay/Source/DisplayClusterConfigurator/Private/Views/Details/DisplayClusterConfiguratorDetailCustomization.h b/Engine/Plugins/Runtime/nDisplay/Source/DisplayClusterConfigurator/Private/Views/Details/DisplayClusterConfiguratorDetailCustomization.h
--- a/Engine/Plugins/Runtime/nDisplay/Source/DisplayClusterConfigurator/Private/Views/Details/DisplayClusterConfiguratorDetailCustomization.h
+++ b/Engine/Plugins/Runtime/nDisplay/Source/DisplayClusterConfigurator/Private/Views/Details/DisplayClusterConfiguratorDetailCustomization.h
@@ -107,30 +107,6 @@ public:
 	/** IDetailCustomization interface */
 	virtual void CustomizeDetails(IDetailLayoutBuilder& InLayoutBuilder) override;
 	/** End IDetailCustomization interface */
-
-private:
-	void ResetCameraOptions();
-
-	TSharedRef<SWidget> CreateCustomCameraWidget();
-
-	TSharedRef<SWidget> MakeCameraOptionComboWidget(TSharedPtr<FString> InItem);
-
-	void OnCameraSelected(TSharedPtr<FString> InCamera, ESelectInfo::Type SelectInfo);
-
-	FText GetSelectedCameraText() const;
-
-private:
-	TArray< TSharedPtr< FString > >	CameraOptions;
-
-	TSharedPtr<IPropertyHandle> CameraHandle;
-
-	TSharedPtr<FString>	NoneOption;
-
-	TWeakObjectPtr<UDisplayClusterConfigurationViewport> ConfigurationViewportPtr;
-
-	TWeakObjectPtr<UDisplayClusterConfigurationData> ConfigurationDataPtr;
-
-	TSharedPtr<SDisplayClusterConfigurationSearchableComboBox> CameraComboBox;
 };
 
 /**
diff --git a/Engine/Plugins/Runtime/nDisplay/Source/DisplayClusterConfigurator/Private/Views/Details/Policies/DisplayClusterConfiguratorPolicyDetailCustomization.cpp b/Engine/Plugins/Runtime/nDisplay/Source/DisplayClusterConfigurator/Private/Views/Details/Policies/DisplayClusterConfiguratorPolicyDetailCustomization.cpp
--- a/Engine/Plugins/Runtime/nDisplay/Source/DisplayClusterConfigurator/Private/Views/Details/Policies/DisplayClusterConfiguratorPolicyDetailCustomization.cpp
+++ b/Engine/Plugins/Runtime/nDisplay/Source/DisplayClusterConfigurator/Private/Views/Details/Policies/DisplayClusterConfiguratorPolicyDetailCustomization.cpp
@@ -11,9 +11,9 @@
 #include "DisplayClusterRootActor.h"
 #include "Blueprints/DisplayClusterBlueprint.h"
 #include "Components/DisplayClusterScreenComponent.h"
-#include "Components/DisplayClusterCameraComponent.h"
+#include "Components/DisplayClusterICVFX_RefCineCameraComponent.h"
 
-#include "CineCameraComponent.h"
+#include "Camera/CameraComponent.h"
 
 #include "DisplayClusterProjectionStrings.h"
 
@@ -364,7 +364,7 @@ void FDisplayClusterConfiguratorProjectionCustomization::CreateSimplePolicy(UDis
 		Blueprint,
 		ConfigurationViewportPtr.Get(),
 		ParametersHandle,
-		UDisplayClusterScreenComponent::StaticClass()));
+		TArray<TSubclassOf<UActorComponent>>{ UDisplayClusterScreenComponent::StaticClass() }));
 }
 
 void FDisplayClusterConfiguratorProjectionCustomization::CreateCameraPolicy(UDisplayClusterBlueprint* Blueprint)
@@ -375,7 +375,8 @@ void FDisplayClusterConfiguratorProjectionCustomization::CreateCameraPolicy(UDis
 		Blueprint,
 		ConfigurationViewportPtr.Get(),
 		ParametersHandle,
-		UCineCameraComponent::StaticClass()));
+		TArray<TSubclassOf<UActorComponent>>{ UCameraComponent::StaticClass(),
+			UDisplayClusterICVFX_RefCineCameraComponent::StaticClass() }));
 
 	CustomPolicyParameters.Add(MakeShared<FPolicyParameterInfoBool>(
 		"Native",
@@ -393,7 +394,7 @@ void FDisplayClusterConfiguratorProjectionCustomization::CreateMeshPolicy(UDispl
 		Blueprint,
 		ConfigurationViewportPtr.Get(),
 		ParametersHandle,
-		UStaticMeshComponent::StaticClass()));
+		TArray<TSubclassOf<UActorComponent>>{ UStaticMeshComponent::StaticClass() }));
 }
 
 void FDisplayClusterConfiguratorProjectionCustomization::CreateDomePolicy(UDisplayClusterBlueprint* Blueprint)
@@ -412,7 +413,7 @@ void FDisplayClusterConfiguratorProjectionCustomization::CreateDomePolicy(UDispl
 		Blueprint,
 		ConfigurationViewportPtr.Get(),
 		ParametersHandle,
-		USceneComponent::StaticClass()));
+		TArray<TSubclassOf<UActorComponent>>{ USceneComponent::StaticClass() }));
 
 	CustomPolicyParameters.Add(MakeShared<FPolicyParameterInfoNumber<int32>>(
 		"Channel",
@@ -438,7 +439,7 @@ void FDisplayClusterConfiguratorProjectionCustomization::CreateVIOSOPolicy(UDisp
 		Blueprint,
 		ConfigurationViewportPtr.Get(),
 		ParametersHandle,
-		USceneComponent::StaticClass()));
+		TArray<TSubclassOf<UActorComponent>>{ USceneComponent::StaticClass() }));
 
 	CustomPolicyParameters.Add(MakeShared<FPolicyParameterInfoMatrix>(
 		"Matrix",
@@ -464,7 +465,7 @@ void FDisplayClusterConfiguratorProjectionCustomization::CreateEasyBlendPolicy(U
 		Blueprint,
 		ConfigurationViewportPtr.Get(),
 		ParametersHandle,
-		USceneComponent::StaticClass()));
+		TArray<TSubclassOf<UActorComponent>>{ USceneComponent::StaticClass() }));
 
 	CustomPolicyParameters.Add(MakeShared<FPolicyParameterInfoNumber<float>>(
 		"Scale",
@@ -749,7 +750,7 @@ void FDisplayClusterConfiguratorProjectionCustomization::CreateMPCDIPolicy(UDisp
 		Blueprint,
 		ConfigurationViewportPtr.Get(),
 		ParametersHandle,
-		USceneComponent::StaticClass());
+		TArray<TSubclassOf<UActorComponent>>{ USceneComponent::StaticClass() });
 }
 
 #undef LOCTEXT_NAMESPACE
diff --git a/Engine/Plugins/Runtime/nDisplay/Source/DisplayClusterConfigurator/Private/Views/Details/Policies/DisplayClusterConfiguratorPolicyParameterCustomization.cpp b/Engine/Plugins/Runtime/nDisplay/Source/DisplayClusterConfigurator/Private/Views/Details/Policies/DisplayClusterConfiguratorPolicyParameterCustomization.cpp
--- a/Engine/Plugins/Runtime/nDisplay/Source/DisplayClusterConfigurator/Private/Views/Details/Policies/DisplayClusterConfiguratorPolicyParameterCustomization.cpp
+++ b/Engine/Plugins/Runtime/nDisplay/Source/DisplayClusterConfigurator/Private/Views/Details/Policies/DisplayClusterConfiguratorPolicyParameterCustomization.cpp
@@ -181,9 +181,9 @@ FPolicyParameterInfoComponentCombo::FPolicyParameterInfoComponentCombo(
 	UDisplayClusterBlueprint* InBlueprint,
 	UDisplayClusterConfigurationViewport* InConfigurationViewport,
 	const TSharedPtr<IPropertyHandle>& InParametersHandle,
-	TSubclassOf<UActorComponent> InComponentClass) : FPolicyParameterInfoCombo(InDisplayName, InKey, InBlueprint, InConfigurationViewport, InParametersHandle, {}, nullptr)
+	const TArray<TSubclassOf<UActorComponent>>& InComponentClasses) : FPolicyParameterInfoCombo(InDisplayName, InKey, InBlueprint, InConfigurationViewport, InParametersHandle, {}, nullptr)
 {
-	ComponentType = InComponentClass;
+	ComponentTypes = InComponentClasses;
 	if (BlueprintEditorPtrCached)
 	{
 		// BlueprintEditorPtrCached can be null when remote control is interfacing with the world instance details panel. At this stage
@@ -196,7 +196,7 @@ FPolicyParameterInfoComponentCombo::FPolicyParameterInfoComponentCombo(
 
 void FPolicyParameterInfoComponentCombo::CreateParameterValues(ADisplayClusterRootActor* RootActor)
 {
-	if (ComponentType)
+	for (const TSubclassOf<UActorComponent>& ComponentType : ComponentTypes)
 	{
 		TArray<UActorComponent*> ActorComponents;
 		RootActor->GetComponents(ComponentType, ActorComponents);
diff --git a/Engine/Plugins/Runtime/nDisplay/Source/DisplayClusterConfigurator/Private/Views/Details/Policies/DisplayClusterConfiguratorPolicyParameterCustomization.h b/Engine/Plugins/Runtime/nDisplay/Source/DisplayClusterConfigurator/Private/Views/Details/Policies/DisplayClusterConfiguratorPolicyParameterCustomization.h
--- a/Engine/Plugins/Runtime/nDisplay/Source/DisplayClusterConfigurator/Private/Views/Details/Policies/DisplayClusterConfiguratorPolicyParameterCustomization.h
+++ b/Engine/Plugins/Runtime/nDisplay/Source/DisplayClusterConfigurator/Private/Views/Details/Policies/DisplayClusterConfiguratorPolicyParameterCustomization.h
@@ -153,17 +153,17 @@ public:
 		UDisplayClusterBlueprint* InBlueprint,
 		UDisplayClusterConfigurationViewport* InConfigurationViewport,
 		const TSharedPtr<IPropertyHandle>& InParametersHandle,
-		TSubclassOf<UActorComponent> InComponentClass);
+		const TArray<TSubclassOf<UActorComponent>>& InComponentClasses);
 
 	/** Populate CustomParameterOptions based on this parameter. */
 	void CreateParameterValues(ADisplayClusterRootActor* RootActor);
 
 	/** The type of component this parameter represents. */
-	TSubclassOf<UActorComponent> GetComponentType() const { return ComponentType; }
+	const TArray<TSubclassOf<UActorComponent>>& GetComponentTypes() const { return ComponentTypes; }
 
 private:
 	/** A component class to use for creating parameter options. */
-	TSubclassOf<UActorComponent> ComponentType;
+	TArray<TSubclassOf<UActorComponent>> ComponentTypes;
 };
 
 /**
diff --git a/Engine/Plugins/Runtime/nDisplay/Source/DisplayClusterConfigurator/Private/Views/Viewport/DisplayClusterConfiguratorSCSEditorViewport.cpp b/Engine/Plugins/Runtime/nDisplay/Source/DisplayClusterConfigurator/Private/Views/Viewport/DisplayClusterConfiguratorSCSEditorViewport.cpp
--- a/Engine/Plugins/Runtime/nDisplay/Source/DisplayClusterConfigurator/Private/Views/Viewport/DisplayClusterConfiguratorSCSEditorViewport.cpp
+++ b/Engine/Plugins/Runtime/nDisplay/Source/DisplayClusterConfigurator/Private/Views/Viewport/DisplayClusterConfiguratorSCSEditorViewport.cpp
@@ -7,12 +7,15 @@
 #include "DisplayClusterConfiguratorUtils.h"
 #include "DisplayClusterProjectionStrings.h"
 #include "DisplayClusterConfiguratorBlueprintEditor.h"
-#include "Components/DisplayClusterScreenComponent.h"
-#include "Components/DisplayClusterICVFX_CineCameraComponent.h"
 #include "Views/DragDrop/DisplayClusterConfiguratorValidatedDragDropOp.h"
 #include "Views/DragDrop/DisplayClusterConfiguratorViewportDragDropOp.h"
 #include "ClusterConfiguration/ViewModels/DisplayClusterConfiguratorProjectionPolicyViewModel.h"
 
+#include "Components/DisplayClusterScreenComponent.h"
+#include "Components/DisplayClusterICVFX_RefCineCameraComponent.h"
+
+#include "Camera/CameraComponent.h"
+
 #include "EngineUtils.h"
 #include "BlueprintEditor.h"
 #include "SSubobjectEditor.h"
@@ -695,9 +698,12 @@ FReply SDisplayClusterConfiguratorSCSEditorViewport::OnDragOver(const FGeometry&
 					}
 					else
 					{
-						if (UDisplayClusterICVFX_CineCameraComponent* CineCameraComponent = Cast<UDisplayClusterICVFX_CineCameraComponent>(MeshComponent->GetAttachParent()))
+						USceneComponent* AttachParent = MeshComponent->GetAttachParent();
+						const bool bIsCamera = AttachParent && (AttachParent->IsA<UCameraComponent>() || AttachParent->IsA<UDisplayClusterICVFX_RefCineCameraComponent>());
+						
+						if (bIsCamera)
 						{
-							ViewportDragDropOp->SetDropAsValid(FText::Format(LOCTEXT("ViewportDragDropOp_CameraMessage", "Project to camera {0}"), FText::FromName(CineCameraComponent->GetFName())));
+							ViewportDragDropOp->SetDropAsValid(FText::Format(LOCTEXT("ViewportDragDropOp_CameraMessage", "Project to camera {0}"), FText::FromName(AttachParent->GetFName())));
 						}
 						else if (!MeshComponent->IsVisualizationComponent())
 						{
@@ -749,11 +755,14 @@ FReply SDisplayClusterConfiguratorSCSEditorViewport::OnDrop(const FGeometry& MyG
 					}
 					else
 					{
-						if (UDisplayClusterICVFX_CineCameraComponent* CineCameraComponent = Cast<UDisplayClusterICVFX_CineCameraComponent>(MeshComponent->GetAttachParent()))
+						USceneComponent* AttachParent = MeshComponent->GetAttachParent();
+						const bool bIsCamera = AttachParent && (AttachParent->IsA<UCameraComponent>() || AttachParent->IsA<UDisplayClusterICVFX_RefCineCameraComponent>());
+						
+						if (bIsCamera)
 						{
 							PolicyType = DisplayClusterProjectionStrings::projection::Camera;
 							ParameterKey = DisplayClusterProjectionStrings::cfg::camera::Component;
-							ComponentName = CineCameraComponent->GetName();
+							ComponentName = AttachParent->GetName();
 						}
 						else if (!MeshComponent->IsVisualizationComponent())
 						{
