nDisplay: Refactor camera settings to a struct. This also prevents a crash from the subobject instance not being maintained in certain scenarios.

diff --git a/Engine/Plugins/Runtime/nDisplay/Source/DisplayCluster/Private/Game/EngineClasses/Scene/DisplayClusterICVFX_CineCameraComponent.cpp b/Engine/Plugins/Runtime/nDisplay/Source/DisplayCluster/Private/Game/EngineClasses/Scene/DisplayClusterICVFX_CineCameraComponent.cpp
--- a/Engine/Plugins/Runtime/nDisplay/Source/DisplayCluster/Private/Game/EngineClasses/Scene/DisplayClusterICVFX_CineCameraComponent.cpp
+++ b/Engine/Plugins/Runtime/nDisplay/Source/DisplayCluster/Private/Game/EngineClasses/Scene/DisplayClusterICVFX_CineCameraComponent.cpp
@@ -8,7 +8,6 @@
 
 UDisplayClusterICVFX_CineCameraComponent::UDisplayClusterICVFX_CineCameraComponent(const FObjectInitializer& ObjectInitializer)
 {
-	IncameraSettings = CreateDefaultSubobject<UDisplayClusterConfigurationICVFX_CameraSettings>(TEXT("IncameraSettings"));
 }
 
 UCameraComponent* UDisplayClusterICVFX_CineCameraComponent::GetCameraComponent()
@@ -26,7 +25,7 @@ FDisplayClusterViewport_CameraMotionBlur UDisplayClusterICVFX_CineCameraComponen
 	FDisplayClusterViewport_CameraMotionBlur OutParameters;
 	OutParameters.Mode = EDisplayClusterViewport_CameraMotionBlur::Undefined;
 
-	switch (IncameraSettings->CameraMotionBlur.MotionBlurMode)
+	switch (IncameraSettings.CameraMotionBlur.MotionBlurMode)
 	{
 	case EDisplayClusterConfigurationCameraMotionBlurMode::Off:
 		OutParameters.Mode = EDisplayClusterViewport_CameraMotionBlur::Off;
@@ -47,7 +46,7 @@ FDisplayClusterViewport_CameraMotionBlur UDisplayClusterICVFX_CineCameraComponen
 					OutParameters.CameraLocation = OuterCamera->K2_GetComponentLocation();
 					OutParameters.CameraRotation = OuterCamera->K2_GetComponentRotation();
 
-					OutParameters.TranslationScale = IncameraSettings->CameraMotionBlur.TranslationScale;
+					OutParameters.TranslationScale = IncameraSettings.CameraMotionBlur.TranslationScale;
 
 					OutParameters.Mode = EDisplayClusterViewport_CameraMotionBlur::Override;
 				}
diff --git a/Engine/Plugins/Runtime/nDisplay/Source/DisplayCluster/Private/Game/EngineClasses/Scene/DisplayClusterICVFX_RefCineCameraComponent.cpp b/Engine/Plugins/Runtime/nDisplay/Source/DisplayCluster/Private/Game/EngineClasses/Scene/DisplayClusterICVFX_RefCineCameraComponent.cpp
--- a/Engine/Plugins/Runtime/nDisplay/Source/DisplayCluster/Private/Game/EngineClasses/Scene/DisplayClusterICVFX_RefCineCameraComponent.cpp
+++ b/Engine/Plugins/Runtime/nDisplay/Source/DisplayCluster/Private/Game/EngineClasses/Scene/DisplayClusterICVFX_RefCineCameraComponent.cpp
@@ -8,7 +8,6 @@
 
 UDisplayClusterICVFX_RefCineCameraComponent::UDisplayClusterICVFX_RefCineCameraComponent(const FObjectInitializer& ObjectInitializer)
 {
-	IncameraSettings = CreateDefaultSubobject<UDisplayClusterConfigurationICVFX_CameraSettings>(TEXT("IncameraSettings"));
 }
 
 FDisplayClusterViewport_CameraMotionBlur UDisplayClusterICVFX_RefCineCameraComponent::GetMotionBlurParameters()
@@ -16,7 +15,7 @@ FDisplayClusterViewport_CameraMotionBlur UDisplayClusterICVFX_RefCineCameraCompo
 	FDisplayClusterViewport_CameraMotionBlur OutParameters;
 	OutParameters.Mode = EDisplayClusterViewport_CameraMotionBlur::Undefined;
 
-	switch (IncameraSettings->CameraMotionBlur.MotionBlurMode)
+	switch (IncameraSettings.CameraMotionBlur.MotionBlurMode)
 	{
 	case EDisplayClusterConfigurationCameraMotionBlurMode::Off:
 		OutParameters.Mode = EDisplayClusterViewport_CameraMotionBlur::Off;
@@ -37,7 +36,7 @@ FDisplayClusterViewport_CameraMotionBlur UDisplayClusterICVFX_RefCineCameraCompo
 				OutParameters.CameraLocation = OuterCamera->K2_GetComponentLocation();
 				OutParameters.CameraRotation = OuterCamera->K2_GetComponentRotation();
 
-				OutParameters.TranslationScale = IncameraSettings->CameraMotionBlur.TranslationScale;
+				OutParameters.TranslationScale = IncameraSettings.CameraMotionBlur.TranslationScale;
 
 				OutParameters.Mode = EDisplayClusterViewport_CameraMotionBlur::Override;
 			}
diff --git a/Engine/Plugins/Runtime/nDisplay/Source/DisplayCluster/Private/Render/Viewport/Configuration/DisplayClusterViewportConfigurationHelpers.h b/Engine/Plugins/Runtime/nDisplay/Source/DisplayCluster/Private/Render/Viewport/Configuration/DisplayClusterViewportConfigurationHelpers.h
--- a/Engine/Plugins/Runtime/nDisplay/Source/DisplayCluster/Private/Render/Viewport/Configuration/DisplayClusterViewportConfigurationHelpers.h
+++ b/Engine/Plugins/Runtime/nDisplay/Source/DisplayCluster/Private/Render/Viewport/Configuration/DisplayClusterViewportConfigurationHelpers.h
@@ -535,12 +535,12 @@ public:
 		}
 	}
 
-	static const FDisplayClusterConfigurationICVFX_ChromakeySettings& GetCameraChromakeySettings(const UDisplayClusterConfigurationICVFX_CameraSettings& InCameraSettings, const UDisplayClusterConfigurationICVFX_StageSettings& InStageSettings)
+	static const FDisplayClusterConfigurationICVFX_ChromakeySettings& GetCameraChromakeySettings(const FDisplayClusterConfigurationICVFX_CameraSettings& InCameraSettings, const UDisplayClusterConfigurationICVFX_StageSettings& InStageSettings)
 	{
 		return (InCameraSettings.CustomChromakey.bEnable) ? InCameraSettings.CustomChromakey.Chromakey : InStageSettings.Chromakey;
 	}
 
-	static void UpdateCameraViewportSetting(FDisplayClusterViewport& DstViewport, const UDisplayClusterConfigurationICVFX_CameraSettings& InCameraSettings, const UDisplayClusterConfigurationICVFX_StageSettings& InStageSettings)
+	static void UpdateCameraViewportSetting(FDisplayClusterViewport& DstViewport, const FDisplayClusterConfigurationICVFX_CameraSettings& InCameraSettings, const UDisplayClusterConfigurationICVFX_StageSettings& InStageSettings)
 	{
 		check(InCameraSettings.bEnable);
 
diff --git a/Engine/Plugins/Runtime/nDisplay/Source/DisplayCluster/Private/Render/Viewport/Configuration/DisplayClusterViewportConfigurationICVFX.cpp b/Engine/Plugins/Runtime/nDisplay/Source/DisplayCluster/Private/Render/Viewport/Configuration/DisplayClusterViewportConfigurationICVFX.cpp
--- a/Engine/Plugins/Runtime/nDisplay/Source/DisplayCluster/Private/Render/Viewport/Configuration/DisplayClusterViewportConfigurationICVFX.cpp
+++ b/Engine/Plugins/Runtime/nDisplay/Source/DisplayCluster/Private/Render/Viewport/Configuration/DisplayClusterViewportConfigurationICVFX.cpp
@@ -23,7 +23,7 @@
 class FDisplayClusterViewportConfigurationCameraViewport
 {
 public:
-	FDisplayClusterViewportConfigurationCameraViewport(const FTransform& InLocal2WorldTransform, FDisplayClusterViewportConfigurationICVFX& InConfigurationICVFX, UCameraComponent* const InCameraComponent, const FString InCameraId, const UDisplayClusterConfigurationICVFX_CameraSettings& InCameraSettings)
+	FDisplayClusterViewportConfigurationCameraViewport(const FTransform& InLocal2WorldTransform, FDisplayClusterViewportConfigurationICVFX& InConfigurationICVFX, UCameraComponent* const InCameraComponent, const FString InCameraId, const FDisplayClusterConfigurationICVFX_CameraSettings& InCameraSettings)
 		: ConfigurationICVFX(InConfigurationICVFX)
 		, CameraComponent(InCameraComponent)
 		, CameraId(InCameraId)
@@ -276,7 +276,7 @@ private:
 	UCameraComponent* const CameraComponent;
 	const FString CameraId;
 
-	const UDisplayClusterConfigurationICVFX_CameraSettings& CameraSettings;
+	const FDisplayClusterConfigurationICVFX_CameraSettings& CameraSettings;
 
 	const FTransform Local2WorldTransform;
 };
@@ -290,15 +290,15 @@ class FDisplayClusterViewportConfigurationCameraICVFX
 public:
 	FDisplayClusterViewportConfigurationCameraICVFX(const FTransform& InLocal2WorldTransform, FDisplayClusterViewportConfigurationICVFX& InConfigurationICVFX, class UDisplayClusterICVFX_CineCameraComponent* const InCameraComponent)
 		: ConfigurationICVFX(InConfigurationICVFX)
-		, CameraViewport(InLocal2WorldTransform, InConfigurationICVFX, InCameraComponent->GetCameraComponent(), InCameraComponent->GetCameraUniqueId(), *(InCameraComponent->GetCameraSettingsICVFX()))
-		, ChromakeySettings(DisplayClusterViewportConfigurationHelpers::GetCameraChromakeySettings(*(InCameraComponent->GetCameraSettingsICVFX()), ConfigurationICVFX.StageSettings))
+		, CameraViewport(InLocal2WorldTransform, InConfigurationICVFX, InCameraComponent->GetCameraComponent(), InCameraComponent->GetCameraUniqueId(), InCameraComponent->GetCameraSettingsICVFX())
+		, ChromakeySettings(DisplayClusterViewportConfigurationHelpers::GetCameraChromakeySettings(InCameraComponent->GetCameraSettingsICVFX(), ConfigurationICVFX.StageSettings))
 		, CameraMotionBlurParameters(InCameraComponent->GetMotionBlurParameters())
 	{}
 
 	FDisplayClusterViewportConfigurationCameraICVFX(const FTransform& InLocal2WorldTransform, class FDisplayClusterViewportConfigurationICVFX& InConfigurationICVFX, class UDisplayClusterICVFX_RefCineCameraComponent* const InCameraComponent)
 		: ConfigurationICVFX(InConfigurationICVFX)
-		, CameraViewport(InLocal2WorldTransform, InConfigurationICVFX, InCameraComponent->GetCameraComponent(), InCameraComponent->GetCameraUniqueId(), *(InCameraComponent->GetCameraSettingsICVFX()))
-		, ChromakeySettings(DisplayClusterViewportConfigurationHelpers::GetCameraChromakeySettings(*(InCameraComponent->GetCameraSettingsICVFX()), ConfigurationICVFX.StageSettings))
+		, CameraViewport(InLocal2WorldTransform, InConfigurationICVFX, InCameraComponent->GetCameraComponent(), InCameraComponent->GetCameraUniqueId(), InCameraComponent->GetCameraSettingsICVFX())
+		, ChromakeySettings(DisplayClusterViewportConfigurationHelpers::GetCameraChromakeySettings(InCameraComponent->GetCameraSettingsICVFX(), ConfigurationICVFX.StageSettings))
 		, CameraMotionBlurParameters(InCameraComponent->GetMotionBlurParameters())
 	{}
 
diff --git a/Engine/Plugins/Runtime/nDisplay/Source/DisplayCluster/Private/Render/Viewport/Configuration/DisplayClusterViewportConfigurationICVFX.h b/Engine/Plugins/Runtime/nDisplay/Source/DisplayCluster/Private/Render/Viewport/Configuration/DisplayClusterViewportConfigurationICVFX.h
--- a/Engine/Plugins/Runtime/nDisplay/Source/DisplayCluster/Private/Render/Viewport/Configuration/DisplayClusterViewportConfigurationICVFX.h
+++ b/Engine/Plugins/Runtime/nDisplay/Source/DisplayCluster/Private/Render/Viewport/Configuration/DisplayClusterViewportConfigurationICVFX.h
@@ -19,7 +19,7 @@ struct FDisplayClusterConfigurationPostRender_GenerateMips;
 struct FDisplayClusterConfigurationICVFX_OverlayAdvancedRenderSettings;
 struct FDisplayClusterConfigurationICVFX_LightcardSettings;
 
-class UDisplayClusterConfigurationICVFX_CameraSettings;
+struct FDisplayClusterConfigurationICVFX_CameraSettings;
 class UDisplayClusterConfigurationICVFX_StageSettings;
 
 
diff --git a/Engine/Plugins/Runtime/nDisplay/Source/DisplayCluster/Public/Components/DisplayClusterICVFX_CineCameraComponent.h b/Engine/Plugins/Runtime/nDisplay/Source/DisplayCluster/Public/Components/DisplayClusterICVFX_CineCameraComponent.h
--- a/Engine/Plugins/Runtime/nDisplay/Source/DisplayCluster/Public/Components/DisplayClusterICVFX_CineCameraComponent.h
+++ b/Engine/Plugins/Runtime/nDisplay/Source/DisplayCluster/Public/Components/DisplayClusterICVFX_CineCameraComponent.h
@@ -23,21 +23,21 @@ public:
 	UDisplayClusterICVFX_CineCameraComponent(const FObjectInitializer& ObjectInitializer);
 
 public:
-	UPROPERTY(EditAnywhere, Category = "NDisplay ICVFX")
-	UDisplayClusterConfigurationICVFX_CameraSettings* IncameraSettings;
+	UPROPERTY(EditAnywhere, Category = "NDisplay ICVFX", meta = (ShowInnerProperties))
+	FDisplayClusterConfigurationICVFX_CameraSettings IncameraSettings;
 
 public:
 	FDisplayClusterViewport_CameraMotionBlur GetMotionBlurParameters();
 
 	bool IsShouldUseICVFX() const
 	{
-		return IncameraSettings->bEnable;
+		return IncameraSettings.bEnable;
 	}
 
 	// Return unique camera name
 	FString GetCameraUniqueId() const;
 
-	const UDisplayClusterConfigurationICVFX_CameraSettings* GetCameraSettingsICVFX() const
+	const FDisplayClusterConfigurationICVFX_CameraSettings& GetCameraSettingsICVFX() const
 	{
 		return IncameraSettings;
 	}
diff --git a/Engine/Plugins/Runtime/nDisplay/Source/DisplayCluster/Public/Components/DisplayClusterICVFX_RefCineCameraComponent.h b/Engine/Plugins/Runtime/nDisplay/Source/DisplayCluster/Public/Components/DisplayClusterICVFX_RefCineCameraComponent.h
--- a/Engine/Plugins/Runtime/nDisplay/Source/DisplayCluster/Public/Components/DisplayClusterICVFX_RefCineCameraComponent.h
+++ b/Engine/Plugins/Runtime/nDisplay/Source/DisplayCluster/Public/Components/DisplayClusterICVFX_RefCineCameraComponent.h
@@ -23,8 +23,8 @@ public:
 	UDisplayClusterICVFX_RefCineCameraComponent(const FObjectInitializer& ObjectInitializer);
 
 public:
-	UPROPERTY(EditAnywhere, Category = "NDisplay ICVFX")
-	UDisplayClusterConfigurationICVFX_CameraSettings* IncameraSettings;
+	UPROPERTY(EditAnywhere, Category = "NDisplay ICVFX", meta = (ShowInnerProperties))
+	FDisplayClusterConfigurationICVFX_CameraSettings IncameraSettings;
 
 	UPROPERTY(EditAnywhere, Category = "NDisplay ICVFX")
 	TSoftObjectPtr<ACineCameraActor> CineCameraActor;
@@ -34,7 +34,7 @@ public:
 
 	bool IsShouldUseICVFX() const
 	{
-		return IsCineCameraActorValid() && IncameraSettings->bEnable;
+		return IsCineCameraActorValid() && IncameraSettings.bEnable;
 	}
 
 	// Return unique camera name
@@ -48,7 +48,7 @@ public:
 		return FString();
 	}
 
-	const UDisplayClusterConfigurationICVFX_CameraSettings* GetCameraSettingsICVFX() const
+	const FDisplayClusterConfigurationICVFX_CameraSettings& GetCameraSettingsICVFX() const
 	{
 		return IncameraSettings;
 	}
diff --git a/Engine/Plugins/Runtime/nDisplay/Source/DisplayClusterConfiguration/Private/DisplayClusterConfigurationTypes.cpp b/Engine/Plugins/Runtime/nDisplay/Source/DisplayClusterConfiguration/Private/DisplayClusterConfigurationTypes.cpp
--- a/Engine/Plugins/Runtime/nDisplay/Source/DisplayClusterConfiguration/Private/DisplayClusterConfigurationTypes.cpp
+++ b/Engine/Plugins/Runtime/nDisplay/Source/DisplayClusterConfiguration/Private/DisplayClusterConfigurationTypes.cpp
@@ -131,7 +131,7 @@ FDisplayClusterConfigurationProjection::FDisplayClusterConfigurationProjection()
 	Type = TEXT("simple");
 }
 
-UDisplayClusterConfigurationICVFX_CameraSettings::UDisplayClusterConfigurationICVFX_CameraSettings()
+FDisplayClusterConfigurationICVFX_CameraSettings::FDisplayClusterConfigurationICVFX_CameraSettings()
 {
 }
 
diff --git a/Engine/Plugins/Runtime/nDisplay/Source/DisplayClusterConfiguration/Public/DisplayClusterConfigurationTypes_ICVFX.h b/Engine/Plugins/Runtime/nDisplay/Source/DisplayClusterConfiguration/Public/DisplayClusterConfigurationTypes_ICVFX.h
--- a/Engine/Plugins/Runtime/nDisplay/Source/DisplayClusterConfiguration/Public/DisplayClusterConfigurationTypes_ICVFX.h
+++ b/Engine/Plugins/Runtime/nDisplay/Source/DisplayClusterConfiguration/Public/DisplayClusterConfigurationTypes_ICVFX.h
@@ -303,14 +303,13 @@ public:
 };
 
 
-UCLASS(ClassGroup = (DisplayClusterICVFX), meta = (BlueprintSpawnableComponent))
-class DISPLAYCLUSTERCONFIGURATION_API UDisplayClusterConfigurationICVFX_CameraSettings
-	 : public UActorComponent
+USTRUCT()
+struct DISPLAYCLUSTERCONFIGURATION_API FDisplayClusterConfigurationICVFX_CameraSettings
 {
 	GENERATED_BODY()
 
 public:
-	UDisplayClusterConfigurationICVFX_CameraSettings();
+	FDisplayClusterConfigurationICVFX_CameraSettings();
 
 public:
 	// Enable this camera
