nDisplay: Add back in camera selection to viewport, mark it hidden when policy is set to camera. Fix property propagation for camera when variable is renamed.

diff --git a/Engine/Plugins/Runtime/nDisplay/Source/DisplayClusterConfiguration/Public/DisplayClusterConfigurationTypes_Viewport.h b/Engine/Plugins/Runtime/nDisplay/Source/DisplayClusterConfiguration/Public/DisplayClusterConfigurationTypes_Viewport.h
--- a/Engine/Plugins/Runtime/nDisplay/Source/DisplayClusterConfiguration/Public/DisplayClusterConfigurationTypes_Viewport.h
+++ b/Engine/Plugins/Runtime/nDisplay/Source/DisplayClusterConfiguration/Public/DisplayClusterConfigurationTypes_Viewport.h
@@ -282,7 +282,7 @@ public:
 	UPROPERTY(EditAnywhere, Category = "NDisplay Viewport")
 	bool bAllowRendering = true;
 
-	UPROPERTY()
+	UPROPERTY(EditAnywhere, Category = "NDisplay Viewport")
 	FString Camera;
 
 	UPROPERTY(EditDefaultsOnly, Category = "NDisplay Viewport")
diff --git a/Engine/Plugins/Runtime/nDisplay/Source/DisplayClusterConfigurator/Private/DisplayClusterConfiguratorBlueprintEditor.cpp b/Engine/Plugins/Runtime/nDisplay/Source/DisplayClusterConfigurator/Private/DisplayClusterConfiguratorBlueprintEditor.cpp
--- a/Engine/Plugins/Runtime/nDisplay/Source/DisplayClusterConfigurator/Private/DisplayClusterConfiguratorBlueprintEditor.cpp
+++ b/Engine/Plugins/Runtime/nDisplay/Source/DisplayClusterConfigurator/Private/DisplayClusterConfiguratorBlueprintEditor.cpp
@@ -690,7 +690,7 @@ void FDisplayClusterConfiguratorBlueprintEditor::OnRenameVariable(UBlueprint* Bl
 
 				if (Viewport->Camera == OldVariableName.ToString())
 				{
-					Viewport->Camera = NewVariableName.ToString();
+					DisplayClusterConfiguratorPropertyUtils::SetPropertyHandleValue(Viewport, GET_MEMBER_NAME_CHECKED(UDisplayClusterConfigurationViewport, Camera), NewVariableName.ToString());
 				}
 			}
 		}
diff --git a/Engine/Plugins/Runtime/nDisplay/Source/DisplayClusterConfigurator/Private/DisplayClusterConfiguratorCompiler.cpp b/Engine/Plugins/Runtime/nDisplay/Source/DisplayClusterConfigurator/Private/DisplayClusterConfiguratorCompiler.cpp
--- a/Engine/Plugins/Runtime/nDisplay/Source/DisplayClusterConfigurator/Private/DisplayClusterConfiguratorCompiler.cpp
+++ b/Engine/Plugins/Runtime/nDisplay/Source/DisplayClusterConfigurator/Private/DisplayClusterConfiguratorCompiler.cpp
@@ -163,6 +163,11 @@ void FDisplayClusterConfiguratorKismetCompilerContext::ValidateConfiguration()
 				{
 					MessageLog.Warning(*LOCTEXT("NoPolicyError", "No projection policy assigned to viewport @@.").ToString(), Viewport.Value);
 				}
+				
+				if (bCameraFound && Viewport.Value->Camera.IsEmpty())
+				{
+					MessageLog.Note(*LOCTEXT("NoCameraForViewportNote", "Using default camera for viewport @@ because none is assigned.").ToString(), Viewport.Value);
+				}
 			}
 		}
 	}
diff --git a/Engine/Plugins/Runtime/nDisplay/Source/DisplayClusterConfigurator/Private/Views/Details/DisplayClusterConfiguratorDetailCustomization.cpp b/Engine/Plugins/Runtime/nDisplay/Source/DisplayClusterConfigurator/Private/Views/Details/DisplayClusterConfiguratorDetailCustomization.cpp
--- a/Engine/Plugins/Runtime/nDisplay/Source/DisplayClusterConfigurator/Private/Views/Details/DisplayClusterConfiguratorDetailCustomization.cpp
+++ b/Engine/Plugins/Runtime/nDisplay/Source/DisplayClusterConfigurator/Private/Views/Details/DisplayClusterConfiguratorDetailCustomization.cpp
@@ -228,6 +228,125 @@ void FDisplayClusterConfiguratorClusterDetailCustomization::CustomizeDetails(IDe
 void FDisplayClusterConfiguratorViewportDetailCustomization::CustomizeDetails(IDetailLayoutBuilder& InLayoutBuilder)
 {
 	Super::CustomizeDetails(InLayoutBuilder);
+
+	ConfigurationViewportPtr = nullptr;
+	NoneOption = MakeShared<FString>("None");
+	
+	// Set config data pointer
+	UDisplayClusterConfigurationData* ConfigurationData = GetConfigData();
+	check(ConfigurationData != nullptr);
+	
+	ConfigurationDataPtr = ConfigurationData;
+	
+	// Get the Editing object
+	const TArray<TWeakObjectPtr<UObject>>& SelectedObjects = InLayoutBuilder.GetSelectedObjects();
+	if (SelectedObjects.Num())
+	{
+		ConfigurationViewportPtr = Cast<UDisplayClusterConfigurationViewport>(SelectedObjects[0]);
+	}
+	check(ConfigurationViewportPtr != nullptr);
+	
+	CameraHandle = InLayoutBuilder.GetProperty(GET_MEMBER_NAME_CHECKED(UDisplayClusterConfigurationViewport, Camera));
+	check(CameraHandle.IsValid());
+
+	if (ConfigurationViewportPtr->ProjectionPolicy.Type == DisplayClusterProjectionStrings::projection::Camera)
+	{
+		CameraHandle->MarkHiddenByCustomization();
+		return;
+	}
+	
+	ResetCameraOptions();
+	
+	BEGIN_CATEGORY(CameraHandle->GetDefaultCategoryName())
+		REPLACE_PROPERTY_WITH_CUSTOM(UDisplayClusterConfigurationViewport, Camera, CreateCustomCameraWidget());
+	END_CATEGORY()
+}
+
+void FDisplayClusterConfiguratorViewportDetailCustomization::ResetCameraOptions()
+{
+	CameraOptions.Reset();
+	UDisplayClusterConfigurationViewport* ConfigurationViewport = ConfigurationViewportPtr.Get();
+	check(ConfigurationViewport != nullptr);
+	
+	AActor* RootActor = GetRootActor();
+	
+	TArray<UActorComponent*> ActorComponents;
+	RootActor->GetComponents(UDisplayClusterCameraComponent::StaticClass(), ActorComponents);
+	for (UActorComponent* ActorComponent : ActorComponents)
+	{
+		const FString ComponentName = ActorComponent->GetName();
+		CameraOptions.Add(MakeShared<FString>(ComponentName));
+	}
+	
+	// Component order not guaranteed, sort for consistency.
+	CameraOptions.Sort([](const TSharedPtr<FString>& A, const TSharedPtr<FString>& B)
+	{
+		// Default sort isn't compatible with TSharedPtr<FString>.
+		return *A < *B;
+	});
+	
+	// Add None option
+	if (!ConfigurationViewport->Camera.IsEmpty())
+	{
+		CameraOptions.Add(NoneOption);
+	}
+}
+
+TSharedRef<SWidget> FDisplayClusterConfiguratorViewportDetailCustomization::CreateCustomCameraWidget()
+{
+	if (CameraComboBox.IsValid())
+	{
+		return CameraComboBox.ToSharedRef();
+	}
+	
+	return SAssignNew(CameraComboBox, SDisplayClusterConfigurationSearchableComboBox)
+		.OptionsSource(&CameraOptions)
+		.OnGenerateWidget(this, &FDisplayClusterConfiguratorViewportDetailCustomization::MakeCameraOptionComboWidget)
+		.OnSelectionChanged(this, &FDisplayClusterConfiguratorViewportDetailCustomization::OnCameraSelected)
+		.ContentPadding(2)
+		.MaxListHeight(200.0f)
+		.Content()
+		[
+			SNew(STextBlock)
+			.Text(this, &FDisplayClusterConfiguratorViewportDetailCustomization::GetSelectedCameraText)
+		];
+}
+
+TSharedRef<SWidget> FDisplayClusterConfiguratorViewportDetailCustomization::MakeCameraOptionComboWidget(TSharedPtr<FString> InItem)
+{
+	return SNew(STextBlock).Text(FText::FromString(*InItem));
+}
+
+void FDisplayClusterConfiguratorViewportDetailCustomization::OnCameraSelected(TSharedPtr<FString> InCamera, ESelectInfo::Type SelectInfo)
+{
+	if (InCamera.IsValid())
+	{
+		UDisplayClusterConfigurationViewport* ConfigurationViewport = ConfigurationViewportPtr.Get();
+		check(ConfigurationViewport != nullptr);
+		// Handle empty case
+		if (InCamera->Equals(*NoneOption.Get()))
+		{
+			CameraHandle->SetValue(TEXT(""));
+		}
+		else
+		{
+			CameraHandle->SetValue(*InCamera.Get());
+		}
+		// Reset available options
+		ResetCameraOptions();
+		CameraComboBox->ResetOptionsSource(&CameraOptions);
+		CameraComboBox->SetIsOpen(false);
+	}
+}
+
+FText FDisplayClusterConfiguratorViewportDetailCustomization::GetSelectedCameraText() const
+{
+	FString SelectedOption = ConfigurationViewportPtr.Get()->Camera;
+	if (SelectedOption.IsEmpty())
+	{
+		SelectedOption = *NoneOption.Get();
+	}
+	return FText::FromString(SelectedOption);
 }
 
 //////////////////////////////////////////////////////////////////////////////////////////////
diff --git a/Engine/Plugins/Runtime/nDisplay/Source/DisplayClusterConfigurator/Private/Views/Details/DisplayClusterConfiguratorDetailCustomization.h b/Engine/Plugins/Runtime/nDisplay/Source/DisplayClusterConfigurator/Private/Views/Details/DisplayClusterConfiguratorDetailCustomization.h
--- a/Engine/Plugins/Runtime/nDisplay/Source/DisplayClusterConfigurator/Private/Views/Details/DisplayClusterConfiguratorDetailCustomization.h
+++ b/Engine/Plugins/Runtime/nDisplay/Source/DisplayClusterConfigurator/Private/Views/Details/DisplayClusterConfiguratorDetailCustomization.h
@@ -107,6 +107,21 @@ public:
 	/** IDetailCustomization interface */
 	virtual void CustomizeDetails(IDetailLayoutBuilder& InLayoutBuilder) override;
 	/** End IDetailCustomization interface */
+
+private:
+	void ResetCameraOptions();
+	TSharedRef<SWidget> CreateCustomCameraWidget();
+	TSharedRef<SWidget> MakeCameraOptionComboWidget(TSharedPtr<FString> InItem);
+	void OnCameraSelected(TSharedPtr<FString> InCamera, ESelectInfo::Type SelectInfo);
+	FText GetSelectedCameraText() const;
+
+private:
+	TArray< TSharedPtr< FString > >	CameraOptions;
+	TSharedPtr<IPropertyHandle> CameraHandle;
+	TSharedPtr<FString>	NoneOption;
+	TWeakObjectPtr<UDisplayClusterConfigurationViewport> ConfigurationViewportPtr;
+	TWeakObjectPtr<UDisplayClusterConfigurationData> ConfigurationDataPtr;
+	TSharedPtr<SDisplayClusterConfigurationSearchableComboBox> CameraComboBox;
 };
 
 /**
