nDisplay - fixed configuration issues. fixed icvfx camera issues

diff --git a/Engine/Plugins/Runtime/nDisplay/Source/DisplayCluster/Private/Game/EngineClasses/Basics/DisplayClusterViewportClient.cpp b/Engine/Plugins/Runtime/nDisplay/Source/DisplayCluster/Private/Game/EngineClasses/Basics/DisplayClusterViewportClient.cpp
--- a/Engine/Plugins/Runtime/nDisplay/Source/DisplayCluster/Private/Game/EngineClasses/Basics/DisplayClusterViewportClient.cpp
+++ b/Engine/Plugins/Runtime/nDisplay/Source/DisplayCluster/Private/Game/EngineClasses/Basics/DisplayClusterViewportClient.cpp
@@ -47,18 +47,6 @@
 
 #include "Config/DisplayClusterConfigManager.h"
 
-// Enable/disable warp&blend
-static TAutoConsoleVariable<int32> CVarWarpBlendEnabled(
-	TEXT("nDisplay.render.WarpBlendEnabled"),
-	1,
-	TEXT("Warp & Blend status\n")
-	TEXT("0 : disabled\n")
-	TEXT("1 : enabled\n")
-	,
-	ECVF_RenderThreadSafe
-);
-
-
 UDisplayClusterViewportClient::UDisplayClusterViewportClient(FVTableHelper& Helper) : Super(Helper)
 {
 	
@@ -602,7 +590,7 @@ void UDisplayClusterViewportClient::Draw(FViewport* InViewport, FCanvas* SceneCa
 		SceneCanvas->Flush_GameThread();
 
 		// After all render target rendered call nDisplay frame rendering
-		RenderFrame.ViewportManager->RenderFrame(CVarWarpBlendEnabled.GetValueOnGameThread() != 0, InViewport);
+		RenderFrame.ViewportManager->RenderFrame(InViewport);
 
 		OnDrawn().Broadcast();
 
diff --git a/Engine/Plugins/Runtime/nDisplay/Source/DisplayCluster/Private/Game/EngineClasses/Scene/DisplayClusterICVFX_CineCameraComponent.cpp b/Engine/Plugins/Runtime/nDisplay/Source/DisplayCluster/Private/Game/EngineClasses/Scene/DisplayClusterICVFX_CineCameraComponent.cpp
--- a/Engine/Plugins/Runtime/nDisplay/Source/DisplayCluster/Private/Game/EngineClasses/Scene/DisplayClusterICVFX_CineCameraComponent.cpp
+++ b/Engine/Plugins/Runtime/nDisplay/Source/DisplayCluster/Private/Game/EngineClasses/Scene/DisplayClusterICVFX_CineCameraComponent.cpp
@@ -12,6 +12,11 @@ UDisplayClusterICVFX_CineCameraComponent::UDisplayClusterICVFX_CineCameraCompone
 
 UCameraComponent* UDisplayClusterICVFX_CineCameraComponent::GetCameraComponent()
 {
+	if (CineCameraActorRef.IsValid() && !CineCameraActorRef.IsNull())
+	{
+		return CineCameraActorRef->GetCameraComponent();
+	}
+
 	return this;
 }
 
@@ -25,7 +30,7 @@ FDisplayClusterViewport_CameraMotionBlur UDisplayClusterICVFX_CineCameraComponen
 	FDisplayClusterViewport_CameraMotionBlur OutParameters;
 	OutParameters.Mode = EDisplayClusterViewport_CameraMotionBlur::Undefined;
 
-	switch (IncameraSettings.CameraMotionBlur.MotionBlurMode)
+	switch (CameraSettings.CameraMotionBlur.MotionBlurMode)
 	{
 	case EDisplayClusterConfigurationCameraMotionBlurMode::Off:
 		OutParameters.Mode = EDisplayClusterViewport_CameraMotionBlur::Off;
@@ -46,7 +51,7 @@ FDisplayClusterViewport_CameraMotionBlur UDisplayClusterICVFX_CineCameraComponen
 					OutParameters.CameraLocation = OuterCamera->K2_GetComponentLocation();
 					OutParameters.CameraRotation = OuterCamera->K2_GetComponentRotation();
 
-					OutParameters.TranslationScale = IncameraSettings.CameraMotionBlur.TranslationScale;
+					OutParameters.TranslationScale = CameraSettings.CameraMotionBlur.TranslationScale;
 
 					OutParameters.Mode = EDisplayClusterViewport_CameraMotionBlur::Override;
 				}
diff --git a/Engine/Plugins/Runtime/nDisplay/Source/DisplayCluster/Private/Game/EngineClasses/Scene/DisplayClusterICVFX_RefCineCameraComponent.cpp b/Engine/Plugins/Runtime/nDisplay/Source/DisplayCluster/Private/Game/EngineClasses/Scene/DisplayClusterICVFX_RefCineCameraComponent.cpp
deleted file mode 100644
diff --git a/Engine/Plugins/Runtime/nDisplay/Source/DisplayCluster/Private/Game/EngineClasses/Scene/DisplayClusterRootActor.cpp b/Engine/Plugins/Runtime/nDisplay/Source/DisplayCluster/Private/Game/EngineClasses/Scene/DisplayClusterRootActor.cpp
--- a/Engine/Plugins/Runtime/nDisplay/Source/DisplayCluster/Private/Game/EngineClasses/Scene/DisplayClusterRootActor.cpp
+++ b/Engine/Plugins/Runtime/nDisplay/Source/DisplayCluster/Private/Game/EngineClasses/Scene/DisplayClusterRootActor.cpp
@@ -15,7 +15,6 @@
 #include "Components/DisplayClusterRootComponent.h"
 #include "Components/DisplayClusterICVFX_CineCameraComponent.h"
 #include "CineCameraComponent.h"
-#include "Components/DisplayClusterICVFX_RefCineCameraComponent.h"
 #include "Components/DisplayClusterSceneComponentSyncThis.h"
 
 #include "Config/IPDisplayClusterConfigManager.h"
@@ -44,8 +43,6 @@
 
 ADisplayClusterRootActor::ADisplayClusterRootActor(const FObjectInitializer& ObjectInitializer)
 	: Super(ObjectInitializer)
-	, bFollowLocalPlayerCamera(false)
-	, bExitOnEsc(true)
 	, OperationMode(EDisplayClusterOperationMode::Disabled)
 {
 	// Root component
@@ -65,12 +62,6 @@ ADisplayClusterRootActor::ADisplayClusterRootActor(const FObjectInitializer& Obj
 	// A helper component to trigger nDisplay Tick() during Tick phase
 	SyncTickComponent = CreateDefaultSubobject<UDisplayClusterSyncTickComponent>(TEXT("DisplayClusterSyncTick"));
 
-	// A ICVFX Stage settings
-	StageSettings = CreateDefaultSubobject<UDisplayClusterConfigurationICVFX_StageSettings>(TEXT("StageSettings"));
-
-	// A render frame settings (allow control whole cluster rendering)
-	RenderFrameSettings = CreateDefaultSubobject<UDisplayClusterConfigurationRenderFrame>(TEXT("RenderFrameSettings"));
-
 	ViewportManager = MakeUnique<FDisplayClusterViewportManager>();
 
 	PrimaryActorTick.bCanEverTick = true;
@@ -107,6 +98,20 @@ bool ADisplayClusterRootActor::IsRunningGameOrPIE() const
 	return true;
 }
 
+const FDisplayClusterConfigurationICVFX_StageSettings& ADisplayClusterRootActor::GetStageSettings() const
+{
+	check(CurrentConfigData);
+
+	return CurrentConfigData->StageSettings;
+}
+
+const FDisplayClusterConfigurationRenderFrame& ADisplayClusterRootActor::GetRenderFrameSettings() const 
+{ 
+	check(CurrentConfigData);
+
+	return CurrentConfigData->RenderFrameSettings;
+}
+
 void ADisplayClusterRootActor::InitializeFromConfig(UDisplayClusterConfigurationData* ConfigData)
 {
 	// Clean up current hierarchy before building a new one
@@ -588,13 +593,13 @@ void ADisplayClusterRootActor::Tick(float DeltaSeconds)
 		OperationMode == EDisplayClusterOperationMode::Editor)
 	{
 		UWorld* const CurWorld = GetWorld();
-		if (CurWorld)
+		if (CurWorld && CurrentConfigData)
 		{
 			APlayerController* const CurPlayerController = CurWorld->GetFirstPlayerController();
 			if (CurPlayerController)
 			{
 				// Depending on the flag state the DCRA follows or not the current player's camera
-				if (bFollowLocalPlayerCamera)
+				if (CurrentConfigData->bFollowLocalPlayerCamera)
 				{
 					APlayerCameraManager* const CurPlayerCameraManager = CurPlayerController->PlayerCameraManager;
 					if (CurPlayerCameraManager)
@@ -603,7 +608,7 @@ void ADisplayClusterRootActor::Tick(float DeltaSeconds)
 					}
 				}
 
-				if (bExitOnEsc)
+				if (CurrentConfigData->bExitOnEsc)
 				{
 					if (CurPlayerController->WasInputKeyJustPressed(EKeys::Escape))
 					{
diff --git a/Engine/Plugins/Runtime/nDisplay/Source/DisplayCluster/Private/Game/EngineClasses/Scene/DisplayClusterRootActorEditor.cpp b/Engine/Plugins/Runtime/nDisplay/Source/DisplayCluster/Private/Game/EngineClasses/Scene/DisplayClusterRootActorEditor.cpp
--- a/Engine/Plugins/Runtime/nDisplay/Source/DisplayCluster/Private/Game/EngineClasses/Scene/DisplayClusterRootActorEditor.cpp
+++ b/Engine/Plugins/Runtime/nDisplay/Source/DisplayCluster/Private/Game/EngineClasses/Scene/DisplayClusterRootActorEditor.cpp
@@ -52,7 +52,6 @@ void ADisplayClusterRootActor::Constructor_Editor()
 {
 	// Allow tick in editor for preview rendering
 	PrimaryActorTick.bStartWithTickEnabled = true;
-	PreviewSettings = CreateDefaultSubobject<UDisplayClusterConfigurationViewportPreview>(TEXT("PreviewSettings"));
 }
 
 void ADisplayClusterRootActor::Destructor_Editor()
@@ -208,19 +207,16 @@ bool ADisplayClusterRootActor::UpdatePreviewConfiguration_Editor()
 {
 	if (ViewportManager.IsValid())
 	{
-		check(PreviewSettings);
-
 		//@todo: make correct GUI implementation for preview settings UObject (special widged in configurator bp, etc)
 		// now just copy RootActor properties to UObject::PreviewSettings
+		FDisplayClusterConfigurationViewportPreview PreviewSettings;
+		if(bPreviewEnable)
 		{
-			PreviewSettings->bEnable = bPreviewEnable;
-			PreviewSettings->PreviewNodeId = PreviewNodeId;
-			PreviewSettings->TickPerFrame = TickPerFrame;
-			PreviewSettings->PreviewRenderTargetRatioMult = PreviewRenderTargetRatioMult;
-		}
+			PreviewSettings.bEnable = bPreviewEnable;
+			PreviewSettings.PreviewNodeId = PreviewNodeId;
+			PreviewSettings.TickPerFrame = TickPerFrame;
+			PreviewSettings.PreviewRenderTargetRatioMult = PreviewRenderTargetRatioMult;
 
-		if (PreviewSettings->bEnable)
-		{
 			return ViewportManager->UpdatePreviewConfiguration(PreviewSettings, this);
 		}
 	}
@@ -340,7 +336,7 @@ void ADisplayClusterRootActor::UpdatePreviewComponents()
 	
 	if (CurrentConfigData != nullptr)
 	{
-		const bool bAllComponentsVisible = PreviewSettings->PreviewNodeId.Equals(DisplayClusterConfigurationStrings::gui::preview::PreviewNodeAll, ESearchCase::IgnoreCase);
+		const bool bAllComponentsVisible = PreviewNodeId.Equals(DisplayClusterConfigurationStrings::gui::preview::PreviewNodeAll, ESearchCase::IgnoreCase);
 
 		for (const TPair<FString, UDisplayClusterConfigurationClusterNode*>& Node : CurrentConfigData->Cluster->Nodes)
 		{
@@ -350,7 +346,7 @@ void ADisplayClusterRootActor::UpdatePreviewComponents()
 			}
 			for (const TPair<FString, UDisplayClusterConfigurationViewport*>& Viewport : Node.Value->Viewports)
 			{
-				if (bAllComponentsVisible || Node.Key.Equals(PreviewSettings->PreviewNodeId, ESearchCase::IgnoreCase))
+				if (bAllComponentsVisible || Node.Key.Equals(PreviewNodeId, ESearchCase::IgnoreCase))
 				{
 					const FString PreviewCompId = GeneratePreviewComponentName(Node.Key, Viewport.Key);
 					UDisplayClusterPreviewComponent* PreviewComp = PreviewComponents.FindRef(PreviewCompId);
diff --git a/Engine/Plugins/Runtime/nDisplay/Source/DisplayCluster/Private/Render/Viewport/Configuration/DisplayClusterViewportConfiguration.cpp b/Engine/Plugins/Runtime/nDisplay/Source/DisplayCluster/Private/Render/Viewport/Configuration/DisplayClusterViewportConfiguration.cpp
--- a/Engine/Plugins/Runtime/nDisplay/Source/DisplayCluster/Private/Render/Viewport/Configuration/DisplayClusterViewportConfiguration.cpp
+++ b/Engine/Plugins/Runtime/nDisplay/Source/DisplayCluster/Private/Render/Viewport/Configuration/DisplayClusterViewportConfiguration.cpp
@@ -72,7 +72,7 @@ bool FDisplayClusterViewportConfiguration::UpdateConfiguration(EDisplayClusterRe
 			ImplUpdateConfiguration(RenderNodes, *RootActor, *ConfigurationData);
 			ImplUpdateConfigurationICVFX(*RootActor, *ConfigurationData);
 			ImplUpdateConfigurationVisibility(*RootActor, *ConfigurationData);
-			ImplUpdateRenderFrameConfiguration(*(RootActor->GetRenderFrameSettings()));
+			ImplUpdateRenderFrameConfiguration(RootActor->GetRenderFrameSettings());
 
 			ImplUpdateConfiguration_PostProcess(InClusterNodeId, *ConfigurationData);
 
@@ -113,7 +113,7 @@ void FDisplayClusterViewportConfiguration::ImplUpdateConfigurationVisibility(ADi
 	}
 }
 
-void FDisplayClusterViewportConfiguration::ImplUpdateRenderFrameConfiguration(const UDisplayClusterConfigurationRenderFrame& InRenderFrameConfiguration)
+void FDisplayClusterViewportConfiguration::ImplUpdateRenderFrameConfiguration(const FDisplayClusterConfigurationRenderFrame& InRenderFrameConfiguration)
 {
 	// Some frame postprocess require additional render targetable resources
 	RenderFrameSettings.bShouldUseAdditionalFrameTargetableResource = ViewportManager.PostProcessManager->ShouldUseAdditionalFrameTargetableResource_PostProcess();
@@ -211,11 +211,10 @@ void FDisplayClusterViewportConfiguration::ImplUpdateConfiguration_PostProcess(c
 }
 
 #if WITH_EDITOR
-bool FDisplayClusterViewportConfiguration::UpdatePreviewConfiguration(class UDisplayClusterConfigurationViewportPreview* PreviewConfiguration)
+bool FDisplayClusterViewportConfiguration::UpdatePreviewConfiguration(const FDisplayClusterConfigurationViewportPreview& InPreviewConfiguration)
 {
 	check(IsInGameThread());
-	check(PreviewConfiguration);
-	check(PreviewConfiguration->bEnable);
+	check(InPreviewConfiguration.bEnable);
 
 	ADisplayClusterRootActor* RootActor = GetRootActor();
 	if (RootActor)
@@ -224,7 +223,7 @@ bool FDisplayClusterViewportConfiguration::UpdatePreviewConfiguration(class UDis
 		if (ConfigurationData)
 		{
 			TArray<FString> RenderNodes;
-			if (PreviewConfiguration->PreviewNodeId == DisplayClusterConfigurationStrings::gui::preview::PreviewNodeAll)
+			if (InPreviewConfiguration.PreviewNodeId == DisplayClusterConfigurationStrings::gui::preview::PreviewNodeAll)
 			{
 				// Collect all nodes from cluster
 				for (const TPair<FString, UDisplayClusterConfigurationClusterNode*>& It : ConfigurationData->Cluster->Nodes)
@@ -234,17 +233,17 @@ bool FDisplayClusterViewportConfiguration::UpdatePreviewConfiguration(class UDis
 			}
 			else
 			{
-				RenderNodes.Add(PreviewConfiguration->PreviewNodeId);
+				RenderNodes.Add(InPreviewConfiguration.PreviewNodeId);
 			}
 
 			ImplUpdateConfiguration(RenderNodes, *RootActor, *ConfigurationData);
 
 			ImplUpdateConfigurationICVFX(*RootActor, *ConfigurationData);
 			ImplUpdateConfigurationVisibility(*RootActor, *ConfigurationData);
-			ImplUpdateRenderFrameConfiguration(*(RootActor->GetRenderFrameSettings()));
+			ImplUpdateRenderFrameConfiguration(RootActor->GetRenderFrameSettings());
 
 			// Downscale resources with PreviewDownscaleRatio
-			RenderFrameSettings.PreviewRenderTargetRatioMult = PreviewConfiguration->PreviewRenderTargetRatioMult;
+			RenderFrameSettings.PreviewRenderTargetRatioMult = InPreviewConfiguration.PreviewRenderTargetRatioMult;
 			RenderFrameSettings.RenderMode = EDisplayClusterRenderFrameMode::PreviewMono;
 
 			RenderFrameSettings.bIsRenderingInEditor = true;
diff --git a/Engine/Plugins/Runtime/nDisplay/Source/DisplayCluster/Private/Render/Viewport/Configuration/DisplayClusterViewportConfiguration.h b/Engine/Plugins/Runtime/nDisplay/Source/DisplayCluster/Private/Render/Viewport/Configuration/DisplayClusterViewportConfiguration.h
--- a/Engine/Plugins/Runtime/nDisplay/Source/DisplayCluster/Private/Render/Viewport/Configuration/DisplayClusterViewportConfiguration.h
+++ b/Engine/Plugins/Runtime/nDisplay/Source/DisplayCluster/Private/Render/Viewport/Configuration/DisplayClusterViewportConfiguration.h
@@ -13,7 +13,8 @@ struct FDisplayClusterRenderFrameSettings;
 class FDisplayClusterViewportManager;
 class ADisplayClusterRootActor;
 class UDisplayClusterConfigurationData;
-class UDisplayClusterConfigurationRenderFrame;
+struct FDisplayClusterConfigurationRenderFrame;
+struct FDisplayClusterConfigurationViewportPreview;
 
 class FDisplayClusterViewportConfiguration
 {
@@ -40,11 +41,11 @@ public:
 	bool UpdateConfiguration(EDisplayClusterRenderFrameMode InRenderMode, const FString& InClusterNodeId);
 
 #if WITH_EDITOR
-	bool UpdatePreviewConfiguration(class UDisplayClusterConfigurationViewportPreview* PreviewConfiguration);
+	bool UpdatePreviewConfiguration(const FDisplayClusterConfigurationViewportPreview& PreviewConfiguration);
 #endif
 
 private:
-	void ImplUpdateRenderFrameConfiguration(const UDisplayClusterConfigurationRenderFrame& InRenderFrameConfiguration);
+	void ImplUpdateRenderFrameConfiguration(const FDisplayClusterConfigurationRenderFrame& InRenderFrameConfiguration);
 	void ImplUpdateConfiguration(const TArray<FString>& InClusterNodeIds, ADisplayClusterRootActor& InRootActor, const UDisplayClusterConfigurationData& InConfigurationData);
 	void ImplUpdateConfigurationICVFX(ADisplayClusterRootActor& InRootActor, const UDisplayClusterConfigurationData& InConfigurationData);
 	void ImplUpdateConfigurationVisibility(ADisplayClusterRootActor& InRootActor, const UDisplayClusterConfigurationData& InConfigurationData);
diff --git a/Engine/Plugins/Runtime/nDisplay/Source/DisplayCluster/Private/Render/Viewport/Configuration/DisplayClusterViewportConfigurationBase.cpp b/Engine/Plugins/Runtime/nDisplay/Source/DisplayCluster/Private/Render/Viewport/Configuration/DisplayClusterViewportConfigurationBase.cpp
--- a/Engine/Plugins/Runtime/nDisplay/Source/DisplayCluster/Private/Render/Viewport/Configuration/DisplayClusterViewportConfigurationBase.cpp
+++ b/Engine/Plugins/Runtime/nDisplay/Source/DisplayCluster/Private/Render/Viewport/Configuration/DisplayClusterViewportConfigurationBase.cpp
@@ -3,11 +3,6 @@
 #include "DisplayClusterViewportConfigurationBase.h"
 #include "DisplayClusterViewportConfigurationHelpers.h"
 
-
-#include "IDisplayClusterProjection.h"
-#include "Render/Projection/IDisplayClusterProjectionPolicy.h"
-
-
 ///////////////////////////////////////////////////////////////////
 // FDisplayClusterViewportConfigurationBase
 ///////////////////////////////////////////////////////////////////
@@ -103,11 +98,7 @@ bool FDisplayClusterViewportConfigurationBase::UpdateViewportConfiguration(FDisp
 	}
 	else
 	{
-		if (!DesiredViewport->ProjectionPolicy.IsValid())
-		{
-			// Try initialize proj policy every tick (mesh deffered load, etc)
-			DesiredViewport->HandleStartScene();
-		}
+		DisplayClusterViewportConfigurationHelpers::UpdateProjectionPolicy(*DesiredViewport);
 	}
 
 	return true;
diff --git a/Engine/Plugins/Runtime/nDisplay/Source/DisplayCluster/Private/Render/Viewport/Configuration/DisplayClusterViewportConfigurationHelpers.h b/Engine/Plugins/Runtime/nDisplay/Source/DisplayCluster/Private/Render/Viewport/Configuration/DisplayClusterViewportConfigurationHelpers.h
--- a/Engine/Plugins/Runtime/nDisplay/Source/DisplayCluster/Private/Render/Viewport/Configuration/DisplayClusterViewportConfigurationHelpers.h
+++ b/Engine/Plugins/Runtime/nDisplay/Source/DisplayCluster/Private/Render/Viewport/Configuration/DisplayClusterViewportConfigurationHelpers.h
@@ -23,6 +23,9 @@
 #include "Render/Viewport/Containers/DisplayClusterViewport_RenderSettingsICVFX.h"
 #include "Render/Viewport/Containers/DisplayClusterViewport_PostRenderSettings.h"
 
+#include "IDisplayClusterProjection.h"
+#include "Render/Projection/IDisplayClusterProjectionPolicy.h"
+
 #include "DisplayClusterSceneViewExtensions.h"
 #include "OpenColorIODisplayExtension.h"
 
@@ -535,12 +538,12 @@ public:
 		}
 	}
 
-	static const FDisplayClusterConfigurationICVFX_ChromakeySettings& GetCameraChromakeySettings(const FDisplayClusterConfigurationICVFX_CameraSettings& InCameraSettings, const UDisplayClusterConfigurationICVFX_StageSettings& InStageSettings)
+	static const FDisplayClusterConfigurationICVFX_ChromakeySettings& GetCameraChromakeySettings(const FDisplayClusterConfigurationICVFX_CameraSettings& InCameraSettings, const FDisplayClusterConfigurationICVFX_StageSettings& InStageSettings)
 	{
 		return (InCameraSettings.CustomChromakey.bEnable) ? InCameraSettings.CustomChromakey.Chromakey : InStageSettings.Chromakey;
 	}
 
-	static void UpdateCameraViewportSetting(FDisplayClusterViewport& DstViewport, const FDisplayClusterConfigurationICVFX_CameraSettings& InCameraSettings, const UDisplayClusterConfigurationICVFX_StageSettings& InStageSettings)
+	static void UpdateCameraViewportSetting(FDisplayClusterViewport& DstViewport, const FDisplayClusterConfigurationICVFX_CameraSettings& InCameraSettings, const FDisplayClusterConfigurationICVFX_StageSettings& InStageSettings)
 	{
 		check(InCameraSettings.bEnable);
 
@@ -554,7 +557,7 @@ public:
 		UpdateViewportOCIOConfiguration(DstViewport, InCameraSettings.OCIO_Configuration);
 
 
-		// UDisplayClusterConfigurationICVFX_CameraSettings
+		// FDisplayClusterConfigurationICVFX_CameraSettings
 		{
 			DstViewport.RenderSettings.CameraId.Empty();
 			DstViewport.RenderSettings.BufferRatio = InCameraSettings.BufferRatio;
@@ -602,7 +605,19 @@ public:
 		}
 	}
 
-	static void UpdateChromakeyViewportSetting(FDisplayClusterViewport& DstViewport, const FDisplayClusterConfigurationICVFX_ChromakeySettings& InChromakeySettings, const UDisplayClusterConfigurationICVFX_StageSettings& InStageSettings)
+	static void UpdateProjectionPolicy(FDisplayClusterViewport& DstViewport)
+	{
+		if (!DstViewport.ProjectionPolicy.IsValid())
+		{
+			if (DstViewport.Owner.IsSceneOpened())
+			{
+				// Try initialize proj policy every tick (mesh deffered load, etc)
+				DstViewport.HandleStartScene();
+			}
+		}
+	}
+
+	static void UpdateChromakeyViewportSetting(FDisplayClusterViewport& DstViewport, const FDisplayClusterConfigurationICVFX_ChromakeySettings& InChromakeySettings, const FDisplayClusterConfigurationICVFX_StageSettings& InStageSettings)
 	{
 		check(InChromakeySettings.Source== EDisplayClusterConfigurationICVFX_ChromakeySource::ChromakeyRenderTexture);
 
@@ -670,7 +685,7 @@ public:
 		return IsVisibilitySettingsDefined(InLightcardSettings.ShowOnlyList);
 	}
 
-	static void UpdateLightcardViewportSetting(FDisplayClusterViewport& DstViewport, FDisplayClusterViewport& BaseViewport, const FDisplayClusterConfigurationICVFX_LightcardSettings& InLightcardSettings, const UDisplayClusterConfigurationICVFX_StageSettings& InStageSettings, bool bIsOpenColorIO)
+	static void UpdateLightcardViewportSetting(FDisplayClusterViewport& DstViewport, FDisplayClusterViewport& BaseViewport, const FDisplayClusterConfigurationICVFX_LightcardSettings& InLightcardSettings, const FDisplayClusterConfigurationICVFX_StageSettings& InStageSettings, bool bIsOpenColorIO)
 	{
 		check(InLightcardSettings.bEnable);
 
diff --git a/Engine/Plugins/Runtime/nDisplay/Source/DisplayCluster/Private/Render/Viewport/Configuration/DisplayClusterViewportConfigurationICVFX.cpp b/Engine/Plugins/Runtime/nDisplay/Source/DisplayCluster/Private/Render/Viewport/Configuration/DisplayClusterViewportConfigurationICVFX.cpp
--- a/Engine/Plugins/Runtime/nDisplay/Source/DisplayCluster/Private/Render/Viewport/Configuration/DisplayClusterViewportConfigurationICVFX.cpp
+++ b/Engine/Plugins/Runtime/nDisplay/Source/DisplayCluster/Private/Render/Viewport/Configuration/DisplayClusterViewportConfigurationICVFX.cpp
@@ -15,7 +15,6 @@
 #include "Render/Viewport/DisplayClusterViewportManager.h"
 
 #include "Components/DisplayClusterICVFX_CineCameraComponent.h"
-#include "Components/DisplayClusterICVFX_RefCineCameraComponent.h"
 
 ////////////////////////////////////////////////////////////////////////////////
 // FDisplayClusterViewportConfigurationCameraViewport
@@ -98,6 +97,9 @@ public:
 		// Setup camera visibility
 		DisplayClusterViewportConfigurationHelpers::UpdateVisibilitySetting(*CameraViewport, EDisplayClusterViewport_VisibilityMode::Hide, ConfigurationICVFX.StageSettings.HideList);
 
+		// Support projection policy update
+		DisplayClusterViewportConfigurationHelpers::UpdateProjectionPolicy(*CameraViewport);
+
 		return true;
 	}
 
@@ -136,6 +138,9 @@ public:
 		// Attach to parent viewport
 		ChromakeyViewport->RenderSettings.AssignParentViewport(CameraViewport->GetId(), CameraViewport->RenderSettings);
 
+		// Support projection policy update
+		DisplayClusterViewportConfigurationHelpers::UpdateProjectionPolicy(*ChromakeyViewport);
+
 		return true;
 	}
 
@@ -295,13 +300,6 @@ public:
 		, CameraMotionBlurParameters(InCameraComponent->GetMotionBlurParameters())
 	{}
 
-	FDisplayClusterViewportConfigurationCameraICVFX(const FTransform& InLocal2WorldTransform, class FDisplayClusterViewportConfigurationICVFX& InConfigurationICVFX, class UDisplayClusterICVFX_RefCineCameraComponent* const InCameraComponent)
-		: ConfigurationICVFX(InConfigurationICVFX)
-		, CameraViewport(InLocal2WorldTransform, InConfigurationICVFX, InCameraComponent->GetCameraComponent(), InCameraComponent->GetCameraUniqueId(), InCameraComponent->GetCameraSettingsICVFX())
-		, ChromakeySettings(DisplayClusterViewportConfigurationHelpers::GetCameraChromakeySettings(InCameraComponent->GetCameraSettingsICVFX(), ConfigurationICVFX.StageSettings))
-		, CameraMotionBlurParameters(InCameraComponent->GetMotionBlurParameters())
-	{}
-
 public:
 	// Create camera resources and intialize targets
 	void Update();
@@ -421,7 +419,7 @@ FDisplayClusterViewportConfigurationICVFX::FDisplayClusterViewportConfigurationI
 	: ViewportManager(InViewportManager)
 	, RootActor(InRootActor)
 	, ConfigurationData(InConfigurationData)
-	, StageSettings(*(InRootActor.GetStageSettings()))
+	, StageSettings(InRootActor.GetStageSettings())
 	, LightcardSettings(StageSettings.Lightcard)
 { }
 
@@ -593,6 +591,9 @@ bool FDisplayClusterViewportConfigurationICVFX::ImplCreateLightcardViewport(FDis
 		// Update configuration
 		DisplayClusterViewportConfigurationHelpers::UpdateLightcardViewportSetting(*LightcardViewport, BaseViewport, LightcardSettings, StageSettings, bIsOpenColorIO);
 
+		// Support projection policy update
+		DisplayClusterViewportConfigurationHelpers::UpdateProjectionPolicy(*LightcardViewport);
+
 		return true;
 	}
 
@@ -652,14 +653,6 @@ void FDisplayClusterViewportConfigurationICVFX::ImplGetCameras(TArray<FDisplayCl
 			{
 				NewCamera = new FDisplayClusterViewportConfigurationCameraICVFX(Local2WorldTransform, *this, CineCameraComponent);
 			}
-			else
-			{
-				UDisplayClusterICVFX_RefCineCameraComponent* RefCineCameraComponent = Cast<UDisplayClusterICVFX_RefCineCameraComponent>(ActorComponentIt);
-				if (RefCineCameraComponent && RefCineCameraComponent->IsShouldUseICVFX())
-				{
-					NewCamera = new FDisplayClusterViewportConfigurationCameraICVFX(Local2WorldTransform, *this, RefCineCameraComponent);
-				}
-			}
 		}
 
 		// If found ICVFX camera, check and add to used
diff --git a/Engine/Plugins/Runtime/nDisplay/Source/DisplayCluster/Private/Render/Viewport/Configuration/DisplayClusterViewportConfigurationICVFX.h b/Engine/Plugins/Runtime/nDisplay/Source/DisplayCluster/Private/Render/Viewport/Configuration/DisplayClusterViewportConfigurationICVFX.h
--- a/Engine/Plugins/Runtime/nDisplay/Source/DisplayCluster/Private/Render/Viewport/Configuration/DisplayClusterViewportConfigurationICVFX.h
+++ b/Engine/Plugins/Runtime/nDisplay/Source/DisplayCluster/Private/Render/Viewport/Configuration/DisplayClusterViewportConfigurationICVFX.h
@@ -18,9 +18,8 @@ struct FDisplayClusterConfigurationPostRender_BlurPostprocess;
 struct FDisplayClusterConfigurationPostRender_GenerateMips;
 struct FDisplayClusterConfigurationICVFX_OverlayAdvancedRenderSettings;
 struct FDisplayClusterConfigurationICVFX_LightcardSettings;
-
 struct FDisplayClusterConfigurationICVFX_CameraSettings;
-class UDisplayClusterConfigurationICVFX_StageSettings;
+struct FDisplayClusterConfigurationICVFX_StageSettings;
 
 
 class ADisplayClusterRootActor;
@@ -60,7 +59,7 @@ private:
 
 public:
 	const UDisplayClusterConfigurationData& ConfigurationData;
-	const UDisplayClusterConfigurationICVFX_StageSettings& StageSettings;
+	const FDisplayClusterConfigurationICVFX_StageSettings& StageSettings;
 	const FDisplayClusterConfigurationICVFX_LightcardSettings& LightcardSettings;
 };
 
diff --git a/Engine/Plugins/Runtime/nDisplay/Source/DisplayCluster/Private/Render/Viewport/DisplayClusterViewportManager.cpp b/Engine/Plugins/Runtime/nDisplay/Source/DisplayCluster/Private/Render/Viewport/DisplayClusterViewportManager.cpp
--- a/Engine/Plugins/Runtime/nDisplay/Source/DisplayCluster/Private/Render/Viewport/DisplayClusterViewportManager.cpp
+++ b/Engine/Plugins/Runtime/nDisplay/Source/DisplayCluster/Private/Render/Viewport/DisplayClusterViewportManager.cpp
@@ -275,9 +275,9 @@ void FDisplayClusterViewportManager::ConfigureViewFamily(const FDisplayClusterRe
 	}
 }
 
-void FDisplayClusterViewportManager::RenderFrame(const bool bWarpBlendEnabled, FViewport* InViewport)
+void FDisplayClusterViewportManager::RenderFrame(FViewport* InViewport)
 {
-	ViewportManagerProxy->ImplRenderFrame(bWarpBlendEnabled, InViewport);
+	ViewportManagerProxy->ImplRenderFrame(InViewport);
 }
 
 bool FDisplayClusterViewportManager::CreateViewport(const FString& InViewportId, const class UDisplayClusterConfigurationViewport* ConfigurationViewport)
diff --git a/Engine/Plugins/Runtime/nDisplay/Source/DisplayCluster/Private/Render/Viewport/DisplayClusterViewportManager.h b/Engine/Plugins/Runtime/nDisplay/Source/DisplayCluster/Private/Render/Viewport/DisplayClusterViewportManager.h
--- a/Engine/Plugins/Runtime/nDisplay/Source/DisplayCluster/Private/Render/Viewport/DisplayClusterViewportManager.h
+++ b/Engine/Plugins/Runtime/nDisplay/Source/DisplayCluster/Private/Render/Viewport/DisplayClusterViewportManager.h
@@ -55,10 +55,10 @@ public:
 
 	virtual void ConfigureViewFamily(const FDisplayClusterRenderFrame::FFrameRenderTarget& InFrameTarget, const FDisplayClusterRenderFrame::FFrameViewFamily& InFrameViewFamily, FSceneViewFamilyContext& InOutViewFamily) override;
 	
-	virtual void RenderFrame(const bool bWarpBlendEnabled, FViewport* InViewport) override;
+	virtual void RenderFrame(FViewport* InViewport) override;
 
 #if WITH_EDITOR
-	virtual bool UpdatePreviewConfiguration(class UDisplayClusterConfigurationViewportPreview* PreviewConfiguration, class ADisplayClusterRootActor* InRootActorPtr) override;
+	virtual bool UpdatePreviewConfiguration(const FDisplayClusterConfigurationViewportPreview& PreviewConfiguration, class ADisplayClusterRootActor* InRootActorPtr) override;
 	virtual bool RenderInEditor(class FDisplayClusterRenderFrame& InRenderFrame, FViewport* InViewport) override;
 	
 	void ImplUpdatePreviewRTTResources();
diff --git a/Engine/Plugins/Runtime/nDisplay/Source/DisplayCluster/Private/Render/Viewport/DisplayClusterViewportManagerProxy.cpp b/Engine/Plugins/Runtime/nDisplay/Source/DisplayCluster/Private/Render/Viewport/DisplayClusterViewportManagerProxy.cpp
--- a/Engine/Plugins/Runtime/nDisplay/Source/DisplayCluster/Private/Render/Viewport/DisplayClusterViewportManagerProxy.cpp
+++ b/Engine/Plugins/Runtime/nDisplay/Source/DisplayCluster/Private/Render/Viewport/DisplayClusterViewportManagerProxy.cpp
@@ -20,6 +20,17 @@
 
 #include "RHIContext.h"
 
+// Enable/disable warp&blend
+static TAutoConsoleVariable<int32> CVarWarpBlendEnabled(
+	TEXT("nDisplay.render.WarpBlendEnabled"),
+	1,
+	TEXT("Warp & Blend status\n")
+	TEXT("0 : disabled\n")
+	TEXT("1 : enabled\n")
+	,
+	ECVF_RenderThreadSafe
+);
+
 static TAutoConsoleVariable<int32> CVarCrossGPUTransfersEnabled(
 	TEXT("nDisplay.render.CrossGPUTransfers"),
 	1,
@@ -136,14 +147,16 @@ void FDisplayClusterViewportManagerProxy::ImplUpdateViewports(const TArray<FDisp
 
 DECLARE_GPU_STAT_NAMED(nDisplay_ViewportManager_RenderFrame, TEXT("nDisplay ViewportManager::RenderFrame"));
 
-void FDisplayClusterViewportManagerProxy::ImplRenderFrame(const bool bWarpBlendEnabled, FViewport* InViewport)
+void FDisplayClusterViewportManagerProxy::ImplRenderFrame(FViewport* InViewport)
 {
 	ENQUEUE_RENDER_COMMAND(DeleteDisplayClusterViewportProxy)(
-		[ViewportManagerProxy = this, bWarpBlendEnabled, InViewport](FRHICommandListImmediate& RHICmdList)
+		[ViewportManagerProxy = this, InViewport](FRHICommandListImmediate& RHICmdList)
 	{
 		SCOPED_GPU_STAT(RHICmdList, nDisplay_ViewportManager_RenderFrame);
 		SCOPED_DRAW_EVENT(RHICmdList, nDisplay_ViewportManager_RenderFrame);
 
+		bool bWarpBlendEnabled = ViewportManagerProxy->RenderFrameSettings.bAllowWarpBlend && CVarWarpBlendEnabled.GetValueOnRenderThread() != 0;
+
 		// mGPU not used for in-editor rendering
 		if(ViewportManagerProxy->RenderFrameSettings.bIsRenderingInEditor == false)
 		{
diff --git a/Engine/Plugins/Runtime/nDisplay/Source/DisplayCluster/Private/Render/Viewport/DisplayClusterViewportManagerProxy.h b/Engine/Plugins/Runtime/nDisplay/Source/DisplayCluster/Private/Render/Viewport/DisplayClusterViewportManagerProxy.h
--- a/Engine/Plugins/Runtime/nDisplay/Source/DisplayCluster/Private/Render/Viewport/DisplayClusterViewportManagerProxy.h
+++ b/Engine/Plugins/Runtime/nDisplay/Source/DisplayCluster/Private/Render/Viewport/DisplayClusterViewportManagerProxy.h
@@ -69,7 +69,7 @@ public:
 	void ImplUpdateViewports(const TArray<FDisplayClusterViewport*>& InViewports);
 	void ImplSafeRelease();
 
-	void ImplRenderFrame(const bool bWarpBlendEnabled, FViewport* InViewport);
+	void ImplRenderFrame(FViewport* InViewport);
 
 private:
 	TSharedPtr<FDisplayClusterRenderTargetManager>        RenderTargetManager;
diff --git a/Engine/Plugins/Runtime/nDisplay/Source/DisplayCluster/Private/Render/Viewport/DisplayClusterViewportManager_EditorPreview.cpp b/Engine/Plugins/Runtime/nDisplay/Source/DisplayCluster/Private/Render/Viewport/DisplayClusterViewportManager_EditorPreview.cpp
--- a/Engine/Plugins/Runtime/nDisplay/Source/DisplayCluster/Private/Render/Viewport/DisplayClusterViewportManager_EditorPreview.cpp
+++ b/Engine/Plugins/Runtime/nDisplay/Source/DisplayCluster/Private/Render/Viewport/DisplayClusterViewportManager_EditorPreview.cpp
@@ -96,7 +96,7 @@ void FDisplayClusterViewportManager::ImplUpdatePreviewRTTResources()
 	}
 }
 
-bool FDisplayClusterViewportManager::UpdatePreviewConfiguration(class UDisplayClusterConfigurationViewportPreview* PreviewConfiguration, class ADisplayClusterRootActor* InRootActorPtr)
+bool FDisplayClusterViewportManager::UpdatePreviewConfiguration(const FDisplayClusterConfigurationViewportPreview& PreviewConfiguration, class ADisplayClusterRootActor* InRootActorPtr)
 {
 	Configuration->SetRootActor(InRootActorPtr);
 	return Configuration->UpdatePreviewConfiguration(PreviewConfiguration);
@@ -188,7 +188,7 @@ bool FDisplayClusterViewportManager::RenderInEditor(class FDisplayClusterRenderF
 		FinalizeNewFrame();
 
 		// After all render target rendered call nDisplay frame rendering:
-		RenderFrame(InRenderFrame.bWarpBlendEnabled, InViewport);
+		RenderFrame(InViewport);
 
 		return true;
 	}
diff --git a/Engine/Plugins/Runtime/nDisplay/Source/DisplayCluster/Public/Components/DisplayClusterICVFX_CineCameraComponent.h b/Engine/Plugins/Runtime/nDisplay/Source/DisplayCluster/Public/Components/DisplayClusterICVFX_CineCameraComponent.h
--- a/Engine/Plugins/Runtime/nDisplay/Source/DisplayCluster/Public/Components/DisplayClusterICVFX_CineCameraComponent.h
+++ b/Engine/Plugins/Runtime/nDisplay/Source/DisplayCluster/Public/Components/DisplayClusterICVFX_CineCameraComponent.h
@@ -4,7 +4,9 @@
 
 #include "CoreMinimal.h"
 
+#include "CineCameraActor.h"
 #include "CineCameraComponent.h"
+
 #include "DisplayClusterConfigurationTypes_ICVFX.h"
 #include "Render/Viewport/Containers/DisplayClusterViewport_CameraMotionBlur.h"
 
@@ -23,15 +25,19 @@ public:
 	UDisplayClusterICVFX_CineCameraComponent(const FObjectInitializer& ObjectInitializer);
 
 public:
-	UPROPERTY(EditAnywhere, Category = "NDisplay ICVFX", meta = (ShowInnerProperties))
-	FDisplayClusterConfigurationICVFX_CameraSettings IncameraSettings;
+	// Use external cine camrea actor ref
+	UPROPERTY(EditAnywhere, Category = NDisplay)
+	TSoftObjectPtr<ACineCameraActor> CineCameraActorRef;
+
+	UPROPERTY(EditAnywhere, Category = NDisplay, meta = (ShowInnerProperties))
+	FDisplayClusterConfigurationICVFX_CameraSettings CameraSettings;
 
 public:
 	FDisplayClusterViewport_CameraMotionBlur GetMotionBlurParameters();
 
 	bool IsShouldUseICVFX() const
 	{
-		return IncameraSettings.bEnable;
+		return CameraSettings.bEnable;
 	}
 
 	// Return unique camera name
@@ -39,7 +45,7 @@ public:
 
 	const FDisplayClusterConfigurationICVFX_CameraSettings& GetCameraSettingsICVFX() const
 	{
-		return IncameraSettings;
+		return CameraSettings;
 	}
 
 	class UCameraComponent* GetCameraComponent();
diff --git a/Engine/Plugins/Runtime/nDisplay/Source/DisplayCluster/Public/Components/DisplayClusterICVFX_RefCineCameraComponent.h b/Engine/Plugins/Runtime/nDisplay/Source/DisplayCluster/Public/Components/DisplayClusterICVFX_RefCineCameraComponent.h
deleted file mode 100644
diff --git a/Engine/Plugins/Runtime/nDisplay/Source/DisplayCluster/Public/DisplayClusterRootActor.h b/Engine/Plugins/Runtime/nDisplay/Source/DisplayCluster/Public/DisplayClusterRootActor.h
--- a/Engine/Plugins/Runtime/nDisplay/Source/DisplayCluster/Public/DisplayClusterRootActor.h
+++ b/Engine/Plugins/Runtime/nDisplay/Source/DisplayCluster/Public/DisplayClusterRootActor.h
@@ -79,8 +79,8 @@ public:
 
 	UDisplayClusterSyncTickComponent* GetSyncTickComponent() const { return SyncTickComponent; }
 
-	UDisplayClusterConfigurationICVFX_StageSettings* GetStageSettings() const { return StageSettings; }
-	UDisplayClusterConfigurationRenderFrame* GetRenderFrameSettings() const { return RenderFrameSettings; }
+	const FDisplayClusterConfigurationICVFX_StageSettings& GetStageSettings() const;
+	const FDisplayClusterConfigurationRenderFrame& GetRenderFrameSettings() const;
 
 	UDisplayClusterConfigurationViewport* GetViewportConfiguration(const FString& ClusterNodeID, const FString& ViewportID);
 
@@ -163,24 +163,12 @@ public:
 	UDisplayClusterSceneComponent* GetComponentById(const FString& ComponentId) const;
 
 public:
-	UPROPERTY(EditAnywhere, Category = "NDisplay", NoClear)
-	UDisplayClusterConfigurationICVFX_StageSettings* StageSettings;
-
-	UPROPERTY(EditAnywhere, Category = "NDisplay", NoClear)
-	UDisplayClusterConfigurationRenderFrame* RenderFrameSettings;
-
 	IDisplayClusterViewportManager* GetViewportManager() const
 	{
 		return ViewportManager.IsValid() ? ViewportManager.Get() : nullptr;
 	}
 
 protected:
-	UPROPERTY(EditAnywhere, Category = "NDisplay", meta = (DisplayName = "Follow Local Player Camera"))
-	bool bFollowLocalPlayerCamera;
-
-	UPROPERTY(EditAnywhere, Category = "NDisplay", meta = (DisplayName = "Exit when ESC pressed"))
-	bool bExitOnEsc;
-
 	// Unique viewport manager for this configuration
 	TUniquePtr<IDisplayClusterViewportManager> ViewportManager;
 
@@ -259,10 +247,6 @@ public:
 	UPROPERTY(EditAnywhere, Category = "NDisplay Preview (Editor only)", meta = (ClampMin = "0.05", UIMin = "0.05", ClampMax = "1", UIMax = "1"))
 	float PreviewRenderTargetRatioMult = 0.25;
 
-	// Extra settings for preview
-	UPROPERTY(EditAnywhere, Category = "NDisplay Preview (Editor only)")
-	UDisplayClusterConfigurationViewportPreview* PreviewSettings;
-
 	UPROPERTY(EditAnywhere, Category = "NDisplay Preview (Editor only)")
 	float XformGizmoScale = 1.0f;
 
diff --git a/Engine/Plugins/Runtime/nDisplay/Source/DisplayCluster/Public/Render/Viewport/IDisplayClusterViewportManager.h b/Engine/Plugins/Runtime/nDisplay/Source/DisplayCluster/Public/Render/Viewport/IDisplayClusterViewportManager.h
--- a/Engine/Plugins/Runtime/nDisplay/Source/DisplayCluster/Public/Render/Viewport/IDisplayClusterViewportManager.h
+++ b/Engine/Plugins/Runtime/nDisplay/Source/DisplayCluster/Public/Render/Viewport/IDisplayClusterViewportManager.h
@@ -14,6 +14,8 @@ class ADisplayClusterRootActor;
 class UDisplayClusterConfigurationViewport;
 class IDisplayClusterViewportManagerProxy;
 
+struct FDisplayClusterConfigurationViewportPreview;
+
 class DISPLAYCLUSTER_API IDisplayClusterViewportManager
 {
 public:
@@ -77,10 +79,10 @@ public:
 	virtual void ConfigureViewFamily(const FDisplayClusterRenderFrame::FFrameRenderTarget& InFrameTarget, const FDisplayClusterRenderFrame::FFrameViewFamily& InFrameViewFamily, FSceneViewFamilyContext& InOutViewFamily) = 0;
 
 	// Send to render thread
-	virtual void RenderFrame(const bool bWarpBlendEnabled, FViewport* InViewport) = 0;
+	virtual void RenderFrame(FViewport* InViewport) = 0;
 
 #if WITH_EDITOR
-	virtual bool UpdatePreviewConfiguration(class UDisplayClusterConfigurationViewportPreview* PreviewConfiguration, ADisplayClusterRootActor* InRootActorPtr) = 0;
+	virtual bool UpdatePreviewConfiguration(const FDisplayClusterConfigurationViewportPreview& PreviewConfiguration, ADisplayClusterRootActor* InRootActorPtr) = 0;
 	virtual bool RenderInEditor(class FDisplayClusterRenderFrame& InRenderFrame, FViewport* InViewport) = 0;
 #endif
 
diff --git a/Engine/Plugins/Runtime/nDisplay/Source/DisplayCluster/Public/Render/Viewport/IDisplayClusterViewportManagerProxy.h b/Engine/Plugins/Runtime/nDisplay/Source/DisplayCluster/Public/Render/Viewport/IDisplayClusterViewportManagerProxy.h
--- a/Engine/Plugins/Runtime/nDisplay/Source/DisplayCluster/Public/Render/Viewport/IDisplayClusterViewportManagerProxy.h
+++ b/Engine/Plugins/Runtime/nDisplay/Source/DisplayCluster/Public/Render/Viewport/IDisplayClusterViewportManagerProxy.h
@@ -13,24 +13,7 @@ public:
 	{ }
 
 public:
-	/** Transfer view results accross GPUs. This is done only once all views have been rendered for RenderFrame
-	* [Rendering thread func]
-	*/
-	//virtual void DoCrossGPUTransfers_RenderThread(FViewport* InViewport, FRHICommandListImmediate& RHICmdList) const = 0;
-
-	/** Resolve render targets to unique viewport inputshader contexts resources
-	* Apply postprocess, generate mips, etc from settings in FDisplayClusterViewporDeferredUpdateSettings
-	* [Rendering thread func]
-	*/
-	//virtual void UpdateDeferredResources_RenderThread(FRHICommandListImmediate& RHICmdList) const = 0;
-
-	/** Apply Warp&Blend
-	* [Rendering thread func]
-	*/
-	//virtual void UpdateFrameResources_RenderThread(FRHICommandListImmediate& RHICmdList, bool bWarpBlendEnabled) const = 0;
-
 	///////////////////////////////////////////////////////////////////////////////////////////////////////////////////////////////////////////////////////////////////////////
-
 	/**
 	* Find viewport render thread proxy object by name
 	* [Rendering thread func]
diff --git a/Engine/Plugins/Runtime/nDisplay/Source/DisplayCluster/Public/Render/Viewport/RenderFrame/DisplayClusterRenderFrame.h b/Engine/Plugins/Runtime/nDisplay/Source/DisplayCluster/Public/Render/Viewport/RenderFrame/DisplayClusterRenderFrame.h
--- a/Engine/Plugins/Runtime/nDisplay/Source/DisplayCluster/Public/Render/Viewport/RenderFrame/DisplayClusterRenderFrame.h
+++ b/Engine/Plugins/Runtime/nDisplay/Source/DisplayCluster/Public/Render/Viewport/RenderFrame/DisplayClusterRenderFrame.h
@@ -108,9 +108,6 @@ public:
 
 	int DesiredNumberOfViews = 0;
 
-	// Allow warp blend
-	bool bWarpBlendEnabled = true;
-
 	IDisplayClusterViewportManager* ViewportManager = nullptr;
 };
 
diff --git a/Engine/Plugins/Runtime/nDisplay/Source/DisplayClusterConfiguration/Private/DisplayClusterConfigurationTypes.cpp b/Engine/Plugins/Runtime/nDisplay/Source/DisplayClusterConfiguration/Private/DisplayClusterConfigurationTypes.cpp
--- a/Engine/Plugins/Runtime/nDisplay/Source/DisplayClusterConfiguration/Private/DisplayClusterConfigurationTypes.cpp
+++ b/Engine/Plugins/Runtime/nDisplay/Source/DisplayClusterConfiguration/Private/DisplayClusterConfigurationTypes.cpp
@@ -131,11 +131,7 @@ FDisplayClusterConfigurationProjection::FDisplayClusterConfigurationProjection()
 	Type = TEXT("simple");
 }
 
-FDisplayClusterConfigurationICVFX_CameraSettings::FDisplayClusterConfigurationICVFX_CameraSettings()
-{
-}
-
-UDisplayClusterConfigurationICVFX_StageSettings::UDisplayClusterConfigurationICVFX_StageSettings()
+FDisplayClusterConfigurationICVFX_StageSettings::FDisplayClusterConfigurationICVFX_StageSettings()
 	: DefaultFrameSize(2560, 1440)
 {
 }
@@ -148,10 +144,6 @@ UDisplayClusterConfigurationViewport::UDisplayClusterConfigurationViewport()
 #endif
 }
 
-UDisplayClusterConfigurationViewportPreview::UDisplayClusterConfigurationViewportPreview()
-{
-}
-
 #if WITH_EDITOR
 
 void UDisplayClusterConfigurationViewport::PostEditChangeChainProperty(FPropertyChangedChainEvent& PropertyChangedEvent)
diff --git a/Engine/Plugins/Runtime/nDisplay/Source/DisplayClusterConfiguration/Public/DisplayClusterConfigurationStrings.h b/Engine/Plugins/Runtime/nDisplay/Source/DisplayClusterConfiguration/Public/DisplayClusterConfigurationStrings.h
--- a/Engine/Plugins/Runtime/nDisplay/Source/DisplayClusterConfiguration/Public/DisplayClusterConfigurationStrings.h
+++ b/Engine/Plugins/Runtime/nDisplay/Source/DisplayClusterConfiguration/Public/DisplayClusterConfigurationStrings.h
@@ -17,10 +17,15 @@ namespace DisplayClusterConfigurationStrings
 	// Property Categories
 	namespace categories
 	{
-		static constexpr auto DefaultCategory = TEXT("NDisplay");
-		static constexpr auto ClusterCategory = TEXT("NDisplay Cluster");
-		static constexpr auto ConfigurationCategory = TEXT("NDisplay Configuration");
-		static constexpr auto PreviewCategory = TEXT("NDisplay Preview (Editor only)");
+		static constexpr auto DefaultCategory              = TEXT("NDisplay");
+
+		static constexpr auto ClusterCategory              = TEXT("NDisplay Cluster");
+		static constexpr auto ClusterConfigurationCategory = TEXT("NDisplay Cluster Configuration");
+		static constexpr auto ClusterPostprocessCategory   = TEXT("NDisplay Cluster Postprocess");
+
+		static constexpr auto ICVFXCategory                = TEXT("NDisplay ICVFX");
+		static constexpr auto ConfigurationCategory        = TEXT("NDisplay Configuration");
+		static constexpr auto PreviewCategory              = TEXT("NDisplay Preview (Editor only)");
 	}
 
 	// Command line arguments
diff --git a/Engine/Plugins/Runtime/nDisplay/Source/DisplayClusterConfiguration/Public/DisplayClusterConfigurationTypes.h b/Engine/Plugins/Runtime/nDisplay/Source/DisplayClusterConfiguration/Public/DisplayClusterConfigurationTypes.h
--- a/Engine/Plugins/Runtime/nDisplay/Source/DisplayClusterConfiguration/Public/DisplayClusterConfigurationTypes.h
+++ b/Engine/Plugins/Runtime/nDisplay/Source/DisplayClusterConfiguration/Public/DisplayClusterConfigurationTypes.h
@@ -525,6 +525,18 @@ public:
 	UPROPERTY(EditAnywhere, Category = NDisplay)
 	FDisplayClusterConfigurationDiagnostics Diagnostics;
 
+	UPROPERTY(EditAnywhere, Category = NDisplay)
+	FDisplayClusterConfigurationICVFX_StageSettings StageSettings;
+
+	UPROPERTY(EditAnywhere, Category = NDisplay)
+	FDisplayClusterConfigurationRenderFrame RenderFrameSettings;
+
+	UPROPERTY(EditAnywhere, Category = NDisplay, meta = (DisplayName = "Follow Local Player Camera"))
+	bool bFollowLocalPlayerCamera = false;
+
+	UPROPERTY(EditAnywhere, Category = NDisplay, meta = (DisplayName = "Exit when ESC pressed"))
+	bool bExitOnEsc = true;
+
 	/** Create empty config data. */
 	static UDisplayClusterConfigurationData* CreateNewConfigData(UObject* Owner = nullptr, EObjectFlags ObjectFlags = RF_NoFlags);
 #if WITH_EDITORONLY_DATA
diff --git a/Engine/Plugins/Runtime/nDisplay/Source/DisplayClusterConfiguration/Public/DisplayClusterConfigurationTypes_ICVFX.h b/Engine/Plugins/Runtime/nDisplay/Source/DisplayClusterConfiguration/Public/DisplayClusterConfigurationTypes_ICVFX.h
--- a/Engine/Plugins/Runtime/nDisplay/Source/DisplayClusterConfiguration/Public/DisplayClusterConfigurationTypes_ICVFX.h
+++ b/Engine/Plugins/Runtime/nDisplay/Source/DisplayClusterConfiguration/Public/DisplayClusterConfigurationTypes_ICVFX.h
@@ -20,15 +20,15 @@ struct DISPLAYCLUSTERCONFIGURATION_API FDisplayClusterConfigurationICVFX_Visibil
 	GENERATED_BODY()
 
 public:
-	UPROPERTY(EditAnywhere, Category = "NDisplay ICVFX")
+	UPROPERTY(EditAnywhere, Category = NDisplay)
 	TArray<FActorLayer> ActorLayers;
 
-	UPROPERTY(EditAnywhere, Category = "NDisplay ICVFX")
+	UPROPERTY(EditAnywhere, Category = NDisplay)
 	TArray<AActor*> Actors;
 
 	//@todo change link, now by names
 	// Reference to RootActor components by names
-	UPROPERTY(EditAnywhere, Category = "NDisplay ICVFX")
+	UPROPERTY(EditAnywhere, Category = NDisplay)
 	TArray<FString> RootActorComponentNames;
 };
 
@@ -39,18 +39,18 @@ struct DISPLAYCLUSTERCONFIGURATION_API FDisplayClusterConfigurationICVFX_Chromak
 
 public:
 	// Allow shromakey markers rendering (Also require not empty MarkerTileRGBA)
-	UPROPERTY(EditAnywhere, Category = "NDisplay ICVFX")
+	UPROPERTY(EditAnywhere, Category = NDisplay)
 	bool bEnable = true;
 
 	// (*required) This texture must by tiled in both directions. Alpha channel used to composing
-	UPROPERTY(EditAnywhere, Category = "NDisplay ICVFX")
+	UPROPERTY(EditAnywhere, Category = NDisplay)
 	UTexture2D* MarkerTileRGBA = nullptr;
 
 	// Scale markers UV source
-	UPROPERTY(EditAnywhere, Category = "NDisplay ICVFX")
+	UPROPERTY(EditAnywhere, Category = NDisplay)
 	float MarkerTileScale = 1;
 
-	UPROPERTY(EditAnywhere, Category = "NDisplay ICVFX")
+	UPROPERTY(EditAnywhere, Category = NDisplay)
 	float MarkerTileDistance = 0;
 };
 
@@ -61,28 +61,28 @@ struct DISPLAYCLUSTERCONFIGURATION_API FDisplayClusterConfigurationICVFX_Overlay
 
 public:
 	// Allow ScreenPercentage 
-	UPROPERTY(EditAnywhere, Category = "NDisplay ICVFX", meta = (ClampMin = "0.05", UIMin = "0.05", ClampMax = "10.0", UIMax = "10.0"))
+	UPROPERTY(EditAnywhere, Category = NDisplay, meta = (ClampMin = "0.05", UIMin = "0.05", ClampMax = "10.0", UIMax = "10.0"))
 	float BufferRatio = 1;
 
 	// Performance: Render to scale RTT, resolved with shader to viewport (Custom value)
-	UPROPERTY(EditAnywhere, Category = "NDisplay ICVFX")
+	UPROPERTY(EditAnywhere, Category = NDisplay)
 	float RenderTargetRatio = 1;
 
 	// Performance, Multi-GPU: Asign GPU for viewport rendering. The Value '-1' used to default gpu mapping (EYE_LEFT and EYE_RIGHT GPU)
-	UPROPERTY(EditAnywhere, Category = "NDisplay ICVFX")
+	UPROPERTY(EditAnywhere, Category = NDisplay)
 	int GPUIndex = -1;
 
 	// Performance, Multi-GPU: Customize GPU for stereo mode second view (EYE_RIGHT GPU)
-	UPROPERTY(EditAnywhere, Category = "NDisplay ICVFX")
+	UPROPERTY(EditAnywhere, Category = NDisplay)
 	int StereoGPUIndex = -1;
 
 	// Performance: force monoscopic render, resolved to stereo viewport
-	UPROPERTY(EditAnywhere, Category = "NDisplay ICVFX")
+	UPROPERTY(EditAnywhere, Category = NDisplay)
 	EDisplayClusterConfigurationViewport_StereoMode StereoMode = EDisplayClusterConfigurationViewport_StereoMode::Default;
 
 	// Experimental: Support special frame builder mode - merge viewports to single viewfamily by group num
 	// [not implemented yet]
-	UPROPERTY(EditAnywhere, Category = "NDisplay ICVFX")
+	UPROPERTY()
 	int RenderFamilyGroup = -1;
 };
 
@@ -93,25 +93,25 @@ struct DISPLAYCLUSTERCONFIGURATION_API FDisplayClusterConfigurationICVFX_Chromak
 
 public:
 	// Debug: override the texture of the camera viewport from this chromakey RTT
-	UPROPERTY(EditAnywhere, Category = "NDisplay ICVFX")
+	UPROPERTY(EditAnywhere, Category = NDisplay)
 	bool bOverrideCameraViewport = false;
 
 	// Render actors from this layers to chromakey texture
-	UPROPERTY(EditAnywhere, Category = "NDisplay ICVFX")
+	UPROPERTY(EditAnywhere, Category = NDisplay)
 	FDisplayClusterConfigurationICVFX_VisibilityList ShowOnlyList;
 
 	// Override viewport render from source texture
-	UPROPERTY(EditAnywhere, Category = "NDisplay ICVFX")
+	UPROPERTY(EditAnywhere, Category = NDisplay)
 	FDisplayClusterConfigurationPostRender_Override Override;
 
-	UPROPERTY(EditAnywhere, Category = "NDisplay ICVFX")
+	UPROPERTY(EditAnywhere, Category = NDisplay)
 	FDisplayClusterConfigurationPostRender_BlurPostprocess PostprocessBlur;
 
-	UPROPERTY(EditAnywhere, Category = "NDisplay ICVFX")
+	UPROPERTY(EditAnywhere, Category = NDisplay)
 	FDisplayClusterConfigurationPostRender_GenerateMips GenerateMips;
 
 	// Advanced render settings
-	UPROPERTY(EditAnywhere, Category = "NDisplay ICVFX")
+	UPROPERTY(EditAnywhere, Category = NDisplay)
 	FDisplayClusterConfigurationICVFX_OverlayAdvancedRenderSettings AdvancedRenderSettings;
 };
 
@@ -120,21 +120,27 @@ struct DISPLAYCLUSTERCONFIGURATION_API FDisplayClusterConfigurationICVFX_Chromak
 {
 	GENERATED_BODY()
 
+public:
+	FDisplayClusterConfigurationICVFX_ChromakeySettings()
+		// Default chromakey color is (0,127,0)
+		: ChromakeyColor(0, 0.5f, 0)
+	{ }
+
 public:
 	// Allow chromakey rendering (also require not empty ChromakeyLayers)
-	UPROPERTY(EditAnywhere, Category = "NDisplay ICVFX")
+	UPROPERTY(EditAnywhere, Category = NDisplay)
 	EDisplayClusterConfigurationICVFX_ChromakeySource Source = EDisplayClusterConfigurationICVFX_ChromakeySource::None;
 
 	// Color to fill camera frame
-	UPROPERTY(EditAnywhere, Category = "NDisplay ICVFX")
-	FLinearColor ChromakeyColor = FLinearColor::Green;
+	UPROPERTY(EditAnywhere, Category = NDisplay)
+	FLinearColor ChromakeyColor;
 
 	// Settings for chromakey texture source rendering
-	UPROPERTY(EditAnywhere, Category = "NDisplay ICVFX")
+	UPROPERTY(EditAnywhere, Category = NDisplay)
 	FDisplayClusterConfigurationICVFX_ChromakeyRenderSettings ChromakeyRenderTexture;
 
 	// Global setup for chromakey markers rendering
-	UPROPERTY(EditAnywhere, Category = "NDisplay ICVFX")
+	UPROPERTY(EditAnywhere, Category = NDisplay)
 	FDisplayClusterConfigurationICVFX_ChromakeyMarkers ChromakeyMarkers;
 };
 
@@ -145,21 +151,21 @@ struct DISPLAYCLUSTERCONFIGURATION_API FDisplayClusterConfigurationICVFX_Lightca
 
 public:
 	// Debug: override the texture of the target viewport from this lightcard RTT
-	UPROPERTY(EditAnywhere, Category = "NDisplay ICVFX")
+	UPROPERTY(EditAnywhere, Category = NDisplay)
 	bool bOverrideViewport = false;
 
 	// Override viewport render from source texture
-	UPROPERTY(EditAnywhere, Category = "NDisplay ICVFX")
+	UPROPERTY(EditAnywhere, Category = NDisplay)
 	FDisplayClusterConfigurationPostRender_Override Override;
 
-	UPROPERTY(EditAnywhere, Category = "NDisplay ICVFX")
+	UPROPERTY(EditAnywhere, Category = NDisplay)
 	FDisplayClusterConfigurationPostRender_BlurPostprocess PostprocessBlur;
 
-	UPROPERTY(EditAnywhere, Category = "NDisplay ICVFX")
+	UPROPERTY(EditAnywhere, Category = NDisplay)
 	FDisplayClusterConfigurationPostRender_GenerateMips GenerateMips;
 
 	// Advanced render settings
-	UPROPERTY(EditAnywhere, Category = "NDisplay ICVFX")
+	UPROPERTY(EditAnywhere, Category = NDisplay)
 	FDisplayClusterConfigurationICVFX_OverlayAdvancedRenderSettings AdvancedRenderSettings;
 };
 
@@ -170,23 +176,23 @@ struct DISPLAYCLUSTERCONFIGURATION_API FDisplayClusterConfigurationICVFX_Lightca
 
 public:
 	// Allow lightcard rendering (also require not empty LightCardLayers)
-	UPROPERTY(EditAnywhere, Category = "NDisplay ICVFX")
+	UPROPERTY(EditAnywhere, Category = NDisplay)
 	bool bEnable = true;
 
 	// Global lighcard rendering mode
-	UPROPERTY(EditAnywhere, Category = "NDisplay ICVFX")
+	UPROPERTY(EditAnywhere, Category = NDisplay)
 	EDisplayClusterConfigurationICVFX_LightcardRenderMode Blendingmode = EDisplayClusterConfigurationICVFX_LightcardRenderMode::Under;
 
 	// Render actors from this layers to lightcard textures
-	UPROPERTY(EditAnywhere, Category = "NDisplay ICVFX")
+	UPROPERTY(EditAnywhere, Category = NDisplay)
 	FDisplayClusterConfigurationICVFX_VisibilityList ShowOnlyList;
 
 	// Configure global render settings for this viewports
-	UPROPERTY(EditAnywhere, Category = "NDisplay ICVFX")
+	UPROPERTY(EditAnywhere, Category = NDisplay)
 	FDisplayClusterConfigurationICVFX_LightcardRenderSettings RenderSettings;
 
 	// OCIO Display look configuration 
-	UPROPERTY(EditAnywhere, Category = "NDisplay ICVFX")
+	UPROPERTY(EditAnywhere, Category = NDisplay)
 	FOpenColorIODisplayConfiguration OCIO_Configuration;
 };
 
@@ -202,11 +208,11 @@ public:
 
 public:
 	// Camera frame size value source
-	UPROPERTY(EditAnywhere, Category = "NDisplay ICVFX")
+	UPROPERTY(EditAnywhere, Category = NDisplay)
 	EDisplayClusterConfigurationICVFX_CameraFrameSizeSource Size = EDisplayClusterConfigurationICVFX_CameraFrameSizeSource::Default;
 
 	// Frame size for this camera, used when selected "Custom size value"
-	UPROPERTY(EditAnywhere, Category = "NDisplay ICVFX")
+	UPROPERTY(EditAnywhere, Category = NDisplay)
 	FIntPoint CustomSizeValue;
 };
 
@@ -217,24 +223,24 @@ struct DISPLAYCLUSTERCONFIGURATION_API FDisplayClusterConfigurationICVFX_CameraA
 
 public:
 	// Performance: Render to scale RTT, resolved with shader to viewport (Custom value)
-	UPROPERTY(EditAnywhere, Category = "NDisplay ICVFX")
+	UPROPERTY(EditAnywhere, Category = NDisplay)
 	float RenderTargetRatio = 1;
 
 	// Performance, Multi-GPU: Asign GPU for viewport rendering. The Value '-1' used to default gpu mapping (EYE_LEFT and EYE_RIGHT GPU)
-	UPROPERTY(EditAnywhere, Category = "NDisplay ICVFX")
+	UPROPERTY(EditAnywhere, Category = NDisplay)
 	int GPUIndex = -1;
 
 	// Performance, Multi-GPU: Customize GPU for stereo mode second view (EYE_RIGHT GPU)
-	UPROPERTY(EditAnywhere, Category = "NDisplay ICVFX")
+	UPROPERTY(EditAnywhere, Category = NDisplay)
 	int StereoGPUIndex = -1;
 
 	// Performance: force monoscopic render, resolved to stereo viewport
-	UPROPERTY(EditAnywhere, Category = "NDisplay ICVFX")
+	UPROPERTY(EditAnywhere, Category = NDisplay)
 	EDisplayClusterConfigurationViewport_StereoMode StereoMode = EDisplayClusterConfigurationViewport_StereoMode::Default;
 
 	// Experimental: Support special frame builder mode - merge viewports to single viewfamily by group num
 	// [not implemented yet]
-	UPROPERTY(EditAnywhere, Category = "NDisplay ICVFX")
+	UPROPERTY()
 	int RenderFamilyGroup = -1;
 };
 
@@ -245,28 +251,28 @@ struct DISPLAYCLUSTERCONFIGURATION_API FDisplayClusterConfigurationICVFX_CameraR
 
 public:
 	// Define camera RTT texture size
-	UPROPERTY(EditAnywhere, Category = "NDisplay ICVFX")
+	UPROPERTY(EditAnywhere, Category = NDisplay)
 	FDisplayClusterConfigurationICVFX_CameraFrameSize FrameSize;
 
 	// Camera render order, bigger value is over
-	UPROPERTY(EditAnywhere, Category = "NDisplay ICVFX")
+	UPROPERTY(EditAnywhere, Category = NDisplay)
 	int RenderOrder = -1;
 
 	UPROPERTY()
 	FDisplayClusterConfigurationViewport_CustomPostprocess CustomPostprocess;
 
 	// Override viewport render from source texture
-	UPROPERTY(EditAnywhere, Category = "NDisplay ICVFX")
+	UPROPERTY(EditAnywhere, Category = NDisplay)
 	FDisplayClusterConfigurationPostRender_Override Override;
 
-	UPROPERTY(EditAnywhere, Category = "NDisplay ICVFX")
+	UPROPERTY(EditAnywhere, Category = NDisplay)
 	FDisplayClusterConfigurationPostRender_BlurPostprocess PostprocessBlur;
 
-	UPROPERTY(EditAnywhere, Category = "NDisplay ICVFX")
+	UPROPERTY(EditAnywhere, Category = NDisplay)
 	FDisplayClusterConfigurationPostRender_GenerateMips GenerateMips;
 
 	// Advanced render settings
-	UPROPERTY(EditAnywhere, Category = "NDisplay ICVFX")
+	UPROPERTY(EditAnywhere, Category = NDisplay)
 	FDisplayClusterConfigurationICVFX_CameraAdvancedRenderSettings AdvancedRenderSettings;
 };
 
@@ -278,10 +284,10 @@ struct DISPLAYCLUSTERCONFIGURATION_API FDisplayClusterConfigurationICVFX_CameraC
 
 public:
 	// Allow use local settings for chromakey and markers
-	UPROPERTY(EditAnywhere, Category = "NDisplay ICVFX")
+	UPROPERTY(EditAnywhere, Category = NDisplay)
 	bool bEnable = false;
 
-	UPROPERTY(EditAnywhere, Category = "NDisplay ICVFX")
+	UPROPERTY(EditAnywhere, Category = NDisplay)
 	FDisplayClusterConfigurationICVFX_ChromakeySettings Chromakey;
 };
 
@@ -293,10 +299,10 @@ struct DISPLAYCLUSTERCONFIGURATION_API FDisplayClusterConfigurationICVFX_CameraM
 	GENERATED_BODY()
 
 public:
-	UPROPERTY(EditAnywhere, Category = "NDisplay ICVFX")
+	UPROPERTY(EditAnywhere, Category = NDisplay)
 	EDisplayClusterConfigurationCameraMotionBlurMode MotionBlurMode = EDisplayClusterConfigurationCameraMotionBlurMode::Off;
 
-	UPROPERTY(EditAnywhere, Category = "NDisplay ICVFX")
+	UPROPERTY(EditAnywhere, Category = NDisplay)
 	float TranslationScale = 1.f;
 
 	// GUI: Add ext camera refs
@@ -308,75 +314,71 @@ struct DISPLAYCLUSTERCONFIGURATION_API FDisplayClusterConfigurationICVFX_CameraS
 {
 	GENERATED_BODY()
 
-public:
-	FDisplayClusterConfigurationICVFX_CameraSettings();
-
 public:
 	// Enable this camera
-	UPROPERTY(EditAnywhere, Category = "NDisplay ICVFX")
+	UPROPERTY(EditAnywhere, Category = NDisplay)
 	bool bEnable = true;
 
 	// Allow ScreenPercentage, for values!=1
-	UPROPERTY(EditAnywhere, Category = "NDisplay ICVFX", meta = (ClampMin = "0.05", UIMin = "0.05", ClampMax = "10.0", UIMax = "10.0"))
+	UPROPERTY(EditAnywhere, Category = NDisplay, meta = (ClampMin = "0.05", UIMin = "0.05", ClampMax = "10.0", UIMax = "10.0"))
 	float BufferRatio = 1;
 
-	UPROPERTY(EditAnywhere, Category = "NDisplay ICVFX")
+	UPROPERTY(EditAnywhere, Category = NDisplay)
 	float FieldOfViewMultiplier = 1.0f;
 
 	// Basic soft edges setup for incamera
-	UPROPERTY(EditAnywhere, Category = "NDisplay ICVFX")
+	UPROPERTY(EditAnywhere, Category = NDisplay)
 	FVector  SoftEdge = FVector::ZeroVector;
 
 	// Rotate incamera frustum on this value to fix broken lens on physic camera
-	UPROPERTY(EditAnywhere, Category = "NDisplay ICVFX")
+	UPROPERTY(EditAnywhere, Category = NDisplay)
 	FRotator  FrustumRotation = FRotator::ZeroRotator;
 
 	// Move incamera frustum on this value to fix broken lens on physic camera
-	UPROPERTY(EditAnywhere, Category = "NDisplay ICVFX")
+	UPROPERTY(EditAnywhere, Category = NDisplay)
 	FVector FrustumOffset = FVector::ZeroVector;
 
-	UPROPERTY(EditAnywhere, Category = "NDisplay ICVFX")
+	UPROPERTY(EditAnywhere, Category = NDisplay)
 	FDisplayClusterConfigurationICVFX_CameraMotionBlur CameraMotionBlur;
 
 	// Configure global render settings for this viewports
-	UPROPERTY(EditAnywhere, Category = "NDisplay ICVFX")
+	UPROPERTY(EditAnywhere, Category = NDisplay)
 	FDisplayClusterConfigurationICVFX_CameraRenderSettings RenderSettings;
 
-	UPROPERTY(EditAnywhere, Category = "NDisplay ICVFX")
+	UPROPERTY(EditAnywhere, Category = NDisplay)
 	FDisplayClusterConfigurationICVFX_CameraCustomChromakeySettings CustomChromakey;
 
 	// OCIO Display look configuration 
-	UPROPERTY(EditAnywhere, Category = "NDisplay ICVFX")
+	UPROPERTY(EditAnywhere, Category = NDisplay)
 	FOpenColorIODisplayConfiguration OCIO_Configuration;
 };
 
-UCLASS(ClassGroup = (DisplayClusterICVFX), meta = (BlueprintSpawnableComponent))
-class  DISPLAYCLUSTERCONFIGURATION_API UDisplayClusterConfigurationICVFX_StageSettings
-	 : public UActorComponent
+USTRUCT()
+struct DISPLAYCLUSTERCONFIGURATION_API FDisplayClusterConfigurationICVFX_StageSettings
 {
 	GENERATED_BODY()
 
 public:
-	UDisplayClusterConfigurationICVFX_StageSettings();
+	FDisplayClusterConfigurationICVFX_StageSettings();
 
 public:
 	// Allow ICVFX features
-	UPROPERTY(EditAnywhere, Category = "NDisplay ICVFX")
+	UPROPERTY(EditAnywhere, Category = NDisplay)
 	bool bEnable = true;
 
 	// Default incameras RTT texture size.
-	UPROPERTY(EditAnywhere, Category = "NDisplay ICVFX")
+	UPROPERTY(EditAnywhere, Category = NDisplay)
 	FIntPoint DefaultFrameSize;
 
-	UPROPERTY(EditAnywhere, Category = "NDisplay ICVFX")
+	UPROPERTY(EditAnywhere, Category = NDisplay)
 	FDisplayClusterConfigurationICVFX_ChromakeySettings Chromakey;
 
-	UPROPERTY(EditAnywhere, Category = "NDisplay ICVFX")
+	UPROPERTY(EditAnywhere, Category = NDisplay)
 	FDisplayClusterConfigurationICVFX_LightcardSettings Lightcard;
 
 	// Should be to add to this list all defined lightcards and chromakeys layers
 	// (This allow to hide all actors from layers for icvfx render logic)
-	UPROPERTY(EditAnywhere, Category = "NDisplay ICVFX")
+	UPROPERTY(EditAnywhere, Category = NDisplay)
 	FDisplayClusterConfigurationICVFX_VisibilityList HideList;
 };
 
diff --git a/Engine/Plugins/Runtime/nDisplay/Source/DisplayClusterConfiguration/Public/DisplayClusterConfigurationTypes_PostRender.h b/Engine/Plugins/Runtime/nDisplay/Source/DisplayClusterConfiguration/Public/DisplayClusterConfigurationTypes_PostRender.h
--- a/Engine/Plugins/Runtime/nDisplay/Source/DisplayClusterConfiguration/Public/DisplayClusterConfigurationTypes_PostRender.h
+++ b/Engine/Plugins/Runtime/nDisplay/Source/DisplayClusterConfiguration/Public/DisplayClusterConfigurationTypes_PostRender.h
@@ -17,13 +17,9 @@
 
 USTRUCT()
 struct DISPLAYCLUSTERCONFIGURATION_API FDisplayClusterConfigurationPostRender_Override
-	//: public UDisplayClusterConfigurationData_Base
 {
 	GENERATED_BODY()
 
-public:
-	//UDisplayClusterConfigurationPostRender_Override();
-
 public:
 	// Disable default render, and resolve SourceTexture to viewport
 	UPROPERTY(EditAnywhere, Category = "NDisplay Render")
@@ -44,13 +40,9 @@ public:
 
 USTRUCT()
 struct DISPLAYCLUSTERCONFIGURATION_API FDisplayClusterConfigurationPostRender_BlurPostprocess
-	//: public UDisplayClusterConfigurationData_Base
 {
 	GENERATED_BODY()
 
-public:
-	//UDisplayClusterConfigurationPostRender_BlurPostprocess();
-
 public:
 	UPROPERTY(EditAnywhere, Category = "NDisplay Render")
 	EDisplayClusterConfiguration_PostRenderBlur Mode = EDisplayClusterConfiguration_PostRenderBlur::None;
@@ -64,13 +56,9 @@ public:
 
 USTRUCT()
 struct DISPLAYCLUSTERCONFIGURATION_API FDisplayClusterConfigurationPostRender_GenerateMips
-	//: public UDisplayClusterConfigurationData_Base
 {
 	GENERATED_BODY()
 
-public:
-	//UDisplayClusterConfigurationPostRender_GenerateMips();
-
 public:
 	// Allow autogenerate num mips for this target
 	UPROPERTY(EditAnywhere, Category = "NDisplay Render")
diff --git a/Engine/Plugins/Runtime/nDisplay/Source/DisplayClusterConfiguration/Public/DisplayClusterConfigurationTypes_Viewport.h b/Engine/Plugins/Runtime/nDisplay/Source/DisplayClusterConfiguration/Public/DisplayClusterConfigurationTypes_Viewport.h
--- a/Engine/Plugins/Runtime/nDisplay/Source/DisplayClusterConfiguration/Public/DisplayClusterConfigurationTypes_Viewport.h
+++ b/Engine/Plugins/Runtime/nDisplay/Source/DisplayClusterConfiguration/Public/DisplayClusterConfigurationTypes_Viewport.h
@@ -146,7 +146,7 @@ public:
 
 	// Experimental: Support special frame builder mode - merge viewports to single viewfamily by group num
 	// [not implemented yet]
-	UPROPERTY(EditAnywhere, Category = "NDisplay Viewport")
+	UPROPERTY()
 	int RenderFamilyGroup = -1;
 };
 
@@ -282,6 +282,7 @@ public:
 	UPROPERTY(EditAnywhere, Category = "NDisplay Viewport")
 	bool bAllowRendering = true;
 
+	// @todo: GUI: Toggle visibility of this property: hide for camera projection policy, and show for other
 	UPROPERTY(EditAnywhere, Category = "NDisplay Viewport")
 	FString Camera;
 
@@ -333,73 +334,65 @@ public:
 };
 
 // This struct now stored in UDisplayClusterConfigurationData, and replicated with MultiUser
-UCLASS(ClassGroup = (DisplayCluster), meta = (BlueprintSpawnableComponent))
-class DISPLAYCLUSTERCONFIGURATION_API UDisplayClusterConfigurationRenderFrame
-	: public UActorComponent
+USTRUCT()
+struct DISPLAYCLUSTERCONFIGURATION_API FDisplayClusterConfigurationRenderFrame
 {
 	GENERATED_BODY()
 
-public:
-	UDisplayClusterConfigurationRenderFrame()
-	{};
-
 public:
 	// Performance: Allow change global MGPU settings
-	UPROPERTY(EditAnywhere, Category = "NDisplay Render")
+	UPROPERTY(EditAnywhere, Category = NDisplay)
 	EDisplayClusterConfigurationRenderMGPUMode MultiGPUMode = EDisplayClusterConfigurationRenderMGPUMode::Enabled;
 
 	// Performance: Allow merge multiple viewports on single RTT with atlasing (required for bAllowViewFamilyMergeOptimization)
-	UPROPERTY(EditAnywhere, Category = "NDisplay Render")
+	// [not implemented yet] Experimental
+	UPROPERTY()
 	bool bAllowRenderTargetAtlasing = false;
 
 	// Performance: Allow viewfamily merge optimization (render multiple viewports contexts within single family)
 	// [not implemented yet] Experimental
-	UPROPERTY(EditAnywhere, Category = "NDisplay Render")
+	UPROPERTY()
 	EDisplayClusterConfigurationRenderFamilyMode ViewFamilyMode = EDisplayClusterConfigurationRenderFamilyMode::None;
 
 	// Performance: Allow to use parent ViewFamily from parent viewport 
 	// (icvfx has child viewports: lightcard and chromakey with prj_view matrices copied from parent viewport. May sense to use same viewfamily?)
 	// [not implemented yet] Experimental
-	UPROPERTY(EditAnywhere, Category = "NDisplay Render")
+	UPROPERTY()
 	bool bShouldUseParentViewportRenderFamily = false;
 
 	// Multiply all downscale ratio inside all viewports settings for whole cluster
-	UPROPERTY(EditAnywhere, Category = "NDisplay Render", meta = (ClampMin = "0.05", UIMin = "0.05", ClampMax = "10", UIMax = "10"))
+	UPROPERTY(EditAnywhere, Category = NDisplay, meta = (ClampMin = "0.05", UIMin = "0.05", ClampMax = "1", UIMax = "1"))
 	float ClusterRenderTargetRatioMult = 1.f;
 
 	// Multiply all buffer ratios for whole cluster by this value
-	UPROPERTY(EditAnywhere, Category = "NDisplay Render", meta = (ClampMin = "0.05", UIMin = "0.05", ClampMax = "10", UIMax = "10"))
+	UPROPERTY(EditAnywhere, Category = NDisplay, meta = (ClampMin = "0.05", UIMin = "0.05", ClampMax = "10", UIMax = "10"))
 	float ClusterBufferRatioMult = 1.f;
 
 	// Allow warpblend render
-	UPROPERTY(EditAnywhere, Category = "NDisplay Render")
+	UPROPERTY(EditAnywhere, Category = NDisplay)
 	bool bAllowWarpBlend = true;
 };
 
-UCLASS(ClassGroup = (DisplayCluster), meta = (BlueprintSpawnableComponent))
-class DISPLAYCLUSTERCONFIGURATION_API UDisplayClusterConfigurationViewportPreview
-	: public UDisplayClusterConfigurationData_Base
+USTRUCT()
+struct DISPLAYCLUSTERCONFIGURATION_API FDisplayClusterConfigurationViewportPreview
 {
 	GENERATED_BODY()
 
-public:
-	UDisplayClusterConfigurationViewportPreview();
-
 public:
 	// Allow preview render
-	UPROPERTY(EditAnywhere, Category = "NDisplay Preview (Editor only)")
+	UPROPERTY()
 	bool bEnable = true;
 
 	// Render single node preview or whole cluster
-	UPROPERTY(EditAnywhere, Category = "NDisplay Preview (Editor only)")
+	UPROPERTY()
 	FString PreviewNodeId = DisplayClusterConfigurationStrings::gui::preview::PreviewNodeAll;
 
 	// Update preview texture period in tick
-	UPROPERTY(EditAnywhere, Category = "NDisplay Preview (Editor only)", meta = (ClampMin = "1", UIMin = "1", ClampMax = "200", UIMax = "200"))
+	UPROPERTY()
 	int TickPerFrame = 1;
 
 	// Preview texture size get from viewport, and scaled by this value
-	UPROPERTY(EditAnywhere, Category = "NDisplay Preview (Editor only)", meta = (ClampMin = "0.05", UIMin = "0.05", ClampMax = "1", UIMax = "1"))
+	UPROPERTY()
 	float PreviewRenderTargetRatioMult = 0.25;
 };
 
diff --git a/Engine/Plugins/Runtime/nDisplay/Source/DisplayClusterConfigurator/Private/DisplayClusterConfiguratorUtils.cpp b/Engine/Plugins/Runtime/nDisplay/Source/DisplayClusterConfigurator/Private/DisplayClusterConfiguratorUtils.cpp
--- a/Engine/Plugins/Runtime/nDisplay/Source/DisplayClusterConfigurator/Private/DisplayClusterConfiguratorUtils.cpp
+++ b/Engine/Plugins/Runtime/nDisplay/Source/DisplayClusterConfigurator/Private/DisplayClusterConfiguratorUtils.cpp
@@ -86,8 +86,6 @@ void FDisplayClusterConfiguratorUtils::AddRootActorComponentsToBlueprint(UDispla
 	// Already added as default subobjects.
 	AllComponents.Remove(RootActor->GetRootComponent());
 	AllComponents.Remove(RootActor->GetSyncTickComponent());
-	AllComponents.Remove(RootActor->GetRenderFrameSettings());
-	AllComponents.Remove(RootActor->GetStageSettings());
 	
 	if (NewRootNode == nullptr)
 	{
diff --git a/Engine/Plugins/Runtime/nDisplay/Source/DisplayClusterConfigurator/Private/Views/Details/DisplayClusterConfiguratorDetailCustomization.cpp b/Engine/Plugins/Runtime/nDisplay/Source/DisplayClusterConfigurator/Private/Views/Details/DisplayClusterConfiguratorDetailCustomization.cpp
--- a/Engine/Plugins/Runtime/nDisplay/Source/DisplayClusterConfigurator/Private/Views/Details/DisplayClusterConfiguratorDetailCustomization.cpp
+++ b/Engine/Plugins/Runtime/nDisplay/Source/DisplayClusterConfigurator/Private/Views/Details/DisplayClusterConfiguratorDetailCustomization.cpp
@@ -171,10 +171,17 @@ void FDisplayClusterConfiguratorDataDetailCustomization::CustomizeDetails(IDetai
 {
 	Super::CustomizeDetails(InLayoutBuilder);
 
+	BEGIN_CATEGORY(DisplayClusterConfigurationStrings::categories::ICVFXCategory)
+		ADD_PROPERTY(UDisplayClusterConfigurationData, StageSettings);
+	END_CATEGORY()
+
 	BEGIN_CATEGORY(DisplayClusterConfigurationStrings::categories::ConfigurationCategory)
-		ADD_EXPANDED_PROPERTY(UDisplayClusterConfigurationData, Info);
-		ADD_EXPANDED_PROPERTY(UDisplayClusterConfigurationData, Diagnostics);
+		ADD_EXPANDED_PROPERTY(UDisplayClusterConfigurationData, RenderFrameSettings);
+		ADD_PROPERTY(UDisplayClusterConfigurationData, Info);
+		ADD_PROPERTY(UDisplayClusterConfigurationData, Diagnostics);
 		ADD_PROPERTY(UDisplayClusterConfigurationData, CustomParameters);
+		ADD_PROPERTY(UDisplayClusterConfigurationData, bFollowLocalPlayerCamera);
+		ADD_PROPERTY(UDisplayClusterConfigurationData, bExitOnEsc);
 	END_CATEGORY()
 }
 
@@ -190,9 +197,18 @@ void FDisplayClusterConfiguratorClusterDetailCustomization::CustomizeDetails(IDe
 	check(ClusterNodesHandle->IsValidHandle());
 
 	BEGIN_CATEGORY(DisplayClusterConfigurationStrings::categories::ClusterCategory)
-		ADD_EXPANDED_PROPERTY(UDisplayClusterConfigurationCluster, MasterNode);
-		ADD_EXPANDED_PROPERTY(UDisplayClusterConfigurationCluster, Sync);
-		ADD_EXPANDED_PROPERTY(UDisplayClusterConfigurationCluster, Network);
+		ADD_PROPERTY(UDisplayClusterConfigurationCluster, Nodes);
+	END_CATEGORY()
+
+	BEGIN_CATEGORY(DisplayClusterConfigurationStrings::categories::ClusterPostprocessCategory)
+		ADD_PROPERTY(UDisplayClusterConfigurationCluster, bUseOverallClusterPostProcess);
+		ADD_PROPERTY(UDisplayClusterConfigurationCluster, OverallClusterPostProcessSettings);
+	END_CATEGORY()
+
+	BEGIN_CATEGORY(DisplayClusterConfigurationStrings::categories::ClusterConfigurationCategory)
+		ADD_PROPERTY(UDisplayClusterConfigurationCluster, MasterNode);
+		ADD_PROPERTY(UDisplayClusterConfigurationCluster, Sync);
+		ADD_PROPERTY(UDisplayClusterConfigurationCluster, Network);
 		ADD_PROPERTY(UDisplayClusterConfigurationCluster, bUseOverallClusterPostProcess);
 		ADD_PROPERTY(UDisplayClusterConfigurationCluster, OverallClusterPostProcessSettings);
 
@@ -218,7 +234,6 @@ void FDisplayClusterConfiguratorClusterDetailCustomization::CustomizeDetails(IDe
 				];
 		}
 
-		ADD_PROPERTY(UDisplayClusterConfigurationCluster, Nodes);
 	END_CATEGORY()
 }
 
diff --git a/Engine/Plugins/Runtime/nDisplay/Source/DisplayClusterConfigurator/Private/Views/Details/Policies/DisplayClusterConfiguratorPolicyDetailCustomization.cpp b/Engine/Plugins/Runtime/nDisplay/Source/DisplayClusterConfigurator/Private/Views/Details/Policies/DisplayClusterConfiguratorPolicyDetailCustomization.cpp
--- a/Engine/Plugins/Runtime/nDisplay/Source/DisplayClusterConfigurator/Private/Views/Details/Policies/DisplayClusterConfiguratorPolicyDetailCustomization.cpp
+++ b/Engine/Plugins/Runtime/nDisplay/Source/DisplayClusterConfigurator/Private/Views/Details/Policies/DisplayClusterConfiguratorPolicyDetailCustomization.cpp
@@ -11,7 +11,7 @@
 #include "DisplayClusterRootActor.h"
 #include "Blueprints/DisplayClusterBlueprint.h"
 #include "Components/DisplayClusterScreenComponent.h"
-#include "Components/DisplayClusterICVFX_RefCineCameraComponent.h"
+#include "Components/DisplayClusterICVFX_CineCameraComponent.h"
 
 #include "Camera/CameraComponent.h"
 
@@ -376,7 +376,7 @@ void FDisplayClusterConfiguratorProjectionCustomization::CreateCameraPolicy(UDis
 		ConfigurationViewportPtr.Get(),
 		ParametersHandle,
 		TArray<TSubclassOf<UActorComponent>>{ UCameraComponent::StaticClass(),
-			UDisplayClusterICVFX_RefCineCameraComponent::StaticClass() }));
+			UDisplayClusterICVFX_CineCameraComponent::StaticClass() }));
 
 	CustomPolicyParameters.Add(MakeShared<FPolicyParameterInfoBool>(
 		"Native",
diff --git a/Engine/Plugins/Runtime/nDisplay/Source/DisplayClusterConfigurator/Private/Views/SCSEditor/SDisplayClusterConfiguratorComponentCombo.cpp b/Engine/Plugins/Runtime/nDisplay/Source/DisplayClusterConfigurator/Private/Views/SCSEditor/SDisplayClusterConfiguratorComponentCombo.cpp
--- a/Engine/Plugins/Runtime/nDisplay/Source/DisplayClusterConfigurator/Private/Views/SCSEditor/SDisplayClusterConfiguratorComponentCombo.cpp
+++ b/Engine/Plugins/Runtime/nDisplay/Source/DisplayClusterConfigurator/Private/Views/SCSEditor/SDisplayClusterConfiguratorComponentCombo.cpp
@@ -12,7 +12,6 @@
 #include "Components/DisplayClusterCameraComponent.h"
 
 #include "Components/DisplayClusterICVFX_CineCameraComponent.h"
-#include "Components/DisplayClusterICVFX_RefCineCameraComponent.h"
 
 #include "Widgets/Layout/SSpacer.h"
 #include "Widgets/Images/SImage.h"
@@ -439,7 +438,6 @@ void SDisplayClusterConfiguratorComponentClassCombo::GenerateComponentClassList(
 
 	AddHeader("nDisplay ICVFX Components");
 	AddDCComp(UDisplayClusterICVFX_CineCameraComponent::StaticClass());
-	AddDCComp(UDisplayClusterICVFX_RefCineCameraComponent::StaticClass());
 
 	int32 CompIterationIdx = 0;
 	for (const TSharedPtr<FComponentClassComboEntry>& Comp : *ComponentClassListPtr)
diff --git a/Engine/Plugins/Runtime/nDisplay/Source/DisplayClusterConfigurator/Private/Views/Viewport/DisplayClusterConfiguratorSCSEditorViewport.cpp b/Engine/Plugins/Runtime/nDisplay/Source/DisplayClusterConfigurator/Private/Views/Viewport/DisplayClusterConfiguratorSCSEditorViewport.cpp
--- a/Engine/Plugins/Runtime/nDisplay/Source/DisplayClusterConfigurator/Private/Views/Viewport/DisplayClusterConfiguratorSCSEditorViewport.cpp
+++ b/Engine/Plugins/Runtime/nDisplay/Source/DisplayClusterConfigurator/Private/Views/Viewport/DisplayClusterConfiguratorSCSEditorViewport.cpp
@@ -12,7 +12,7 @@
 #include "ClusterConfiguration/ViewModels/DisplayClusterConfiguratorProjectionPolicyViewModel.h"
 
 #include "Components/DisplayClusterScreenComponent.h"
-#include "Components/DisplayClusterICVFX_RefCineCameraComponent.h"
+#include "Components/DisplayClusterICVFX_CineCameraComponent.h"
 
 #include "Camera/CameraComponent.h"
 
@@ -699,7 +699,7 @@ FReply SDisplayClusterConfiguratorSCSEditorViewport::OnDragOver(const FGeometry&
 					else
 					{
 						USceneComponent* AttachParent = MeshComponent->GetAttachParent();
-						const bool bIsCamera = AttachParent && (AttachParent->IsA<UCameraComponent>() || AttachParent->IsA<UDisplayClusterICVFX_RefCineCameraComponent>());
+						const bool bIsCamera = AttachParent && (AttachParent->IsA<UCameraComponent>());
 						
 						if (bIsCamera)
 						{
@@ -756,7 +756,7 @@ FReply SDisplayClusterConfiguratorSCSEditorViewport::OnDrop(const FGeometry& MyG
 					else
 					{
 						USceneComponent* AttachParent = MeshComponent->GetAttachParent();
-						const bool bIsCamera = AttachParent && (AttachParent->IsA<UCameraComponent>() || AttachParent->IsA<UDisplayClusterICVFX_RefCineCameraComponent>());
+						const bool bIsCamera = AttachParent && (AttachParent->IsA<UCameraComponent>());
 						
 						if (bIsCamera)
 						{
diff --git a/Engine/Plugins/Runtime/nDisplay/Source/DisplayClusterConfigurator/Private/Views/Viewport/DisplayClusterConfiguratorSCSEditorViewportClient.cpp b/Engine/Plugins/Runtime/nDisplay/Source/DisplayClusterConfigurator/Private/Views/Viewport/DisplayClusterConfiguratorSCSEditorViewportClient.cpp
--- a/Engine/Plugins/Runtime/nDisplay/Source/DisplayClusterConfigurator/Private/Views/Viewport/DisplayClusterConfiguratorSCSEditorViewportClient.cpp
+++ b/Engine/Plugins/Runtime/nDisplay/Source/DisplayClusterConfigurator/Private/Views/Viewport/DisplayClusterConfiguratorSCSEditorViewportClient.cpp
@@ -1087,7 +1087,6 @@ void FDisplayClusterConfiguratorSCSEditorViewportClient::SyncShowPreview()
 			if (CorrectShowPreviewValue != Actor->PreviewNodeId)
 			{
 				Actor->PreviewNodeId = CorrectShowPreviewValue;
-				Actor->PreviewSettings->PreviewNodeId = CorrectShowPreviewValue;
 				Actor->UpdatePreviewComponents();
 			}
 		}
diff --git a/Engine/Plugins/Runtime/nDisplay/Source/DisplayClusterEditor/Private/DetailsCustomization/DisplayClusterRootActorDetailsCustomization.cpp b/Engine/Plugins/Runtime/nDisplay/Source/DisplayClusterEditor/Private/DetailsCustomization/DisplayClusterRootActorDetailsCustomization.cpp
--- a/Engine/Plugins/Runtime/nDisplay/Source/DisplayClusterEditor/Private/DetailsCustomization/DisplayClusterRootActorDetailsCustomization.cpp
+++ b/Engine/Plugins/Runtime/nDisplay/Source/DisplayClusterEditor/Private/DetailsCustomization/DisplayClusterRootActorDetailsCustomization.cpp
@@ -166,20 +166,23 @@ bool FDisplayClusterRootActorDetailsCustomization::RebuildNodeIdOptionsList()
 
 void FDisplayClusterRootActorDetailsCustomization::UpdateNodeIdSelection()
 {
-	const FString& CurrentPreviewNode = EditedObject->PreviewNodeId;
-	TSharedPtr<FString>* FoundItem = NodeIdOptions.FindByPredicate([CurrentPreviewNode](const TSharedPtr<FString>& Item)
+	if (NodeIdComboBox.IsValid())
 	{
-		return Item->Equals(CurrentPreviewNode, ESearchCase::IgnoreCase);
-	});
+		const FString& CurrentPreviewNode = EditedObject->PreviewNodeId;
+		TSharedPtr<FString>* FoundItem = NodeIdOptions.FindByPredicate([CurrentPreviewNode](const TSharedPtr<FString>& Item)
+		{
+			return Item->Equals(CurrentPreviewNode, ESearchCase::IgnoreCase);
+		});
 
-	if (FoundItem)
-	{
-		NodeIdComboBox->SetSelectedItem(*FoundItem);
-	}
-	else
-	{
-		// Set combobox selected item (options list is not empty here)
-		NodeIdComboBox->SetSelectedItem(NodeIdOptionAll);
+		if (FoundItem)
+		{
+			NodeIdComboBox->SetSelectedItem(*FoundItem);
+		}
+		else
+		{
+			// Set combobox selected item (options list is not empty here)
+			NodeIdComboBox->SetSelectedItem(NodeIdOptionAll);
+		}
 	}
 }
 
@@ -191,8 +194,13 @@ void FDisplayClusterRootActorDetailsCustomization::OnNodeIdSelected(TSharedPtr<F
 
 FText FDisplayClusterRootActorDetailsCustomization::GetSelectedNodeIdText() const
 {
-	TSharedPtr<FString> CurSelection = NodeIdComboBox->GetSelectedItem();
-	return FText::FromString(CurSelection.IsValid() ? *CurSelection : *NodeIdOptionNone);
+	if (NodeIdComboBox.IsValid())
+	{
+		TSharedPtr<FString> CurSelection = NodeIdComboBox->GetSelectedItem();
+		return FText::FromString(CurSelection.IsValid() ? *CurSelection : *NodeIdOptionNone);
+	}
+
+	return FText::FromString(*NodeIdOptionNone);
 }
 
 
diff --git a/Engine/Plugins/Runtime/nDisplay/Source/DisplayClusterProjection/Private/Policy/Camera/DisplayClusterProjectionCameraPolicy.cpp b/Engine/Plugins/Runtime/nDisplay/Source/DisplayClusterProjection/Private/Policy/Camera/DisplayClusterProjectionCameraPolicy.cpp
--- a/Engine/Plugins/Runtime/nDisplay/Source/DisplayClusterProjection/Private/Policy/Camera/DisplayClusterProjectionCameraPolicy.cpp
+++ b/Engine/Plugins/Runtime/nDisplay/Source/DisplayClusterProjection/Private/Policy/Camera/DisplayClusterProjectionCameraPolicy.cpp
@@ -15,6 +15,8 @@
 #include "Camera/CameraComponent.h"
 #include "ComposurePostMoves.h"
 
+#include "Components/DisplayClusterICVFX_CineCameraComponent.h"
+
 #include "Render/Viewport/IDisplayClusterViewport.h"
 #include "Render/Viewport/IDisplayClusterViewportManager.h"
 
@@ -34,18 +36,30 @@ bool FDisplayClusterProjectionCameraPolicy::HandleStartScene(class IDisplayClust
 {
 	check(IsInGameThread());
 
-	UCameraComponent* CfgCamera = nullptr;
+	USceneComponent* CfgComponent = nullptr;
 	FDisplayClusterProjectionCameraPolicySettings CfgCameraSettings;
 
-	if (!GetSettingsFromConfig(InViewport, CfgCamera, CfgCameraSettings))
+	if (!GetSettingsFromConfig(InViewport, CfgComponent, CfgCameraSettings))
 	{
 		UE_LOG(LogDisplayClusterProjectionCamera, Error, TEXT("Invalid camera settings for viewport: %s"), *InViewport->GetId());
 		return false;
 	}
 
-	if (CfgCamera) 
+	if (CfgComponent)
 	{
-		SetCamera(CfgCamera, CfgCameraSettings);
+		UDisplayClusterICVFX_CineCameraComponent* CameraICVFX = Cast<UDisplayClusterICVFX_CineCameraComponent>(CfgComponent);
+		if (CameraICVFX)
+		{
+			AssignedICVFXCameraRef.SetSceneComponent(CameraICVFX);
+			UpdateICVFXCamera(CameraICVFX);
+		}
+		else
+		{
+			AssignedICVFXCameraRef.ResetSceneComponent();
+
+			UCameraComponent* Camera = StaticCast<UCameraComponent*>(CfgComponent);
+			SetCamera(Camera, CfgCameraSettings);
+		}
 	}
 
 	return true;
@@ -55,10 +69,10 @@ void FDisplayClusterProjectionCameraPolicy::HandleEndScene(class IDisplayCluster
 {
 	check(IsInGameThread());
 
-	AssignedCamera = nullptr;
+	AssignedCameraRef.ResetSceneComponent();
+	AssignedICVFXCameraRef.ResetSceneComponent();
 }
 
-
 APlayerCameraManager* const GetCurPlayerCameraManager(IDisplayClusterViewport* InViewport)
 {
 	UWorld* World = InViewport->GetOwner().GetCurrentWorld();
@@ -74,18 +88,93 @@ APlayerCameraManager* const GetCurPlayerCameraManager(IDisplayClusterViewport* I
 	return nullptr;
 }
 
+bool FDisplayClusterProjectionCameraPolicy::UpdateICVFXCamera(UDisplayClusterICVFX_CineCameraComponent* InComponentICVFX)
+{
+	if (InComponentICVFX)
+	{
+		FDisplayClusterProjectionCameraPolicySettings PolicyCameraSettings;
+		{
+			const FDisplayClusterConfigurationICVFX_CameraSettings& CameraSettingsICVFX = InComponentICVFX->GetCameraSettingsICVFX();
+			PolicyCameraSettings.FOVMultiplier = CameraSettingsICVFX.FieldOfViewMultiplier;
+
+			// Lens correction
+			PolicyCameraSettings.FrustumRotation = CameraSettingsICVFX.FrustumRotation;
+			PolicyCameraSettings.FrustumOffset   = CameraSettingsICVFX.FrustumOffset;
+		}
+
+		// Update referenced component and settings
+		SetCamera(InComponentICVFX->GetCameraComponent(), PolicyCameraSettings);
+
+		return true;
+	}
+
+	return false;
+}
+
+UCameraComponent* FDisplayClusterProjectionCameraPolicy::GetCameraComponent()
+{
+	// Support runtime update for ICVFX camera component
+	UDisplayClusterICVFX_CineCameraComponent* CameraICVFX = ImplGetAssignedICVFXComponent();
+	if (CameraICVFX)
+	{
+		UpdateICVFXCamera(CameraICVFX);
+	}
+
+	return ImplGetAssignedCamera();
+}
+
+UDisplayClusterICVFX_CineCameraComponent* FDisplayClusterProjectionCameraPolicy::ImplGetAssignedICVFXComponent()
+{
+	if (AssignedICVFXCameraRef.IsDefinedSceneComponent())
+	{
+		USceneComponent* SceneComponent = AssignedICVFXCameraRef.GetOrFindSceneComponent();
+		if (SceneComponent)
+		{
+			UDisplayClusterICVFX_CineCameraComponent* CameraComponentICVFX = Cast<UDisplayClusterICVFX_CineCameraComponent>(SceneComponent);
+			if (CameraComponentICVFX)
+			{
+				return CameraComponentICVFX;
+			}
+		}
+
+		//@todo: handle error: reference to deleted ICVFX component
+	}
+
+	return nullptr;
+}
+
+UCameraComponent* FDisplayClusterProjectionCameraPolicy::ImplGetAssignedCamera()
+{
+	if (AssignedCameraRef.IsDefinedSceneComponent())
+	{
+		USceneComponent* SceneComponent = AssignedCameraRef.GetOrFindSceneComponent();
+		if (SceneComponent)
+		{
+			UCameraComponent* CameraComponent = Cast<UCameraComponent>(SceneComponent);
+			if (CameraComponent)
+			{
+				return CameraComponent;
+			}
+		}
+
+		//@todo: handle error: reference to deleted component
+	}
+
+	return nullptr;
+}
+
 bool FDisplayClusterProjectionCameraPolicy::CalculateView(class IDisplayClusterViewport* InViewport, const uint32 InContextNum, FVector& InOutViewLocation, FRotator& InOutViewRotation, const FVector& ViewOffset, const float WorldToMeters, const float NCP, const float FCP)
 {
 	check(IsInGameThread());
 
 	InOutViewLocation = FVector::ZeroVector;
 	InOutViewRotation = FRotator::ZeroRotator;
-
+	
 	// Use transform of an assigned camera
-	if (AssignedCamera)
+	if (UCameraComponent* AssignedCameraComponent = GetCameraComponent())
 	{
-		InOutViewLocation = AssignedCamera->GetComponentLocation();
-		InOutViewRotation = AssignedCamera->GetComponentRotation();
+		InOutViewLocation = AssignedCameraComponent->GetComponentLocation();
+		InOutViewRotation = AssignedCameraComponent->GetComponentRotation();
 	}
 	// Otherwise default UE camera is used
 	else
@@ -111,9 +200,9 @@ bool FDisplayClusterProjectionCameraPolicy::GetProjectionMatrix(IDisplayClusterV
 
 	FComposurePostMoveSettings ComposureSettings;
 
-	if (AssignedCamera)
+	if (UCameraComponent* AssignedCameraComponent = GetCameraComponent())
 	{
-		OutPrjMatrix = ComposureSettings.GetProjectionMatrix(AssignedCamera->FieldOfView * CameraSettings.FOVMultiplier, AssignedCamera->AspectRatio);
+		OutPrjMatrix = ComposureSettings.GetProjectionMatrix(AssignedCameraComponent->FieldOfView * CameraSettings.FOVMultiplier, AssignedCameraComponent->AspectRatio);
 		return true;
 	}
 	else
@@ -130,30 +219,22 @@ bool FDisplayClusterProjectionCameraPolicy::GetProjectionMatrix(IDisplayClusterV
 }
 
 
-bool FDisplayClusterProjectionCameraPolicy::GetSettingsFromConfig(class IDisplayClusterViewport* InViewport, UCameraComponent*& OutCamera, FDisplayClusterProjectionCameraPolicySettings& OutCameraSettings)
+bool FDisplayClusterProjectionCameraPolicy::GetSettingsFromConfig(class IDisplayClusterViewport* InViewport, USceneComponent*& OutSceneComponent, FDisplayClusterProjectionCameraPolicySettings& OutCameraSettings)
 {
 	check(InViewport);
-	
-	ADisplayClusterRootActor* const RootActor = InViewport->GetOwner().GetRootActor();
-	if (!RootActor)
-	{
-		UE_LOG(LogDisplayClusterProjectionCamera, Error, TEXT("Couldn't get a DisplayClusterRootActor root object"));
-		return false;
-	}
 
 	FString CameraComponentId;
 	// Get assigned camera ID
 	if (!DisplayClusterHelpers::map::template ExtractValue(GetParameters(), DisplayClusterProjectionStrings::cfg::camera::Component, CameraComponentId))
 	{
+		// use default cameras
+		return true;
+	}
 
-#if WITH_EDITOR
-		if (CameraComponentId.IsEmpty())
-		{
-			return false;
-		}
-#endif
-
-		UE_LOG(LogDisplayClusterProjectionCamera, Error, TEXT("No camera component ID '%s' specified for projection policy '%s'"), *CameraComponentId, *GetId());
+	ADisplayClusterRootActor* const RootActor = InViewport->GetOwner().GetRootActor();
+	if (!RootActor)
+	{
+		UE_LOG(LogDisplayClusterProjectionCamera, Error, TEXT("Couldn't get a DisplayClusterRootActor root object"));
 		return false;
 	}
 
@@ -164,11 +245,24 @@ bool FDisplayClusterProjectionCameraPolicy::GetSettingsFromConfig(class IDisplay
 	{
 		if (Comp->GetName() == CameraComponentId)
 		{
-			OutCamera = Comp;
+			OutSceneComponent = Comp;
 			return true;
 		}
 	}
 
+	// Get ICVFX camera component
+	TArray<UDisplayClusterICVFX_CineCameraComponent*> ICVFXCameraComps;
+	RootActor->GetComponents<UDisplayClusterICVFX_CineCameraComponent>(ICVFXCameraComps);
+	for (UCameraComponent* Comp : ICVFXCameraComps)
+	{
+		if (Comp->GetName() == CameraComponentId)
+		{
+			OutSceneComponent = Comp;
+			return true;
+		}
+	}
+	
+
 	UE_LOG(LogDisplayClusterProjectionCamera, Error, TEXT("Camera component ID '%s' not found for projection policy '%s'"), *CameraComponentId, *GetId());
 
 	return false;
@@ -185,10 +279,11 @@ void FDisplayClusterProjectionCameraPolicy::SetCamera(UCameraComponent* NewCamer
 	if (NewCamera)
 	{
 		UE_LOG(LogDisplayClusterProjectionCamera, Verbose, TEXT("New camera set: %s"), *NewCamera->GetFullName());
-		AssignedCamera = NewCamera;
+		AssignedCameraRef.SetSceneComponent(NewCamera);
 	}
 	else
 	{
+		AssignedCameraRef.ResetSceneComponent();
 		UE_LOG(LogDisplayClusterProjectionCamera, Warning, TEXT("Trying to set nullptr camera pointer"));
 	}
 
diff --git a/Engine/Plugins/Runtime/nDisplay/Source/DisplayClusterProjection/Private/Policy/Camera/DisplayClusterProjectionCameraPolicy.h b/Engine/Plugins/Runtime/nDisplay/Source/DisplayClusterProjection/Private/Policy/Camera/DisplayClusterProjectionCameraPolicy.h
--- a/Engine/Plugins/Runtime/nDisplay/Source/DisplayClusterProjection/Private/Policy/Camera/DisplayClusterProjectionCameraPolicy.h
+++ b/Engine/Plugins/Runtime/nDisplay/Source/DisplayClusterProjection/Private/Policy/Camera/DisplayClusterProjectionCameraPolicy.h
@@ -4,7 +4,9 @@
 
 #include "Policy/DisplayClusterProjectionPolicyBase.h"
 #include "Containers/DisplayClusterProjectionCameraPolicySettings.h"
+#include "Misc/DisplayClusterObjectRef.h"
 
+class UDisplayClusterICVFX_CineCameraComponent;
 class UCameraComponent;
 class UWorld;
 
@@ -38,12 +40,19 @@ public:
 	void SetCamera(UCameraComponent* const NewCamera, const FDisplayClusterProjectionCameraPolicySettings& InCameraSettings);
 
 protected:
-	bool GetSettingsFromConfig(class IDisplayClusterViewport* InViewport, UCameraComponent* &OutCamera, FDisplayClusterProjectionCameraPolicySettings& OutCameraSettings);
+	bool GetSettingsFromConfig(class IDisplayClusterViewport* InViewport, USceneComponent*& OutSceneComponent, FDisplayClusterProjectionCameraPolicySettings& OutCameraSettings);
+	bool UpdateICVFXCamera(UDisplayClusterICVFX_CineCameraComponent* InComponentICVFX);
 
+	UCameraComponent* GetCameraComponent();
+
+private:
+	UCameraComponent* ImplGetAssignedCamera();
+	UDisplayClusterICVFX_CineCameraComponent* ImplGetAssignedICVFXComponent();
 
 private:
 	// Camera to use for rendering
-	UCameraComponent* AssignedCamera = nullptr;
+	FDisplayClusterSceneComponentRef AssignedCameraRef;
+	FDisplayClusterSceneComponentRef AssignedICVFXCameraRef;
 
 	FDisplayClusterProjectionCameraPolicySettings CameraSettings;
 };
