Do not early out in the HandleSyncEvent

diff --git a/Engine/Plugins/Runtime/nDisplay/Source/DisplayClusterMessageInterceptor/Private/DisplayClusterMessageInterceptionModule.cpp b/Engine/Plugins/Runtime/nDisplay/Source/DisplayClusterMessageInterceptor/Private/DisplayClusterMessageInterceptionModule.cpp
--- a/Engine/Plugins/Runtime/nDisplay/Source/DisplayClusterMessageInterceptor/Private/DisplayClusterMessageInterceptionModule.cpp
+++ b/Engine/Plugins/Runtime/nDisplay/Source/DisplayClusterMessageInterceptor/Private/DisplayClusterMessageInterceptionModule.cpp
@@ -333,10 +333,6 @@ private:
 	void HandleWorkspaceSyncEvent(const FDisplayClusterClusterEventJson& InEvent)
 	{
 		UE_LOG(LogDisplayClusterInterception, Display, TEXT("Handle multi-user workspace sync -> %s."), *InEvent.Name);
-		if (bCanFinalizeWorkspace)
-		{
-			return;
-		}
 
 		IDisplayClusterClusterManager* ClusterManager = IDisplayCluster::Get().GetClusterMgr();
 		NodesReady.Add(InEvent.Name);
