nDisplay: Fixed compile issue caused by ambiguous boolean statement

diff --git a/Engine/Plugins/Runtime/nDisplay/Source/DisplayClusterConfigurator/Private/Views/Details/DisplayClusterConfiguratorDetailCustomization.cpp b/Engine/Plugins/Runtime/nDisplay/Source/DisplayClusterConfigurator/Private/Views/Details/DisplayClusterConfiguratorDetailCustomization.cpp
--- a/Engine/Plugins/Runtime/nDisplay/Source/DisplayClusterConfigurator/Private/Views/Details/DisplayClusterConfiguratorDetailCustomization.cpp
+++ b/Engine/Plugins/Runtime/nDisplay/Source/DisplayClusterConfigurator/Private/Views/Details/DisplayClusterConfiguratorDetailCustomization.cpp
@@ -101,7 +101,7 @@ void FDisplayClusterConfiguratorDetailCustomization::CustomizeDetails(IDetailLay
 
 				const bool bShouldHide = PropertyHandle->HasMetaData("nDisplayHidden") ||
 					(IsRunningForBlueprintEditor() ?
-						PropertyHandle->HasMetaData("nDisplayInstanceOnly") || Property && Property->HasAnyPropertyFlags(CPF_DisableEditOnTemplate) :
+						PropertyHandle->HasMetaData("nDisplayInstanceOnly") || Property->HasAnyPropertyFlags(CPF_DisableEditOnTemplate) :
 						false);
 
 				if (bShouldHide)
