nDisplay: Add profiling support to preview components. Optimize output mapping update frequency.

diff --git a/Engine/Plugins/Runtime/nDisplay/Source/DisplayCluster/Private/Game/EngineClasses/Scene/DisplayClusterPreviewComponent.cpp b/Engine/Plugins/Runtime/nDisplay/Source/DisplayCluster/Private/Game/EngineClasses/Scene/DisplayClusterPreviewComponent.cpp
--- a/Engine/Plugins/Runtime/nDisplay/Source/DisplayCluster/Private/Game/EngineClasses/Scene/DisplayClusterPreviewComponent.cpp
+++ b/Engine/Plugins/Runtime/nDisplay/Source/DisplayCluster/Private/Game/EngineClasses/Scene/DisplayClusterPreviewComponent.cpp
@@ -8,6 +8,7 @@
 #include "DisplayClusterRootActor.h"
 #include "DisplayClusterConfigurationTypes.h"
 
+#include "Misc/DisplayClusterLog.h"
 #include "Misc/DisplayClusterStrings.h"
 #include "Render/Projection/IDisplayClusterProjectionPolicy.h"
 #include "Render/Projection/IDisplayClusterProjectionPolicyFactory.h"
@@ -48,6 +49,8 @@ const uint32 UDisplayClusterPreviewComponent::MaxRenderTargetDimension = 2048;
 
 void UDisplayClusterPreviewComponent::OnComponentCreated()
 {
+	DECLARE_SCOPE_CYCLE_COUNTER(TEXT("UDisplayClusterPreviewComponent::OnComponentCreated"), STAT_OnComponentCreated, STATGROUP_NDisplay);
+	
 	Super::OnComponentCreated();
 
 	InitializeInternals();
@@ -55,6 +58,8 @@ void UDisplayClusterPreviewComponent::OnComponentCreated()
 
 void UDisplayClusterPreviewComponent::DestroyComponent(bool bPromoteChildren)
 {
+	DECLARE_SCOPE_CYCLE_COUNTER(TEXT("UDisplayClusterPreviewComponent::DestroyComponent"), STAT_DestroyComponent, STATGROUP_NDisplay);
+	
 	if (PreviewMesh)
 	{
 		bUseMeshUsePreviewMaterialInstance = false;
@@ -88,6 +93,8 @@ bool UDisplayClusterPreviewComponent::InitializePreviewComponent(ADisplayCluster
 
 bool UDisplayClusterPreviewComponent::UpdatePreviewMesh()
 {
+	DECLARE_SCOPE_CYCLE_COUNTER(TEXT("UDisplayClusterPreviewComponent::UpdatePreviewMesh"), STAT_UpdatePreviewMesh, STATGROUP_NDisplay);
+	
 	check(IsInGameThread());
 
 	// And search for new mesh reference
@@ -146,6 +153,8 @@ bool UDisplayClusterPreviewComponent::UpdatePreviewMesh()
 
 void UDisplayClusterPreviewComponent::UpdatePreviewResources()
 {
+	DECLARE_SCOPE_CYCLE_COUNTER(TEXT("UDisplayClusterPreviewComponent::UpdatePreviewResources"), STAT_UpdatePreviewResources, STATGROUP_NDisplay);
+	
 	if (PreviewMesh && PreviewMesh->GetName().Find(TEXT("TRASH_")) != INDEX_NONE)
 	{
 		// Screen components are regenerated from construction scripts, but preview components are added in dynamically. This preview component may end up
@@ -195,6 +204,8 @@ void UDisplayClusterPreviewComponent::InitializeInternals()
 
 void UDisplayClusterPreviewComponent::UpdatePreviewRenderTarget()
 {
+	DECLARE_SCOPE_CYCLE_COUNTER(TEXT("UDisplayClusterPreviewComponent::UpdatePreviewRenderTarget"), STAT_UpdatePreviewRenderTarget, STATGROUP_NDisplay);
+	
 	FIntPoint TextureSize(1,1);
 	float     TextureGamma = 1.f;
 
@@ -294,6 +305,8 @@ void UDisplayClusterPreviewComponent::RemovePreviewTexture()
 
 bool UDisplayClusterPreviewComponent::UpdatePreviewTexture()
 {
+	DECLARE_SCOPE_CYCLE_COUNTER(TEXT("UDisplayClusterPreviewComponent::UpdatePreviewTexture"), STAT_UpdatePreviewTexture, STATGROUP_NDisplay);
+	
 	check(RenderTarget);
 
 	TArray<FColor> SurfData;
@@ -358,7 +371,7 @@ bool UDisplayClusterPreviewComponent::UpdatePreviewTexture()
 	return true;
 }
 
-void UDisplayClusterPreviewComponent::HandleRenderTargetTextureDefferedUpdate()
+void UDisplayClusterPreviewComponent::HandleRenderTargetTextureDeferredUpdate()
 {
 	check(IsInGameThread());
 
@@ -376,7 +389,10 @@ UTexture2D* UDisplayClusterPreviewComponent::GetOrCreateRenderTexture2D()
 	else
 	if (RenderTarget && bIsRenderTargetSurfaceChanged)
 	{
-		UpdatePreviewTexture();
+		if (UpdatePreviewTexture())
+		{
+			bIsRenderTargetSurfaceChanged = false;
+		}
 	}
 
 	return PreviewTexture;
diff --git a/Engine/Plugins/Runtime/nDisplay/Source/DisplayCluster/Private/Game/EngineClasses/Scene/DisplayClusterRootActorEditor.cpp b/Engine/Plugins/Runtime/nDisplay/Source/DisplayCluster/Private/Game/EngineClasses/Scene/DisplayClusterRootActorEditor.cpp
--- a/Engine/Plugins/Runtime/nDisplay/Source/DisplayCluster/Private/Game/EngineClasses/Scene/DisplayClusterRootActorEditor.cpp
+++ b/Engine/Plugins/Runtime/nDisplay/Source/DisplayCluster/Private/Game/EngineClasses/Scene/DisplayClusterRootActorEditor.cpp
@@ -126,17 +126,25 @@ bool ADisplayClusterRootActor::IsPreviewEnabled() const
 }
 
 // Return all RTT RHI resources for preview
-void ADisplayClusterRootActor::GetPreviewRenderTargetableTextures(const TArray<FString>& InViewportNames, TArray<FTextureRHIRef>& OutTextures) const
+void ADisplayClusterRootActor::GetPreviewRenderTargetableTextures(const TArray<FString>& InViewportNames, TArray<FTextureRHIRef>& OutTextures)
 {
 	check(IsInGameThread());
 
 	if(IsPreviewEnabled())
 	{
+		const bool bPerformDeferredUpdate = PreviewRenderTargetUpdatesRequired > 0;
+		if (bPerformDeferredUpdate)
+		{
+			// Optimization to limit number of deferred updates. This process is expensive
+			// and only needs to happen for a few frames after the BP was modified.
+			--PreviewRenderTargetUpdatesRequired;
+		}
+
 		for(const TPair<FString, UDisplayClusterPreviewComponent*>& PreviewComponentIt : PreviewComponents)
 		{
 			if (PreviewComponentIt.Value)
 			{
-				int OutTextureIndex = InViewportNames.Find(PreviewComponentIt.Value->GetViewportId());
+				const int32 OutTextureIndex = InViewportNames.Find(PreviewComponentIt.Value->GetViewportId());
 				if (OutTextureIndex != INDEX_NONE)
 				{
 					// Add scope for func GetRenderTargetTexture()
@@ -148,8 +156,11 @@ void ADisplayClusterRootActor::GetPreviewRenderTargetableTextures(const TArray<F
 						{
 							OutTextures[OutTextureIndex] = DstRenderTarget->TextureRHI;
 
-							// handle configurator logic: raise deffered update for preview component RTT
-							PreviewComponentIt.Value->HandleRenderTargetTextureDefferedUpdate();
+							if(bPerformDeferredUpdate)
+							{
+								// handle configurator logic: raise deffered update for preview component RTT
+								PreviewComponentIt.Value->HandleRenderTargetTextureDeferredUpdate();
+							}
 						}
 					}
 				}
@@ -386,6 +397,12 @@ void ADisplayClusterRootActor::UpdatePreviewComponents()
 			ExistingComp->DestroyComponent();
 		}
 	}
+
+	// Update render target for output mapping.
+	// Compile - 3 frames
+	// Load    - 2 frames
+	// Modify  - 1 frame
+	PreviewRenderTargetUpdatesRequired = 3;
 }
 
 void ADisplayClusterRootActor::ReleasePreviewComponents()
diff --git a/Engine/Plugins/Runtime/nDisplay/Source/DisplayCluster/Private/Misc/DisplayClusterLog.h b/Engine/Plugins/Runtime/nDisplay/Source/DisplayCluster/Private/Misc/DisplayClusterLog.h
--- a/Engine/Plugins/Runtime/nDisplay/Source/DisplayCluster/Private/Misc/DisplayClusterLog.h
+++ b/Engine/Plugins/Runtime/nDisplay/Source/DisplayCluster/Private/Misc/DisplayClusterLog.h
@@ -3,6 +3,7 @@
 #pragma once
 
 #include "CoreMinimal.h"
+#include "Stats/Stats2.h"
 
 // Module-wide log categories
 DECLARE_LOG_CATEGORY_EXTERN(LogDisplayClusterEngine,     Log, All);
@@ -16,3 +17,6 @@ DECLARE_LOG_CATEGORY_EXTERN(LogDisplayClusterRender,     Log, All);
 DECLARE_LOG_CATEGORY_EXTERN(LogDisplayClusterRenderSync, Log, All);
 DECLARE_LOG_CATEGORY_EXTERN(LogDisplayClusterViewport,   Log, All);
 DECLARE_LOG_CATEGORY_EXTERN(LogDisplayClusterBlueprint,  Log, All);
+
+
+DECLARE_STATS_GROUP(TEXT("nDisplay"), STATGROUP_NDisplay, STATCAT_Advanced)
\ No newline at end of file
diff --git a/Engine/Plugins/Runtime/nDisplay/Source/DisplayCluster/Public/Components/DisplayClusterPreviewComponent.h b/Engine/Plugins/Runtime/nDisplay/Source/DisplayCluster/Public/Components/DisplayClusterPreviewComponent.h
--- a/Engine/Plugins/Runtime/nDisplay/Source/DisplayCluster/Public/Components/DisplayClusterPreviewComponent.h
+++ b/Engine/Plugins/Runtime/nDisplay/Source/DisplayCluster/Public/Components/DisplayClusterPreviewComponent.h
@@ -54,7 +54,7 @@ public:
 	}
 
 	void UpdatePreviewResources();
-	void HandleRenderTargetTextureDefferedUpdate();
+	void HandleRenderTargetTextureDeferredUpdate();
 
 	UMeshComponent* GetPreviewMesh() const
 	{
diff --git a/Engine/Plugins/Runtime/nDisplay/Source/DisplayCluster/Public/DisplayClusterRootActor.h b/Engine/Plugins/Runtime/nDisplay/Source/DisplayCluster/Public/DisplayClusterRootActor.h
--- a/Engine/Plugins/Runtime/nDisplay/Source/DisplayCluster/Public/DisplayClusterRootActor.h
+++ b/Engine/Plugins/Runtime/nDisplay/Source/DisplayCluster/Public/DisplayClusterRootActor.h
@@ -314,7 +314,7 @@ public:
 
 	IDisplayClusterViewport* FindPreviewViewport(const FString& InViewportId) const;
 
-	void GetPreviewRenderTargetableTextures(const TArray<FString>& InViewportNames, TArray<FTextureRHIRef>& OutTextures) const;
+	void GetPreviewRenderTargetableTextures(const TArray<FString>& InViewportNames, TArray<FTextureRHIRef>& OutTextures);
 
 	float GetXformGizmoScale() const;
 	bool GetXformGizmoVisibility() const;
@@ -325,6 +325,8 @@ public:
 	void SetIsSelectedInEditor(bool bValue);
 
 private:
+	/** The number of times to update the render target via deferred update. */
+	int32 PreviewRenderTargetUpdatesRequired = 0;
 	bool bIsSelectedInEditor = false;
 	
 protected:
