fixing linux cross compile target by fixing case
#rnx

diff --git a/Engine/Plugins/Runtime/nDisplay/Source/DisplayClusterConfigurator/Private/Views/OutputMapping/Widgets/SDisplayClusterConfiguratorExternalImage.cpp b/Engine/Plugins/Runtime/nDisplay/Source/DisplayClusterConfigurator/Private/Views/OutputMapping/Widgets/SDisplayClusterConfiguratorExternalImage.cpp
--- a/Engine/Plugins/Runtime/nDisplay/Source/DisplayClusterConfigurator/Private/Views/OutputMapping/Widgets/SDisplayClusterConfiguratorExternalImage.cpp
+++ b/Engine/Plugins/Runtime/nDisplay/Source/DisplayClusterConfigurator/Private/Views/OutputMapping/Widgets/SDisplayClusterConfiguratorExternalImage.cpp
@@ -6,7 +6,7 @@
 #include "IImageWrapperModule.h"
 #include "EditorStyleSet.h"
 #include "Framework/Application/SlateApplication.h"
-#include "HAL/PlatformFilemanager.h"
+#include "HAL/PlatformFileManager.h"
 #include "Misc/FileHelper.h"
 #include "Modules/ModuleManager.h"
 #include "Brushes/SlateDynamicImageBrush.h"
