nDisplay: ICVFX camera fixes and custom details panel

diff --git a/Engine/Plugins/Runtime/nDisplay/Source/DisplayCluster/Private/Game/EngineClasses/Scene/DisplayClusterICVFXCameraComponent.cpp b/Engine/Plugins/Runtime/nDisplay/Source/DisplayCluster/Private/Game/EngineClasses/Scene/DisplayClusterICVFXCameraComponent.cpp
new file mode 100644
--- /dev/null
+++ b/Engine/Plugins/Runtime/nDisplay/Source/DisplayCluster/Private/Game/EngineClasses/Scene/DisplayClusterICVFXCameraComponent.cpp
@@ -0,0 +1,52 @@
+// Copyright Epic Games, Inc. All Rights Reserved.
+
+#include "Components/DisplayClusterICVFXCameraComponent.h"
+
+#include "Render/Viewport/Containers/DisplayClusterViewport_CameraMotionBlur.h"
+#include "Components/DisplayClusterCameraComponent.h"
+#include "DisplayClusterRootActor.h"
+
+
+UCameraComponent* UDisplayClusterICVFXCameraComponent::GetCameraComponent()
+{
+	return ExternalCameraActor.IsValid() ? ExternalCameraActor->GetCameraComponent() : this;
+}
+
+FString UDisplayClusterICVFXCameraComponent::GetCameraUniqueId() const
+{
+	return GetFName().ToString();
+}
+
+FDisplayClusterViewport_CameraMotionBlur UDisplayClusterICVFXCameraComponent::GetMotionBlurParameters()
+{
+	FDisplayClusterViewport_CameraMotionBlur OutParameters;
+	OutParameters.Mode = EDisplayClusterViewport_CameraMotionBlur::Undefined;
+
+	switch (CameraSettings.CameraMotionBlur.MotionBlurMode)
+	{
+	case EDisplayClusterConfigurationCameraMotionBlurMode::Off:
+		OutParameters.Mode = EDisplayClusterViewport_CameraMotionBlur::Off;
+		break;
+
+	case EDisplayClusterConfigurationCameraMotionBlurMode::On:
+		OutParameters.Mode = EDisplayClusterViewport_CameraMotionBlur::On;
+		break;
+
+	case EDisplayClusterConfigurationCameraMotionBlurMode::Override:
+		ADisplayClusterRootActor* RootActor = static_cast<ADisplayClusterRootActor*>(GetOwner());
+		if (RootActor)
+		{
+			UDisplayClusterCameraComponent* OuterCamera = RootActor->GetDefaultCamera();
+			if (OuterCamera)
+			{
+				OutParameters.CameraLocation   = OuterCamera->GetComponentLocation();
+				OutParameters.CameraRotation   = OuterCamera->GetComponentRotation();
+				OutParameters.TranslationScale = CameraSettings.CameraMotionBlur.TranslationScale;
+				OutParameters.Mode             = EDisplayClusterViewport_CameraMotionBlur::Override;
+			}
+		}
+		break;
+	}
+
+	return OutParameters;
+}
diff --git a/Engine/Plugins/Runtime/nDisplay/Source/DisplayCluster/Private/Game/EngineClasses/Scene/DisplayClusterICVFX_CineCameraComponent.cpp b/Engine/Plugins/Runtime/nDisplay/Source/DisplayCluster/Private/Game/EngineClasses/Scene/DisplayClusterICVFX_CineCameraComponent.cpp
deleted file mode 100644
diff --git a/Engine/Plugins/Runtime/nDisplay/Source/DisplayCluster/Private/Game/EngineClasses/Scene/DisplayClusterRootActor.cpp b/Engine/Plugins/Runtime/nDisplay/Source/DisplayCluster/Private/Game/EngineClasses/Scene/DisplayClusterRootActor.cpp
--- a/Engine/Plugins/Runtime/nDisplay/Source/DisplayCluster/Private/Game/EngineClasses/Scene/DisplayClusterRootActor.cpp
+++ b/Engine/Plugins/Runtime/nDisplay/Source/DisplayCluster/Private/Game/EngineClasses/Scene/DisplayClusterRootActor.cpp
@@ -13,7 +13,7 @@
 #include "Components/DisplayClusterSyncTickComponent.h"
 
 #include "Components/DisplayClusterRootComponent.h"
-#include "Components/DisplayClusterICVFX_CineCameraComponent.h"
+#include "Components/DisplayClusterICVFXCameraComponent.h"
 #include "CineCameraComponent.h"
 #include "Components/DisplayClusterSceneComponentSyncThis.h"
 
diff --git a/Engine/Plugins/Runtime/nDisplay/Source/DisplayCluster/Private/Render/Viewport/Configuration/DisplayClusterViewportConfigurationHelpers.h b/Engine/Plugins/Runtime/nDisplay/Source/DisplayCluster/Private/Render/Viewport/Configuration/DisplayClusterViewportConfigurationHelpers.h
--- a/Engine/Plugins/Runtime/nDisplay/Source/DisplayCluster/Private/Render/Viewport/Configuration/DisplayClusterViewportConfigurationHelpers.h
+++ b/Engine/Plugins/Runtime/nDisplay/Source/DisplayCluster/Private/Render/Viewport/Configuration/DisplayClusterViewportConfigurationHelpers.h
@@ -666,7 +666,7 @@ public:
 		}
 	}
 
-	static bool IsShouldUseLightcard(const FDisplayClusterConfigurationICVFX_LightcardSettings& InLightcardSettings)
+	static bool ShouldUseLightcard(const FDisplayClusterConfigurationICVFX_LightcardSettings& InLightcardSettings)
 	{
 		if (InLightcardSettings.bEnable == false)
 		{
diff --git a/Engine/Plugins/Runtime/nDisplay/Source/DisplayCluster/Private/Render/Viewport/Configuration/DisplayClusterViewportConfigurationICVFX.cpp b/Engine/Plugins/Runtime/nDisplay/Source/DisplayCluster/Private/Render/Viewport/Configuration/DisplayClusterViewportConfigurationICVFX.cpp
--- a/Engine/Plugins/Runtime/nDisplay/Source/DisplayCluster/Private/Render/Viewport/Configuration/DisplayClusterViewportConfigurationICVFX.cpp
+++ b/Engine/Plugins/Runtime/nDisplay/Source/DisplayCluster/Private/Render/Viewport/Configuration/DisplayClusterViewportConfigurationICVFX.cpp
@@ -14,7 +14,8 @@
 #include "Render/Viewport/DisplayClusterViewport.h"
 #include "Render/Viewport/DisplayClusterViewportManager.h"
 
-#include "Components/DisplayClusterICVFX_CineCameraComponent.h"
+#include "Components/DisplayClusterICVFXCameraComponent.h"
+
 
 ////////////////////////////////////////////////////////////////////////////////
 // FDisplayClusterViewportConfigurationCameraViewport
@@ -293,7 +294,7 @@ private:
 class FDisplayClusterViewportConfigurationCameraICVFX
 {
 public:
-	FDisplayClusterViewportConfigurationCameraICVFX(const FTransform& InLocal2WorldTransform, FDisplayClusterViewportConfigurationICVFX& InConfigurationICVFX, class UDisplayClusterICVFX_CineCameraComponent* const InCameraComponent)
+	FDisplayClusterViewportConfigurationCameraICVFX(const FTransform& InLocal2WorldTransform, FDisplayClusterViewportConfigurationICVFX& InConfigurationICVFX, class UDisplayClusterICVFXCameraComponent* const InCameraComponent)
 		: ConfigurationICVFX(InConfigurationICVFX)
 		, CameraViewport(InLocal2WorldTransform, InConfigurationICVFX, InCameraComponent->GetCameraComponent(), InCameraComponent->GetCameraUniqueId(), InCameraComponent->GetCameraSettingsICVFX())
 		, ChromakeySettings(DisplayClusterViewportConfigurationHelpers::GetCameraChromakeySettings(InCameraComponent->GetCameraSettingsICVFX(), ConfigurationICVFX.StageSettings))
@@ -517,7 +518,7 @@ void FDisplayClusterViewportConfigurationICVFX::Update()
 			if ((TargetViewportsFlags & ViewportICVFX_DisableLightcard) == 0)
 			{
 				// Allocate and assign lightcard resources
-				if (DisplayClusterViewportConfigurationHelpers::IsShouldUseLightcard(LightcardSettings))
+				if (DisplayClusterViewportConfigurationHelpers::ShouldUseLightcard(LightcardSettings))
 				{
 					for (FDisplayClusterViewport* TargetIt : TargetViewports)
 					{
@@ -648,8 +649,8 @@ void FDisplayClusterViewportConfigurationICVFX::ImplGetCameras(TArray<FDisplayCl
 		// Try to create ICVFX camera from component:
 		if (ActorComponentIt)
 		{
-			UDisplayClusterICVFX_CineCameraComponent* CineCameraComponent = Cast<UDisplayClusterICVFX_CineCameraComponent>(ActorComponentIt);
-			if (CineCameraComponent && CineCameraComponent->IsShouldUseICVFX())
+			UDisplayClusterICVFXCameraComponent* CineCameraComponent = Cast<UDisplayClusterICVFXCameraComponent>(ActorComponentIt);
+			if (CineCameraComponent && CineCameraComponent->IsICVFXEnabled())
 			{
 				NewCamera = new FDisplayClusterViewportConfigurationCameraICVFX(Local2WorldTransform, *this, CineCameraComponent);
 			}
diff --git a/Engine/Plugins/Runtime/nDisplay/Source/DisplayCluster/Public/Components/DisplayClusterICVFXCameraComponent.h b/Engine/Plugins/Runtime/nDisplay/Source/DisplayCluster/Public/Components/DisplayClusterICVFXCameraComponent.h
new file mode 100644
--- /dev/null
+++ b/Engine/Plugins/Runtime/nDisplay/Source/DisplayCluster/Public/Components/DisplayClusterICVFXCameraComponent.h
@@ -0,0 +1,57 @@
+// Copyright Epic Games, Inc. All Rights Reserved.
+
+#pragma once
+
+#include "CoreMinimal.h"
+
+#include "CineCameraActor.h"
+#include "CineCameraComponent.h"
+
+#include "DisplayClusterConfigurationTypes_ICVFX.h"
+#include "Render/Viewport/Containers/DisplayClusterViewport_CameraMotionBlur.h"
+
+#include "DisplayClusterICVFXCameraComponent.generated.h"
+
+
+class UCameraComponent;
+
+
+/**
+ * ICVFX camera with configuration
+ */
+UCLASS(ClassGroup = (DisplayCluster), meta = (DisplayName="ICVFX Camera"))
+class DISPLAYCLUSTER_API UDisplayClusterICVFXCameraComponent
+	: public UCineCameraComponent
+{
+	GENERATED_BODY()
+
+public:
+	UDisplayClusterICVFXCameraComponent(const FObjectInitializer& ObjectInitializer)
+	{ }
+
+public:
+	// Use external cine camera actor
+	UPROPERTY(EditAnywhere, Category = NDisplay)
+	TSoftObjectPtr<ACineCameraActor> ExternalCameraActor;
+
+	UPROPERTY(EditAnywhere, Category = NDisplay, meta = (ShowInnerProperties))
+	FDisplayClusterConfigurationICVFX_CameraSettings CameraSettings;
+
+public:
+	FDisplayClusterViewport_CameraMotionBlur GetMotionBlurParameters();
+
+	bool IsICVFXEnabled() const
+	{
+		return CameraSettings.bEnable;
+	}
+
+	// Return unique camera name
+	FString GetCameraUniqueId() const;
+
+	const FDisplayClusterConfigurationICVFX_CameraSettings& GetCameraSettingsICVFX() const
+	{
+		return CameraSettings;
+	}
+
+	UCameraComponent* GetCameraComponent();
+};
diff --git a/Engine/Plugins/Runtime/nDisplay/Source/DisplayCluster/Public/Components/DisplayClusterICVFX_CineCameraComponent.h b/Engine/Plugins/Runtime/nDisplay/Source/DisplayCluster/Public/Components/DisplayClusterICVFX_CineCameraComponent.h
deleted file mode 100644
diff --git a/Engine/Plugins/Runtime/nDisplay/Source/DisplayClusterConfigurator/Private/Views/Details/Policies/DisplayClusterConfiguratorPolicyDetailCustomization.cpp b/Engine/Plugins/Runtime/nDisplay/Source/DisplayClusterConfigurator/Private/Views/Details/Policies/DisplayClusterConfiguratorPolicyDetailCustomization.cpp
--- a/Engine/Plugins/Runtime/nDisplay/Source/DisplayClusterConfigurator/Private/Views/Details/Policies/DisplayClusterConfiguratorPolicyDetailCustomization.cpp
+++ b/Engine/Plugins/Runtime/nDisplay/Source/DisplayClusterConfigurator/Private/Views/Details/Policies/DisplayClusterConfiguratorPolicyDetailCustomization.cpp
@@ -11,7 +11,7 @@
 #include "DisplayClusterRootActor.h"
 #include "Blueprints/DisplayClusterBlueprint.h"
 #include "Components/DisplayClusterScreenComponent.h"
-#include "Components/DisplayClusterICVFX_CineCameraComponent.h"
+#include "Components/DisplayClusterICVFXCameraComponent.h"
 
 #include "Camera/CameraComponent.h"
 
@@ -376,7 +376,7 @@ void FDisplayClusterConfiguratorProjectionCustomization::CreateCameraPolicy(UDis
 		ConfigurationViewportPtr.Get(),
 		ParametersHandle,
 		TArray<TSubclassOf<UActorComponent>>{ UCameraComponent::StaticClass(),
-			UDisplayClusterICVFX_CineCameraComponent::StaticClass() }));
+			UDisplayClusterICVFXCameraComponent::StaticClass() }));
 
 	CustomPolicyParameters.Add(MakeShared<FPolicyParameterInfoBool>(
 		"Native",
diff --git a/Engine/Plugins/Runtime/nDisplay/Source/DisplayClusterConfigurator/Private/Views/SCSEditor/SDisplayClusterConfiguratorComponentCombo.cpp b/Engine/Plugins/Runtime/nDisplay/Source/DisplayClusterConfigurator/Private/Views/SCSEditor/SDisplayClusterConfiguratorComponentCombo.cpp
--- a/Engine/Plugins/Runtime/nDisplay/Source/DisplayClusterConfigurator/Private/Views/SCSEditor/SDisplayClusterConfiguratorComponentCombo.cpp
+++ b/Engine/Plugins/Runtime/nDisplay/Source/DisplayClusterConfigurator/Private/Views/SCSEditor/SDisplayClusterConfiguratorComponentCombo.cpp
@@ -11,7 +11,7 @@
 #include "Components/DisplayClusterXformComponent.h"
 #include "Components/DisplayClusterCameraComponent.h"
 
-#include "Components/DisplayClusterICVFX_CineCameraComponent.h"
+#include "Components/DisplayClusterICVFXCameraComponent.h"
 
 #include "Widgets/Layout/SSpacer.h"
 #include "Widgets/Images/SImage.h"
@@ -437,7 +437,7 @@ void SDisplayClusterConfiguratorComponentClassCombo::GenerateComponentClassList(
 	AddDCComp(UDisplayClusterCameraComponent::StaticClass());
 
 	AddHeader("nDisplay ICVFX Components");
-	AddDCComp(UDisplayClusterICVFX_CineCameraComponent::StaticClass());
+	AddDCComp(UDisplayClusterICVFXCameraComponent::StaticClass());
 
 	int32 CompIterationIdx = 0;
 	for (const TSharedPtr<FComponentClassComboEntry>& Comp : *ComponentClassListPtr)
diff --git a/Engine/Plugins/Runtime/nDisplay/Source/DisplayClusterConfigurator/Private/Views/Viewport/DisplayClusterConfiguratorSCSEditorViewport.cpp b/Engine/Plugins/Runtime/nDisplay/Source/DisplayClusterConfigurator/Private/Views/Viewport/DisplayClusterConfiguratorSCSEditorViewport.cpp
--- a/Engine/Plugins/Runtime/nDisplay/Source/DisplayClusterConfigurator/Private/Views/Viewport/DisplayClusterConfiguratorSCSEditorViewport.cpp
+++ b/Engine/Plugins/Runtime/nDisplay/Source/DisplayClusterConfigurator/Private/Views/Viewport/DisplayClusterConfiguratorSCSEditorViewport.cpp
@@ -12,7 +12,7 @@
 #include "ClusterConfiguration/ViewModels/DisplayClusterConfiguratorProjectionPolicyViewModel.h"
 
 #include "Components/DisplayClusterScreenComponent.h"
-#include "Components/DisplayClusterICVFX_CineCameraComponent.h"
+#include "Components/DisplayClusterICVFXCameraComponent.h"
 
 #include "Camera/CameraComponent.h"
 
diff --git a/Engine/Plugins/Runtime/nDisplay/Source/DisplayClusterEditor/DisplayClusterEditor.Build.cs b/Engine/Plugins/Runtime/nDisplay/Source/DisplayClusterEditor/DisplayClusterEditor.Build.cs
--- a/Engine/Plugins/Runtime/nDisplay/Source/DisplayClusterEditor/DisplayClusterEditor.Build.cs
+++ b/Engine/Plugins/Runtime/nDisplay/Source/DisplayClusterEditor/DisplayClusterEditor.Build.cs
@@ -9,6 +9,7 @@ public class DisplayClusterEditor : ModuleRules
 	{
 		PrivateDependencyModuleNames.AddRange(
 			new string[] {
+				"CinematicCamera",
 				"Core",
 				"CoreUObject",
 				"DisplayCluster",
@@ -18,7 +19,7 @@ public class DisplayClusterEditor : ModuleRules
 				"PropertyEditor",
 				"Slate",
 				"SlateCore",
-				"UnrealEd"
+				"UnrealEd",
 			});
 	}
 }
diff --git a/Engine/Plugins/Runtime/nDisplay/Source/DisplayClusterEditor/Private/DetailsCustomization/DisplayClusterICVFXCameraComponentDetailsCustomization.cpp b/Engine/Plugins/Runtime/nDisplay/Source/DisplayClusterEditor/Private/DetailsCustomization/DisplayClusterICVFXCameraComponentDetailsCustomization.cpp
new file mode 100644
--- /dev/null
+++ b/Engine/Plugins/Runtime/nDisplay/Source/DisplayClusterEditor/Private/DetailsCustomization/DisplayClusterICVFXCameraComponentDetailsCustomization.cpp
@@ -0,0 +1,64 @@
+// Copyright Epic Games, Inc. All Rights Reserved.
+
+#include "DetailsCustomization/DisplayClusterICVFXCameraComponentDetailsCustomization.h"
+
+#include "Components/DisplayClusterICVFXCameraComponent.h"
+
+#include "CineCameraActor.h"
+#include "UObject/SoftObjectPtr.h"
+
+#include "DetailLayoutBuilder.h"
+
+
+TSharedRef<IDetailCustomization> FDisplayClusterICVFXCameraComponentDetailsCustomization::MakeInstance()
+{
+	return MakeShared<FDisplayClusterICVFXCameraComponentDetailsCustomization>();
+}
+
+void FDisplayClusterICVFXCameraComponentDetailsCustomization::CustomizeDetails(IDetailLayoutBuilder& Builder)
+{
+	DetailLayout = &Builder;
+
+	if (!EditedObject.IsValid())
+	{
+		TArray<TWeakObjectPtr<UObject>> Objects;
+		Builder.GetObjectsBeingCustomized(Objects);
+
+		for (TWeakObjectPtr<UObject> Object : Objects)
+		{
+			if (Object->IsA<UDisplayClusterICVFXCameraComponent>())
+			{
+				EditedObject = Cast<UDisplayClusterICVFXCameraComponent>(Object.Get());
+			}
+		}
+	}
+
+	// Hide some groups if an external CineCameraActor is set
+	if (EditedObject.IsValid() && EditedObject->ExternalCameraActor.IsValid())
+	{
+		Builder.HideCategory(TEXT("Current Camera Settings"));
+		Builder.HideCategory(TEXT("CameraOptions"));
+		Builder.HideCategory(TEXT("Camera"));
+		Builder.HideCategory(TEXT("PostProcess"));
+		Builder.HideCategory(TEXT("Lens"));
+		Builder.HideCategory(TEXT("LOD"));
+		Builder.HideCategory(TEXT("ColorGrading"));
+		Builder.HideCategory(TEXT("RenderingFeatures"));
+		Builder.HideCategory(TEXT("Tags"));
+		Builder.HideCategory(TEXT("ComponentTick"));
+		Builder.HideCategory(TEXT("Color Grading"));
+		Builder.HideCategory(TEXT("Rendering Features"));
+		Builder.HideCategory(TEXT("Activation"));
+	}
+
+	// Always hide the following categories
+	Builder.HideCategory(TEXT("AssetUserData"));
+	Builder.HideCategory(TEXT("Collision"));
+	Builder.HideCategory(TEXT("Cooking"));
+	Builder.HideCategory(TEXT("ComponentReplication"));
+	Builder.HideCategory(TEXT("Events"));
+	Builder.HideCategory(TEXT("Physics"));
+	Builder.HideCategory(TEXT("Sockets"));
+}
+
+#undef LOCTEXT_NAMESPACE
diff --git a/Engine/Plugins/Runtime/nDisplay/Source/DisplayClusterEditor/Private/DetailsCustomization/DisplayClusterICVFXCameraComponentDetailsCustomization.h b/Engine/Plugins/Runtime/nDisplay/Source/DisplayClusterEditor/Private/DetailsCustomization/DisplayClusterICVFXCameraComponentDetailsCustomization.h
new file mode 100644
--- /dev/null
+++ b/Engine/Plugins/Runtime/nDisplay/Source/DisplayClusterEditor/Private/DetailsCustomization/DisplayClusterICVFXCameraComponentDetailsCustomization.h
@@ -0,0 +1,27 @@
+// Copyright Epic Games, Inc. All Rights Reserved.
+
+#pragma once
+
+#include "IDetailCustomization.h"
+
+#include "CoreMinimal.h"
+#include "Types/SlateEnums.h"
+
+class UDisplayClusterICVFXCameraComponent;
+class ACineCameraActor;
+class IDetailLayoutBuilder;
+
+class FDisplayClusterICVFXCameraComponentDetailsCustomization : public IDetailCustomization
+{
+public:
+	static TSharedRef<IDetailCustomization> MakeInstance();
+
+protected:
+	// IDetailCustomization interface
+	virtual void CustomizeDetails(class IDetailLayoutBuilder& LayoutBuilder) override;
+
+protected:
+	TWeakObjectPtr<UDisplayClusterICVFXCameraComponent> EditedObject;
+	/** Keep a reference to force refresh the layout */
+	IDetailLayoutBuilder* DetailLayout = nullptr;
+};
diff --git a/Engine/Plugins/Runtime/nDisplay/Source/DisplayClusterEditor/Private/DisplayClusterEditorModule.cpp b/Engine/Plugins/Runtime/nDisplay/Source/DisplayClusterEditor/Private/DisplayClusterEditorModule.cpp
--- a/Engine/Plugins/Runtime/nDisplay/Source/DisplayClusterEditor/Private/DisplayClusterEditorModule.cpp
+++ b/Engine/Plugins/Runtime/nDisplay/Source/DisplayClusterEditor/Private/DisplayClusterEditorModule.cpp
@@ -3,7 +3,9 @@
 #include "DisplayClusterEditorModule.h"
 
 #include "DisplayClusterRootActor.h"
+#include "Components/DisplayClusterICVFXCameraComponent.h"
 #include "DetailsCustomization/DisplayClusterRootActorDetailsCustomization.h"
+#include "DetailsCustomization/DisplayClusterICVFXCameraComponentDetailsCustomization.h"
 
 #include "Components/DisplayClusterPreviewComponent.h"
 
@@ -60,6 +62,7 @@ void FDisplayClusterEditorModule::RegisterCustomizations()
 	FPropertyEditorModule& PropertyEditorModule = FModuleManager::Get().LoadModuleChecked<FPropertyEditorModule>("PropertyEditor");
 
 	PropertyEditorModule.RegisterCustomClassLayout(ADisplayClusterRootActor::StaticClass()->GetFName(), FOnGetDetailCustomizationInstance::CreateStatic(&FDisplayClusterRootActorDetailsCustomization::MakeInstance));
+	PropertyEditorModule.RegisterCustomClassLayout(UDisplayClusterICVFXCameraComponent::StaticClass()->GetFName(), FOnGetDetailCustomizationInstance::CreateStatic(&FDisplayClusterICVFXCameraComponentDetailsCustomization::MakeInstance));
 }
 
 void FDisplayClusterEditorModule::UnregisterCustomizations()
@@ -67,6 +70,7 @@ void FDisplayClusterEditorModule::UnregisterCustomizations()
 	FPropertyEditorModule& PropertyEditorModule = FModuleManager::Get().LoadModuleChecked<FPropertyEditorModule>("PropertyEditor");
 
 	PropertyEditorModule.UnregisterCustomClassLayout(ADisplayClusterRootActor::StaticClass()->GetFName());
+	PropertyEditorModule.UnregisterCustomClassLayout(UDisplayClusterICVFXCameraComponent::StaticClass()->GetFName());
 }
 
 IMPLEMENT_MODULE(FDisplayClusterEditorModule, DisplayClusterEditor);
diff --git a/Engine/Plugins/Runtime/nDisplay/Source/DisplayClusterProjection/Private/Policy/Camera/DisplayClusterProjectionCameraPolicy.cpp b/Engine/Plugins/Runtime/nDisplay/Source/DisplayClusterProjection/Private/Policy/Camera/DisplayClusterProjectionCameraPolicy.cpp
--- a/Engine/Plugins/Runtime/nDisplay/Source/DisplayClusterProjection/Private/Policy/Camera/DisplayClusterProjectionCameraPolicy.cpp
+++ b/Engine/Plugins/Runtime/nDisplay/Source/DisplayClusterProjection/Private/Policy/Camera/DisplayClusterProjectionCameraPolicy.cpp
@@ -15,7 +15,7 @@
 #include "Camera/CameraComponent.h"
 #include "ComposurePostMoves.h"
 
-#include "Components/DisplayClusterICVFX_CineCameraComponent.h"
+#include "Components/DisplayClusterICVFXCameraComponent.h"
 
 #include "Render/Viewport/IDisplayClusterViewport.h"
 #include "Render/Viewport/IDisplayClusterViewportManager.h"
@@ -47,7 +47,7 @@ bool FDisplayClusterProjectionCameraPolicy::HandleStartScene(class IDisplayClust
 
 	if (CfgComponent)
 	{
-		UDisplayClusterICVFX_CineCameraComponent* CameraICVFX = Cast<UDisplayClusterICVFX_CineCameraComponent>(CfgComponent);
+		UDisplayClusterICVFXCameraComponent* CameraICVFX = Cast<UDisplayClusterICVFXCameraComponent>(CfgComponent);
 		if (CameraICVFX)
 		{
 			AssignedICVFXCameraRef.SetSceneComponent(CameraICVFX);
@@ -88,7 +88,7 @@ APlayerCameraManager* const GetCurPlayerCameraManager(IDisplayClusterViewport* I
 	return nullptr;
 }
 
-bool FDisplayClusterProjectionCameraPolicy::UpdateICVFXCamera(UDisplayClusterICVFX_CineCameraComponent* InComponentICVFX)
+bool FDisplayClusterProjectionCameraPolicy::UpdateICVFXCamera(UDisplayClusterICVFXCameraComponent* InComponentICVFX)
 {
 	if (InComponentICVFX)
 	{
@@ -114,7 +114,7 @@ bool FDisplayClusterProjectionCameraPolicy::UpdateICVFXCamera(UDisplayClusterICV
 UCameraComponent* FDisplayClusterProjectionCameraPolicy::GetCameraComponent()
 {
 	// Support runtime update for ICVFX camera component
-	UDisplayClusterICVFX_CineCameraComponent* CameraICVFX = ImplGetAssignedICVFXComponent();
+	UDisplayClusterICVFXCameraComponent* CameraICVFX = ImplGetAssignedICVFXComponent();
 	if (CameraICVFX)
 	{
 		UpdateICVFXCamera(CameraICVFX);
@@ -123,14 +123,14 @@ UCameraComponent* FDisplayClusterProjectionCameraPolicy::GetCameraComponent()
 	return ImplGetAssignedCamera();
 }
 
-UDisplayClusterICVFX_CineCameraComponent* FDisplayClusterProjectionCameraPolicy::ImplGetAssignedICVFXComponent()
+UDisplayClusterICVFXCameraComponent* FDisplayClusterProjectionCameraPolicy::ImplGetAssignedICVFXComponent()
 {
 	if (AssignedICVFXCameraRef.IsDefinedSceneComponent())
 	{
 		USceneComponent* SceneComponent = AssignedICVFXCameraRef.GetOrFindSceneComponent();
 		if (SceneComponent)
 		{
-			UDisplayClusterICVFX_CineCameraComponent* CameraComponentICVFX = Cast<UDisplayClusterICVFX_CineCameraComponent>(SceneComponent);
+			UDisplayClusterICVFXCameraComponent* CameraComponentICVFX = Cast<UDisplayClusterICVFXCameraComponent>(SceneComponent);
 			if (CameraComponentICVFX)
 			{
 				return CameraComponentICVFX;
@@ -251,8 +251,8 @@ bool FDisplayClusterProjectionCameraPolicy::GetSettingsFromConfig(class IDisplay
 	}
 
 	// Get ICVFX camera component
-	TArray<UDisplayClusterICVFX_CineCameraComponent*> ICVFXCameraComps;
-	RootActor->GetComponents<UDisplayClusterICVFX_CineCameraComponent>(ICVFXCameraComps);
+	TArray<UDisplayClusterICVFXCameraComponent*> ICVFXCameraComps;
+	RootActor->GetComponents<UDisplayClusterICVFXCameraComponent>(ICVFXCameraComps);
 	for (UCameraComponent* Comp : ICVFXCameraComps)
 	{
 		if (Comp->GetName() == CameraComponentId)
diff --git a/Engine/Plugins/Runtime/nDisplay/Source/DisplayClusterProjection/Private/Policy/Camera/DisplayClusterProjectionCameraPolicy.h b/Engine/Plugins/Runtime/nDisplay/Source/DisplayClusterProjection/Private/Policy/Camera/DisplayClusterProjectionCameraPolicy.h
--- a/Engine/Plugins/Runtime/nDisplay/Source/DisplayClusterProjection/Private/Policy/Camera/DisplayClusterProjectionCameraPolicy.h
+++ b/Engine/Plugins/Runtime/nDisplay/Source/DisplayClusterProjection/Private/Policy/Camera/DisplayClusterProjectionCameraPolicy.h
@@ -6,7 +6,7 @@
 #include "Containers/DisplayClusterProjectionCameraPolicySettings.h"
 #include "Misc/DisplayClusterObjectRef.h"
 
-class UDisplayClusterICVFX_CineCameraComponent;
+class UDisplayClusterICVFXCameraComponent;
 class UCameraComponent;
 class UWorld;
 
@@ -41,13 +41,13 @@ public:
 
 protected:
 	bool GetSettingsFromConfig(class IDisplayClusterViewport* InViewport, USceneComponent*& OutSceneComponent, FDisplayClusterProjectionCameraPolicySettings& OutCameraSettings);
-	bool UpdateICVFXCamera(UDisplayClusterICVFX_CineCameraComponent* InComponentICVFX);
+	bool UpdateICVFXCamera(UDisplayClusterICVFXCameraComponent* InComponentICVFX);
 
 	UCameraComponent* GetCameraComponent();
 
 private:
 	UCameraComponent* ImplGetAssignedCamera();
-	UDisplayClusterICVFX_CineCameraComponent* ImplGetAssignedICVFXComponent();
+	UDisplayClusterICVFXCameraComponent* ImplGetAssignedICVFXComponent();
 
 private:
 	// Camera to use for rendering
