nDisplay: Fixed issue where moving multiple items on the Output Mapping graph editor at the same time caused unexpected movement in some items whem attempting to place any of the items next to a stationary sibling when overlapping is not allowed. The cause was that graph nodes were not being moved as a cohesive group, instead each node was computing its own movement offset, which caused issues when one node was being blocked from moving by a sibling while the other nodes were not.

diff --git a/Engine/Plugins/Runtime/nDisplay/Source/DisplayClusterConfigurator/Private/Views/OutputMapping/EdNodes/DisplayClusterConfiguratorBaseNode.cpp b/Engine/Plugins/Runtime/nDisplay/Source/DisplayClusterConfigurator/Private/Views/OutputMapping/EdNodes/DisplayClusterConfiguratorBaseNode.cpp
--- a/Engine/Plugins/Runtime/nDisplay/Source/DisplayClusterConfigurator/Private/Views/OutputMapping/EdNodes/DisplayClusterConfiguratorBaseNode.cpp
+++ b/Engine/Plugins/Runtime/nDisplay/Source/DisplayClusterConfigurator/Private/Views/OutputMapping/EdNodes/DisplayClusterConfiguratorBaseNode.cpp
@@ -561,7 +561,7 @@ FVector2D UDisplayClusterConfiguratorBaseNode::FindNonOverlappingSize(UDisplayCl
 	return BestSize;
 }
 
-FVector2D UDisplayClusterConfiguratorBaseNode::FindNonOverlappingOffsetFromParent(const FVector2D& InDesiredOffset)
+FVector2D UDisplayClusterConfiguratorBaseNode::FindNonOverlappingOffsetFromParent(const FVector2D& InDesiredOffset, const TSet<UDisplayClusterConfiguratorBaseNode*>& NodesToIgnore)
 {
 	FVector2D BestOffset = InDesiredOffset;
 
@@ -578,7 +578,7 @@ FVector2D UDisplayClusterConfiguratorBaseNode::FindNonOverlappingOffsetFromParen
 		}
 
 		// Skip siblings that are not visible or enabled, as they are not visible in the graph editor and is confusing when a node is blocked by something invisible.
-		if (!SiblingNode->IsNodeEnabled() || !SiblingNode->IsNodeVisible())
+		if (!SiblingNode->IsNodeEnabled() || !SiblingNode->IsNodeVisible() || NodesToIgnore.Contains(SiblingNode))
 		{
 			continue;
 		}
@@ -602,7 +602,7 @@ FVector2D UDisplayClusterConfiguratorBaseNode::FindNonOverlappingOffsetFromParen
 		}
 
 		// Skip siblings that are not visible or enabled, as they are not visible in the graph editor and is confusing when a node is blocked by something invisible.
-		if (!SiblingNode->IsNodeEnabled() || !SiblingNode->IsNodeVisible())
+		if (!SiblingNode->IsNodeEnabled() || !SiblingNode->IsNodeVisible() || NodesToIgnore.Contains(SiblingNode))
 		{
 			continue;
 		}
@@ -874,13 +874,13 @@ FVector2D UDisplayClusterConfiguratorBaseNode::FindBoundedSizeFromChildren(const
 	return BestSize;
 }
 
-FNodeAlignmentPair UDisplayClusterConfiguratorBaseNode::GetTranslationAlignments(const FVector2D& InOffset, const FNodeAlignmentParams& AlignmentParams) const
+FNodeAlignmentPair UDisplayClusterConfiguratorBaseNode::GetTranslationAlignments(const FVector2D& InOffset, const FNodeAlignmentParams& AlignmentParams, const TSet<UDisplayClusterConfiguratorBaseNode*>& NodesToIgnore) const
 {
 	const FNodeAlignmentAnchors ShiftedAnchors = GetNodeAlignmentAnchors().ShiftBy(InOffset);
-	return GetAlignments(ShiftedAnchors, AlignmentParams);
+	return GetAlignments(ShiftedAnchors, AlignmentParams, NodesToIgnore);
 }
 
-FNodeAlignmentPair UDisplayClusterConfiguratorBaseNode::GetResizeAlignments(const FVector2D& InSizeChange, const FNodeAlignmentParams& AlignmentParams) const
+FNodeAlignmentPair UDisplayClusterConfiguratorBaseNode::GetResizeAlignments(const FVector2D& InSizeChange, const FNodeAlignmentParams& AlignmentParams, const TSet<UDisplayClusterConfiguratorBaseNode*>& NodesToIgnore) const
 {
 	const FNodeAlignmentAnchors ExpandedAnchors = GetNodeAlignmentAnchors().ExpandBy(InSizeChange);
 
@@ -891,10 +891,10 @@ FNodeAlignmentPair UDisplayClusterConfiguratorBaseNode::GetResizeAlignments(cons
 	NewParams.AnchorsToIgnore.Add(EAlignmentAnchor::Top);
 	NewParams.AnchorsToIgnore.Add(EAlignmentAnchor::Left);
 
-	return GetAlignments(ExpandedAnchors, NewParams);
+	return GetAlignments(ExpandedAnchors, NewParams, NodesToIgnore);
 }
 
-FNodeAlignmentPair UDisplayClusterConfiguratorBaseNode::GetAlignments(const FNodeAlignmentAnchors& TransformedAnchors, const FNodeAlignmentParams& AlignmentParams) const
+FNodeAlignmentPair UDisplayClusterConfiguratorBaseNode::GetAlignments(const FNodeAlignmentAnchors& TransformedAnchors, const FNodeAlignmentParams& AlignmentParams, const TSet<UDisplayClusterConfiguratorBaseNode*>& NodesToIgnore) const
 {
 	FDisplayClusterConfiguratorNodeAlignmentHelper AlignmentHelper(this, TransformedAnchors, AlignmentParams);
 
@@ -912,7 +912,7 @@ FNodeAlignmentPair UDisplayClusterConfiguratorBaseNode::GetAlignments(const FNod
 		}
 
 		// Skip siblings that are not visible or enabled, as they are not visible in the graph editor and is confusing when a node snap aligns to an invisible node.
-		if (!SiblingNode->IsNodeEnabled() || !SiblingNode->IsNodeVisible())
+		if (!SiblingNode->IsNodeEnabled() || !SiblingNode->IsNodeVisible() || NodesToIgnore.Contains(SiblingNode))
 		{
 			continue;
 		}
diff --git a/Engine/Plugins/Runtime/nDisplay/Source/DisplayClusterConfigurator/Private/Views/OutputMapping/EdNodes/DisplayClusterConfiguratorBaseNode.h b/Engine/Plugins/Runtime/nDisplay/Source/DisplayClusterConfigurator/Private/Views/OutputMapping/EdNodes/DisplayClusterConfiguratorBaseNode.h
--- a/Engine/Plugins/Runtime/nDisplay/Source/DisplayClusterConfigurator/Private/Views/OutputMapping/EdNodes/DisplayClusterConfiguratorBaseNode.h
+++ b/Engine/Plugins/Runtime/nDisplay/Source/DisplayClusterConfigurator/Private/Views/OutputMapping/EdNodes/DisplayClusterConfiguratorBaseNode.h
@@ -100,6 +100,9 @@ public:
 	virtual bool IsNodeEnabled() const { return true; }
 	virtual bool IsNodeAutoPositioned() const { return false; }
 	virtual bool IsNodeAutosized() const { return false; }
+	virtual bool CanNodeOverlapSiblings() const { return true; }
+	virtual bool CanNodeExceedParentBounds() const { return true; }
+	virtual bool CanNodeEncroachChildBounds() const { return true; }
 
 	virtual void FillParent(bool bRepositionNode = true);
 	virtual void SizeToChildren(bool bRepositionNode = true);
@@ -113,8 +116,9 @@ public:
 	virtual void TickPosition() { }
 
 	virtual bool IsUserInteractingWithNode(bool bCheckDescendents = false) const;
-	virtual void MarkUserInteractingWithNode() { bIsUserInteractingWithNode = true; }
-	virtual void ClearUserInteractingWithNode() { bIsUserInteractingWithNode = false; }
+	virtual bool IsUserDirectlyInteractingWithNode() const { return bIsDirectInteraction; }
+	virtual void MarkUserInteractingWithNode(bool bInIsDirectInteraction) { bIsUserInteractingWithNode = true; bIsDirectInteraction = bInIsDirectInteraction; }
+	virtual void ClearUserInteractingWithNode() { bIsUserInteractingWithNode = bIsDirectInteraction = false;  }
 
 	virtual void UpdateNode();
 	virtual void UpdateObject();
@@ -126,18 +130,18 @@ public:
 	virtual FVector2D FindNonOverlappingOffset(UDisplayClusterConfiguratorBaseNode* InNode, const FVector2D& InDesiredOffset) const;
 	virtual FVector2D FindNonOverlappingSize(UDisplayClusterConfiguratorBaseNode* InNode, const FVector2D& InDesiredSize, const bool bFixedApsectRatio) const;
 
-	virtual FVector2D FindNonOverlappingOffsetFromParent(const FVector2D& InDesiredOffset);
+	virtual FVector2D FindNonOverlappingOffsetFromParent(const FVector2D& InDesiredOffset, const TSet<UDisplayClusterConfiguratorBaseNode*>& NodesToIgnore = TSet<UDisplayClusterConfiguratorBaseNode*>());
 	virtual FVector2D FindBoundedOffsetFromParent(const FVector2D& InDesiredOffset);
 	virtual FVector2D FindNonOverlappingSizeFromParent(const FVector2D& InDesiredSize, const bool bFixedApsectRatio);
 	virtual FVector2D FindBoundedSizeFromParent(const FVector2D& InDesiredSize, const bool bFixedApsectRatio);
 	virtual FVector2D FindBoundedSizeFromChildren(const FVector2D& InDesiredSize, const bool bFixedApsectRatio);
 
-	virtual FNodeAlignmentPair GetTranslationAlignments(const FVector2D& InOffset, const FNodeAlignmentParams& AlignmentParams) const;
-	virtual FNodeAlignmentPair GetResizeAlignments(const FVector2D& InSizeChange, const FNodeAlignmentParams& AlignmentParams) const;
+	virtual FNodeAlignmentPair GetTranslationAlignments(const FVector2D& InOffset, const FNodeAlignmentParams& AlignmentParams, const TSet<UDisplayClusterConfiguratorBaseNode*>& NodesToIgnore = TSet<UDisplayClusterConfiguratorBaseNode*>()) const;
+	virtual FNodeAlignmentPair GetResizeAlignments(const FVector2D& InSizeChange, const FNodeAlignmentParams& AlignmentParams, const TSet<UDisplayClusterConfiguratorBaseNode*>& NodesToIgnore = TSet<UDisplayClusterConfiguratorBaseNode*>()) const;
 
 protected:
 	virtual bool CanAlignWithParent() const { return false; }
-	virtual FNodeAlignmentPair GetAlignments(const FNodeAlignmentAnchors& TransformedAnchors, const FNodeAlignmentParams& AlignmentParams) const;
+	virtual FNodeAlignmentPair GetAlignments(const FNodeAlignmentAnchors& TransformedAnchors, const FNodeAlignmentParams& AlignmentParams, const TSet<UDisplayClusterConfiguratorBaseNode*>& NodesToIgnore) const;
 
 	virtual void WriteNodeStateToObject() { }
 	virtual void ReadNodeStateFromObject() { }
@@ -171,4 +175,5 @@ protected:
 	int32 NodeZIndex;
 
 	bool bIsUserInteractingWithNode;
+	bool bIsDirectInteraction;
 };
diff --git a/Engine/Plugins/Runtime/nDisplay/Source/DisplayClusterConfigurator/Private/Views/OutputMapping/EdNodes/DisplayClusterConfiguratorHostNode.h b/Engine/Plugins/Runtime/nDisplay/Source/DisplayClusterConfigurator/Private/Views/OutputMapping/EdNodes/DisplayClusterConfiguratorHostNode.h
--- a/Engine/Plugins/Runtime/nDisplay/Source/DisplayClusterConfigurator/Private/Views/OutputMapping/EdNodes/DisplayClusterConfiguratorHostNode.h
+++ b/Engine/Plugins/Runtime/nDisplay/Source/DisplayClusterConfigurator/Private/Views/OutputMapping/EdNodes/DisplayClusterConfiguratorHostNode.h
@@ -42,6 +42,9 @@ public:
 	virtual bool IsNodeEnabled() const override;
 	virtual bool IsNodeAutoPositioned() const { return !CanUserMoveNode(); }
 	virtual bool IsNodeAutosized() const override { return !CanUserResizeNode(); }
+	virtual bool CanNodeOverlapSiblings() const override { return false; }
+	virtual bool CanNodeEncroachChildBounds() const { return false; }
+
 	virtual void DeleteObject() override;
 
 protected:
diff --git a/Engine/Plugins/Runtime/nDisplay/Source/DisplayClusterConfigurator/Private/Views/OutputMapping/EdNodes/DisplayClusterConfiguratorViewportNode.h b/Engine/Plugins/Runtime/nDisplay/Source/DisplayClusterConfigurator/Private/Views/OutputMapping/EdNodes/DisplayClusterConfiguratorViewportNode.h
--- a/Engine/Plugins/Runtime/nDisplay/Source/DisplayClusterConfigurator/Private/Views/OutputMapping/EdNodes/DisplayClusterConfiguratorViewportNode.h
+++ b/Engine/Plugins/Runtime/nDisplay/Source/DisplayClusterConfigurator/Private/Views/OutputMapping/EdNodes/DisplayClusterConfiguratorViewportNode.h
@@ -32,6 +32,8 @@ public:
 	//~ Begin UDisplayClusterConfiguratorBaseNode Interface
 	virtual bool IsNodeVisible() const override;
 	virtual bool IsNodeEnabled() const override;
+	virtual bool CanNodeOverlapSiblings() const override { return false; }
+
 	virtual void DeleteObject() override;
 
 protected:
diff --git a/Engine/Plugins/Runtime/nDisplay/Source/DisplayClusterConfigurator/Private/Views/OutputMapping/EdNodes/DisplayClusterConfiguratorWindowNode.cpp b/Engine/Plugins/Runtime/nDisplay/Source/DisplayClusterConfigurator/Private/Views/OutputMapping/EdNodes/DisplayClusterConfiguratorWindowNode.cpp
--- a/Engine/Plugins/Runtime/nDisplay/Source/DisplayClusterConfigurator/Private/Views/OutputMapping/EdNodes/DisplayClusterConfiguratorWindowNode.cpp
+++ b/Engine/Plugins/Runtime/nDisplay/Source/DisplayClusterConfigurator/Private/Views/OutputMapping/EdNodes/DisplayClusterConfiguratorWindowNode.cpp
@@ -3,6 +3,7 @@
 #include "Views/OutputMapping/EdNodes/DisplayClusterConfiguratorWindowNode.h"
 
 #include "DisplayClusterConfiguratorBlueprintEditor.h"
+#include "Interfaces/Views/OutputMapping/IDisplayClusterConfiguratorViewOutputMapping.h"
 #include "ClusterConfiguration/DisplayClusterConfiguratorClusterUtils.h"
 #include "ClusterConfiguration/ViewModels/DisplayClusterConfiguratorClusterNodeViewModel.h"
 #include "Views/OutputMapping/EdNodes/DisplayClusterConfiguratorCanvasNode.h"
@@ -121,6 +122,15 @@ bool UDisplayClusterConfiguratorWindowNode::IsNodeEnabled() const
 	return bIsChildEnabled;
 }
 
+bool UDisplayClusterConfiguratorWindowNode::CanNodeExceedParentBounds() const
+{
+	TSharedPtr<FDisplayClusterConfiguratorBlueprintEditor> Toolkit = ToolkitPtr.Pin();
+	check(Toolkit.IsValid());
+
+	TSharedRef<IDisplayClusterConfiguratorViewOutputMapping> OutputMapping = Toolkit->GetViewOutputMapping();
+	return !OutputMapping->GetOutputMappingSettings().bKeepClusterNodesInHosts;
+}
+
 void UDisplayClusterConfiguratorWindowNode::DeleteObject()
 {
 	UDisplayClusterConfigurationClusterNode* ClusterNode = GetObjectChecked<UDisplayClusterConfigurationClusterNode>();
diff --git a/Engine/Plugins/Runtime/nDisplay/Source/DisplayClusterConfigurator/Private/Views/OutputMapping/EdNodes/DisplayClusterConfiguratorWindowNode.h b/Engine/Plugins/Runtime/nDisplay/Source/DisplayClusterConfigurator/Private/Views/OutputMapping/EdNodes/DisplayClusterConfiguratorWindowNode.h
--- a/Engine/Plugins/Runtime/nDisplay/Source/DisplayClusterConfigurator/Private/Views/OutputMapping/EdNodes/DisplayClusterConfiguratorWindowNode.h
+++ b/Engine/Plugins/Runtime/nDisplay/Source/DisplayClusterConfigurator/Private/Views/OutputMapping/EdNodes/DisplayClusterConfiguratorWindowNode.h
@@ -42,6 +42,9 @@ public:
 	//~ Begin UDisplayClusterConfiguratorBaseNode Interface
 	virtual bool IsNodeVisible() const override;
 	virtual bool IsNodeEnabled() const override;
+	virtual bool CanNodeOverlapSiblings() const override { return false; }
+	virtual bool CanNodeExceedParentBounds() const override;
+
 	virtual void DeleteObject() override;
 
 protected:
diff --git a/Engine/Plugins/Runtime/nDisplay/Source/DisplayClusterConfigurator/Private/Views/OutputMapping/GraphNodes/SDisplayClusterConfiguratorBaseNode.cpp b/Engine/Plugins/Runtime/nDisplay/Source/DisplayClusterConfigurator/Private/Views/OutputMapping/GraphNodes/SDisplayClusterConfiguratorBaseNode.cpp
--- a/Engine/Plugins/Runtime/nDisplay/Source/DisplayClusterConfigurator/Private/Views/OutputMapping/GraphNodes/SDisplayClusterConfiguratorBaseNode.cpp
+++ b/Engine/Plugins/Runtime/nDisplay/Source/DisplayClusterConfigurator/Private/Views/OutputMapping/GraphNodes/SDisplayClusterConfiguratorBaseNode.cpp
@@ -166,30 +166,58 @@ void SDisplayClusterConfiguratorBaseNode::MoveTo(const FVector2D& NewPosition, F
 	UDisplayClusterConfiguratorBaseNode* EdNode = GetGraphNodeChecked<UDisplayClusterConfiguratorBaseNode>();
 	UDisplayClusterConfiguratorBaseNode* ParentEdNode = EdNode->GetParent();
 
+	const FGraphPanelSelectionSet& SelectedNodes = GetOwnerPanel()->SelectionManager.GetSelectedNodes();
+
+	TSet<UDisplayClusterConfiguratorBaseNode*> SelectedBaseNodes;
+	for (UObject* SelectedObject : SelectedNodes)
+	{
+		if (UDisplayClusterConfiguratorBaseNode* BaseNode = Cast<UDisplayClusterConfiguratorBaseNode>(SelectedObject))
+		{
+			SelectedBaseNodes.Add(BaseNode);
+		}
+	}
+
 	// If the parent node is also being moved, we don't want to move this node; otherwise, weird translations happen as the parent tries to update its child positions.
-	const bool bIsParentSelected = GetOwnerPanel()->SelectionManager.SelectedNodes.Contains(EdNode->GetParent());
+	const bool bIsParentSelected = SelectedBaseNodes.Contains(EdNode->GetParent());
 	if (bIsParentSelected)
 	{
 		return;
 	}
 
 	const bool bIsNodeFiltered = NodeFilter.Contains(SharedThis(this));
-	if (!bIsNodeFiltered)
+	if (!bIsNodeFiltered && !EdNode->IsUserInteractingWithNode())
 	{
 		BeginUserInteraction();
 	}
 
-	const bool bIsOverlappingAllowed = OutputMapping->GetOutputMappingSettings().bAllowClusterItemOverlap;
-
-	FVector2D BestOffset = Offset;
-	if (!CanNodeExceedParentBounds() && !bIsNodeFiltered)
+	// If the user is not directly interacting with this node (clicking it directly to drag it), don't update the node's position.
+	// In order to correctly move multiple selected nodes at once while preventing overlap or aligning nodes,  the node being directly
+	// moved will handle the move offset and apply it to the entire group of selected nodes instead of allowing each node to determine
+	// its own movement.
+	if (!EdNode->IsUserDirectlyInteractingWithNode())
 	{
-		BestOffset = EdNode->FindBoundedOffsetFromParent(BestOffset);
+		return;
 	}
 
-	if (!CanNodeOverlapSiblings() && !bIsOverlappingAllowed && !bIsNodeFiltered)
+	FVector2D BestOffset = Offset;
+	if (!bIsNodeFiltered)
 	{
-		BestOffset = EdNode->FindNonOverlappingOffsetFromParent(BestOffset);
+		const bool bIsOverlappingAllowed = OutputMapping->GetOutputMappingSettings().bAllowClusterItemOverlap;
+
+		// Check for overlapping across the entire group of nodes that are being moved, and iteratively update the best offset to prevent
+		// any node being moved from overlapping with any of its siblings
+		for (UDisplayClusterConfiguratorBaseNode* BaseNode : SelectedBaseNodes)
+		{
+			if (!BaseNode->CanNodeExceedParentBounds())
+			{
+				BestOffset = BaseNode->FindBoundedOffsetFromParent(BestOffset);
+			}
+
+			if (!BaseNode->CanNodeOverlapSiblings() && !bIsOverlappingAllowed)
+			{
+				BestOffset = BaseNode->FindNonOverlappingOffsetFromParent(BestOffset, SelectedBaseNodes);
+			}
+		}
 	}
 
 	FVector2D AlignmentOffset = FVector2D::ZeroVector;
@@ -203,7 +231,7 @@ void SDisplayClusterConfiguratorBaseNode::MoveTo(const FVector2D& NewPosition, F
 		Params.SnapProximity = AlignmentSettings.SnapProximity;
 		Params.SnapAdjacentEdgesPadding = AlignmentSettings.AdjacentEdgesSnapPadding;
 
-		FNodeAlignmentPair Alignments = EdNode->GetTranslationAlignments(BestOffset, Params);
+		FNodeAlignmentPair Alignments = EdNode->GetTranslationAlignments(BestOffset, Params, SelectedBaseNodes);
 		AlignmentOffset = Alignments.GetOffset();
 
 		UpdateAlignmentTarget(XAlignmentTarget, Alignments.XAlignment, Alignments.XAlignment.TargetNode == ParentEdNode);
@@ -217,6 +245,26 @@ void SDisplayClusterConfiguratorBaseNode::MoveTo(const FVector2D& NewPosition, F
 
 	SGraphNode::MoveTo(CurrentPosition + BestOffset + AlignmentOffset, NodeFilter);
 
+	// Nodes that aren't being directly interacted with aren't allowed to update their positions themselves since that might cause
+	// issues with overlapping and alignment, so the node being directly dragged and moved needs to update the positions of all selected
+	// nodes to move them as a cohesive group
+	if (!bIsNodeFiltered)
+	{
+		for (UDisplayClusterConfiguratorBaseNode* BaseNode : SelectedBaseNodes)
+		{
+			if (BaseNode != EdNode)
+			{
+				BaseNode->Modify();
+				BaseNode->NodePosX += BestOffset.X + AlignmentOffset.X;
+				BaseNode->NodePosY += BestOffset.Y + AlignmentOffset.Y;
+
+				BaseNode->UpdateObject();
+				BaseNode->UpdateChildNodes();
+			}
+		}
+	}
+
+
 	if (!bIsNodeFiltered)
 	{
 		// If the parent node is being auto-positioned, add it to the undo stack here because we need to store its old position with the this node's old position
@@ -235,8 +283,16 @@ void SDisplayClusterConfiguratorBaseNode::EndUserInteraction() const
 {
 	SGraphNode::EndUserInteraction();
 
-	UDisplayClusterConfiguratorBaseNode* EdNode = GetGraphNodeChecked<UDisplayClusterConfiguratorBaseNode>();
-	EdNode->ClearUserInteractingWithNode();
+	// EndUserInteraction is only called on the widget the user initially clicked on, but since nodes are marked as being 
+	// interacted with if they are moved, all selected nodes need their interaction flag cleared.
+	const FGraphPanelSelectionSet& SelectedNodes = GetOwnerPanel()->SelectionManager.GetSelectedNodes();
+	for (UObject* SelectedNode : SelectedNodes)
+	{
+		if (UDisplayClusterConfiguratorBaseNode* BaseNode = Cast<UDisplayClusterConfiguratorBaseNode>(SelectedNode))
+		{
+			BaseNode->ClearUserInteractingWithNode();
+		}
+	}
 
 	XAlignmentTarget.TargetNode = nullptr;
 	YAlignmentTarget.TargetNode = nullptr;
@@ -279,7 +335,34 @@ TArray<FOverlayWidgetInfo> SDisplayClusterConfiguratorBaseNode::GetOverlayWidget
 void SDisplayClusterConfiguratorBaseNode::BeginUserInteraction() const
 {
 	UDisplayClusterConfiguratorBaseNode* EdNode = GetGraphNodeChecked<UDisplayClusterConfiguratorBaseNode>();
-	EdNode->MarkUserInteractingWithNode();
+
+	bool bIsInteractingDirectly = false;
+	const FVector2D CursorPos = GetOwnerPanel()->PanelCoordToGraphCoord(GetOwnerPanel()->GetTickSpaceGeometry().AbsoluteToLocal(FSlateApplication::Get().GetCursorPos()));
+
+	if (EdNode->GetNodeBounds().IsInside(CursorPos))
+	{
+		bIsInteractingDirectly = true;
+
+		// Check that there isn't another node above this node that the user might be clicking on
+		const FGraphPanelSelectionSet& SelectedNodes = GetOwnerPanel()->SelectionManager.GetSelectedNodes();
+		const int32 EdNodeLayer = EdNode->GetNodeLayer(SelectedNodes);
+
+		for (UObject* SelectedNode : SelectedNodes)
+		{
+			if (UDisplayClusterConfiguratorBaseNode* BaseNode = Cast<UDisplayClusterConfiguratorBaseNode>(SelectedNode))
+			{
+				if (BaseNode->GetNodeBounds().IsInside(CursorPos))
+				{
+					if (BaseNode->GetNodeLayer(SelectedNodes) > EdNodeLayer)
+					{
+						bIsInteractingDirectly = false;
+					}
+				}
+			}
+		}
+	}
+
+	EdNode->MarkUserInteractingWithNode(bIsInteractingDirectly);
 }
 
 void SDisplayClusterConfiguratorBaseNode::SetNodeSize(const FVector2D InLocalSize, bool bFixedAspectRatio)
@@ -296,17 +379,17 @@ void SDisplayClusterConfiguratorBaseNode::SetNodeSize(const FVector2D InLocalSiz
 	const FVector2D CurrentSize = EdNode->GetNodeSize();
 
 	FVector2D BestSize = InLocalSize;
-	if (!CanNodeExceedParentBounds())
+	if (!EdNode->CanNodeExceedParentBounds())
 	{
 		BestSize = EdNode->FindBoundedSizeFromParent(BestSize, bFixedAspectRatio);
 	}
 
-	if (!CanNodeEncroachChildBounds())
+	if (!EdNode->CanNodeEncroachChildBounds())
 	{
 		BestSize = EdNode->FindBoundedSizeFromChildren(BestSize, bFixedAspectRatio);
 	}
 
-	if (!CanNodeOverlapSiblings() && !bIsOverlappingAllowed)
+	if (!EdNode->CanNodeOverlapSiblings() && !bIsOverlappingAllowed)
 	{
 		BestSize = EdNode->FindNonOverlappingSizeFromParent(BestSize, bFixedAspectRatio);
 	}
@@ -421,8 +504,7 @@ EVisibility SDisplayClusterConfiguratorBaseNode::GetResizeHandleVisibility() con
 
 bool SDisplayClusterConfiguratorBaseNode::CanSnapAlign() const
 {
-	bool bAreMultipleNodesSelected = GetOwnerPanel()->SelectionManager.SelectedNodes.Num() > 1;
-	return FSlateApplication::Get().GetModifierKeys().IsShiftDown() && !bAreMultipleNodesSelected;
+	return FSlateApplication::Get().GetModifierKeys().IsShiftDown();
 }
 
 void SDisplayClusterConfiguratorBaseNode::UpdateAlignmentTarget(FAlignmentRulerTarget& OutTarget, const FNodeAlignment& Alignment, bool bIsTargetingParent)
diff --git a/Engine/Plugins/Runtime/nDisplay/Source/DisplayClusterConfigurator/Private/Views/OutputMapping/GraphNodes/SDisplayClusterConfiguratorBaseNode.h b/Engine/Plugins/Runtime/nDisplay/Source/DisplayClusterConfigurator/Private/Views/OutputMapping/GraphNodes/SDisplayClusterConfiguratorBaseNode.h
--- a/Engine/Plugins/Runtime/nDisplay/Source/DisplayClusterConfigurator/Private/Views/OutputMapping/GraphNodes/SDisplayClusterConfiguratorBaseNode.h
+++ b/Engine/Plugins/Runtime/nDisplay/Source/DisplayClusterConfigurator/Private/Views/OutputMapping/GraphNodes/SDisplayClusterConfiguratorBaseNode.h
@@ -109,21 +109,6 @@ public:
 	 */
 	virtual FVector2D GetSize() const;
 
-	/**
-	 * @return Whether this node can overlap its sibling nodes in the graph editor 
-	 */
-	virtual bool CanNodeOverlapSiblings() const { return true; }
-
-	/**
-	 * @return Whether this node can be placed outside of its parent node's bounds in the graph editor
-	 */
-	virtual bool CanNodeExceedParentBounds() const { return true; }
-
-	/**
-	 * @return Whether this node can be be smaller than the total bounds of its children
-	 */
-	virtual bool CanNodeEncroachChildBounds() const { return true; }
-
 	/**
 	 * @return Wether this node can be snap aligned when the user activates snap aligning
 	 */
diff --git a/Engine/Plugins/Runtime/nDisplay/Source/DisplayClusterConfigurator/Private/Views/OutputMapping/GraphNodes/SDisplayClusterConfiguratorHostNode.h b/Engine/Plugins/Runtime/nDisplay/Source/DisplayClusterConfigurator/Private/Views/OutputMapping/GraphNodes/SDisplayClusterConfiguratorHostNode.h
--- a/Engine/Plugins/Runtime/nDisplay/Source/DisplayClusterConfigurator/Private/Views/OutputMapping/GraphNodes/SDisplayClusterConfiguratorHostNode.h
+++ b/Engine/Plugins/Runtime/nDisplay/Source/DisplayClusterConfigurator/Private/Views/OutputMapping/GraphNodes/SDisplayClusterConfiguratorHostNode.h
@@ -30,8 +30,6 @@ public:
 
 	//~ Begin SDisplayClusterConfiguratorBaseNode interface
 	virtual void SetNodeSize(const FVector2D InLocalSize, bool bFixedAspectRatio) override;
-	virtual bool CanNodeOverlapSiblings() const override { return false; }
-	virtual bool CanNodeEncroachChildBounds() const { return false; }
 	virtual bool CanNodeBeSnapAligned() const override { return true; }
 	virtual bool CanNodeBeResized() const;
 	//~ End SDisplayClusterConfiguratorBaseNode interface
diff --git a/Engine/Plugins/Runtime/nDisplay/Source/DisplayClusterConfigurator/Private/Views/OutputMapping/GraphNodes/SDisplayClusterConfiguratorViewportNode.h b/Engine/Plugins/Runtime/nDisplay/Source/DisplayClusterConfigurator/Private/Views/OutputMapping/GraphNodes/SDisplayClusterConfiguratorViewportNode.h
--- a/Engine/Plugins/Runtime/nDisplay/Source/DisplayClusterConfigurator/Private/Views/OutputMapping/GraphNodes/SDisplayClusterConfiguratorViewportNode.h
+++ b/Engine/Plugins/Runtime/nDisplay/Source/DisplayClusterConfigurator/Private/Views/OutputMapping/GraphNodes/SDisplayClusterConfiguratorViewportNode.h
@@ -35,7 +35,6 @@ public:
 
 	//~ Begin SDisplayClusterConfiguratorBaseNode interface
 	virtual bool IsNodeVisible() const override;
-	virtual bool CanNodeOverlapSiblings() const override { return false; }
 	virtual bool CanNodeBeSnapAligned() const override { return true; }
 	virtual bool CanNodeBeResized() const { return !IsViewportLocked(); }
 	virtual bool IsAspectRatioFixed() const override;
diff --git a/Engine/Plugins/Runtime/nDisplay/Source/DisplayClusterConfigurator/Private/Views/OutputMapping/GraphNodes/SDisplayClusterConfiguratorWindowNode.cpp b/Engine/Plugins/Runtime/nDisplay/Source/DisplayClusterConfigurator/Private/Views/OutputMapping/GraphNodes/SDisplayClusterConfiguratorWindowNode.cpp
--- a/Engine/Plugins/Runtime/nDisplay/Source/DisplayClusterConfigurator/Private/Views/OutputMapping/GraphNodes/SDisplayClusterConfiguratorWindowNode.cpp
+++ b/Engine/Plugins/Runtime/nDisplay/Source/DisplayClusterConfigurator/Private/Views/OutputMapping/GraphNodes/SDisplayClusterConfiguratorWindowNode.cpp
@@ -375,15 +375,6 @@ FVector2D SDisplayClusterConfiguratorWindowNode::GetPosition() const
 	return SDisplayClusterConfiguratorBaseNode::GetPosition() * WindowScaleFactor;
 }
 
-bool SDisplayClusterConfiguratorWindowNode::CanNodeExceedParentBounds() const
-{
-	TSharedPtr<FDisplayClusterConfiguratorBlueprintEditor> Toolkit = ToolkitPtr.Pin();
-	check(Toolkit.IsValid());
-
-	TSharedRef<IDisplayClusterConfiguratorViewOutputMapping> OutputMapping = Toolkit->GetViewOutputMapping();
-	return !OutputMapping->GetOutputMappingSettings().bKeepClusterNodesInHosts;
-}
-
 bool SDisplayClusterConfiguratorWindowNode::IsAspectRatioFixed() const
 {
 	UDisplayClusterConfiguratorWindowNode* WindowEdNode = GetGraphNodeChecked<UDisplayClusterConfiguratorWindowNode>();
diff --git a/Engine/Plugins/Runtime/nDisplay/Source/DisplayClusterConfigurator/Private/Views/OutputMapping/GraphNodes/SDisplayClusterConfiguratorWindowNode.h b/Engine/Plugins/Runtime/nDisplay/Source/DisplayClusterConfigurator/Private/Views/OutputMapping/GraphNodes/SDisplayClusterConfiguratorWindowNode.h
--- a/Engine/Plugins/Runtime/nDisplay/Source/DisplayClusterConfigurator/Private/Views/OutputMapping/GraphNodes/SDisplayClusterConfiguratorWindowNode.h
+++ b/Engine/Plugins/Runtime/nDisplay/Source/DisplayClusterConfigurator/Private/Views/OutputMapping/GraphNodes/SDisplayClusterConfiguratorWindowNode.h
@@ -40,8 +40,6 @@ public:
 	//~ End SGraphNode interface
 
 	//~ Begin SDisplayClusterConfiguratorBaseNode interface
-	virtual bool CanNodeOverlapSiblings() const override { return false; }
-	virtual bool CanNodeExceedParentBounds() const override;
 	virtual bool CanNodeBeSnapAligned() const override { return true; }
 	virtual bool CanNodeBeResized() const { return !IsClusterNodeLocked(); }
 	virtual bool IsAspectRatioFixed() const override;
