RC apply property in Editor when Interceptor enabled

diff --git a/Engine/Plugins/Runtime/nDisplay/Source/DisplayClusterRemoteControlInterceptor/Private/DisplayClusterRemoteControlInterceptorModule.cpp b/Engine/Plugins/Runtime/nDisplay/Source/DisplayClusterRemoteControlInterceptor/Private/DisplayClusterRemoteControlInterceptorModule.cpp
--- a/Engine/Plugins/Runtime/nDisplay/Source/DisplayClusterRemoteControlInterceptor/Private/DisplayClusterRemoteControlInterceptorModule.cpp
+++ b/Engine/Plugins/Runtime/nDisplay/Source/DisplayClusterRemoteControlInterceptor/Private/DisplayClusterRemoteControlInterceptorModule.cpp
@@ -7,20 +7,30 @@
 #include "Modules/ModuleManager.h"
 #include "Features/IModularFeatures.h"
 #include "IRemoteControlInterceptionFeature.h"
+#include "IDisplayCluster.h"
 
 
 void FDisplayClusterRemoteControlInterceptorModule::StartupModule()
 {
 	// Instantiate the interceptor feature on module start
 	Interceptor = MakeUnique<FDisplayClusterRemoteControlInterceptor>();
-	// Register the interceptor feature
-	IModularFeatures::Get().RegisterModularFeature(IRemoteControlInterceptionFeatureInterceptor::GetName(), Interceptor.Get());
+
+	// Add Interceptor only in Cluster mode
+	if (IDisplayCluster::Get().GetOperationMode() == EDisplayClusterOperationMode::Cluster)
+	{
+		// Register the interceptor feature
+		IModularFeatures::Get().RegisterModularFeature(IRemoteControlInterceptionFeatureInterceptor::GetName(), Interceptor.Get());	
+	}
 }
 
 void FDisplayClusterRemoteControlInterceptorModule::ShutdownModule()
 {
-	// Unregister the interceptor feature on module shutdown
-	IModularFeatures::Get().UnregisterModularFeature(IRemoteControlInterceptionFeatureInterceptor::GetName(), Interceptor.Get());
+	// Add Interceptor only in Cluster mode
+	if (IDisplayCluster::Get().GetOperationMode() == EDisplayClusterOperationMode::Cluster)
+	{
+		// Unregister the interceptor feature on module shutdown
+		IModularFeatures::Get().UnregisterModularFeature(IRemoteControlInterceptionFeatureInterceptor::GetName(), Interceptor.Get());
+	}
 }
 
 IMPLEMENT_MODULE(FDisplayClusterRemoteControlInterceptorModule, DisplayClusterRemoteControlInterceptor);
