nDisplay: Fix duplicate camera in policy selection.

diff --git a/Engine/Plugins/Runtime/nDisplay/Source/DisplayClusterConfigurator/Private/Views/Details/Policies/DisplayClusterConfiguratorPolicyDetailCustomization.cpp b/Engine/Plugins/Runtime/nDisplay/Source/DisplayClusterConfigurator/Private/Views/Details/Policies/DisplayClusterConfiguratorPolicyDetailCustomization.cpp
--- a/Engine/Plugins/Runtime/nDisplay/Source/DisplayClusterConfigurator/Private/Views/Details/Policies/DisplayClusterConfiguratorPolicyDetailCustomization.cpp
+++ b/Engine/Plugins/Runtime/nDisplay/Source/DisplayClusterConfigurator/Private/Views/Details/Policies/DisplayClusterConfiguratorPolicyDetailCustomization.cpp
@@ -375,8 +375,7 @@ void FDisplayClusterConfiguratorProjectionCustomization::CreateCameraPolicy(UDis
 		Blueprint,
 		ConfigurationViewportPtr.Get(),
 		ParametersHandle,
-		TArray<TSubclassOf<UActorComponent>>{ UCameraComponent::StaticClass(),
-			UDisplayClusterICVFXCameraComponent::StaticClass() }));
+		TArray<TSubclassOf<UActorComponent>>{ UCameraComponent::StaticClass() } ));
 
 	CustomPolicyParameters.Add(MakeShared<FPolicyParameterInfoBool>(
 		"Native",
