nDisplay: Fix CIS warning about uninitialized bool properties.

diff --git a/Engine/Plugins/Runtime/nDisplay/Source/DisplayClusterConfiguration/Public/DisplayClusterConfigurationTypes_Viewport.h b/Engine/Plugins/Runtime/nDisplay/Source/DisplayClusterConfiguration/Public/DisplayClusterConfigurationTypes_Viewport.h
--- a/Engine/Plugins/Runtime/nDisplay/Source/DisplayClusterConfiguration/Public/DisplayClusterConfigurationTypes_Viewport.h
+++ b/Engine/Plugins/Runtime/nDisplay/Source/DisplayClusterConfiguration/Public/DisplayClusterConfigurationTypes_Viewport.h
@@ -192,7 +192,11 @@ struct DISPLAYCLUSTERCONFIGURATION_API FDisplayClusterConfigurationViewport_PerV
 {
 	GENERATED_BODY()
 
-	FDisplayClusterConfigurationViewport_PerViewportSettings() {};
+	FDisplayClusterConfigurationViewport_PerViewportSettings() 
+		: bOverride_WhiteTemp(0)
+		, bOverride_WhiteTint(0)
+		, bOverride_AutoExposureBias(0)
+	{};
 
 	UPROPERTY(EditAnywhere, Category = "Overrides", meta = (PinHiddenByDefault, InlineEditConditionToggle))
 	uint8 bOverride_WhiteTemp:1;
