RC and nD. Fix replication in Cluster
Register Interceptor only after Engine Loop Complete

diff --git a/Engine/Plugins/Runtime/nDisplay/Source/DisplayClusterRemoteControlInterceptor/Private/DisplayClusterRemoteControlInterceptorModule.cpp b/Engine/Plugins/Runtime/nDisplay/Source/DisplayClusterRemoteControlInterceptor/Private/DisplayClusterRemoteControlInterceptorModule.cpp
--- a/Engine/Plugins/Runtime/nDisplay/Source/DisplayClusterRemoteControlInterceptor/Private/DisplayClusterRemoteControlInterceptorModule.cpp
+++ b/Engine/Plugins/Runtime/nDisplay/Source/DisplayClusterRemoteControlInterceptor/Private/DisplayClusterRemoteControlInterceptorModule.cpp
@@ -5,22 +5,26 @@
 #include "DisplayClusterRemoteControlInterceptor.h"
 
 #include "Modules/ModuleManager.h"
+#include "Misc/CoreDelegates.h"
 #include "Features/IModularFeatures.h"
 #include "IRemoteControlInterceptionFeature.h"
 #include "IDisplayCluster.h"
 
-
 void FDisplayClusterRemoteControlInterceptorModule::StartupModule()
 {
 	// Instantiate the interceptor feature on module start
 	Interceptor = MakeUnique<FDisplayClusterRemoteControlInterceptor>();
 
-	// Add Interceptor only in Cluster mode
-	if (IDisplayCluster::Get().GetOperationMode() == EDisplayClusterOperationMode::Cluster)
+	// Register Interceptor only after Engine Loop Complete 
+	FCoreDelegates::OnFEngineLoopInitComplete.AddLambda([this]()
 	{
-		// Register the interceptor feature
-		IModularFeatures::Get().RegisterModularFeature(IRemoteControlInterceptionFeatureInterceptor::GetName(), Interceptor.Get());	
-	}
+		// Add Interceptor only in Cluster mode
+		if (IDisplayCluster::Get().GetOperationMode() == EDisplayClusterOperationMode::Cluster)
+		{
+			// Register the interceptor feature
+			IModularFeatures::Get().RegisterModularFeature(IRemoteControlInterceptionFeatureInterceptor::GetName(), Interceptor.Get());	
+		}
+	});
 }
 
 void FDisplayClusterRemoteControlInterceptorModule::ShutdownModule()
