nDisplay: Added groups to the root actor's post process settings category that display each viewport's post process settings.

diff --git a/Engine/Plugins/Runtime/nDisplay/Source/DisplayClusterConfigurator/Private/Views/Details/DisplayClusterConfiguratorDetailCustomization.cpp b/Engine/Plugins/Runtime/nDisplay/Source/DisplayClusterConfigurator/Private/Views/Details/DisplayClusterConfiguratorDetailCustomization.cpp
--- a/Engine/Plugins/Runtime/nDisplay/Source/DisplayClusterConfigurator/Private/Views/Details/DisplayClusterConfiguratorDetailCustomization.cpp
+++ b/Engine/Plugins/Runtime/nDisplay/Source/DisplayClusterConfigurator/Private/Views/Details/DisplayClusterConfiguratorDetailCustomization.cpp
@@ -4,6 +4,7 @@
 
 #include "DisplayClusterConfiguratorBlueprintEditor.h"
 #include "DisplayClusterConfigurationTypes.h"
+#include "Views/Details/DisplayClusterConfiguratorDetailCustomizationUtils.h"
 #include "Views/Details/Widgets/SDisplayClusterConfigurationSearchableComboBox.h"
 #include "Views/OutputMapping/Widgets/SDisplayClusterConfiguratorExternalImagePicker.h"
 #include "DisplayClusterConfiguratorUtils.h"
@@ -20,6 +21,7 @@
 #include "DetailWidgetRow.h"
 #include "DisplayClusterConfiguratorPropertyUtils.h"
 #include "IDetailChildrenBuilder.h"
+#include "IDetailGroup.h"
 #include "IPropertyUtilities.h"
 #include "PropertyHandle.h"
 #include "Kismet2/BlueprintEditorUtils.h"
@@ -30,31 +32,6 @@
 #include "Widgets/Text/STextBlock.h"
 #include "Widgets/SBoxPanel.h"
 
-#define BEGIN_CATEGORY(CategoryName) { \
-	IDetailCategoryBuilder& CurrentCategory = InLayoutBuilder.EditCategory(CategoryName);
-
-#define END_CATEGORY() }
-
-#define ADD_PROPERTY(ClassName, PropertyName) { \
-	TSharedRef<IPropertyHandle> PropertyHandle = InLayoutBuilder.GetProperty(GET_MEMBER_NAME_CHECKED(ClassName, PropertyName)); \
-	check(PropertyHandle->IsValidHandle()); \
-	CurrentCategory.AddProperty(PropertyHandle); \
-}
-
-#define ADD_EXPANDED_PROPERTY(ClassName, PropertyName) { \
-	TSharedRef<IPropertyHandle> PropertyHandle = InLayoutBuilder.GetProperty(GET_MEMBER_NAME_CHECKED(ClassName, PropertyName)); \
-	check(PropertyHandle->IsValidHandle()); \
-	CurrentCategory.AddProperty(PropertyHandle).ShouldAutoExpand(true); \
-}
-
-#define ADD_CUSTOM_PROPERTY(FilterText) CurrentCategory.AddCustomRow(FilterText)
-
-#define REPLACE_PROPERTY_WITH_CUSTOM(ClassName, PropertyName, Widget) { \
-	TSharedRef<IPropertyHandle> PropertyHandle = InLayoutBuilder.GetProperty(GET_MEMBER_NAME_CHECKED(ClassName, PropertyName)); \
-	check(PropertyHandle->IsValidHandle()); \
-	InLayoutBuilder.HideProperty(PropertyHandle); \
-	CurrentCategory.AddCustomRow(PropertyHandle->GetPropertyDisplayName()).NameContent()[PropertyHandle->CreatePropertyNameWidget()].ValueContent()[Widget]; \
-}
 
 #define LOCTEXT_NAMESPACE "FDisplayClusterConfiguratorDetailCustomization"
 
@@ -196,14 +173,35 @@ void FDisplayClusterConfiguratorClusterDetailCustomization::CustomizeDetails(IDe
 	ClusterNodesHandle = InLayoutBuilder.GetProperty(GET_MEMBER_NAME_CHECKED(UDisplayClusterConfigurationCluster, Nodes));
 	check(ClusterNodesHandle->IsValidHandle());
 
+	FDisplayClusterConfiguratorNestedPropertyHelper NestedPropertyHelper(InLayoutBuilder);
+
 	BEGIN_CATEGORY(DisplayClusterConfigurationStrings::categories::ClusterCategory)
 		ADD_PROPERTY(UDisplayClusterConfigurationCluster, Nodes);
 	END_CATEGORY()
 
 	BEGIN_CATEGORY(DisplayClusterConfigurationStrings::categories::ClusterPostprocessCategory)
-		ADD_PROPERTY(UDisplayClusterConfigurationCluster, bUseOverallClusterPostProcess);
-		ADD_PROPERTY(UDisplayClusterConfigurationCluster, OverallClusterPostProcessSettings);
-	END_CATEGORY()
+		BEGIN_GROUP(TEXT("GlobalPostProcess"), LOCTEXT("GlobalPostprocessLabel", "All Viewports"))
+			ADD_GROUP_PROPERTY(UDisplayClusterConfigurationCluster, bUseOverallClusterPostProcess);
+			ADD_GROUP_PROPERTY(UDisplayClusterConfigurationCluster, OverallClusterPostProcessSettings);
+		END_GROUP();
+
+		TArray<FString> ViewportNames;
+		NestedPropertyHelper.GetNestedPropertyKeys(TEXT("Nodes.Viewports"), ViewportNames);
+
+		TArray<TSharedPtr<IPropertyHandle>> ViewportPostProcessSettings;
+		NestedPropertyHelper.GetNestedProperties(TEXT("Nodes.Viewports.PostProcessSettings"), ViewportPostProcessSettings);
+
+		check(ViewportNames.Num() == ViewportPostProcessSettings.Num());
+
+		for (int32 Index = 0; Index < ViewportNames.Num(); ++Index)
+		{
+			BEGIN_GROUP(FName(*ViewportNames[Index]), FText::FromString(ViewportNames[Index]))
+				CurrentGroup.AddPropertyRow(ViewportPostProcessSettings[Index]->GetChildHandle(GET_MEMBER_NAME_CHECKED(FDisplayClusterConfigurationViewport_PostProcessSettings, bExcludeFromOverallClusterPostProcess)).ToSharedRef());
+				CurrentGroup.AddPropertyRow(ViewportPostProcessSettings[Index]->GetChildHandle(GET_MEMBER_NAME_CHECKED(FDisplayClusterConfigurationViewport_PostProcessSettings, bIsEnabled)).ToSharedRef());
+				CurrentGroup.AddPropertyRow(ViewportPostProcessSettings[Index]->GetChildHandle(GET_MEMBER_NAME_CHECKED(FDisplayClusterConfigurationViewport_PostProcessSettings, ViewportSettings)).ToSharedRef()).ShouldAutoExpand(true);
+			END_GROUP();
+		}
+	END_CATEGORY();
 
 	BEGIN_CATEGORY(DisplayClusterConfigurationStrings::categories::ClusterConfigurationCategory)
 		ADD_PROPERTY(UDisplayClusterConfigurationCluster, MasterNode);
diff --git a/Engine/Plugins/Runtime/nDisplay/Source/DisplayClusterConfigurator/Private/Views/Details/DisplayClusterConfiguratorDetailCustomizationUtils.cpp b/Engine/Plugins/Runtime/nDisplay/Source/DisplayClusterConfigurator/Private/Views/Details/DisplayClusterConfiguratorDetailCustomizationUtils.cpp
new file mode 100644
--- /dev/null
+++ b/Engine/Plugins/Runtime/nDisplay/Source/DisplayClusterConfigurator/Private/Views/Details/DisplayClusterConfiguratorDetailCustomizationUtils.cpp
@@ -0,0 +1,224 @@
+// Copyright Epic Games, Inc. All Rights Reserved.
+
+#include "DisplayClusterConfiguratorDetailCustomizationUtils.h"
+
+#include "DetailLayoutBuilder.h"
+#include "IDetailGroup.h"
+#include "PropertyHandle.h"
+
+TSharedPtr<IPropertyHandle> FDisplayClusterConfiguratorNestedPropertyHelper::GetNestedProperty(const FString& PropertyPath)
+{
+	if (CachedPropertyHandles.Contains(PropertyPath))
+	{
+		return CachedPropertyHandles[PropertyPath];
+	}
+
+	TArray<FString> Properties;
+	PropertyPath.ParseIntoArray(Properties, TEXT("."));
+
+	if (Properties.Num())
+	{
+		int32 CurrentPropertyIndex = 0;
+		FString CurrentPropertyPath;
+		TSharedPtr<IPropertyHandle> CurrentPropertyHandle = FindCachedHandle(Properties, CurrentPropertyPath, CurrentPropertyIndex);
+
+		if (!CurrentPropertyHandle.IsValid())
+		{
+			CurrentPropertyIndex = 0;
+			CurrentPropertyPath = Properties[CurrentPropertyIndex];
+			CurrentPropertyHandle = LayoutBuilder.GetProperty(FName(*CurrentPropertyPath));
+		}
+
+		++CurrentPropertyIndex;
+
+		while (CurrentPropertyIndex < Properties.Num() && CurrentPropertyHandle)
+		{
+			CurrentPropertyHandle = GetProperty(Properties[CurrentPropertyIndex], CurrentPropertyHandle, CurrentPropertyPath);
+			CurrentPropertyPath += TEXT(".") + Properties[CurrentPropertyIndex];
+			++CurrentPropertyIndex;
+		}
+
+		return CurrentPropertyHandle;
+	}
+
+	return TSharedPtr<IPropertyHandle>();
+}
+
+void FDisplayClusterConfiguratorNestedPropertyHelper::GetNestedProperties(const FString& PropertyPath, TArray<TSharedPtr<IPropertyHandle>>& OutPropertyHandles)
+{
+	TArray<FString> Properties;
+	PropertyPath.ParseIntoArray(Properties, TEXT("."));
+
+	if (Properties.Num())
+	{
+		int32 CurrentPropertyIndex = 0;
+
+		TMap<FString, TSharedPtr<IPropertyHandle>> CurrentPropertyHandles;
+
+		FString CachedPropertyPath;
+		if (TSharedPtr<IPropertyHandle> CachedPropertyHandle = FindCachedHandle(Properties, CachedPropertyPath, CurrentPropertyIndex))
+		{
+			CurrentPropertyHandles.Add(CachedPropertyPath, CachedPropertyHandle);
+		}
+		else
+		{
+			CurrentPropertyHandles.Add(Properties[CurrentPropertyIndex], LayoutBuilder.GetProperty(FName(*Properties[CurrentPropertyIndex])));
+		}
+
+		do
+		{
+			// First, if the current nested properties are lists, expand them into list elements.
+			TMap<FString, TSharedPtr<IPropertyHandle>> NextPropertyHandles;
+			for (const TPair<FString, TSharedPtr<IPropertyHandle>>& CurrentPropertyPair : CurrentPropertyHandles)
+			{
+				if (IsListType(CurrentPropertyPair.Value))
+				{
+					uint32 NumElements;
+					CurrentPropertyPair.Value->GetNumChildren(NumElements);
+
+					for (uint32 Index = 0; Index < NumElements; ++Index)
+					{
+						if (TSharedPtr<IPropertyHandle> ElementHandle = CurrentPropertyPair.Value->GetChildHandle(Index))
+						{
+							const FString NextPropertyPath = CurrentPropertyPair.Key + FString::Format(TEXT("[{0}]"), { Index });
+							NextPropertyHandles.Add(NextPropertyPath, ElementHandle);
+						}
+					}
+				}
+				else
+				{
+					NextPropertyHandles.Add(CurrentPropertyPair);
+				}
+			}
+
+			CurrentPropertyHandles.Empty();
+			++CurrentPropertyIndex;
+
+			// Next, grab the next set of properties from the expanded list elements. For now, assume that
+			// the elements of any lists are not lists themselves.
+			if (CurrentPropertyIndex < Properties.Num())
+			{
+				for (const TPair<FString, TSharedPtr<IPropertyHandle>>& NextPropertyPair : NextPropertyHandles)
+				{
+					if (TSharedPtr<IPropertyHandle> PropertyHandle = GetProperty(Properties[CurrentPropertyIndex], NextPropertyPair.Value, NextPropertyPair.Key))
+					{
+						const FString CurrentPropertyPath = NextPropertyPair.Key + TEXT(".") + Properties[CurrentPropertyIndex];
+						CurrentPropertyHandles.Add(CurrentPropertyPath, PropertyHandle);
+					}
+				}
+			}
+			else
+			{
+				CurrentPropertyHandles.Append(NextPropertyHandles);
+			}
+
+		} while (CurrentPropertyIndex < Properties.Num() && CurrentPropertyHandles.Num());
+
+		CurrentPropertyHandles.GenerateValueArray(OutPropertyHandles);
+	}
+}
+
+void FDisplayClusterConfiguratorNestedPropertyHelper::GetNestedPropertyKeys(const FString& PropertyPath, TArray<FString>& OutPropertyKeys)
+{
+	TArray<TSharedPtr<IPropertyHandle>> PropertyHandles;
+	GetNestedProperties(PropertyPath, PropertyHandles);
+
+	for (const TSharedPtr<IPropertyHandle>& PropertyHandle : PropertyHandles)
+	{
+		TSharedPtr<IPropertyHandle> KeyHandle = PropertyHandle->GetKeyHandle();
+
+		FString KeyValue;
+		if (KeyHandle.IsValid() && KeyHandle->GetValue(KeyValue) == FPropertyAccess::Success)
+		{
+			OutPropertyKeys.Add(KeyValue);
+		}
+	}
+}
+
+TSharedPtr<IPropertyHandle> FDisplayClusterConfiguratorNestedPropertyHelper::GetProperty(const FString& PropertyName, const TSharedPtr<IPropertyHandle>& Parent, const FString& ParentPath)
+{
+	const FString PropertyPath = ParentPath + TEXT(".") + PropertyName;
+	if (CachedPropertyHandles.Contains(PropertyPath))
+	{
+		return CachedPropertyHandles[PropertyPath];
+	}
+
+	if (PropertyName.Contains(TEXT("[")) && PropertyName.Contains(TEXT("]")))
+	{
+		FString ListPropertyName;
+		FString IndexStr;
+		PropertyName.Split(TEXT("["), &ListPropertyName, &IndexStr);
+
+		int32 Index = FCString::Atoi(*IndexStr);
+
+		const FString ListPropertyPath = ParentPath + TEXT(".") + ListPropertyName;
+
+		TSharedPtr<IPropertyHandle> ListPropertyHandle;
+		if (CachedPropertyHandles.Contains(ListPropertyPath))
+		{
+			ListPropertyHandle = CachedPropertyHandles[ListPropertyPath];
+		}
+		else
+		{
+			ListPropertyHandle = Parent->GetChildHandle(FName(*ListPropertyName));
+		}
+
+		if (IsListType(ListPropertyHandle))
+		{
+			CachedPropertyHandles.Add(ListPropertyPath, ListPropertyHandle);
+
+			uint32 NumElements;
+			ListPropertyHandle->GetNumChildren(NumElements);
+
+			if (Index >= 0 && Index < (int32)NumElements)
+			{
+				if (TSharedPtr<IPropertyHandle> ElementHandle = ListPropertyHandle->GetChildHandle(Index))
+				{
+					CachedPropertyHandles.Add(PropertyPath, ElementHandle);
+					return ElementHandle;
+				}
+			}
+		}
+
+		return TSharedPtr<IPropertyHandle>();
+	}
+	else
+	{
+		if (TSharedPtr<IPropertyHandle> PropertyHandle = Parent->GetChildHandle(FName(*PropertyName)))
+		{
+			CachedPropertyHandles.Add(PropertyPath, PropertyHandle);
+			return PropertyHandle;
+		}
+
+		return TSharedPtr<IPropertyHandle>();
+	}
+
+}
+
+TSharedPtr<IPropertyHandle> FDisplayClusterConfiguratorNestedPropertyHelper::FindCachedHandle(const TArray<FString>& PropertyPath, FString& OutFoundPath, int32& OutFoundIndex)
+{
+	FString CurrentPath = FString::Join(PropertyPath, TEXT("."));
+	int32 CurrentIndex = PropertyPath.Num();
+
+	while (!CachedPropertyHandles.Contains(CurrentPath) && CurrentIndex > 0)
+	{
+		--CurrentIndex;
+		CurrentPath.RemoveFromEnd(PropertyPath[CurrentIndex]);
+		CurrentPath.RemoveFromEnd(TEXT("."));
+	}
+
+	// If the cache contains a handle for a property somewhere in the desired path, return it
+	if (CachedPropertyHandles.Contains(CurrentPath))
+	{
+		OutFoundPath = CurrentPath;
+		OutFoundIndex = CurrentIndex - 1;
+		return CachedPropertyHandles[CurrentPath];
+	}
+
+	return TSharedPtr<IPropertyHandle>();
+}
+
+bool FDisplayClusterConfiguratorNestedPropertyHelper::IsListType(const TSharedPtr<IPropertyHandle>& PropertyHandle)
+{
+	return PropertyHandle->AsArray() || PropertyHandle->AsSet() || PropertyHandle->AsMap();
+}
\ No newline at end of file
diff --git a/Engine/Plugins/Runtime/nDisplay/Source/DisplayClusterConfigurator/Private/Views/Details/DisplayClusterConfiguratorDetailCustomizationUtils.h b/Engine/Plugins/Runtime/nDisplay/Source/DisplayClusterConfigurator/Private/Views/Details/DisplayClusterConfiguratorDetailCustomizationUtils.h
new file mode 100644
--- /dev/null
+++ b/Engine/Plugins/Runtime/nDisplay/Source/DisplayClusterConfigurator/Private/Views/Details/DisplayClusterConfiguratorDetailCustomizationUtils.h
@@ -0,0 +1,66 @@
+// Copyright Epic Games, Inc. All Rights Reserved.
+
+#pragma once
+
+#include "CoreMinimal.h"
+
+class IDetailLayoutBuilder;
+class IPropertyHandle;
+
+#define BEGIN_CATEGORY(CategoryName) { \
+	IDetailCategoryBuilder& CurrentCategory = InLayoutBuilder.EditCategory(CategoryName);
+
+#define END_CATEGORY() }
+
+#define ADD_PROPERTY(ClassName, PropertyName) { \
+	TSharedRef<IPropertyHandle> PropertyHandle = InLayoutBuilder.GetProperty(GET_MEMBER_NAME_CHECKED(ClassName, PropertyName)); \
+	check(PropertyHandle->IsValidHandle()); \
+	CurrentCategory.AddProperty(PropertyHandle); \
+}
+
+#define ADD_EXPANDED_PROPERTY(ClassName, PropertyName) { \
+	TSharedRef<IPropertyHandle> PropertyHandle = InLayoutBuilder.GetProperty(GET_MEMBER_NAME_CHECKED(ClassName, PropertyName)); \
+	check(PropertyHandle->IsValidHandle()); \
+	CurrentCategory.AddProperty(PropertyHandle).ShouldAutoExpand(true); \
+}
+
+#define ADD_CUSTOM_PROPERTY(FilterText) CurrentCategory.AddCustomRow(FilterText)
+
+#define REPLACE_PROPERTY_WITH_CUSTOM(ClassName, PropertyName, Widget) { \
+	TSharedRef<IPropertyHandle> PropertyHandle = InLayoutBuilder.GetProperty(GET_MEMBER_NAME_CHECKED(ClassName, PropertyName)); \
+	check(PropertyHandle->IsValidHandle()); \
+	InLayoutBuilder.HideProperty(PropertyHandle); \
+	CurrentCategory.AddCustomRow(PropertyHandle->GetPropertyDisplayName()).NameContent()[PropertyHandle->CreatePropertyNameWidget()].ValueContent()[Widget]; \
+}
+
+#define BEGIN_GROUP(GroupName, GroupLabel) { \
+	IDetailGroup& CurrentGroup = CurrentCategory.AddGroup(GroupName, GroupLabel, false, true);
+
+#define END_GROUP() }
+
+#define ADD_GROUP_PROPERTY(ClassName, PropertyName) { \
+	TSharedRef<IPropertyHandle> PropertyHandle = InLayoutBuilder.GetProperty(GET_MEMBER_NAME_CHECKED(ClassName, PropertyName)); \
+	check(PropertyHandle->IsValidHandle()); \
+	CurrentGroup.AddPropertyRow(PropertyHandle); \
+}
+
+struct FDisplayClusterConfiguratorNestedPropertyHelper
+{
+public:
+	FDisplayClusterConfiguratorNestedPropertyHelper(IDetailLayoutBuilder& InLayoutBuilder) :
+		LayoutBuilder(InLayoutBuilder)
+	{ }
+
+	TSharedPtr<IPropertyHandle> GetNestedProperty(const FString& PropertyPath);
+	void GetNestedProperties(const FString& PropertyPath, TArray<TSharedPtr<IPropertyHandle>>& OutPropertyHandles);
+	void GetNestedPropertyKeys(const FString& PropertyPath, TArray<FString>& OutPropertyKeys);
+
+private:
+	TSharedPtr<IPropertyHandle> GetProperty(const FString& PropertyName, const TSharedPtr<IPropertyHandle>& Parent, const FString& ParentPath);
+	TSharedPtr<IPropertyHandle> FindCachedHandle(const TArray<FString>& PropertyPath, FString& OutFoundPath, int32& OutFoundIndex);
+	bool IsListType(const TSharedPtr<IPropertyHandle>& PropertyHandle);
+
+private:
+	TMap<FString, TSharedPtr<IPropertyHandle>> CachedPropertyHandles;
+	IDetailLayoutBuilder& LayoutBuilder;
+};
\ No newline at end of file
diff --git a/Engine/Plugins/Runtime/nDisplay/Source/DisplayClusterEditor/Private/DetailsCustomization/DisplayClusterRootActorDetailsCustomization.cpp b/Engine/Plugins/Runtime/nDisplay/Source/DisplayClusterEditor/Private/DetailsCustomization/DisplayClusterRootActorDetailsCustomization.cpp
--- a/Engine/Plugins/Runtime/nDisplay/Source/DisplayClusterEditor/Private/DetailsCustomization/DisplayClusterRootActorDetailsCustomization.cpp
+++ b/Engine/Plugins/Runtime/nDisplay/Source/DisplayClusterEditor/Private/DetailsCustomization/DisplayClusterRootActorDetailsCustomization.cpp
@@ -94,6 +94,7 @@ void FDisplayClusterRootActorDetailsCustomization::BuildLayout(IDetailLayoutBuil
 	InLayoutBuilder.EditCategory(DisplayClusterConfigurationStrings::categories::DefaultCategory);
 	InLayoutBuilder.EditCategory(DisplayClusterConfigurationStrings::categories::ConfigurationCategory);
 	InLayoutBuilder.EditCategory(DisplayClusterConfigurationStrings::categories::ClusterCategory);
+	InLayoutBuilder.EditCategory(DisplayClusterConfigurationStrings::categories::ClusterPostprocessCategory);
 	InLayoutBuilder.EditCategory(DisplayClusterConfigurationStrings::categories::PreviewCategory);
 
 	// Add custom properties and lay out/order properties into their correct categories.
