nD Remote Control.
clean delegates

diff --git a/Engine/Plugins/Runtime/nDisplay/Source/DisplayClusterRemoteControlInterceptor/Private/DisplayClusterRemoteControlInterceptorModule.cpp b/Engine/Plugins/Runtime/nDisplay/Source/DisplayClusterRemoteControlInterceptor/Private/DisplayClusterRemoteControlInterceptorModule.cpp
--- a/Engine/Plugins/Runtime/nDisplay/Source/DisplayClusterRemoteControlInterceptor/Private/DisplayClusterRemoteControlInterceptorModule.cpp
+++ b/Engine/Plugins/Runtime/nDisplay/Source/DisplayClusterRemoteControlInterceptor/Private/DisplayClusterRemoteControlInterceptorModule.cpp
@@ -29,6 +29,9 @@ void FDisplayClusterRemoteControlInterceptorModule::StartupModule()
 
 void FDisplayClusterRemoteControlInterceptorModule::ShutdownModule()
 {
+	// Clean delegates
+	FCoreDelegates::OnFEngineLoopInitComplete.RemoveAll(this);
+
 	// Add Interceptor only in Cluster mode
 	if (IDisplayCluster::Get().GetOperationMode() == EDisplayClusterOperationMode::Cluster)
 	{
