RC and nD replication queue

diff --git a/Engine/Plugins/Runtime/nDisplay/Source/DisplayClusterRemoteControlInterceptor/Private/DisplayClusterRemoteControlInterceptor.cpp b/Engine/Plugins/Runtime/nDisplay/Source/DisplayClusterRemoteControlInterceptor/Private/DisplayClusterRemoteControlInterceptor.cpp
--- a/Engine/Plugins/Runtime/nDisplay/Source/DisplayClusterRemoteControlInterceptor/Private/DisplayClusterRemoteControlInterceptor.cpp
+++ b/Engine/Plugins/Runtime/nDisplay/Source/DisplayClusterRemoteControlInterceptor/Private/DisplayClusterRemoteControlInterceptor.cpp
@@ -6,17 +6,14 @@
 #include "IDisplayCluster.h"
 #include "Cluster/DisplayClusterClusterEvent.h"
 
-#include "Backends/CborStructDeserializerBackend.h"
-#include "Backends/JsonStructDeserializerBackend.h"
 #include "Serialization/MemoryWriter.h"
 #include "Serialization/MemoryReader.h"
-#include "UObject/FieldPath.h"
 
 #include "Features/IModularFeatures.h"
 #include "HAL/IConsoleManager.h"
+#include "Misc/CoreDelegates.h"
 #include "Misc/CommandLine.h"
 
-
 // Data interception source
 static TAutoConsoleVariable<int32> CVarInterceptOnMasterOnly(
 	TEXT("nDisplay.RemoteControlInterceptor.MasterOnly"),
@@ -30,10 +27,7 @@ static TAutoConsoleVariable<int32> CVarInterceptOnMasterOnly(
 
 // Magic numbers for now. Unfortunately there is no any ID management for binary events yet.
 // This will be refactored once we have some global events registry to prevent any ID conflicts.
-const int32 EventId_SetObjectProperties		= 0xaabb0701;
-const int32 EventId_ResetObjectProperties	= 0xaabb0702;
-const int32 EventId_InvokeCall				= 0xaabb0703;
-
+const int32 EventId_InterceptorQueue		= 0xaabb0704;
 
 FDisplayClusterRemoteControlInterceptor::FDisplayClusterRemoteControlInterceptor()
 	: bInterceptOnMasterOnly(CVarInterceptOnMasterOnly.GetValueOnGameThread() == 1)
@@ -45,6 +39,8 @@ FDisplayClusterRemoteControlInterceptor::FDisplayClusterRemoteControlInterceptor
 	EventsListener.BindRaw(this, &FDisplayClusterRemoteControlInterceptor::OnClusterEventBinaryHandler);
 	// Subscribe for cluster events
 	IDisplayCluster::Get().GetClusterMgr()->AddClusterEventBinaryListener(EventsListener);
+	// Send the queue of replication events at the end of the tick
+	FCoreDelegates::OnEndFrame.AddRaw(this, &FDisplayClusterRemoteControlInterceptor::SendReplicationQueue);
 
 	UE_LOG(LogDisplayClusterRemoteControlInterceptor, Log, TEXT("DisplayClusterRemoteControlInterceptor has been registered"));
 }
@@ -55,6 +51,8 @@ FDisplayClusterRemoteControlInterceptor::~FDisplayClusterRemoteControlIntercepto
 	IDisplayCluster::Get().GetClusterMgr()->RemoveClusterEventBinaryListener(EventsListener);
 	// Unbind cluster events handler
 	EventsListener.Unbind();
+	// Remove delegates from end of the engine frame
+	FCoreDelegates::OnEndFrame.RemoveAll(this);
 
 	UE_LOG(LogDisplayClusterRemoteControlInterceptor, Log, TEXT("DisplayClusterRemoteControlInterceptor has been unregistered"));
 }
@@ -66,9 +64,8 @@ ERCIResponse FDisplayClusterRemoteControlInterceptor::SetObjectProperties(FRCIPr
 	FMemoryWriter MemoryWriter(Buffer);
 	MemoryWriter << InProperties;
 
-	// Replicate data
-	static const FString EventName = FString(TEXT("SetObjectProperties"));
-	EmitReplicationEvent(EventId_SetObjectProperties, Buffer, EventName);
+	// Queue interception sending
+	QueueInterceptEvent(FRCIPropertiesMetadata::Name, InProperties.GetUniquePath(), MoveTemp(Buffer));
 
 	if (bForceApply)
 	{
@@ -82,13 +79,12 @@ ERCIResponse FDisplayClusterRemoteControlInterceptor::ResetObjectProperties(FRCI
 {
 	// Serialize command data to binary
 	TArray<uint8> Buffer;
-	FMemoryWriter MemoryWiter(Buffer);
-	MemoryWiter << InObject;
-
-	// Replicate data
-	static const FString EventName = FString(TEXT("ResetObjectProperties"));
-	EmitReplicationEvent(EventId_ResetObjectProperties, Buffer, EventName);
+	FMemoryWriter MemoryWriter(Buffer);
+	MemoryWriter << InObject;
 
+	// Queue interception sending
+	QueueInterceptEvent(FRCIObjectMetadata::Name, InObject.GetUniquePath(), MoveTemp(Buffer));
+	
 	if (bForceApply)
 	{
 		return ERCIResponse::Apply;
@@ -101,12 +97,11 @@ ERCIResponse FDisplayClusterRemoteControlInterceptor::InvokeCall(FRCIFunctionMet
 {
 	// Serialize command data to binary
 	TArray<uint8> Buffer;
-	FMemoryWriter MemoryWiter(Buffer);
-	MemoryWiter << InFunction;
+	FMemoryWriter MemoryWriter(Buffer);
+	MemoryWriter << InFunction;
 
-	// Replicate data
-	static const FString EventName = FString(TEXT("InvokeCall"));
-	EmitReplicationEvent(EventId_InvokeCall, Buffer, EventName);
+	// Queue interception sending
+	QueueInterceptEvent(FRCIFunctionMetadata::Name, InFunction.GetUniquePath(), MoveTemp(Buffer));
 
 	if (bForceApply)
 	{
@@ -116,52 +111,53 @@ ERCIResponse FDisplayClusterRemoteControlInterceptor::InvokeCall(FRCIFunctionMet
 	return ERCIResponse::Intercept;
 }
 
-void FDisplayClusterRemoteControlInterceptor::EmitReplicationEvent(int32 EventId, TArray<uint8>& Buffer, const FString& EventName)
-{
-	UE_LOG(LogDisplayClusterRemoteControlInterceptor, VeryVerbose, TEXT("Sending replication event %s (0x%x): %d bytes"), *EventName, EventId, Buffer.Num());
-
-	// Cluster event instance
-	FDisplayClusterClusterEventBinary Event;
-
-	// Fill the event with data
-	Event.EventId                = EventId;
-	Event.bIsSystemEvent         = true;
-	Event.bShouldDiscardOnRepeat = false;
-	Event.EventData              = MoveTemp(Buffer);
-
-	// Emit cluster event (or not, it depends on the bInterceptOnMasterOnly value and the role of this cluster node)
-	IDisplayCluster::Get().GetClusterMgr()->EmitClusterEventBinary(Event, bInterceptOnMasterOnly);
-}
-
 void FDisplayClusterRemoteControlInterceptor::OnClusterEventBinaryHandler(const FDisplayClusterClusterEventBinary& Event)
 {
 	// Dispatch data to a proper handler
 	if (Event.bIsSystemEvent)
 	{
-		switch (Event.EventId)
-		{
-		case EventId_SetObjectProperties:
-			OnReplication_SetObjectProperties(Event.EventData);
-			break;
+		// Deserialize command data
+		TMap<FName, TMap<FString, TArray<uint8>>> ReceivedInterceptQueueMap;
+		FMemoryReader MemoryReader(Event.EventData);
+		MemoryReader << ReceivedInterceptQueueMap;
 
-		case EventId_ResetObjectProperties:
-			OnReplication_ResetObjectProperties(Event.EventData);
-			break;
-
-		case EventId_InvokeCall:
-			OnReplication_InvokeCall(Event.EventData);
-			break;
-
-		default:
-			// Unsupported event, skipping it
-			break;
+		for (const TPair<FName, TMap<FString, TArray<uint8>>>& InterceptorTypeMapPair : ReceivedInterceptQueueMap)
+		{
+			const FName MetadataName = InterceptorTypeMapPair.Key;
+			const TMap<FString, TArray<uint8>>& EventPayloadMap = InterceptorTypeMapPair.Value;
+
+			for (const TPair<FString, TArray<uint8>>& EventPayloadPair : EventPayloadMap)
+			{
+				const TArray<uint8>& EventPayload = EventPayloadPair.Value;
+				
+				if (MetadataName == FRCIPropertiesMetadata::Name)
+				{
+					OnReplication_SetObjectProperties(EventPayload);
+				}
+				else if (MetadataName == FRCIObjectMetadata::Name)
+				{
+					OnReplication_ResetObjectProperties(EventPayload);
+				}
+				else if (MetadataName == FRCIFunctionMetadata::Name)
+				{
+					OnReplication_InvokeCall(EventPayload);
+				}
+			}
 		}
 	}
 }
 
+void FDisplayClusterRemoteControlInterceptor::QueueInterceptEvent(const FName& InterceptEventType, const FName& InUniquePath, TArray<uint8>&& InBuffer)
+{
+	check(IsInGameThread());
+	
+	TMap<FName, TArray<uint8>>& InterceptorTypeMap = InterceptQueueMap.FindOrAdd(InterceptEventType);
+	InterceptorTypeMap.Add(InUniquePath, MoveTemp(InBuffer));
+}
+
 void FDisplayClusterRemoteControlInterceptor::OnReplication_SetObjectProperties(const TArray<uint8>& Buffer)
 {
-	UE_LOG(LogDisplayClusterRemoteControlInterceptor, VeryVerbose, TEXT("Processing replication event SetObjectProperties (0x%d): %d bytes"), EventId_SetObjectProperties, Buffer.Num());
+	UE_LOG(LogDisplayClusterRemoteControlInterceptor, VeryVerbose, TEXT("Processing replication event SetObjectProperties: %d bytes"), Buffer.Num());
 
 	// Deserialize command data
 	FMemoryReader MemoryReader(Buffer);
@@ -186,7 +182,7 @@ void FDisplayClusterRemoteControlInterceptor::OnReplication_SetObjectProperties(
 
 void FDisplayClusterRemoteControlInterceptor::OnReplication_ResetObjectProperties(const TArray<uint8>& Buffer)
 {
-	UE_LOG(LogDisplayClusterRemoteControlInterceptor, VeryVerbose, TEXT("Processing replication event ResetObjectProperties (0x%d): %d bytes"), EventId_ResetObjectProperties, Buffer.Num());
+	UE_LOG(LogDisplayClusterRemoteControlInterceptor, VeryVerbose, TEXT("Processing replication event ResetObjectProperties: %d bytes"), Buffer.Num());
 
 	// Deserialize command data
 	FMemoryReader MemoryReader(Buffer);
@@ -211,7 +207,7 @@ void FDisplayClusterRemoteControlInterceptor::OnReplication_ResetObjectPropertie
 
 void FDisplayClusterRemoteControlInterceptor::OnReplication_InvokeCall(const TArray<uint8>& Buffer)
 {
-	UE_LOG(LogDisplayClusterRemoteControlInterceptor, VeryVerbose, TEXT("Processing replication event InvokeCall (0x%d): %d bytes"), EventId_InvokeCall, Buffer.Num());
+	UE_LOG(LogDisplayClusterRemoteControlInterceptor, VeryVerbose, TEXT("Processing replication event InvokeCall: %d bytes"), Buffer.Num());
 
 	// Deserialize command data
 	FMemoryReader MemoryReader(Buffer);
@@ -233,3 +229,30 @@ void FDisplayClusterRemoteControlInterceptor::OnReplication_InvokeCall(const TAr
 		}
 	}
 }
+
+void FDisplayClusterRemoteControlInterceptor::SendReplicationQueue()
+{
+	if (!InterceptQueueMap.Num())
+	{
+		return;
+	}
+	
+	// Serialize the InterceptQueueMap into the buffer
+	TArray<uint8> Buffer;
+	FMemoryWriter MemoryWriter(Buffer);
+	MemoryWriter << InterceptQueueMap;
+
+	// Cluster event instance
+	FDisplayClusterClusterEventBinary Event;
+
+	// Fill the event with data
+	Event.EventId                = EventId_InterceptorQueue;
+	Event.bIsSystemEvent         = true;
+	Event.bShouldDiscardOnRepeat = false;
+	Event.EventData              = MoveTemp(Buffer);
+
+	// Emit cluster event (or not, it depends on the bInterceptOnMasterOnly value and the role of this cluster node)
+	IDisplayCluster::Get().GetClusterMgr()->EmitClusterEventBinary(Event, bInterceptOnMasterOnly);
+
+	InterceptQueueMap.Empty();
+}
diff --git a/Engine/Plugins/Runtime/nDisplay/Source/DisplayClusterRemoteControlInterceptor/Private/DisplayClusterRemoteControlInterceptor.h b/Engine/Plugins/Runtime/nDisplay/Source/DisplayClusterRemoteControlInterceptor/Private/DisplayClusterRemoteControlInterceptor.h
--- a/Engine/Plugins/Runtime/nDisplay/Source/DisplayClusterRemoteControlInterceptor/Private/DisplayClusterRemoteControlInterceptor.h
+++ b/Engine/Plugins/Runtime/nDisplay/Source/DisplayClusterRemoteControlInterceptor/Private/DisplayClusterRemoteControlInterceptor.h
@@ -26,11 +26,12 @@ public:
 	// ~IRemoteControlInterceptionCommands interface
 
 private:
-	// Helper function to perform events emission
-	void EmitReplicationEvent(int32 EventId, TArray<uint8>& Buffer, const FString& EventName);
 	// Cluster events handler/dispatcher
 	void OnClusterEventBinaryHandler(const FDisplayClusterClusterEventBinary& Event);
 
+	// Queue unique intercept event for sending on next engine tick
+	void QueueInterceptEvent(const FName& InterceptEventType, const FName& InUniquePath, TArray<uint8>&& InBuffer);
+
 private:
 	// Process SetObjectProperties command replication data
 	void OnReplication_SetObjectProperties   (const TArray<uint8>& Buffer);
@@ -38,6 +39,8 @@ private:
 	void OnReplication_ResetObjectProperties (const TArray<uint8>& Buffer);
 	// Process InvokeCall command replication data
 	void OnReplication_InvokeCall (const TArray<uint8>& Buffer);
+	// Send the queue of replication events at the end of the tick
+	void SendReplicationQueue();
 
 private:
 	// CVar value MasterOnly
@@ -46,4 +49,11 @@ private:
 	FOnClusterEventBinaryListener EventsListener;
 	// Force Apply the ERCIResponse
 	bool bForceApply;
+
+	/**
+	 * FName name of the metadata struct
+	 * FName unique object path + field path
+	 * TArray<uint8> Metadata struct buffer
+	 */
+	TMap<FName, TMap<FName, TArray<uint8>>> InterceptQueueMap;
 };
diff --git a/Engine/Plugins/VirtualProduction/RemoteControlInterception/Source/RemoteControlInterception/Public/IRemoteControlInterceptionCommands.h b/Engine/Plugins/VirtualProduction/RemoteControlInterception/Source/RemoteControlInterception/Public/IRemoteControlInterceptionCommands.h
--- a/Engine/Plugins/VirtualProduction/RemoteControlInterception/Source/RemoteControlInterception/Public/IRemoteControlInterceptionCommands.h
+++ b/Engine/Plugins/VirtualProduction/RemoteControlInterception/Source/RemoteControlInterception/Public/IRemoteControlInterceptionCommands.h
@@ -94,11 +94,17 @@ public:
 		, PropertyPath(InPropertyPath)
 		, PropertyPathInfo(InPropertyPathInfo)
 		, Access(InAccess)
-	{ }
+	{
+		UniquePath = *FString::Printf(TEXT("%s:%s"), *ObjectPath, *PropertyPath);
+	}
 
 	virtual ~FRCIObjectMetadata() = default;
 
 public:
+	/** Get Unique Path for Object metadata */
+	const FName& GetUniquePath() const
+	{ return UniquePath; }
+	
 	/** Returns true if reference is valid */
 	virtual bool IsValid() const
 	{
@@ -134,6 +140,14 @@ public:
 
 	/** Property access type */
 	ERCIAccess Access = ERCIAccess::NO_ACCESS;
+
+public:
+	/** Structure Name */
+	static constexpr TCHAR const* Name = TEXT("RCIObjectMetadata");
+	
+private:
+	/** Object Path + Property Path */
+    FName UniquePath;
 };
 
 
@@ -187,6 +201,10 @@ public:
 
 	/** Property payload to intercept */
 	const TArray<uint8> Payload;
+
+public:
+	/** Structure Name */
+	static constexpr TCHAR const* Name = TEXT("RCIPropertiesMetadata");
 };
 
 /**
@@ -206,9 +224,15 @@ public:
 		, bGenerateTransaction(bInGenerateTransaction)
 		, PayloadType(InPayloadType)
 		, Payload(InPayload)
-	{ }
+	{
+		UniquePath = *FString::Printf(TEXT("%s:%s"), *ObjectPath, *FunctionPath);
+	}
 
 public:
+	/** Get Unique Path for Function metadata */
+	const FName& GetUniquePath() const
+	{ return UniquePath; }
+
 	/** Returns true if reference is valid */
 	bool IsValid() const
 	{
@@ -249,4 +273,12 @@ public:
 
 	/** Property payload to intercept */
 	const TArray<uint8> Payload;
+
+public:
+	/** Structure Name */
+	static constexpr TCHAR const* Name = TEXT("RCIFunctionMetadata");
+
+private:
+	/** Object Path + Function Path */
+	FName UniquePath;
 };
