nDisplay: Added UI enforced minimum and maximum sizes for viewports.

diff --git a/Engine/Plugins/Runtime/nDisplay/Source/DisplayClusterConfiguration/Private/DisplayClusterConfigurationTypes.cpp b/Engine/Plugins/Runtime/nDisplay/Source/DisplayClusterConfiguration/Private/DisplayClusterConfigurationTypes.cpp
--- a/Engine/Plugins/Runtime/nDisplay/Source/DisplayClusterConfiguration/Private/DisplayClusterConfigurationTypes.cpp
+++ b/Engine/Plugins/Runtime/nDisplay/Source/DisplayClusterConfiguration/Private/DisplayClusterConfigurationTypes.cpp
@@ -136,6 +136,9 @@ FDisplayClusterConfigurationICVFX_StageSettings::FDisplayClusterConfigurationICV
 {
 }
 
+const float UDisplayClusterConfigurationViewport::ViewportMinimumSize = 1.0f;
+const float UDisplayClusterConfigurationViewport::ViewportMaximumSize = 15360.0f;
+
 UDisplayClusterConfigurationViewport::UDisplayClusterConfigurationViewport()
 {
 #if WITH_EDITORONLY_DATA
diff --git a/Engine/Plugins/Runtime/nDisplay/Source/DisplayClusterConfiguration/Public/DisplayClusterConfigurationTypes_Viewport.h b/Engine/Plugins/Runtime/nDisplay/Source/DisplayClusterConfiguration/Public/DisplayClusterConfigurationTypes_Viewport.h
--- a/Engine/Plugins/Runtime/nDisplay/Source/DisplayClusterConfiguration/Public/DisplayClusterConfigurationTypes_Viewport.h
+++ b/Engine/Plugins/Runtime/nDisplay/Source/DisplayClusterConfiguration/Public/DisplayClusterConfigurationTypes_Viewport.h
@@ -343,6 +343,10 @@ public:
 	UPROPERTY(EditDefaultsOnly, Category = "NDisplay", meta = (nDisplayHidden))
 	bool bIsVisible = true;
 #endif
+
+public:
+	static const float ViewportMinimumSize;
+	static const float ViewportMaximumSize;
 };
 
 // This struct now stored in UDisplayClusterConfigurationData, and replicated with MultiUser
diff --git a/Engine/Plugins/Runtime/nDisplay/Source/DisplayClusterConfigurator/Private/Views/Details/DisplayClusterConfiguratorDetailCustomization.cpp b/Engine/Plugins/Runtime/nDisplay/Source/DisplayClusterConfigurator/Private/Views/Details/DisplayClusterConfiguratorDetailCustomization.cpp
--- a/Engine/Plugins/Runtime/nDisplay/Source/DisplayClusterConfigurator/Private/Views/Details/DisplayClusterConfiguratorDetailCustomization.cpp
+++ b/Engine/Plugins/Runtime/nDisplay/Source/DisplayClusterConfigurator/Private/Views/Details/DisplayClusterConfiguratorDetailCustomization.cpp
@@ -273,6 +273,24 @@ void FDisplayClusterConfiguratorViewportDetailCustomization::CustomizeDetails(ID
 	BEGIN_CATEGORY(CameraHandle->GetDefaultCategoryName())
 		REPLACE_PROPERTY_WITH_CUSTOM(UDisplayClusterConfigurationViewport, Camera, CreateCustomCameraWidget());
 	END_CATEGORY()
+
+	FDisplayClusterConfiguratorNestedPropertyHelper NestedPropertyHelper(InLayoutBuilder);
+
+	// Update the metadata for the viewport's width and height. Must set this here instead of in the UPROPERTY specifier because
+	// the Region property is a generic FDisplayClusterConfigurationRectangle struct which is used in lots of places, most of
+	// which don't make sense to have a minimum or maximum limit
+	TSharedPtr<IPropertyHandle> WidthHandle = NestedPropertyHelper.GetNestedProperty(TEXT("Region.W"));
+	TSharedPtr<IPropertyHandle> HeightHandle = NestedPropertyHelper.GetNestedProperty(TEXT("Region.H"));
+
+	WidthHandle->SetInstanceMetaData(TEXT("ClampMin"), FString::SanitizeFloat(UDisplayClusterConfigurationViewport::ViewportMinimumSize));
+	WidthHandle->SetInstanceMetaData(TEXT("UIMin"), FString::SanitizeFloat(UDisplayClusterConfigurationViewport::ViewportMinimumSize));
+	WidthHandle->SetInstanceMetaData(TEXT("ClampMax"), FString::SanitizeFloat(UDisplayClusterConfigurationViewport::ViewportMaximumSize));
+	WidthHandle->SetInstanceMetaData(TEXT("UIMax"), FString::SanitizeFloat(UDisplayClusterConfigurationViewport::ViewportMaximumSize));
+
+	HeightHandle->SetInstanceMetaData(TEXT("ClampMin"), FString::SanitizeFloat(UDisplayClusterConfigurationViewport::ViewportMinimumSize));
+	HeightHandle->SetInstanceMetaData(TEXT("UIMin"), FString::SanitizeFloat(UDisplayClusterConfigurationViewport::ViewportMinimumSize));
+	HeightHandle->SetInstanceMetaData(TEXT("ClampMax"), FString::SanitizeFloat(UDisplayClusterConfigurationViewport::ViewportMaximumSize));
+	HeightHandle->SetInstanceMetaData(TEXT("UIMax"), FString::SanitizeFloat(UDisplayClusterConfigurationViewport::ViewportMaximumSize));
 }
 
 void FDisplayClusterConfiguratorViewportDetailCustomization::ResetCameraOptions()
diff --git a/Engine/Plugins/Runtime/nDisplay/Source/DisplayClusterConfigurator/Private/Views/OutputMapping/GraphNodes/SDisplayClusterConfiguratorBaseNode.cpp b/Engine/Plugins/Runtime/nDisplay/Source/DisplayClusterConfigurator/Private/Views/OutputMapping/GraphNodes/SDisplayClusterConfiguratorBaseNode.cpp
--- a/Engine/Plugins/Runtime/nDisplay/Source/DisplayClusterConfigurator/Private/Views/OutputMapping/GraphNodes/SDisplayClusterConfiguratorBaseNode.cpp
+++ b/Engine/Plugins/Runtime/nDisplay/Source/DisplayClusterConfigurator/Private/Views/OutputMapping/GraphNodes/SDisplayClusterConfiguratorBaseNode.cpp
@@ -141,6 +141,8 @@ void SDisplayClusterConfiguratorBaseNode::UpdateGraphNode()
 			.Visibility(this, &SDisplayClusterConfiguratorBaseNode::GetResizeHandleVisibility)
 			[
 				SNew(SDisplayClusterConfiguratorResizer, ToolkitPtr.Pin().ToSharedRef(), SharedThis(this))
+				.MinimumSize(this, &SDisplayClusterConfiguratorBaseNode::GetNodeMinimumSize)
+				.MaximumSize(this, &SDisplayClusterConfiguratorBaseNode::GetNodeMaximumSize)
 				.IsFixedAspectRatio(this, &SDisplayClusterConfiguratorBaseNode::IsAspectRatioFixed)
 			]
 		];
diff --git a/Engine/Plugins/Runtime/nDisplay/Source/DisplayClusterConfigurator/Private/Views/OutputMapping/GraphNodes/SDisplayClusterConfiguratorBaseNode.h b/Engine/Plugins/Runtime/nDisplay/Source/DisplayClusterConfigurator/Private/Views/OutputMapping/GraphNodes/SDisplayClusterConfiguratorBaseNode.h
--- a/Engine/Plugins/Runtime/nDisplay/Source/DisplayClusterConfigurator/Private/Views/OutputMapping/GraphNodes/SDisplayClusterConfiguratorBaseNode.h
+++ b/Engine/Plugins/Runtime/nDisplay/Source/DisplayClusterConfigurator/Private/Views/OutputMapping/GraphNodes/SDisplayClusterConfiguratorBaseNode.h
@@ -117,6 +117,12 @@ public:
 	/** @return Whether this node can be resized using the resize widget. */
 	virtual bool CanNodeBeResized() const { return true; }
 
+	/** @return The minimum size this node can be resized to. */
+	virtual float GetNodeMinimumSize() const { return 0; }
+
+	/** @return The maximum size this node can be resized to. */
+	virtual float GetNodeMaximumSize() const { return FLT_MAX; }
+
 	/** @return Whether this node's size is fixed to a specific aspect ratio. */
 	virtual bool IsAspectRatioFixed() const { return false; }
 
diff --git a/Engine/Plugins/Runtime/nDisplay/Source/DisplayClusterConfigurator/Private/Views/OutputMapping/GraphNodes/SDisplayClusterConfiguratorViewportNode.cpp b/Engine/Plugins/Runtime/nDisplay/Source/DisplayClusterConfigurator/Private/Views/OutputMapping/GraphNodes/SDisplayClusterConfiguratorViewportNode.cpp
--- a/Engine/Plugins/Runtime/nDisplay/Source/DisplayClusterConfigurator/Private/Views/OutputMapping/GraphNodes/SDisplayClusterConfiguratorViewportNode.cpp
+++ b/Engine/Plugins/Runtime/nDisplay/Source/DisplayClusterConfigurator/Private/Views/OutputMapping/GraphNodes/SDisplayClusterConfiguratorViewportNode.cpp
@@ -175,6 +175,16 @@ bool SDisplayClusterConfiguratorViewportNode::IsNodeVisible() const
 	return SDisplayClusterConfiguratorBaseNode::IsNodeVisible() && bIsVisible;
 }
 
+float SDisplayClusterConfiguratorViewportNode::GetNodeMinimumSize() const
+{
+	return UDisplayClusterConfigurationViewport::ViewportMinimumSize;
+}
+
+float SDisplayClusterConfiguratorViewportNode::GetNodeMaximumSize() const
+{
+	return UDisplayClusterConfigurationViewport::ViewportMaximumSize;
+}
+
 bool SDisplayClusterConfiguratorViewportNode::IsAspectRatioFixed() const
 {
 	UDisplayClusterConfiguratorViewportNode* ViewportEdNode = GetGraphNodeChecked<UDisplayClusterConfiguratorViewportNode>();
diff --git a/Engine/Plugins/Runtime/nDisplay/Source/DisplayClusterConfigurator/Private/Views/OutputMapping/GraphNodes/SDisplayClusterConfiguratorViewportNode.h b/Engine/Plugins/Runtime/nDisplay/Source/DisplayClusterConfigurator/Private/Views/OutputMapping/GraphNodes/SDisplayClusterConfiguratorViewportNode.h
--- a/Engine/Plugins/Runtime/nDisplay/Source/DisplayClusterConfigurator/Private/Views/OutputMapping/GraphNodes/SDisplayClusterConfiguratorViewportNode.h
+++ b/Engine/Plugins/Runtime/nDisplay/Source/DisplayClusterConfigurator/Private/Views/OutputMapping/GraphNodes/SDisplayClusterConfiguratorViewportNode.h
@@ -37,6 +37,8 @@ public:
 	virtual bool IsNodeVisible() const override;
 	virtual bool CanNodeBeSnapAligned() const override { return true; }
 	virtual bool CanNodeBeResized() const { return !IsViewportLocked(); }
+	virtual float GetNodeMinimumSize() const override;
+	virtual float GetNodeMaximumSize() const override;
 	virtual bool IsAspectRatioFixed() const override;
 	//~ End of SDisplayClusterConfiguratorBaseNode interface
 
diff --git a/Engine/Plugins/Runtime/nDisplay/Source/DisplayClusterConfigurator/Private/Views/OutputMapping/Widgets/SDisplayClusterConfiguratorResizer.cpp b/Engine/Plugins/Runtime/nDisplay/Source/DisplayClusterConfigurator/Private/Views/OutputMapping/Widgets/SDisplayClusterConfiguratorResizer.cpp
--- a/Engine/Plugins/Runtime/nDisplay/Source/DisplayClusterConfigurator/Private/Views/OutputMapping/Widgets/SDisplayClusterConfiguratorResizer.cpp
+++ b/Engine/Plugins/Runtime/nDisplay/Source/DisplayClusterConfigurator/Private/Views/OutputMapping/Widgets/SDisplayClusterConfiguratorResizer.cpp
@@ -20,6 +20,8 @@ void SDisplayClusterConfiguratorResizer::Construct(const FArguments& InArgs, con
 	CurrentAspectRatio = 1;
 	bResizing = false;
 
+	MinimumSize = InArgs._MinimumSize;
+	MaximumSize = InArgs._MaximumSize;
 	IsFixedAspectRatio = InArgs._IsFixedAspectRatio;
 
 	ScopedTransaction.Reset();
@@ -112,15 +114,12 @@ FReply SDisplayClusterConfiguratorResizer::OnMouseMove(const FGeometry& MyGeomet
 
 		NewNodeSize /= (GraphPanel->GetZoomAmount() * DPIScale);
 
-		// Never node size less then 0
-		if (NewNodeSize.X < 0.f)
-		{
-			NewNodeSize.X = 0.f;
-		}
-		if (NewNodeSize.Y < 0.f)
-		{
-			NewNodeSize.Y = 0.f;
-		}
+		// Bound the node's size between the minimum and maximum. Also never let the size become negative.
+		float Minimum = FMath::Max(MinimumSize.Get(0.0f), 0.0f);
+		float Maximum = FMath::Max(MaximumSize.Get(FLT_MAX), Minimum);
+		
+		NewNodeSize.X = FMath::Clamp(NewNodeSize.X, Minimum, Maximum);
+		NewNodeSize.Y = FMath::Clamp(NewNodeSize.Y, Minimum, Maximum);
 
 		// If we don't have a scoped transaction for the resize, create a new one.
 		if (!ScopedTransaction.IsValid())
diff --git a/Engine/Plugins/Runtime/nDisplay/Source/DisplayClusterConfigurator/Private/Views/OutputMapping/Widgets/SDisplayClusterConfiguratorResizer.h b/Engine/Plugins/Runtime/nDisplay/Source/DisplayClusterConfigurator/Private/Views/OutputMapping/Widgets/SDisplayClusterConfiguratorResizer.h
--- a/Engine/Plugins/Runtime/nDisplay/Source/DisplayClusterConfigurator/Private/Views/OutputMapping/Widgets/SDisplayClusterConfiguratorResizer.h
+++ b/Engine/Plugins/Runtime/nDisplay/Source/DisplayClusterConfigurator/Private/Views/OutputMapping/Widgets/SDisplayClusterConfiguratorResizer.h
@@ -17,7 +17,12 @@ class SDisplayClusterConfiguratorResizer
 public:
 
 	SLATE_BEGIN_ARGS(SDisplayClusterConfiguratorResizer)
-		: _IsFixedAspectRatio(false) {}
+		: _MinimumSize(0)
+		, _MaximumSize(FLT_MAX)
+		, _IsFixedAspectRatio(false)
+	{}
+		SLATE_ATTRIBUTE(float, MinimumSize)
+		SLATE_ATTRIBUTE(float, MaximumSize)
 		SLATE_ATTRIBUTE(bool, IsFixedAspectRatio)
 	SLATE_END_ARGS()
 
@@ -40,6 +45,8 @@ private:
 	bool bResizing;
 	float CurrentAspectRatio;
 
+	TAttribute<float> MinimumSize;
+	TAttribute<float> MaximumSize;
 	TAttribute<bool> IsFixedAspectRatio;
 };
 
