nDisplay: Details enhancements.
- ICFVX: Add global setting to disable Inner Frustum. Add global priority sorting.
- Add Chroma key group.
- Light Cards: Enable and Blending Mode
- Move Editor details to Configurator module.

diff --git a/Engine/Plugins/Runtime/nDisplay/Source/DisplayCluster/Private/Game/EngineClasses/Scene/DisplayClusterRootActorEditor.cpp b/Engine/Plugins/Runtime/nDisplay/Source/DisplayCluster/Private/Game/EngineClasses/Scene/DisplayClusterRootActorEditor.cpp
--- a/Engine/Plugins/Runtime/nDisplay/Source/DisplayCluster/Private/Game/EngineClasses/Scene/DisplayClusterRootActorEditor.cpp
+++ b/Engine/Plugins/Runtime/nDisplay/Source/DisplayCluster/Private/Game/EngineClasses/Scene/DisplayClusterRootActorEditor.cpp
@@ -10,6 +10,7 @@
 #include "Components/DisplayClusterXformComponent.h"
 #include "Components/DisplayClusterSceneComponentSyncParent.h"
 #include "Components/DisplayClusterPreviewComponent.h"
+#include "Components/DisplayClusterICVFXCameraComponent.h"
 
 #include "Config/IPDisplayClusterConfigManager.h"
 #include "DisplayClusterConfigurationStrings.h"
@@ -107,6 +108,8 @@ void ADisplayClusterRootActor::RerunConstructionScripts_Editor()
 {
 	/* We need to reinitialize since our components are being regenerated here. */
 	InitializeRootActor();
+
+	UpdateInnerFrustumPriority();
 }
 
 bool ADisplayClusterRootActor::IsPreviewEnabled() const
@@ -194,6 +197,69 @@ void ADisplayClusterRootActor::UpdateXformGizmos()
 	}
 }
 
+void ADisplayClusterRootActor::UpdateInnerFrustumPriority()
+{
+	if (InnerFrustumPriority.Num() == 0)
+	{
+		ResetInnerFrustumPriority();
+		return;
+	}
+	
+	TArray<UDisplayClusterICVFXCameraComponent*> Components;
+	GetComponents<UDisplayClusterICVFXCameraComponent>(Components);
+
+	TArray<FString> ValidCameras;
+	for (UDisplayClusterICVFXCameraComponent* Camera : Components)
+	{
+		FString CameraName = Camera->GetName();
+		InnerFrustumPriority.AddUnique(CameraName);
+		ValidCameras.Add(CameraName);
+	}
+
+	// Removes invalid cameras or duplicate cameras.
+	InnerFrustumPriority.RemoveAll([ValidCameras, this](const FDisplayClusterComponentRef& CameraRef)
+	{
+		return !ValidCameras.Contains(CameraRef.Name) || InnerFrustumPriority.FilterByPredicate([CameraRef](const FDisplayClusterComponentRef& CameraRefN2)
+		{
+			return CameraRef == CameraRefN2;
+		}).Num() > 1;
+	});
+	
+	for (int32 Idx = 0; Idx < InnerFrustumPriority.Num(); ++Idx)
+	{
+		if (UDisplayClusterICVFXCameraComponent* CameraComponent = FindObjectFast<UDisplayClusterICVFXCameraComponent>(this, *InnerFrustumPriority[Idx].Name))
+		{
+			CameraComponent->CameraSettings.RenderSettings.RenderOrder = Idx;
+		}
+	}
+}
+
+void ADisplayClusterRootActor::ResetInnerFrustumPriority()
+{
+	TArray<UDisplayClusterICVFXCameraComponent*> Components;
+	GetComponents<UDisplayClusterICVFXCameraComponent>(Components);
+
+	InnerFrustumPriority.Reset(Components.Num());
+	for (UDisplayClusterICVFXCameraComponent* Camera : Components)
+	{
+		InnerFrustumPriority.Add(Camera->GetName());
+	}
+	
+	// Initialize based on current render priority.
+	InnerFrustumPriority.Sort([this](const FDisplayClusterComponentRef& CameraA, const FDisplayClusterComponentRef& CameraB)
+	{
+		UDisplayClusterICVFXCameraComponent* CameraComponentA = FindObjectFast<UDisplayClusterICVFXCameraComponent>(this, *CameraA.Name);
+		UDisplayClusterICVFXCameraComponent* CameraComponentB = FindObjectFast<UDisplayClusterICVFXCameraComponent>(this, *CameraB.Name);
+
+		if (CameraComponentA && CameraComponentB)
+		{
+			return CameraComponentA->CameraSettings.RenderSettings.RenderOrder < CameraComponentB->CameraSettings.RenderSettings.RenderOrder;
+		}
+
+		return false;
+	});
+}
+
 bool ADisplayClusterRootActor::IsSelectedInEditor() const
 {
 	return bIsSelectedInEditor || Super::IsSelectedInEditor();
@@ -288,6 +354,10 @@ void ADisplayClusterRootActor::PostEditChangeProperty(FPropertyChangedEvent& Pro
 	{
 		UpdateXformGizmos();
 	}
+	else if (PropertyName == GET_MEMBER_NAME_CHECKED(ADisplayClusterRootActor, InnerFrustumPriority))
+	{
+		ResetInnerFrustumPriority();
+	}
 
 	if (bReinitializeActor)
 	{
diff --git a/Engine/Plugins/Runtime/nDisplay/Source/DisplayCluster/Private/Render/Viewport/Configuration/DisplayClusterViewportConfigurationHelpers.h b/Engine/Plugins/Runtime/nDisplay/Source/DisplayCluster/Private/Render/Viewport/Configuration/DisplayClusterViewportConfigurationHelpers.h
--- a/Engine/Plugins/Runtime/nDisplay/Source/DisplayCluster/Private/Render/Viewport/Configuration/DisplayClusterViewportConfigurationHelpers.h
+++ b/Engine/Plugins/Runtime/nDisplay/Source/DisplayCluster/Private/Render/Viewport/Configuration/DisplayClusterViewportConfigurationHelpers.h
@@ -489,7 +489,10 @@ public:
 		{
 			EDisplayClusterViewportICVFXFlags& TargetFlags = DstViewport.RenderSettingsICVFX.Flags;
 
-			if (InConfigurationViewport.ICVFX.bAllowICVFX)
+			ADisplayClusterRootActor* RootActor = InConfigurationViewport.GetTypedOuter<ADisplayClusterRootActor>();
+			check(RootActor);
+			
+			if (InConfigurationViewport.ICVFX.bAllowICVFX && RootActor->EnableInnerFrustum)
 			{
 				TargetFlags |= ViewportICVFX_Enable;
 
diff --git a/Engine/Plugins/Runtime/nDisplay/Source/DisplayCluster/Public/DisplayClusterRootActor.h b/Engine/Plugins/Runtime/nDisplay/Source/DisplayCluster/Public/DisplayClusterRootActor.h
--- a/Engine/Plugins/Runtime/nDisplay/Source/DisplayCluster/Public/DisplayClusterRootActor.h
+++ b/Engine/Plugins/Runtime/nDisplay/Source/DisplayCluster/Public/DisplayClusterRootActor.h
@@ -37,6 +37,7 @@ class UDisplayClusterXformComponent;
 class UDisplayClusterSyncTickComponent;
 class UDisplayClusterPreviewComponent;
 
+
 /**
  * VR root. This contains nDisplay VR hierarchy in the game.
  */
@@ -227,30 +228,30 @@ private:
 //////////////////////////////////////////////////////////////////////////////////////////////
 #if WITH_EDITORONLY_DATA 
 public:
+	// Allow preview render
+	UPROPERTY(EditAnywhere, Category = "Editor Preview")
+	bool bPreviewEnable = true;
+	
 	// Render single node preview or whole cluster
-	UPROPERTY(EditAnywhere, Category = "NDisplay Preview (Editor only)")
+	UPROPERTY(EditAnywhere, Category = "Editor Preview")
 	FString PreviewNodeId = DisplayClusterConfigurationStrings::gui::preview::PreviewNodeAll;
 
 	// Render mode for PIE
-	UPROPERTY(EditAnywhere, Category = "NDisplay Preview (Editor only)")
+	UPROPERTY(EditAnywhere, Category = "Editor Preview")
 	EDisplayClusterConfigurationRenderMode RenderMode = EDisplayClusterConfigurationRenderMode::Mono;
 
-	// Allow preview render
-	UPROPERTY(EditAnywhere, Category = "NDisplay Preview (Editor only)")
-	bool bPreviewEnable = true;
-
 	// Update preview texture period in tick
-	UPROPERTY(EditAnywhere, Category = "NDisplay Preview (Editor only)", meta = (ClampMin = "1", UIMin = "1", ClampMax = "200", UIMax = "200"))
+	UPROPERTY(EditAnywhere, Category = "Editor Preview", meta = (ClampMin = "1", UIMin = "1", ClampMax = "200", UIMax = "200"))
 	int TickPerFrame = 1;
 
 	// Preview texture size get from viewport, and scaled by this value
-	UPROPERTY(EditAnywhere, Category = "NDisplay Preview (Editor only)", meta = (ClampMin = "0.05", UIMin = "0.05", ClampMax = "1", UIMax = "1"))
+	UPROPERTY(EditAnywhere, Category = "Editor Preview", meta = (ClampMin = "0.05", UIMin = "0.05", ClampMax = "1", UIMax = "1"))
 	float PreviewRenderTargetRatioMult = 0.25;
 
-	UPROPERTY(EditAnywhere, Category = "NDisplay Preview (Editor only)")
+	UPROPERTY(EditAnywhere, Category = "Editor Preview")
 	float XformGizmoScale = 1.0f;
 
-	UPROPERTY(EditAnywhere, Category = "NDisplay Preview (Editor only)")
+	UPROPERTY(EditAnywhere, Category = "Editor Preview")
 	bool bAreXformGizmosVisible = true;
 
 private:
@@ -321,9 +322,19 @@ public:
 
 	void UpdateXformGizmos();
 
+	void UpdateInnerFrustumPriority();
+	void ResetInnerFrustumPriority();
+	
 	virtual bool IsSelectedInEditor() const override;
 	void SetIsSelectedInEditor(bool bValue);
 
+public:
+	UPROPERTY(EditInstanceOnly, Category = NDisplay)
+	bool EnableInnerFrustum = true;
+
+	UPROPERTY(EditInstanceOnly, EditFixedSize, Category = "In Camera ICVFX", meta = (TitleProperty = "Name"))
+	TArray<FDisplayClusterComponentRef> InnerFrustumPriority;
+		
 private:
 	/** The number of times to update the render target via deferred update. */
 	int32 PreviewRenderTargetUpdatesRequired = 0;
diff --git a/Engine/Plugins/Runtime/nDisplay/Source/DisplayCluster/Public/Misc/DisplayClusterObjectRef.h b/Engine/Plugins/Runtime/nDisplay/Source/DisplayCluster/Public/Misc/DisplayClusterObjectRef.h
--- a/Engine/Plugins/Runtime/nDisplay/Source/DisplayCluster/Public/Misc/DisplayClusterObjectRef.h
+++ b/Engine/Plugins/Runtime/nDisplay/Source/DisplayCluster/Public/Misc/DisplayClusterObjectRef.h
@@ -13,6 +13,8 @@
 #include "Editor.h"
 #endif
 
+#include "DisplayClusterObjectRef.generated.h"
+
 class DISPLAYCLUSTER_API FDisplayClusterActorRef
 {
 public:
@@ -233,3 +235,24 @@ private:
 	FName   ComponentName;
 	mutable TWeakObjectPtr<USceneComponent> ComponentPtr;
 };
+
+USTRUCT()
+struct DISPLAYCLUSTER_API FDisplayClusterComponentRef
+{
+	GENERATED_BODY()
+	
+	FDisplayClusterComponentRef() {}
+	
+	FDisplayClusterComponentRef(const FString& InName)
+	{
+		Name = InName;
+	}
+	
+	UPROPERTY(VisibleAnywhere, Category = Component)
+	FString Name;
+	
+	FORCEINLINE bool operator==(const FDisplayClusterComponentRef& OtherRef) const
+	{
+		return Name == OtherRef.Name;
+	}
+};
\ No newline at end of file
diff --git a/Engine/Plugins/Runtime/nDisplay/Source/DisplayClusterConfiguration/Public/DisplayClusterConfigurationStrings.h b/Engine/Plugins/Runtime/nDisplay/Source/DisplayClusterConfiguration/Public/DisplayClusterConfigurationStrings.h
--- a/Engine/Plugins/Runtime/nDisplay/Source/DisplayClusterConfiguration/Public/DisplayClusterConfigurationStrings.h
+++ b/Engine/Plugins/Runtime/nDisplay/Source/DisplayClusterConfiguration/Public/DisplayClusterConfigurationStrings.h
@@ -21,11 +21,14 @@ namespace DisplayClusterConfigurationStrings
 
 		static constexpr auto ClusterCategory              = TEXT("NDisplay Cluster");
 		static constexpr auto ClusterConfigurationCategory = TEXT("NDisplay Cluster Configuration");
-		static constexpr auto ClusterPostprocessCategory   = TEXT("NDisplay Cluster Postprocess");
-
-		static constexpr auto ICVFXCategory                = TEXT("NDisplay ICVFX");
+		static constexpr auto ClusterPostprocessCategory   = TEXT("Post Process");
+		static constexpr auto ChromaKeyCategory            = TEXT("Chromakey");
+		static constexpr auto LightcardCategory            = TEXT("Light Cards");
+		
+		static constexpr auto ICVFXCategory                = TEXT("In Camera ICVFX");
 		static constexpr auto ConfigurationCategory        = TEXT("NDisplay Configuration");
-		static constexpr auto PreviewCategory              = TEXT("NDisplay Preview (Editor only)");
+		static constexpr auto PreviewCategory              = TEXT("Editor Preview");
+		static constexpr auto AdvancedCategory              = TEXT("Advanced");
 	}
 
 	// Command line arguments
diff --git a/Engine/Plugins/Runtime/nDisplay/Source/DisplayClusterConfiguration/Public/DisplayClusterConfigurationTypes.h b/Engine/Plugins/Runtime/nDisplay/Source/DisplayClusterConfiguration/Public/DisplayClusterConfigurationTypes.h
--- a/Engine/Plugins/Runtime/nDisplay/Source/DisplayClusterConfiguration/Public/DisplayClusterConfigurationTypes.h
+++ b/Engine/Plugins/Runtime/nDisplay/Source/DisplayClusterConfiguration/Public/DisplayClusterConfigurationTypes.h
@@ -510,31 +510,31 @@ public:
 public:
 	FDisplayClusterConfigurationDataMetaInfo Meta;
 
-	UPROPERTY(EditAnywhere, Category = NDisplay)
+	UPROPERTY(EditAnywhere, Category = Advanced)
 	FDisplayClusterConfigurationInfo Info;
 
 	UPROPERTY()
 	UDisplayClusterConfigurationScene* Scene;
 
-	UPROPERTY(VisibleInstanceOnly, Instanced, Category = NDisplay, meta = (DisplayThumbnail = false, ShowInnerProperties))
+	UPROPERTY(VisibleInstanceOnly, Instanced, Category = Advanced, meta = (DisplayThumbnail = false, ShowInnerProperties))
 	UDisplayClusterConfigurationCluster* Cluster;
 
-	UPROPERTY(EditAnywhere, Category = NDisplay)
+	UPROPERTY(EditAnywhere, Category = Advanced)
 	TMap<FString, FString> CustomParameters;
 
-	UPROPERTY(EditAnywhere, Category = NDisplay)
+	UPROPERTY(EditAnywhere, Category = Advanced)
 	FDisplayClusterConfigurationDiagnostics Diagnostics;
 
 	UPROPERTY(EditAnywhere, Category = NDisplay)
 	FDisplayClusterConfigurationICVFX_StageSettings StageSettings;
 
-	UPROPERTY(EditAnywhere, Category = NDisplay)
+	UPROPERTY(EditAnywhere, Category = Advanced)
 	FDisplayClusterConfigurationRenderFrame RenderFrameSettings;
 
-	UPROPERTY(EditAnywhere, Category = NDisplay, meta = (DisplayName = "Follow Local Player Camera"))
+	UPROPERTY(EditAnywhere, Category = Advanced, meta = (DisplayName = "Follow Local Player Camera"))
 	bool bFollowLocalPlayerCamera = false;
 
-	UPROPERTY(EditAnywhere, Category = NDisplay, meta = (DisplayName = "Exit when ESC pressed"))
+	UPROPERTY(EditAnywhere, Category = Advanced, meta = (DisplayName = "Exit when ESC pressed"))
 	bool bExitOnEsc = true;
 
 	/** Create empty config data. */
diff --git a/Engine/Plugins/Runtime/nDisplay/Source/DisplayClusterConfiguration/Public/DisplayClusterConfigurationTypes_Enums.h b/Engine/Plugins/Runtime/nDisplay/Source/DisplayClusterConfiguration/Public/DisplayClusterConfigurationTypes_Enums.h
--- a/Engine/Plugins/Runtime/nDisplay/Source/DisplayClusterConfiguration/Public/DisplayClusterConfigurationTypes_Enums.h
+++ b/Engine/Plugins/Runtime/nDisplay/Source/DisplayClusterConfiguration/Public/DisplayClusterConfigurationTypes_Enums.h
@@ -75,10 +75,10 @@ UENUM()
 enum class EDisplayClusterConfigurationICVFX_LightcardRenderMode : uint8
 {
 	// Render incamera frame over lightcard
-	Over    UMETA(DisplayName = "Lightcard OVER"),
+	Over    UMETA(DisplayName = "Lightcard Over Frustum"),
 
 	// Over lightcard over incamera frame
-	Under   UMETA(DisplayName = "Lightcard UNDER"),
+	Under   UMETA(DisplayName = "Lightcard Under Frustum"),
 };
 
 UENUM()
diff --git a/Engine/Plugins/Runtime/nDisplay/Source/DisplayClusterConfiguration/Public/DisplayClusterConfigurationTypes_ICVFX.h b/Engine/Plugins/Runtime/nDisplay/Source/DisplayClusterConfiguration/Public/DisplayClusterConfigurationTypes_ICVFX.h
--- a/Engine/Plugins/Runtime/nDisplay/Source/DisplayClusterConfiguration/Public/DisplayClusterConfigurationTypes_ICVFX.h
+++ b/Engine/Plugins/Runtime/nDisplay/Source/DisplayClusterConfiguration/Public/DisplayClusterConfigurationTypes_ICVFX.h
@@ -180,7 +180,7 @@ public:
 	bool bEnable = true;
 
 	// Global lighcard rendering mode
-	UPROPERTY(EditAnywhere, Category = NDisplay)
+	UPROPERTY(EditAnywhere, Category = NDisplay, meta = (DisplayName = "Blending Mode"))
 	EDisplayClusterConfigurationICVFX_LightcardRenderMode Blendingmode = EDisplayClusterConfigurationICVFX_LightcardRenderMode::Under;
 
 	// Render actors from this layers to lightcard textures
diff --git a/Engine/Plugins/Runtime/nDisplay/Source/DisplayClusterConfigurator/Private/DisplayClusterConfiguratorModule.cpp b/Engine/Plugins/Runtime/nDisplay/Source/DisplayClusterConfigurator/Private/DisplayClusterConfiguratorModule.cpp
--- a/Engine/Plugins/Runtime/nDisplay/Source/DisplayClusterConfigurator/Private/DisplayClusterConfiguratorModule.cpp
+++ b/Engine/Plugins/Runtime/nDisplay/Source/DisplayClusterConfigurator/Private/DisplayClusterConfiguratorModule.cpp
@@ -6,21 +6,27 @@
 #include "DisplayClusterConfiguratorAssetTypeActions.h"
 #include "DisplayClusterConfiguratorVersionUtils.h"
 #include "Settings/DisplayClusterConfiguratorSettings.h"
+#include "Views/Details/DisplayClusterRootActorDetailsCustomization.h"
 
 #include "Views/Details/DisplayClusterConfiguratorDetailCustomization.h"
 #include "Views/Details/Policies/DisplayClusterConfiguratorPolicyDetailCustomization.h"
 
 #include "Blueprints/DisplayClusterBlueprint.h"
 #include "Components/DisplayClusterSceneComponent.h"
+#include "Components/DisplayClusterScreenComponent.h"
+#include "Components/DisplayClusterICVFXCameraComponent.h"
+#include "Misc/DisplayClusterObjectRef.h"
+#include "DisplayClusterRootActor.h"
 
 #include "AssetToolsModule.h"
 #include "IAssetTools.h"
 #include "ISettingsModule.h"
 #include "AssetRegistry/AssetRegistryModule.h"
 #include "AssetTypeCategories.h"
-#include "Components/DisplayClusterScreenComponent.h"
 #include "HAL/IConsoleManager.h"
 #include "Modules/ModuleManager.h"
+#include "Views/Details/DisplayClusterICVFXCameraComponentDetailsCustomization.h"
+
 
 #define LOCTEXT_NAMESPACE "DisplayClusterConfigurator"
 
@@ -175,6 +181,12 @@ void FDisplayClusterConfiguratorModule::RegisterCustomLayouts()
 	/**
 	 * CLASSES
 	 */
+	{
+		const FName LayoutName = UDisplayClusterConfigurationData::StaticClass()->GetFName();
+		RegisteredClassLayoutNames.Add(LayoutName);
+		PropertyModule.RegisterCustomClassLayout(ADisplayClusterRootActor::StaticClass()->GetFName(), FOnGetDetailCustomizationInstance::CreateStatic(&FDisplayClusterRootActorDetailsCustomization::MakeInstance));
+	}
+	
 	{
 		const FName LayoutName = UDisplayClusterConfigurationData::StaticClass()->GetFName();
 		RegisteredClassLayoutNames.Add(LayoutName);
@@ -209,7 +221,14 @@ void FDisplayClusterConfiguratorModule::RegisterCustomLayouts()
 		PropertyModule.RegisterCustomClassLayout(LayoutName,
 			FOnGetDetailCustomizationInstance::CreateStatic(&FDisplayClusterConfiguratorScreenDetailCustomization::MakeInstance));
 	}
-
+	
+	{
+		const FName LayoutName = UDisplayClusterICVFXCameraComponent::StaticClass()->GetFName();
+		RegisteredClassLayoutNames.Add(LayoutName);
+		PropertyModule.RegisterCustomClassLayout(LayoutName,
+			FOnGetDetailCustomizationInstance::CreateStatic(&FDisplayClusterICVFXCameraComponentDetailsCustomization::MakeInstance));
+	}
+	
 	/**
 	 * STRUCTS
 	 */
@@ -254,6 +273,14 @@ void FDisplayClusterConfiguratorModule::RegisterCustomLayouts()
 		PropertyModule.RegisterCustomPropertyTypeLayout(LayoutName,
 			FOnGetPropertyTypeCustomizationInstance::CreateStatic(&FDisplayClusterConfiguratorTypeCustomization::MakeInstance<FDisplayClusterConfiguratorExternalImageTypeCustomization>));
 	}
+
+	{
+		const FName LayoutName = FDisplayClusterComponentRef::StaticStruct()->GetFName();
+		RegisteredPropertyLayoutNames.Add(LayoutName);
+
+		PropertyModule.RegisterCustomPropertyTypeLayout(LayoutName,
+			FOnGetPropertyTypeCustomizationInstance::CreateStatic(&FDisplayClusterConfiguratorComponentRefCustomization::MakeInstance<FDisplayClusterConfiguratorComponentRefCustomization>));
+	}
 }
 
 void FDisplayClusterConfiguratorModule::UnregisterCustomLayouts()
diff --git a/Engine/Plugins/Runtime/nDisplay/Source/DisplayClusterConfigurator/Private/Views/Details/DisplayClusterConfiguratorDetailCustomization.cpp b/Engine/Plugins/Runtime/nDisplay/Source/DisplayClusterConfigurator/Private/Views/Details/DisplayClusterConfiguratorDetailCustomization.cpp
--- a/Engine/Plugins/Runtime/nDisplay/Source/DisplayClusterConfigurator/Private/Views/Details/DisplayClusterConfiguratorDetailCustomization.cpp
+++ b/Engine/Plugins/Runtime/nDisplay/Source/DisplayClusterConfigurator/Private/Views/Details/DisplayClusterConfiguratorDetailCustomization.cpp
@@ -148,10 +148,6 @@ void FDisplayClusterConfiguratorDataDetailCustomization::CustomizeDetails(IDetai
 {
 	Super::CustomizeDetails(InLayoutBuilder);
 
-	BEGIN_CATEGORY(DisplayClusterConfigurationStrings::categories::ICVFXCategory)
-		ADD_PROPERTY(UDisplayClusterConfigurationData, StageSettings);
-	END_CATEGORY()
-
 	BEGIN_CATEGORY(DisplayClusterConfigurationStrings::categories::ConfigurationCategory)
 		ADD_EXPANDED_PROPERTY(UDisplayClusterConfigurationData, RenderFrameSettings);
 		ADD_PROPERTY(UDisplayClusterConfigurationData, Info);
@@ -174,11 +170,7 @@ void FDisplayClusterConfiguratorClusterDetailCustomization::CustomizeDetails(IDe
 	check(ClusterNodesHandle->IsValidHandle());
 
 	FDisplayClusterConfiguratorNestedPropertyHelper NestedPropertyHelper(InLayoutBuilder);
-
-	BEGIN_CATEGORY(DisplayClusterConfigurationStrings::categories::ClusterCategory)
-		ADD_PROPERTY(UDisplayClusterConfigurationCluster, Nodes);
-	END_CATEGORY()
-
+	
 	BEGIN_CATEGORY(DisplayClusterConfigurationStrings::categories::ClusterPostprocessCategory)
 		BEGIN_GROUP(TEXT("GlobalPostProcess"), LOCTEXT("GlobalPostprocessLabel", "All Viewports"))
 			ADD_GROUP_PROPERTY(UDisplayClusterConfigurationCluster, bUseOverallClusterPostProcess);
@@ -203,7 +195,7 @@ void FDisplayClusterConfiguratorClusterDetailCustomization::CustomizeDetails(IDe
 		}
 	END_CATEGORY();
 
-	BEGIN_CATEGORY(DisplayClusterConfigurationStrings::categories::ClusterConfigurationCategory)
+	BEGIN_CATEGORY(DisplayClusterConfigurationStrings::categories::AdvancedCategory)
 		ADD_PROPERTY(UDisplayClusterConfigurationCluster, MasterNode);
 		ADD_PROPERTY(UDisplayClusterConfigurationCluster, Sync);
 		ADD_PROPERTY(UDisplayClusterConfigurationCluster, Network);
@@ -232,6 +224,7 @@ void FDisplayClusterConfiguratorClusterDetailCustomization::CustomizeDetails(IDe
 				];
 		}
 
+		ADD_PROPERTY(UDisplayClusterConfigurationCluster, Nodes);
 	END_CATEGORY()
 }
 
@@ -1178,4 +1171,20 @@ void FDisplayClusterConfiguratorExternalImageTypeCustomization::CustomizeChildre
 	FDisplayClusterConfiguratorTypeCustomization::CustomizeChildren(InPropertyHandle, InChildBuilder, CustomizationUtils);
 }
 
+void FDisplayClusterConfiguratorComponentRefCustomization::CustomizeHeader(TSharedRef<IPropertyHandle> PropertyHandle,
+	FDetailWidgetRow& HeaderRow, IPropertyTypeCustomizationUtils& CustomizationUtils)
+{
+	// This will prevent the struct from being expanded.
+
+	HeaderRow.NameContent()
+	[
+		PropertyHandle->CreatePropertyNameWidget()
+	]
+	.ValueContent()
+	[
+		// Automatically retrieves the TitleProperty
+		PropertyHandle->CreatePropertyValueWidget(false)
+	];
+}
+
 #undef LOCTEXT_NAMESPACE
diff --git a/Engine/Plugins/Runtime/nDisplay/Source/DisplayClusterConfigurator/Private/Views/Details/DisplayClusterConfiguratorDetailCustomization.h b/Engine/Plugins/Runtime/nDisplay/Source/DisplayClusterConfigurator/Private/Views/Details/DisplayClusterConfiguratorDetailCustomization.h
--- a/Engine/Plugins/Runtime/nDisplay/Source/DisplayClusterConfigurator/Private/Views/Details/DisplayClusterConfiguratorDetailCustomization.h
+++ b/Engine/Plugins/Runtime/nDisplay/Source/DisplayClusterConfigurator/Private/Views/Details/DisplayClusterConfiguratorDetailCustomization.h
@@ -400,4 +400,18 @@ private:
 	TSharedPtr<IPropertyHandle> ImagePathHandle;
 };
 
+
+/**
+ * Component Ref Type Customization (Prevents struct from expanding)
+ */
+class FDisplayClusterConfiguratorComponentRefCustomization final
+	: public FDisplayClusterConfiguratorTypeCustomization
+{
+protected:
+	//~ IPropertyTypeCustomization interface begin
+	virtual void CustomizeHeader(TSharedRef<IPropertyHandle> PropertyHandle, FDetailWidgetRow& HeaderRow, IPropertyTypeCustomizationUtils& CustomizationUtils) override;
+	//~ IPropertyTypeCustomization interface end
+};
+
+
 #undef CONSTRUCT_CUSTOMIZATION
\ No newline at end of file
diff --git a/Engine/Plugins/Runtime/nDisplay/Source/DisplayClusterConfigurator/Private/Views/Details/DisplayClusterICVFXCameraComponentDetailsCustomization.cpp b/Engine/Plugins/Runtime/nDisplay/Source/DisplayClusterConfigurator/Private/Views/Details/DisplayClusterICVFXCameraComponentDetailsCustomization.cpp
new file mode 100644
--- /dev/null
+++ b/Engine/Plugins/Runtime/nDisplay/Source/DisplayClusterConfigurator/Private/Views/Details/DisplayClusterICVFXCameraComponentDetailsCustomization.cpp
@@ -0,0 +1,64 @@
+// Copyright Epic Games, Inc. All Rights Reserved.
+
+#include "DisplayClusterICVFXCameraComponentDetailsCustomization.h"
+
+#include "Components/DisplayClusterICVFXCameraComponent.h"
+
+#include "CineCameraActor.h"
+#include "UObject/SoftObjectPtr.h"
+
+#include "DetailLayoutBuilder.h"
+
+
+TSharedRef<IDetailCustomization> FDisplayClusterICVFXCameraComponentDetailsCustomization::MakeInstance()
+{
+	return MakeShared<FDisplayClusterICVFXCameraComponentDetailsCustomization>();
+}
+
+void FDisplayClusterICVFXCameraComponentDetailsCustomization::CustomizeDetails(IDetailLayoutBuilder& Builder)
+{
+	DetailLayout = &Builder;
+
+	if (!EditedObject.IsValid())
+	{
+		TArray<TWeakObjectPtr<UObject>> Objects;
+		Builder.GetObjectsBeingCustomized(Objects);
+
+		for (TWeakObjectPtr<UObject> Object : Objects)
+		{
+			if (Object->IsA<UDisplayClusterICVFXCameraComponent>())
+			{
+				EditedObject = Cast<UDisplayClusterICVFXCameraComponent>(Object.Get());
+			}
+		}
+	}
+
+	// Hide some groups if an external CineCameraActor is set
+	if (EditedObject.IsValid() && EditedObject->ExternalCameraActor.IsValid())
+	{
+		Builder.HideCategory(TEXT("Current Camera Settings"));
+		Builder.HideCategory(TEXT("CameraOptions"));
+		Builder.HideCategory(TEXT("Camera"));
+		Builder.HideCategory(TEXT("PostProcess"));
+		Builder.HideCategory(TEXT("Lens"));
+		Builder.HideCategory(TEXT("LOD"));
+		Builder.HideCategory(TEXT("ColorGrading"));
+		Builder.HideCategory(TEXT("RenderingFeatures"));
+		Builder.HideCategory(TEXT("Tags"));
+		Builder.HideCategory(TEXT("ComponentTick"));
+		Builder.HideCategory(TEXT("Color Grading"));
+		Builder.HideCategory(TEXT("Rendering Features"));
+		Builder.HideCategory(TEXT("Activation"));
+	}
+
+	// Always hide the following categories
+	Builder.HideCategory(TEXT("AssetUserData"));
+	Builder.HideCategory(TEXT("Collision"));
+	Builder.HideCategory(TEXT("Cooking"));
+	Builder.HideCategory(TEXT("ComponentReplication"));
+	Builder.HideCategory(TEXT("Events"));
+	Builder.HideCategory(TEXT("Physics"));
+	Builder.HideCategory(TEXT("Sockets"));
+}
+
+#undef LOCTEXT_NAMESPACE
diff --git a/Engine/Plugins/Runtime/nDisplay/Source/DisplayClusterConfigurator/Private/Views/Details/DisplayClusterICVFXCameraComponentDetailsCustomization.h b/Engine/Plugins/Runtime/nDisplay/Source/DisplayClusterConfigurator/Private/Views/Details/DisplayClusterICVFXCameraComponentDetailsCustomization.h
new file mode 100644
--- /dev/null
+++ b/Engine/Plugins/Runtime/nDisplay/Source/DisplayClusterConfigurator/Private/Views/Details/DisplayClusterICVFXCameraComponentDetailsCustomization.h
@@ -0,0 +1,27 @@
+// Copyright Epic Games, Inc. All Rights Reserved.
+
+#pragma once
+
+#include "IDetailCustomization.h"
+
+#include "CoreMinimal.h"
+#include "Types/SlateEnums.h"
+
+class UDisplayClusterICVFXCameraComponent;
+class ACineCameraActor;
+class IDetailLayoutBuilder;
+
+class FDisplayClusterICVFXCameraComponentDetailsCustomization : public IDetailCustomization
+{
+public:
+	static TSharedRef<IDetailCustomization> MakeInstance();
+
+protected:
+	// IDetailCustomization interface
+	virtual void CustomizeDetails(class IDetailLayoutBuilder& LayoutBuilder) override;
+
+protected:
+	TWeakObjectPtr<UDisplayClusterICVFXCameraComponent> EditedObject;
+	/** Keep a reference to force refresh the layout */
+	IDetailLayoutBuilder* DetailLayout = nullptr;
+};
diff --git a/Engine/Plugins/Runtime/nDisplay/Source/DisplayClusterConfigurator/Private/Views/Details/DisplayClusterRootActorDetailsCustomization.cpp b/Engine/Plugins/Runtime/nDisplay/Source/DisplayClusterConfigurator/Private/Views/Details/DisplayClusterRootActorDetailsCustomization.cpp
new file mode 100644
--- /dev/null
+++ b/Engine/Plugins/Runtime/nDisplay/Source/DisplayClusterConfigurator/Private/Views/Details/DisplayClusterRootActorDetailsCustomization.cpp
@@ -0,0 +1,358 @@
+// Copyright Epic Games, Inc. All Rights Reserved.
+
+#include "DisplayClusterRootActorDetailsCustomization.h"
+#include "DisplayClusterConfiguratorDetailCustomizationUtils.h"
+
+#include "DetailCategoryBuilder.h"
+#include "DetailLayoutBuilder.h"
+#include "DetailWidgetRow.h"
+#include "PropertyHandle.h"
+
+#include "SSearchableComboBox.h"
+#include "Widgets/Text/STextBlock.h"
+#include "IDetailGroup.h"
+
+#include "DisplayClusterRootActor.h"
+#include "DisplayClusterConfigurationStrings.h"
+#include "DisplayClusterConfigurationTypes.h"
+#include "DisplayClusterConfigurationTypes_Viewport.h"
+
+
+#define LOCTEXT_NAMESPACE "DisplayClusterRootActorDetailsCustomization"
+
+
+TSharedRef<IDetailCustomization> FDisplayClusterRootActorDetailsCustomization::MakeInstance()
+{
+	return MakeShared<FDisplayClusterRootActorDetailsCustomization>();
+}
+
+void FDisplayClusterRootActorDetailsCustomization::CustomizeDetails(IDetailLayoutBuilder& InLayoutBuilder)
+{
+	// Hide the following categories, we don't really need them
+	InLayoutBuilder.HideCategory(TEXT("Rendering"));
+	InLayoutBuilder.HideCategory(TEXT("Replication"));
+	InLayoutBuilder.HideCategory(TEXT("Collision"));
+	InLayoutBuilder.HideCategory(TEXT("Input"));
+	InLayoutBuilder.HideCategory(TEXT("Actor"));
+	InLayoutBuilder.HideCategory(TEXT("LOD"));
+	InLayoutBuilder.HideCategory(TEXT("Cooking"));
+	
+	// Only single selection is allowed
+	TArray<TWeakObjectPtr<UObject>> SelectedObjects = InLayoutBuilder.GetSelectedObjects();
+	if (SelectedObjects.Num() != 1)
+	{
+		return;
+	}
+
+	// Store the object we're working with
+	EditedObject = Cast<ADisplayClusterRootActor>(SelectedObjects[0].Get());
+	if (!EditedObject.IsValid())
+	{
+		return;
+	}
+
+	// Store the NodeId property handle to use later
+	PropertyNodeId = InLayoutBuilder.GetProperty(GET_MEMBER_NAME_CHECKED(ADisplayClusterRootActor, PreviewNodeId), ADisplayClusterRootActor::StaticClass());
+	check(PropertyNodeId->IsValidHandle());
+
+	if (EditedObject.IsValid() && !EditedObject->IsTemplate(RF_ClassDefaultObject))
+	{
+		InLayoutBuilder.HideCategory(TEXT("NDisplay"));
+		InLayoutBuilder.HideCategory(TEXT("NDisplay Configuration"));
+		InLayoutBuilder.HideCategory(TEXT("NDisplay Cluster"));
+		InLayoutBuilder.HideCategory(TEXT("NDisplay Cluster Configuration"));
+	}
+	
+	// Lay out all of the root actor properties into their correct orders and categories. Because we are adding custom properties,
+	// most of the actor's properties will need to be laid out manually here even if they aren't being customized.
+	BuildLayout(InLayoutBuilder);
+
+	// Update the selected item in the NodeId combo box to match the current value on the root actor
+	UpdateNodeIdSelection();
+}
+
+void FDisplayClusterRootActorDetailsCustomization::BuildLayout(IDetailLayoutBuilder& InLayoutBuilder)
+{
+	InLayoutBuilder.EditCategory(DisplayClusterConfigurationStrings::categories::DefaultCategory);
+	
+	FDisplayClusterConfiguratorNestedPropertyHelper NestedPropertyHelper(InLayoutBuilder);
+	
+	if (EditedObject.IsValid() && !EditedObject->IsTemplate(RF_ClassDefaultObject))
+	{
+		// Build ICVFX category.
+		
+		TArray<FString> ViewportNames;
+		NestedPropertyHelper.GetNestedPropertyKeys(TEXT("CurrentConfigData.Cluster.Nodes.Viewports"), ViewportNames);
+
+		BEGIN_CATEGORY(DisplayClusterConfigurationStrings::categories::ICVFXCategory)
+
+			const TSharedRef<IPropertyHandle> EnableInnerFrustumPropertyHandle = InLayoutBuilder.GetProperty(GET_MEMBER_NAME_CHECKED(ADisplayClusterRootActor, EnableInnerFrustum));
+			{
+				check(EnableInnerFrustumPropertyHandle->IsValidHandle());
+				CurrentCategory.AddProperty(EnableInnerFrustumPropertyHandle);
+			}
+
+			const TAttribute<bool> AllowICVFXEditCondition = TAttribute<bool>::Create([this, EnableInnerFrustumPropertyHandle]()
+			{
+				bool bFrustumEnabled = false;
+				EnableInnerFrustumPropertyHandle->GetValue(bFrustumEnabled);
+				return bFrustumEnabled;
+			});
+		
+			if (ViewportNames.Num() > 0)
+			{
+				TArray<TSharedPtr<IPropertyHandle>> AllowICVFXHandles;
+				NestedPropertyHelper.GetNestedProperties(TEXT("CurrentConfigData.Cluster.Nodes.Viewports.ICVFX.bAllowICVFX"), AllowICVFXHandles);
+
+				BEGIN_GROUP("InnerFrustumEnabledInViewports", LOCTEXT("InnerFrustumEnabledInViewports", "Inner Frustum Enabled in Viewports"))
+					for (int32 VPIdx = 0; VPIdx < AllowICVFXHandles.Num(); ++VPIdx)
+					{
+						TSharedPtr<IPropertyHandle>& Handle = AllowICVFXHandles[VPIdx];
+
+						Handle->SetPropertyDisplayName(FText::FromString(ViewportNames[VPIdx]));
+						IDetailPropertyRow& PropertyRow = CurrentGroup.AddPropertyRow(Handle.ToSharedRef());
+						PropertyRow.EditCondition(AllowICVFXEditCondition, nullptr);
+					}
+				END_GROUP();
+
+				ADD_PROPERTY(ADisplayClusterRootActor, InnerFrustumPriority);
+
+				{
+					TArray<TSharedPtr<IPropertyHandle>> GlobalScreenPercentageHandles;
+					NestedPropertyHelper.GetNestedProperties(TEXT("CurrentConfigData.RenderFrameSettings.ClusterRenderTargetRatioMult"), GlobalScreenPercentageHandles);
+					check(GlobalScreenPercentageHandles.Num() == 1);
+
+					GlobalScreenPercentageHandles[0]->SetPropertyDisplayName(LOCTEXT("OuterViewportScreenPercentageMult", "Outer Viewport Screen Percentage Multiplier"));
+					CurrentCategory.AddProperty(GlobalScreenPercentageHandles[0]);
+				}
+				
+				TArray<TSharedPtr<IPropertyHandle>> ScreenPercentageHandles;
+				NestedPropertyHelper.GetNestedProperties(TEXT("CurrentConfigData.Cluster.Nodes.Viewports.RenderSettings.BufferRatio"), ScreenPercentageHandles);
+
+				BEGIN_GROUP("OuterViewportScreenPercentage", LOCTEXT("OuterViewportScreenPercentage", "Outer Viewport Screen Percentage"))
+					for (int32 VPIdx = 0; VPIdx < ScreenPercentageHandles.Num(); ++VPIdx)
+					{
+						TSharedPtr<IPropertyHandle>& Handle = ScreenPercentageHandles[VPIdx];
+
+						Handle->SetPropertyDisplayName(FText::FromString(ViewportNames[VPIdx]));
+						CurrentGroup.AddPropertyRow(Handle.ToSharedRef());
+					}
+				END_GROUP();
+			}
+
+		END_CATEGORY();
+
+		BEGIN_CATEGORY(DisplayClusterConfigurationStrings::categories::ChromaKeyCategory)
+		{
+			TArray<TSharedPtr<IPropertyHandle>> Handles;
+			NestedPropertyHelper.GetNestedProperties(TEXT("CurrentConfigData.StageSettings.Chromakey"), Handles);
+			
+			check(Handles.Num() == 1);
+			CurrentCategory.AddProperty(Handles[0]).ShouldAutoExpand(true);
+		}
+		END_CATEGORY();
+		
+		BEGIN_CATEGORY(DisplayClusterConfigurationStrings::categories::LightcardCategory)
+		{
+			TArray<TSharedPtr<IPropertyHandle>> Handles;
+			NestedPropertyHelper.GetNestedProperties(TEXT("CurrentConfigData.StageSettings.Lightcard.bEnable"), Handles);
+			check(Handles.Num() == 1);
+			Handles[0]->SetPropertyDisplayName(LOCTEXT("LightCardDisplayName", "Enable Light Cards"));
+			CurrentCategory.AddProperty(Handles[0]);
+		}
+		{
+			TArray<TSharedPtr<IPropertyHandle>> Handles;
+			NestedPropertyHelper.GetNestedProperties(TEXT("CurrentConfigData.StageSettings.Lightcard.BlendingMode"), Handles);
+			check(Handles.Num() == 1);
+
+			CurrentCategory.AddProperty(Handles[0]);
+		}
+		END_CATEGORY();
+	}
+	
+	// Force a particular order that the property categories show up in the details panel
+
+	InLayoutBuilder.EditCategory(DisplayClusterConfigurationStrings::categories::ConfigurationCategory);
+	InLayoutBuilder.EditCategory(DisplayClusterConfigurationStrings::categories::ClusterCategory);
+	InLayoutBuilder.EditCategory(DisplayClusterConfigurationStrings::categories::ClusterPostprocessCategory);
+	InLayoutBuilder.EditCategory(DisplayClusterConfigurationStrings::categories::PreviewCategory);
+	InLayoutBuilder.EditCategory(DisplayClusterConfigurationStrings::categories::AdvancedCategory).SetCategoryVisibility(false); // Re-enable for access to more adv control.
+	
+	// Add custom properties and lay out/order properties into their correct categories.
+	BEGIN_CATEGORY(DisplayClusterConfigurationStrings::categories::PreviewCategory)
+		if (RebuildNodeIdOptionsList())
+		{
+			REPLACE_PROPERTY_WITH_CUSTOM(ADisplayClusterRootActor, PreviewNodeId, CreateCustomNodeIdWidget());
+		}
+	END_CATEGORY();
+}
+
+
+//////////////////////////////////////////////////////////////////////////////////////////////
+// Preview node ID
+//////////////////////////////////////////////////////////////////////////////////////////////
+TSharedRef<SWidget> FDisplayClusterRootActorDetailsCustomization::CreateCustomNodeIdWidget()
+{
+	if (NodeIdComboBox.IsValid())
+	{
+		return NodeIdComboBox.ToSharedRef();
+	}
+
+	return SAssignNew(NodeIdComboBox, SSearchableComboBox)
+		.OptionsSource(&NodeIdOptions)
+		.OnGenerateWidget(this, &FDisplayClusterRootActorDetailsCustomization::CreateComboWidget)
+		.OnSelectionChanged(this, &FDisplayClusterRootActorDetailsCustomization::OnNodeIdSelected)
+		.ContentPadding(2)
+		.Content()
+		[
+			SNew(STextBlock)
+			.Text(this, &FDisplayClusterRootActorDetailsCustomization::GetSelectedNodeIdText)
+		];
+}
+
+bool FDisplayClusterRootActorDetailsCustomization::RebuildNodeIdOptionsList()
+{
+	// Get current configuration data
+	const UDisplayClusterConfigurationData* ConfigurationData = EditedObject->GetConfigData();
+	if (!ConfigurationData)
+	{
+		return false;
+	}
+
+	// Initialize special options
+	NodeIdOptionAll  = MakeShared<FString>(DisplayClusterConfigurationStrings::gui::preview::PreviewNodeAll);
+	NodeIdOptionNone = MakeShared<FString>(DisplayClusterConfigurationStrings::gui::preview::PreviewNodeNone);
+
+	// Fill combobox with the options
+	NodeIdOptions.Reset();
+	NodeIdOptions.Emplace(NodeIdOptionNone);
+	NodeIdOptions.Emplace(NodeIdOptionAll);
+	for (const TPair<FString, UDisplayClusterConfigurationClusterNode*>& Node : ConfigurationData->Cluster->Nodes)
+	{
+		if (!Node.Key.IsEmpty())
+		{
+			NodeIdOptions.Emplace(MakeShared<FString>(Node.Key));
+		}
+	}
+
+	// Set 'None' each time we update the preview config file
+	if (NodeIdComboBox.IsValid())
+	{
+		NodeIdComboBox->SetSelectedItem(NodeIdOptionNone);
+	}
+
+	// Make sure we've got at least one cluster node in the config
+	// (None+node or None+All+node1+node2+...)
+	return NodeIdOptions.Num() >= 3;
+}
+
+void FDisplayClusterRootActorDetailsCustomization::UpdateNodeIdSelection()
+{
+	if (NodeIdComboBox.IsValid())
+	{
+		const FString& CurrentPreviewNode = EditedObject->PreviewNodeId;
+		TSharedPtr<FString>* FoundItem = NodeIdOptions.FindByPredicate([CurrentPreviewNode](const TSharedPtr<FString>& Item)
+		{
+			return Item->Equals(CurrentPreviewNode, ESearchCase::IgnoreCase);
+		});
+
+		if (FoundItem)
+		{
+			NodeIdComboBox->SetSelectedItem(*FoundItem);
+		}
+		else
+		{
+			// Set combobox selected item (options list is not empty here)
+			NodeIdComboBox->SetSelectedItem(NodeIdOptionAll);
+		}
+	}
+}
+
+void FDisplayClusterRootActorDetailsCustomization::OnNodeIdSelected(TSharedPtr<FString> PreviewNodeId, ESelectInfo::Type SelectInfo)
+{
+	const FString NewValue = (PreviewNodeId.IsValid() ? *PreviewNodeId : *NodeIdOptionNone);
+	PropertyNodeId->SetValue(NewValue);
+}
+
+FText FDisplayClusterRootActorDetailsCustomization::GetSelectedNodeIdText() const
+{
+	if (NodeIdComboBox.IsValid())
+	{
+		TSharedPtr<FString> CurSelection = NodeIdComboBox->GetSelectedItem();
+		return FText::FromString(CurSelection.IsValid() ? *CurSelection : *NodeIdOptionNone);
+	}
+
+	return FText::FromString(*NodeIdOptionNone);
+}
+
+
+//////////////////////////////////////////////////////////////////////////////////////////////
+// Preview default camera ID
+//////////////////////////////////////////////////////////////////////////////////////////////
+void FDisplayClusterRootActorDetailsCustomization::AddDefaultCameraRow()
+{
+	//! remove 
+}
+
+bool FDisplayClusterRootActorDetailsCustomization::RebuildDefaultCameraOptionsList()
+{
+	// Get current configuration data
+	const UDisplayClusterConfigurationData* ConfigurationData = EditedObject->GetConfigData();
+	if (!ConfigurationData)
+	{
+		return false;
+	}
+
+	// Fill combobox with the options
+	DefaultCameraOptions.Reset();
+	for (const TPair<FString, UDisplayClusterConfigurationSceneComponentCamera*>& Node : ConfigurationData->Scene->Cameras)
+	{
+		if (!Node.Key.IsEmpty())
+		{
+			DefaultCameraOptions.Emplace(MakeShared<FString>(Node.Key));
+		}
+	}
+
+	if (DefaultCameraOptions.Num() < 1)
+	{
+		return false;
+	}
+
+	// Set 'None' each time we update the preview config file
+	if (DefaultCameraComboBox.IsValid())
+	{
+		DefaultCameraComboBox->SetSelectedItem(DefaultCameraOptions[0]);
+	}
+
+	return true;
+}
+
+void FDisplayClusterRootActorDetailsCustomization::OnDefaultCameraSelected(TSharedPtr<FString> CameraId, ESelectInfo::Type SelectInfo)
+{
+	const FString NewValue = (CameraId.IsValid() ? *CameraId : FString());
+	PropertyDefaultCamera->SetValue(NewValue);
+}
+
+FText FDisplayClusterRootActorDetailsCustomization::GetSelectedDefaultCameraText() const
+{
+	TSharedPtr<FString> CurSelection = DefaultCameraComboBox->GetSelectedItem();
+	return FText::FromString(CurSelection.IsValid() ? *CurSelection : *FString());
+}
+
+
+//////////////////////////////////////////////////////////////////////////////////////////////
+// Internals
+//////////////////////////////////////////////////////////////////////////////////////////////
+void FDisplayClusterRootActorDetailsCustomization::OnPreviewConfigChanged()
+{
+	// We need to update options for the comboboxes
+	RebuildNodeIdOptionsList();
+	RebuildDefaultCameraOptionsList();
+}
+
+TSharedRef<SWidget> FDisplayClusterRootActorDetailsCustomization::CreateComboWidget(TSharedPtr<FString> InItem)
+{
+	return SNew(STextBlock).Text(FText::FromString(*InItem));
+}
+
+#undef LOCTEXT_NAMESPACE
diff --git a/Engine/Plugins/Runtime/nDisplay/Source/DisplayClusterConfigurator/Private/Views/Details/DisplayClusterRootActorDetailsCustomization.h b/Engine/Plugins/Runtime/nDisplay/Source/DisplayClusterConfigurator/Private/Views/Details/DisplayClusterRootActorDetailsCustomization.h
new file mode 100644
--- /dev/null
+++ b/Engine/Plugins/Runtime/nDisplay/Source/DisplayClusterConfigurator/Private/Views/Details/DisplayClusterRootActorDetailsCustomization.h
@@ -0,0 +1,77 @@
+// Copyright Epic Games, Inc. All Rights Reserved.
+
+#pragma once
+
+#include "IDetailCustomization.h"
+
+#include "CoreMinimal.h"
+#include "Types/SlateEnums.h"
+
+class ADisplayClusterRootActor;
+class IDetailCategoryBuilder;
+class IDetailLayoutBuilder;
+class IPropertyHandle;
+class SWidget;
+class SSearchableComboBox;
+
+
+class FDisplayClusterRootActorDetailsCustomization : public IDetailCustomization
+{
+public:
+	static TSharedRef<IDetailCustomization> MakeInstance();
+
+protected:
+	// IDetailCustomization interface
+	virtual void CustomizeDetails(IDetailLayoutBuilder& InLayoutBuilder) override;
+	// End IDetailCustomization interface
+
+protected:
+	void BuildLayout(IDetailLayoutBuilder& InLayoutBuilder);
+
+protected:
+	//////////////////////////////////////////////////////////////////////////////////////////////
+	// Node ID
+	//////////////////////////////////////////////////////////////////////////////////////////////
+	TSharedRef<SWidget> CreateCustomNodeIdWidget();
+	bool RebuildNodeIdOptionsList();
+	void UpdateNodeIdSelection();
+
+	void OnNodeIdSelected(TSharedPtr<FString> NodeId, ESelectInfo::Type SelectInfo);
+	FText GetSelectedNodeIdText() const;
+
+protected:
+	TSharedPtr<FString> NodeIdOptionNone;
+	TSharedPtr<FString> NodeIdOptionAll;
+
+	TSharedPtr<IPropertyHandle>     PropertyNodeId;
+	TArray<TSharedPtr<FString>>     NodeIdOptions;
+	TSharedPtr<SSearchableComboBox> NodeIdComboBox;
+
+protected:
+	//////////////////////////////////////////////////////////////////////////////////////////////
+	// Default camera ID
+	//////////////////////////////////////////////////////////////////////////////////////////////
+	void AddDefaultCameraRow();
+	bool RebuildDefaultCameraOptionsList();
+	void OnDefaultCameraSelected(TSharedPtr<FString> CameraId, ESelectInfo::Type SelectInfo);
+	FText GetSelectedDefaultCameraText() const;
+
+protected:
+	TSharedPtr<IPropertyHandle>     PropertyDefaultCamera;
+	TArray<TSharedPtr<FString>>     DefaultCameraOptions;
+	TSharedPtr<SSearchableComboBox> DefaultCameraComboBox;
+
+protected:
+	//////////////////////////////////////////////////////////////////////////////////////////////
+	// Internals
+	//////////////////////////////////////////////////////////////////////////////////////////////
+
+	// Handle config changes
+	void OnPreviewConfigChanged();
+	// Create combobox widget
+	TSharedRef<SWidget> CreateComboWidget(TSharedPtr<FString> InItem);
+
+protected:
+	// ADisplayClusterRootActor on which we're acting
+	TWeakObjectPtr<ADisplayClusterRootActor> EditedObject;
+};
diff --git a/Engine/Plugins/Runtime/nDisplay/Source/DisplayClusterEditor/DisplayClusterEditor.Build.cs b/Engine/Plugins/Runtime/nDisplay/Source/DisplayClusterEditor/DisplayClusterEditor.Build.cs
--- a/Engine/Plugins/Runtime/nDisplay/Source/DisplayClusterEditor/DisplayClusterEditor.Build.cs
+++ b/Engine/Plugins/Runtime/nDisplay/Source/DisplayClusterEditor/DisplayClusterEditor.Build.cs
@@ -9,17 +9,12 @@ public class DisplayClusterEditor : ModuleRules
 	{
 		PrivateDependencyModuleNames.AddRange(
 			new string[] {
-				"CinematicCamera",
+				"DisplayCluster",
+
 				"Core",
 				"CoreUObject",
-				"DisplayCluster",
-				"DisplayClusterConfiguration",
 				"Engine",
-				"EditorWidgets",
-				"PropertyEditor",
-				"Slate",
-				"SlateCore",
-				"UnrealEd",
+				"UnrealEd"
 			});
 	}
 }
diff --git a/Engine/Plugins/Runtime/nDisplay/Source/DisplayClusterEditor/Private/DetailsCustomization/DisplayClusterICVFXCameraComponentDetailsCustomization.cpp b/Engine/Plugins/Runtime/nDisplay/Source/DisplayClusterEditor/Private/DetailsCustomization/DisplayClusterICVFXCameraComponentDetailsCustomization.cpp
deleted file mode 100644
diff --git a/Engine/Plugins/Runtime/nDisplay/Source/DisplayClusterEditor/Private/DetailsCustomization/DisplayClusterICVFXCameraComponentDetailsCustomization.h b/Engine/Plugins/Runtime/nDisplay/Source/DisplayClusterEditor/Private/DetailsCustomization/DisplayClusterICVFXCameraComponentDetailsCustomization.h
deleted file mode 100644
diff --git a/Engine/Plugins/Runtime/nDisplay/Source/DisplayClusterEditor/Private/DetailsCustomization/DisplayClusterRootActorDetailsCustomization.cpp b/Engine/Plugins/Runtime/nDisplay/Source/DisplayClusterEditor/Private/DetailsCustomization/DisplayClusterRootActorDetailsCustomization.cpp
deleted file mode 100644
diff --git a/Engine/Plugins/Runtime/nDisplay/Source/DisplayClusterEditor/Private/DetailsCustomization/DisplayClusterRootActorDetailsCustomization.h b/Engine/Plugins/Runtime/nDisplay/Source/DisplayClusterEditor/Private/DetailsCustomization/DisplayClusterRootActorDetailsCustomization.h
deleted file mode 100644
diff --git a/Engine/Plugins/Runtime/nDisplay/Source/DisplayClusterEditor/Private/DisplayClusterEditorModule.cpp b/Engine/Plugins/Runtime/nDisplay/Source/DisplayClusterEditor/Private/DisplayClusterEditorModule.cpp
--- a/Engine/Plugins/Runtime/nDisplay/Source/DisplayClusterEditor/Private/DisplayClusterEditorModule.cpp
+++ b/Engine/Plugins/Runtime/nDisplay/Source/DisplayClusterEditor/Private/DisplayClusterEditorModule.cpp
@@ -2,17 +2,9 @@
 
 #include "DisplayClusterEditorModule.h"
 
-#include "DisplayClusterRootActor.h"
-#include "Components/DisplayClusterICVFXCameraComponent.h"
-#include "DetailsCustomization/DisplayClusterRootActorDetailsCustomization.h"
-#include "DetailsCustomization/DisplayClusterICVFXCameraComponentDetailsCustomization.h"
-
-#include "Components/DisplayClusterPreviewComponent.h"
-
 #include "Settings/DisplayClusterEditorSettings.h"
 
 #include "ISettingsModule.h"
-#include "PropertyEditorModule.h"
 
 
 #define LOCTEXT_NAMESPACE "DisplayClusterEditor"
@@ -20,7 +12,6 @@
 void FDisplayClusterEditorModule::StartupModule()
 {
 	RegisterSettings();
-	RegisterCustomizations();
 }
 
 void FDisplayClusterEditorModule::ShutdownModule()
@@ -28,7 +19,6 @@ void FDisplayClusterEditorModule::ShutdownModule()
 	if (UObjectInitialized())
 	{
 		UnregisterSettings();
-		UnregisterCustomizations();
 	}
 }
 
@@ -56,23 +46,6 @@ void FDisplayClusterEditorModule::UnregisterSettings()
 	}
 }
 
-
-void FDisplayClusterEditorModule::RegisterCustomizations()
-{
-	FPropertyEditorModule& PropertyEditorModule = FModuleManager::Get().LoadModuleChecked<FPropertyEditorModule>("PropertyEditor");
-
-	PropertyEditorModule.RegisterCustomClassLayout(ADisplayClusterRootActor::StaticClass()->GetFName(), FOnGetDetailCustomizationInstance::CreateStatic(&FDisplayClusterRootActorDetailsCustomization::MakeInstance));
-	PropertyEditorModule.RegisterCustomClassLayout(UDisplayClusterICVFXCameraComponent::StaticClass()->GetFName(), FOnGetDetailCustomizationInstance::CreateStatic(&FDisplayClusterICVFXCameraComponentDetailsCustomization::MakeInstance));
-}
-
-void FDisplayClusterEditorModule::UnregisterCustomizations()
-{
-	FPropertyEditorModule& PropertyEditorModule = FModuleManager::Get().LoadModuleChecked<FPropertyEditorModule>("PropertyEditor");
-
-	PropertyEditorModule.UnregisterCustomClassLayout(ADisplayClusterRootActor::StaticClass()->GetFName());
-	PropertyEditorModule.UnregisterCustomClassLayout(UDisplayClusterICVFXCameraComponent::StaticClass()->GetFName());
-}
-
 IMPLEMENT_MODULE(FDisplayClusterEditorModule, DisplayClusterEditor);
 
 #undef LOCTEXT_NAMESPACE
diff --git a/Engine/Plugins/Runtime/nDisplay/Source/DisplayClusterEditor/Private/DisplayClusterEditorModule.h b/Engine/Plugins/Runtime/nDisplay/Source/DisplayClusterEditor/Private/DisplayClusterEditorModule.h
--- a/Engine/Plugins/Runtime/nDisplay/Source/DisplayClusterEditor/Private/DisplayClusterEditorModule.h
+++ b/Engine/Plugins/Runtime/nDisplay/Source/DisplayClusterEditor/Private/DisplayClusterEditorModule.h
@@ -19,7 +19,4 @@ public:
 protected:
 	void RegisterSettings();
 	void UnregisterSettings();
-
-	void RegisterCustomizations();
-	void UnregisterCustomizations();
 };
