nDisplay: Fixed crashes happening when attempting to undo adding a cluster node to a new host. The issue was caused by attempting to access an object that was removed during the undo operation, and was fixed by adding validation checks to the cluster item graph nodes

diff --git a/Engine/Plugins/Runtime/nDisplay/Source/DisplayClusterConfigurator/Private/Views/OutputMapping/EdNodes/DisplayClusterConfiguratorBaseNode.cpp b/Engine/Plugins/Runtime/nDisplay/Source/DisplayClusterConfigurator/Private/Views/OutputMapping/EdNodes/DisplayClusterConfiguratorBaseNode.cpp
--- a/Engine/Plugins/Runtime/nDisplay/Source/DisplayClusterConfigurator/Private/Views/OutputMapping/EdNodes/DisplayClusterConfiguratorBaseNode.cpp
+++ b/Engine/Plugins/Runtime/nDisplay/Source/DisplayClusterConfigurator/Private/Views/OutputMapping/EdNodes/DisplayClusterConfiguratorBaseNode.cpp
@@ -28,14 +28,19 @@ void UDisplayClusterConfiguratorBaseNode::PostEditUndo()
 {
 	Super::PostEditUndo();
 
-	UpdateObject();
-
-	// Don't update the child nodes if this node is auto-positioned because this node is probably in the undo stack as part of
-	// a change to a child, and we don't want to overwrite the child's undo. If the children need updating, it will be handled on
-	// the next position tick.
-	if (!IsNodeAutoPositioned())
+	// Make sure the undo operation isn't going to a state where this node or the object it edits isn't valid
+	// before attempting to update the object or its children
+	if (IsObjectValid())
 	{
-		UpdateChildNodes();
+		UpdateObject();
+
+		// Don't update the child nodes if this node is auto-positioned because this node is probably in the undo stack as part of
+		// a change to a child, and we don't want to overwrite the child's undo. If the children need updating, it will be handled on
+		// the next position tick.
+		if (!IsNodeAutoPositioned())
+		{
+			UpdateChildNodes();
+		}
 	}
 }
 #endif
@@ -417,7 +422,7 @@ bool UDisplayClusterConfiguratorBaseNode::IsUserInteractingWithNode(bool bCheckD
 
 void UDisplayClusterConfiguratorBaseNode::UpdateNode()
 {
-	if (!IsPendingKill())
+	if (IsObjectValid())
 	{
 		ReadNodeStateFromObject();
 
@@ -427,7 +432,7 @@ void UDisplayClusterConfiguratorBaseNode::UpdateNode()
 
 void UDisplayClusterConfiguratorBaseNode::UpdateObject()
 {
-	if (!IsPendingKill())
+	if (IsObjectValid())
 	{
 		WriteNodeStateToObject();
 
@@ -438,6 +443,11 @@ void UDisplayClusterConfiguratorBaseNode::UpdateObject()
 	}
 }
 
+bool UDisplayClusterConfiguratorBaseNode::IsObjectValid() const
+{
+	return !IsPendingKill() && ObjectToEdit.IsValid();
+}
+
 void UDisplayClusterConfiguratorBaseNode::OnNodeAligned(bool bUpdateChildren)
 {
 	UpdateObject();
diff --git a/Engine/Plugins/Runtime/nDisplay/Source/DisplayClusterConfigurator/Private/Views/OutputMapping/EdNodes/DisplayClusterConfiguratorBaseNode.h b/Engine/Plugins/Runtime/nDisplay/Source/DisplayClusterConfigurator/Private/Views/OutputMapping/EdNodes/DisplayClusterConfiguratorBaseNode.h
--- a/Engine/Plugins/Runtime/nDisplay/Source/DisplayClusterConfigurator/Private/Views/OutputMapping/EdNodes/DisplayClusterConfiguratorBaseNode.h
+++ b/Engine/Plugins/Runtime/nDisplay/Source/DisplayClusterConfigurator/Private/Views/OutputMapping/EdNodes/DisplayClusterConfiguratorBaseNode.h
@@ -123,6 +123,7 @@ public:
 	virtual void UpdateNode();
 	virtual void UpdateObject();
 	virtual void DeleteObject() { }
+	virtual bool IsObjectValid() const;
 
 	virtual void OnNodeAligned(bool bUpdateChildren = false);
 
diff --git a/Engine/Plugins/Runtime/nDisplay/Source/DisplayClusterConfigurator/Private/Views/OutputMapping/EdNodes/DisplayClusterConfiguratorHostNode.cpp b/Engine/Plugins/Runtime/nDisplay/Source/DisplayClusterConfigurator/Private/Views/OutputMapping/EdNodes/DisplayClusterConfiguratorHostNode.cpp
--- a/Engine/Plugins/Runtime/nDisplay/Source/DisplayClusterConfigurator/Private/Views/OutputMapping/EdNodes/DisplayClusterConfiguratorHostNode.cpp
+++ b/Engine/Plugins/Runtime/nDisplay/Source/DisplayClusterConfigurator/Private/Views/OutputMapping/EdNodes/DisplayClusterConfiguratorHostNode.cpp
@@ -41,18 +41,33 @@ TSharedPtr<SGraphNode> UDisplayClusterConfiguratorHostNode::CreateVisualWidget()
 
 FLinearColor UDisplayClusterConfiguratorHostNode::GetHostColor() const
 {
+	if (!IsObjectValid())
+	{
+		return FLinearColor();
+	}
+
 	UDisplayClusterConfigurationHostDisplayData* HostDisplayData = GetObjectChecked<UDisplayClusterConfigurationHostDisplayData>();
 	return HostDisplayData->Color;
 }
 
 FText UDisplayClusterConfiguratorHostNode::GetHostName() const
 {
+	if (!IsObjectValid())
+	{
+		return FText::GetEmpty();
+	}
+
 	UDisplayClusterConfigurationHostDisplayData* HostDisplayData = GetObjectChecked<UDisplayClusterConfigurationHostDisplayData>();
 	return HostDisplayData->HostName;
 }
 
 FVector2D UDisplayClusterConfiguratorHostNode::GetHostOrigin(bool bInGlobalCoordinates) const
 {
+	if (!IsObjectValid())
+	{
+		return FVector2D::ZeroVector;
+	}
+
 	UDisplayClusterConfigurationHostDisplayData* HostDisplayData = GetObjectChecked<UDisplayClusterConfigurationHostDisplayData>();
 	float Scale = bInGlobalCoordinates ? GetViewScale() : 1.0f;
 
@@ -61,6 +76,11 @@ FVector2D UDisplayClusterConfiguratorHostNode::GetHostOrigin(bool bInGlobalCoord
 
 void UDisplayClusterConfiguratorHostNode::SetHostOrigin(const FVector2D& NewOrigin, bool bInGlobalCoordinates)
 {
+	if (!IsObjectValid())
+	{
+		return;
+	}
+
 	UDisplayClusterConfigurationHostDisplayData* HostDisplayData = GetObjectChecked<UDisplayClusterConfigurationHostDisplayData>();
 
 	float Scale = bInGlobalCoordinates ? GetViewScale() : 1.0f;
@@ -74,12 +94,22 @@ void UDisplayClusterConfiguratorHostNode::SetHostOrigin(const FVector2D& NewOrig
 
 bool UDisplayClusterConfiguratorHostNode::CanUserMoveNode() const
 {
+	if (!IsObjectValid())
+	{
+		return false;
+	}
+
 	UDisplayClusterConfigurationHostDisplayData* HostDisplayData = GetObjectChecked<UDisplayClusterConfigurationHostDisplayData>();
 	return HostDisplayData->bAllowManualPlacement;
 }
 
 bool UDisplayClusterConfiguratorHostNode::CanUserResizeNode() const
 {
+	if (!IsObjectValid())
+	{
+		return false;
+	}
+
 	UDisplayClusterConfigurationHostDisplayData* HostDisplayData = GetObjectChecked<UDisplayClusterConfigurationHostDisplayData>();
 	return HostDisplayData->bAllowManualSizing;
 }
@@ -132,6 +162,11 @@ FNodeAlignmentAnchors UDisplayClusterConfiguratorHostNode::GetNodeAlignmentAncho
 }
 bool UDisplayClusterConfiguratorHostNode::IsNodeVisible() const
 {
+	if (!IsObjectValid())
+	{
+		return false;
+	}
+
 	UDisplayClusterConfigurationHostDisplayData* HostDisplayData = GetObjectChecked<UDisplayClusterConfigurationHostDisplayData>();
 
 	if (HostDisplayData->bIsVisible)
@@ -155,6 +190,11 @@ bool UDisplayClusterConfiguratorHostNode::IsNodeVisible() const
 
 bool UDisplayClusterConfiguratorHostNode::IsNodeEnabled() const
 {
+	if (!IsObjectValid())
+	{
+		return false;
+	}
+
 	UDisplayClusterConfigurationHostDisplayData* HostDisplayData = GetObjectChecked<UDisplayClusterConfigurationHostDisplayData>();
 
 	if (HostDisplayData->bIsEnabled)
@@ -218,6 +258,12 @@ void UDisplayClusterConfiguratorHostNode::OnPostEditChangeChainProperty(const FP
 		return;
 	}
 
+	// If the object is no longer valid, don't attempt to sync properties
+	if (!IsObjectValid())
+	{
+		return;
+	}
+
 	UDisplayClusterConfigurationHostDisplayData* HostDisplayData = GetObjectChecked<UDisplayClusterConfigurationHostDisplayData>();
 
 	const FName& PropertyName = PropertyChangedEvent.PropertyChain.GetActiveMemberNode()->GetValue()->GetFName();
diff --git a/Engine/Plugins/Runtime/nDisplay/Source/DisplayClusterConfigurator/Private/Views/OutputMapping/EdNodes/DisplayClusterConfiguratorViewportNode.cpp b/Engine/Plugins/Runtime/nDisplay/Source/DisplayClusterConfigurator/Private/Views/OutputMapping/EdNodes/DisplayClusterConfiguratorViewportNode.cpp
--- a/Engine/Plugins/Runtime/nDisplay/Source/DisplayClusterConfigurator/Private/Views/OutputMapping/EdNodes/DisplayClusterConfiguratorViewportNode.cpp
+++ b/Engine/Plugins/Runtime/nDisplay/Source/DisplayClusterConfigurator/Private/Views/OutputMapping/EdNodes/DisplayClusterConfiguratorViewportNode.cpp
@@ -40,18 +40,33 @@ TSharedPtr<SGraphNode> UDisplayClusterConfiguratorViewportNode::CreateVisualWidg
 
 bool UDisplayClusterConfiguratorViewportNode::IsNodeVisible() const
 {
+	if (!IsObjectValid())
+	{
+		return false;
+	}
+
 	UDisplayClusterConfigurationViewport* Viewport = GetObjectChecked<UDisplayClusterConfigurationViewport>();
 	return Viewport->bIsVisible;
 }
 
 bool UDisplayClusterConfiguratorViewportNode::IsNodeEnabled() const
 {
+	if (!IsObjectValid())
+	{
+		return false;
+	}
+
 	UDisplayClusterConfigurationViewport* Viewport = GetObjectChecked<UDisplayClusterConfigurationViewport>();
 	return Viewport->bIsEnabled;
 }
 
 void UDisplayClusterConfiguratorViewportNode::DeleteObject()
 {
+	if (!IsObjectValid())
+	{
+		return;
+	}
+
 	UDisplayClusterConfigurationViewport* Viewport = GetObjectChecked<UDisplayClusterConfigurationViewport>();
 	FDisplayClusterConfiguratorClusterUtils::RemoveViewportFromClusterNode(Viewport);
 }
@@ -86,6 +101,11 @@ const FDisplayClusterConfigurationRectangle& UDisplayClusterConfiguratorViewport
 
 bool UDisplayClusterConfiguratorViewportNode::IsFixedAspectRatio() const
 {
+	if (!IsObjectValid())
+	{
+		return false;
+	}
+
 	UDisplayClusterConfigurationViewport* CfgViewport = GetObjectChecked<UDisplayClusterConfigurationViewport>();
 	return CfgViewport->bFixedAspectRatio;
 }
@@ -118,6 +138,12 @@ void UDisplayClusterConfiguratorViewportNode::OnPostEditChangeChainProperty(cons
 	{
 		return;
 	}
+	
+	// If the object is no longer valid, don't attempt to sync properties
+	if (!IsObjectValid())
+	{
+		return;
+	}
 
 	const UDisplayClusterConfigurationViewport* CfgViewport = GetObjectChecked<UDisplayClusterConfigurationViewport>();
 
diff --git a/Engine/Plugins/Runtime/nDisplay/Source/DisplayClusterConfigurator/Private/Views/OutputMapping/EdNodes/DisplayClusterConfiguratorWindowNode.cpp b/Engine/Plugins/Runtime/nDisplay/Source/DisplayClusterConfigurator/Private/Views/OutputMapping/EdNodes/DisplayClusterConfiguratorWindowNode.cpp
--- a/Engine/Plugins/Runtime/nDisplay/Source/DisplayClusterConfigurator/Private/Views/OutputMapping/EdNodes/DisplayClusterConfiguratorWindowNode.cpp
+++ b/Engine/Plugins/Runtime/nDisplay/Source/DisplayClusterConfigurator/Private/Views/OutputMapping/EdNodes/DisplayClusterConfiguratorWindowNode.cpp
@@ -44,6 +44,11 @@ const FDisplayClusterConfigurationRectangle& UDisplayClusterConfiguratorWindowNo
 
 FString UDisplayClusterConfiguratorWindowNode::GetCfgHost() const
 {
+	if (!IsObjectValid())
+	{
+		return TEXT("");
+	}
+
 	UDisplayClusterConfigurationClusterNode* CfgClusterNode = GetObjectChecked<UDisplayClusterConfigurationClusterNode>();
 	return CfgClusterNode->Host;
 }
@@ -56,12 +61,22 @@ const FString& UDisplayClusterConfiguratorWindowNode::GetPreviewImagePath() cons
 
 bool UDisplayClusterConfiguratorWindowNode::IsFixedAspectRatio() const
 {
+	if (!IsObjectValid())
+	{
+		return false;
+	}
+
 	UDisplayClusterConfigurationClusterNode* CfgClusterNode = GetObjectChecked<UDisplayClusterConfigurationClusterNode>();
 	return CfgClusterNode->bFixedAspectRatio;
 }
 
 bool UDisplayClusterConfiguratorWindowNode::IsMaster() const
 {
+	if (!IsObjectValid())
+	{
+		return false;
+	}
+
 	UDisplayClusterConfigurationClusterNode* ClusterNode = GetObjectChecked<UDisplayClusterConfigurationClusterNode>();
 	return FDisplayClusterConfiguratorClusterUtils::IsClusterNodeMaster(ClusterNode);
 }
@@ -78,6 +93,11 @@ void UDisplayClusterConfiguratorWindowNode::UnregisterOnPreviewImageChanged(FDel
 
 bool UDisplayClusterConfiguratorWindowNode::IsNodeVisible() const
 {
+	if (!IsObjectValid())
+	{
+		return false;
+	}
+
 	UDisplayClusterConfigurationClusterNode* ClusterNode = GetObjectChecked<UDisplayClusterConfigurationClusterNode>();
 
 	if (ClusterNode->bIsVisible)
@@ -101,6 +121,11 @@ bool UDisplayClusterConfiguratorWindowNode::IsNodeVisible() const
 
 bool UDisplayClusterConfiguratorWindowNode::IsNodeEnabled() const
 {
+	if (!IsObjectValid())
+	{
+		return false;
+	}
+
 	UDisplayClusterConfigurationClusterNode* ClusterNode = GetObjectChecked<UDisplayClusterConfigurationClusterNode>();
 
 	if (ClusterNode->bIsEnabled)
@@ -133,6 +158,11 @@ bool UDisplayClusterConfiguratorWindowNode::CanNodeExceedParentBounds() const
 
 void UDisplayClusterConfiguratorWindowNode::DeleteObject()
 {
+	if (!IsObjectValid())
+	{
+		return;
+	}
+
 	UDisplayClusterConfigurationClusterNode* ClusterNode = GetObjectChecked<UDisplayClusterConfigurationClusterNode>();
 	FDisplayClusterConfiguratorClusterUtils::RemoveClusterNodeFromCluster(ClusterNode);
 }
@@ -168,6 +198,12 @@ void UDisplayClusterConfiguratorWindowNode::OnPostEditChangeChainProperty(const
 		return;
 	}
 
+	// If the object is no longer valid, don't attempt to sync properties
+	if (!IsObjectValid())
+	{
+		return;
+	}
+
 	UDisplayClusterConfigurationClusterNode* CfgClusterNode = GetObjectChecked<UDisplayClusterConfigurationClusterNode>();
 
 	const FName& PropertyName = PropertyChangedEvent.GetPropertyName();
