nDisplay: Added delegate that allows the level instance property details panels to refresh their details customization panels when an nDisplay config's cluster heirarchy is changed

diff --git a/Engine/Plugins/Runtime/nDisplay/Source/DisplayClusterConfigurator/Private/DisplayClusterConfiguratorBlueprintEditor.cpp b/Engine/Plugins/Runtime/nDisplay/Source/DisplayClusterConfigurator/Private/DisplayClusterConfiguratorBlueprintEditor.cpp
--- a/Engine/Plugins/Runtime/nDisplay/Source/DisplayClusterConfigurator/Private/DisplayClusterConfiguratorBlueprintEditor.cpp
+++ b/Engine/Plugins/Runtime/nDisplay/Source/DisplayClusterConfigurator/Private/DisplayClusterConfiguratorBlueprintEditor.cpp
@@ -36,6 +36,7 @@
 #include "GameFramework/Actor.h"
 
 #include "EditorDirectories.h"
+#include "EditorSupportDelegates.h"
 #include "EditorViewportTabContent.h"
 #include "ISCSEditorUICustomization.h"
 #include "SBlueprintEditorToolbar.h"
@@ -197,6 +198,30 @@ void FDisplayClusterConfiguratorBlueprintEditor::InitDisplayClusterBlueprintEdit
 	FocusChangedHandle = FSlateApplication::Get().OnFocusChanging().AddSP(this, &FDisplayClusterConfiguratorBlueprintEditor::OnFocusChanged);
 }
 
+void FDisplayClusterConfiguratorBlueprintEditor::PostUndo(bool bSuccess)
+{
+	FBlueprintEditor::PostUndo(bSuccess);
+
+	// Make sure to force any property window displaying this actor class to refresh in case 
+	// the cluster hierarchy was changed in the undo transaction
+	if (ADisplayClusterRootActor* Actor = Cast<ADisplayClusterRootActor>(GetPreviewActor()))
+	{
+		FEditorSupportDelegates::ForcePropertyWindowRebuild.Broadcast(Actor->GetClass());
+	}
+}
+
+void FDisplayClusterConfiguratorBlueprintEditor::PostRedo(bool bSuccess)
+{
+	FBlueprintEditor::PostRedo(bSuccess);
+
+	// Make sure to force any property window displaying this actor class to refresh in case 
+	// the cluster hierarchy was changed in the redo transaction
+	if (ADisplayClusterRootActor* Actor = Cast<ADisplayClusterRootActor>(GetPreviewActor()))
+	{
+		FEditorSupportDelegates::ForcePropertyWindowRebuild.Broadcast(Actor->GetClass());
+	}
+}
+
 UDisplayClusterConfigurationData* FDisplayClusterConfiguratorBlueprintEditor::GetEditorData() const
 {
 	if (!LoadedBlueprint.IsValid() || !LoadedBlueprint->GeneratedClass)
@@ -375,6 +400,7 @@ void FDisplayClusterConfiguratorBlueprintEditor::ClusterChanged(bool bStructureC
 	if (ADisplayClusterRootActor* Actor = Cast<ADisplayClusterRootActor>(GetPreviewActor()))
 	{
 		Actor->UpdatePreviewComponents();
+		FEditorSupportDelegates::ForcePropertyWindowRebuild.Broadcast(Actor->GetClass());
 	}
 
 	OnClusterChanged.Broadcast();
diff --git a/Engine/Plugins/Runtime/nDisplay/Source/DisplayClusterConfigurator/Private/DisplayClusterConfiguratorBlueprintEditor.h b/Engine/Plugins/Runtime/nDisplay/Source/DisplayClusterConfigurator/Private/DisplayClusterConfiguratorBlueprintEditor.h
--- a/Engine/Plugins/Runtime/nDisplay/Source/DisplayClusterConfigurator/Private/DisplayClusterConfiguratorBlueprintEditor.h
+++ b/Engine/Plugins/Runtime/nDisplay/Source/DisplayClusterConfigurator/Private/DisplayClusterConfiguratorBlueprintEditor.h
@@ -39,6 +39,11 @@ public:
 	virtual bool IsObjectSelected(UObject* Obj) const override;
 	//~ End IDisplayClusterConfiguratorBlueprintEditor Interface
 
+	//~ Begin FEditorUndoClient interface
+	virtual void PostUndo(bool bSuccess) override;
+	virtual void PostRedo(bool bSuccess) override;
+	//~ End FEditorUndoClient interface
+
 	virtual void SelectObjects(TArray<UObject*>& InSelectedObjects, bool bFullRefresh = false);
 	virtual void SelectAncillaryComponents(const TArray<FString>& ComponentNames);
 	virtual void SelectAncillaryViewports(const TArray<FString>& ComponentNames);
diff --git a/Engine/Plugins/Runtime/nDisplay/Source/DisplayClusterConfigurator/Private/Views/Details/DisplayClusterRootActorDetailsCustomization.cpp b/Engine/Plugins/Runtime/nDisplay/Source/DisplayClusterConfigurator/Private/Views/Details/DisplayClusterRootActorDetailsCustomization.cpp
--- a/Engine/Plugins/Runtime/nDisplay/Source/DisplayClusterConfigurator/Private/Views/Details/DisplayClusterRootActorDetailsCustomization.cpp
+++ b/Engine/Plugins/Runtime/nDisplay/Source/DisplayClusterConfigurator/Private/Views/Details/DisplayClusterRootActorDetailsCustomization.cpp
@@ -8,6 +8,7 @@
 #include "DetailWidgetRow.h"
 #include "PropertyHandle.h"
 
+#include "EditorSupportDelegates.h"
 #include "SSearchableComboBox.h"
 #include "Widgets/Text/STextBlock.h"
 #include "IDetailGroup.h"
@@ -26,8 +27,16 @@ TSharedRef<IDetailCustomization> FDisplayClusterRootActorDetailsCustomization::M
 	return MakeShared<FDisplayClusterRootActorDetailsCustomization>();
 }
 
+FDisplayClusterRootActorDetailsCustomization::~FDisplayClusterRootActorDetailsCustomization()
+{
+	FEditorSupportDelegates::ForcePropertyWindowRebuild.Remove(ForcePropertyWindowRebuildHandle);
+}
+
 void FDisplayClusterRootActorDetailsCustomization::CustomizeDetails(IDetailLayoutBuilder& InLayoutBuilder)
 {
+	LayoutBuilder = &InLayoutBuilder;
+	ForcePropertyWindowRebuildHandle = FEditorSupportDelegates::ForcePropertyWindowRebuild.AddSP(this, &FDisplayClusterRootActorDetailsCustomization::OnForcePropertyWindowRebuild);
+
 	// Hide the following categories, we don't really need them
 	InLayoutBuilder.HideCategory(TEXT("Rendering"));
 	InLayoutBuilder.HideCategory(TEXT("Replication"));
@@ -355,4 +364,15 @@ TSharedRef<SWidget> FDisplayClusterRootActorDetailsCustomization::CreateComboWid
 	return SNew(STextBlock).Text(FText::FromString(*InItem));
 }
 
+void FDisplayClusterRootActorDetailsCustomization::OnForcePropertyWindowRebuild(UObject* Object)
+{
+	if (LayoutBuilder && EditedObject.IsValid())
+	{
+		if (EditedObject->GetClass() == Object)
+		{
+			LayoutBuilder->ForceRefreshDetails();
+		}
+	}
+}
+
 #undef LOCTEXT_NAMESPACE
diff --git a/Engine/Plugins/Runtime/nDisplay/Source/DisplayClusterConfigurator/Private/Views/Details/DisplayClusterRootActorDetailsCustomization.h b/Engine/Plugins/Runtime/nDisplay/Source/DisplayClusterConfigurator/Private/Views/Details/DisplayClusterRootActorDetailsCustomization.h
--- a/Engine/Plugins/Runtime/nDisplay/Source/DisplayClusterConfigurator/Private/Views/Details/DisplayClusterRootActorDetailsCustomization.h
+++ b/Engine/Plugins/Runtime/nDisplay/Source/DisplayClusterConfigurator/Private/Views/Details/DisplayClusterRootActorDetailsCustomization.h
@@ -20,6 +20,8 @@ class FDisplayClusterRootActorDetailsCustomization : public IDetailCustomization
 public:
 	static TSharedRef<IDetailCustomization> MakeInstance();
 
+	~FDisplayClusterRootActorDetailsCustomization();
+
 protected:
 	// IDetailCustomization interface
 	virtual void CustomizeDetails(IDetailLayoutBuilder& InLayoutBuilder) override;
@@ -71,7 +73,12 @@ protected:
 	// Create combobox widget
 	TSharedRef<SWidget> CreateComboWidget(TSharedPtr<FString> InItem);
 
+	void OnForcePropertyWindowRebuild(UObject* Object);
+
 protected:
 	// ADisplayClusterRootActor on which we're acting
 	TWeakObjectPtr<ADisplayClusterRootActor> EditedObject;
+	IDetailLayoutBuilder* LayoutBuilder;
+
+	FDelegateHandle ForcePropertyWindowRebuildHandle;
 };
