RC fix replication with if condition of the event

diff --git a/Engine/Plugins/Runtime/nDisplay/Source/DisplayClusterRemoteControlInterceptor/Private/DisplayClusterRemoteControlInterceptor.cpp b/Engine/Plugins/Runtime/nDisplay/Source/DisplayClusterRemoteControlInterceptor/Private/DisplayClusterRemoteControlInterceptor.cpp
--- a/Engine/Plugins/Runtime/nDisplay/Source/DisplayClusterRemoteControlInterceptor/Private/DisplayClusterRemoteControlInterceptor.cpp
+++ b/Engine/Plugins/Runtime/nDisplay/Source/DisplayClusterRemoteControlInterceptor/Private/DisplayClusterRemoteControlInterceptor.cpp
@@ -114,7 +114,7 @@ ERCIResponse FDisplayClusterRemoteControlInterceptor::InvokeCall(FRCIFunctionMet
 void FDisplayClusterRemoteControlInterceptor::OnClusterEventBinaryHandler(const FDisplayClusterClusterEventBinary& Event)
 {
 	// Dispatch data to a proper handler
-	if (Event.bIsSystemEvent)
+	if (Event.bIsSystemEvent && Event.EventId == EventId_InterceptorQueue)
 	{
 		// Deserialize command data
 		TMap<FName, TMap<FString, TArray<uint8>>> ReceivedInterceptQueueMap;
