nDisplay: Add preview option to disable ICVFX camera preview window.

diff --git a/Engine/Plugins/Runtime/nDisplay/Source/DisplayCluster/Private/Game/EngineClasses/Scene/DisplayClusterICVFXCameraComponent.cpp b/Engine/Plugins/Runtime/nDisplay/Source/DisplayCluster/Private/Game/EngineClasses/Scene/DisplayClusterICVFXCameraComponent.cpp
--- a/Engine/Plugins/Runtime/nDisplay/Source/DisplayCluster/Private/Game/EngineClasses/Scene/DisplayClusterICVFXCameraComponent.cpp
+++ b/Engine/Plugins/Runtime/nDisplay/Source/DisplayCluster/Private/Game/EngineClasses/Scene/DisplayClusterICVFXCameraComponent.cpp
@@ -17,6 +17,18 @@ FString UDisplayClusterICVFXCameraComponent::GetCameraUniqueId() const
 	return GetFName().ToString();
 }
 
+#if WITH_EDITOR
+bool UDisplayClusterICVFXCameraComponent::GetEditorPreviewInfo(float DeltaTime, FMinimalViewInfo& ViewOut)
+{
+	if (ADisplayClusterRootActor* RootActor = Cast<ADisplayClusterRootActor>(GetOwner()))
+	{
+		return RootActor->bEnableICVFXCameraPreview;
+	}
+
+	return Super::GetEditorPreviewInfo(DeltaTime, ViewOut);
+}
+#endif
+
 FDisplayClusterViewport_CameraMotionBlur UDisplayClusterICVFXCameraComponent::GetMotionBlurParameters()
 {
 	FDisplayClusterViewport_CameraMotionBlur OutParameters;
diff --git a/Engine/Plugins/Runtime/nDisplay/Source/DisplayCluster/Public/Components/DisplayClusterICVFXCameraComponent.h b/Engine/Plugins/Runtime/nDisplay/Source/DisplayCluster/Public/Components/DisplayClusterICVFXCameraComponent.h
--- a/Engine/Plugins/Runtime/nDisplay/Source/DisplayCluster/Public/Components/DisplayClusterICVFXCameraComponent.h
+++ b/Engine/Plugins/Runtime/nDisplay/Source/DisplayCluster/Public/Components/DisplayClusterICVFXCameraComponent.h
@@ -37,6 +37,10 @@ public:
 	UPROPERTY(EditAnywhere, Category = NDisplay, meta = (ShowInnerProperties))
 	FDisplayClusterConfigurationICVFX_CameraSettings CameraSettings;
 
+#if WITH_EDITOR
+	virtual bool GetEditorPreviewInfo(float DeltaTime, FMinimalViewInfo& ViewOut) override;
+#endif
+	
 public:
 	FDisplayClusterViewport_CameraMotionBlur GetMotionBlurParameters();
 
diff --git a/Engine/Plugins/Runtime/nDisplay/Source/DisplayCluster/Public/DisplayClusterRootActor.h b/Engine/Plugins/Runtime/nDisplay/Source/DisplayCluster/Public/DisplayClusterRootActor.h
--- a/Engine/Plugins/Runtime/nDisplay/Source/DisplayCluster/Public/DisplayClusterRootActor.h
+++ b/Engine/Plugins/Runtime/nDisplay/Source/DisplayCluster/Public/DisplayClusterRootActor.h
@@ -255,6 +255,10 @@ public:
 	UPROPERTY(EditAnywhere, Category = "Editor Preview", meta = (ClampMin = "0.05", UIMin = "0.05", ClampMax = "1", UIMax = "1"))
 	float PreviewRenderTargetRatioMult = 0.25;
 
+	/* Enable the on-screen camera preview widget when selecting this actor. Changing this value requires reselecting the actor. */
+	UPROPERTY(EditAnywhere, Category = "Editor Preview")
+	bool bEnableICVFXCameraPreview = false;
+	
 	UPROPERTY(EditAnywhere, Category = "Editor Preview")
 	float XformGizmoScale = 1.0f;
 
