nDisplay: Fixed issue that was causing erroneous logs about missing camera to be output when a display cluster root actor was being placed in a level.

diff --git a/Engine/Plugins/Runtime/nDisplay/Source/DisplayCluster/Private/Game/EngineClasses/Scene/DisplayClusterRootActor.cpp b/Engine/Plugins/Runtime/nDisplay/Source/DisplayCluster/Private/Game/EngineClasses/Scene/DisplayClusterRootActor.cpp
--- a/Engine/Plugins/Runtime/nDisplay/Source/DisplayCluster/Private/Game/EngineClasses/Scene/DisplayClusterRootActor.cpp
+++ b/Engine/Plugins/Runtime/nDisplay/Source/DisplayCluster/Private/Game/EngineClasses/Scene/DisplayClusterRootActor.cpp
@@ -481,23 +481,28 @@ bool ADisplayClusterRootActor::BuildHierarchy()
 		}
 	}
 
-	// If no default camera set, try to set the first one
-	if (DefaultCameraComponent.IsDefinedSceneComponent() == false)
+	// Only check for the default camera if the construction scripts have been run, since camera components created in the blueprint editor
+	// won't have been created before then
+	if (bHasRerunConstructionScripts)
 	{
-		TMap<FString, UDisplayClusterCameraComponent*> Cameras;
-		GetTypedComponents<UDisplayClusterCameraComponent>(Cameras, CameraComponents);
-		if (Cameras.Num() > 0)
+		// If no default camera set, try to set the first one
+		if (DefaultCameraComponent.IsDefinedSceneComponent() == false)
 		{
-			// There is no guarantee that default camera is the first one listed in a config file
-			SetDefaultCamera(Cameras.CreateConstIterator()->Key);
-		}
-		else
-		{
-			UE_LOG(LogDisplayClusterGame, Error, TEXT("No cameras found"));
-			return false;
+			TMap<FString, UDisplayClusterCameraComponent*> Cameras;
+			GetTypedComponents<UDisplayClusterCameraComponent>(Cameras, CameraComponents);
+			if (Cameras.Num() > 0)
+			{
+				// There is no guarantee that default camera is the first one listed in a config file
+				SetDefaultCamera(Cameras.CreateConstIterator()->Key);
+			}
+			else
+			{
+				UE_LOG(LogDisplayClusterGame, Error, TEXT("No cameras found"));
+				return false;
+			}
 		}
 	}
-	
+
 	return true;
 }
 
@@ -695,6 +700,8 @@ void ADisplayClusterRootActor::RerunConstructionScripts()
 {
 	Super::RerunConstructionScripts();
 
+	bHasRerunConstructionScripts = true;
+
 #if WITH_EDITOR
 	RerunConstructionScripts_Editor();
 #endif
diff --git a/Engine/Plugins/Runtime/nDisplay/Source/DisplayCluster/Public/DisplayClusterRootActor.h b/Engine/Plugins/Runtime/nDisplay/Source/DisplayCluster/Public/DisplayClusterRootActor.h
--- a/Engine/Plugins/Runtime/nDisplay/Source/DisplayCluster/Public/DisplayClusterRootActor.h
+++ b/Engine/Plugins/Runtime/nDisplay/Source/DisplayCluster/Public/DisplayClusterRootActor.h
@@ -210,6 +210,8 @@ private:
 	TMap<FString, FDisplayClusterSceneComponentRef*> MeshComponents;
 	FDisplayClusterSceneComponentRef DefaultCameraComponent;
 
+	bool bHasRerunConstructionScripts = false;
+
 private:
 	template <typename TComp>
 	void GetTypedPrimitives(TSet<FPrimitiveComponentId>& OutPrimitives, const TArray<FString>* InCompNames = nullptr, bool bCollectChildrenVisualizationComponent = true) const;
