nDisplay: Added metadata flags to PPS color grading vectors to allow them to use the color grading UI in details panels, and set the global PPS enabled setting to default to true instead of false.

diff --git a/Engine/Plugins/Runtime/nDisplay/Source/DisplayClusterConfiguration/Public/DisplayClusterConfigurationTypes.h b/Engine/Plugins/Runtime/nDisplay/Source/DisplayClusterConfiguration/Public/DisplayClusterConfigurationTypes.h
--- a/Engine/Plugins/Runtime/nDisplay/Source/DisplayClusterConfiguration/Public/DisplayClusterConfigurationTypes.h
+++ b/Engine/Plugins/Runtime/nDisplay/Source/DisplayClusterConfiguration/Public/DisplayClusterConfigurationTypes.h
@@ -441,7 +441,7 @@ public:
 
 	// Apply the global cluster post process settings to all viewports
 	UPROPERTY(EditAnywhere, Category = NDisplay)
-	bool bUseOverallClusterPostProcess = false;
+	bool bUseOverallClusterPostProcess = true;
 
 	// Global cluster post process settings
 	UPROPERTY(EditAnywhere, Category = NDisplay, meta = (EditCondition = "bUseOverallClusterPostProcess"))
diff --git a/Engine/Plugins/Runtime/nDisplay/Source/DisplayClusterConfiguration/Public/DisplayClusterConfigurationTypes_Viewport.h b/Engine/Plugins/Runtime/nDisplay/Source/DisplayClusterConfiguration/Public/DisplayClusterConfigurationTypes_Viewport.h
--- a/Engine/Plugins/Runtime/nDisplay/Source/DisplayClusterConfiguration/Public/DisplayClusterConfigurationTypes_Viewport.h
+++ b/Engine/Plugins/Runtime/nDisplay/Source/DisplayClusterConfiguration/Public/DisplayClusterConfigurationTypes_Viewport.h
@@ -175,23 +175,23 @@ struct DISPLAYCLUSTERCONFIGURATION_API FDisplayClusterConfigurationViewport_Colo
 	uint8 bOverride_Offset:1;
 
 	// Saturation
-	UPROPERTY(EditAnywhere, Category = "Color Grading", meta = (UIMin = "0.0", UIMax = "2.0", Delta = "0.01", EditCondition = "bOverride_Saturation"))
+	UPROPERTY(EditAnywhere, Category = "Color Grading", meta = (UIMin = "0.0", UIMax = "2.0", Delta = "0.01", EditCondition = "bOverride_Saturation", ColorGradingMode = "saturation"))
 	FVector4 Saturation = FVector4(1.0f, 1.0f, 1.0f, 1.0f);
 
 	// Contrast
-	UPROPERTY(EditAnywhere, Category = "Color Grading", meta = (UIMin = "0.0", UIMax = "2.0", Delta = "0.01", EditCondition = "bOverride_Contrast"))
+	UPROPERTY(EditAnywhere, Category = "Color Grading", meta = (UIMin = "0.0", UIMax = "2.0", Delta = "0.01", EditCondition = "bOverride_Contrast", ColorGradingMode = "contrast"))
 	FVector4 Contrast = FVector4(1.0f, 1.0f, 1.0f, 1.0f);
 
 	// Gamma
-	UPROPERTY(EditAnywhere, Category = "Color Grading", meta = (UIMin = "0.0", UIMax = "2.0", Delta = "0.01", EditCondition = "bOverride_Gamma"))
+	UPROPERTY(EditAnywhere, Category = "Color Grading", meta = (UIMin = "0.0", UIMax = "2.0", Delta = "0.01", EditCondition = "bOverride_Gamma", ColorGradingMode = "gamma"))
 	FVector4 Gamma = FVector4(1.0f, 1.0f, 1.0f, 1.0f);
 
 	// Gain
-	UPROPERTY(EditAnywhere, Category = "Color Grading", meta = (UIMin = "0.0", UIMax = "2.0", Delta = "0.01", EditCondition = "bOverride_Gain"))
+	UPROPERTY(EditAnywhere, Category = "Color Grading", meta = (UIMin = "0.0", UIMax = "2.0", Delta = "0.01", EditCondition = "bOverride_Gain", ColorGradingMode = "gain"))
 	FVector4 Gain = FVector4(1.0f, 1.0f, 1.0f, 1.0f);
 
 	// Offset
-	UPROPERTY(EditAnywhere, Category = "Color Grading", meta = (UIMin = "-1.0", UIMax = "1.0", Delta = "0.001", EditCondition = "bOverride_Offset"))
+	UPROPERTY(EditAnywhere, Category = "Color Grading", meta = (UIMin = "-1.0", UIMax = "1.0", Delta = "0.001", EditCondition = "bOverride_Offset", ColorGradingMode = "offset"))
 	FVector4 Offset = FVector4(0.0f, 0.0f, 0.0f, 0.0f);
 };
 
