nDisplay: Fix details level crash when selecting an empty root actor. This is only a partial fix. The issue is the native root actor is allowed to be placed in the world when it shouldn't be.

diff --git a/Engine/Plugins/Runtime/nDisplay/Source/DisplayClusterConfigurator/Private/Views/Details/DisplayClusterRootActorDetailsCustomization.cpp b/Engine/Plugins/Runtime/nDisplay/Source/DisplayClusterConfigurator/Private/Views/Details/DisplayClusterRootActorDetailsCustomization.cpp
--- a/Engine/Plugins/Runtime/nDisplay/Source/DisplayClusterConfigurator/Private/Views/Details/DisplayClusterRootActorDetailsCustomization.cpp
+++ b/Engine/Plugins/Runtime/nDisplay/Source/DisplayClusterConfigurator/Private/Views/Details/DisplayClusterRootActorDetailsCustomization.cpp
@@ -86,7 +86,7 @@ void FDisplayClusterRootActorDetailsCustomization::BuildLayout(IDetailLayoutBuil
 	
 	FDisplayClusterConfiguratorNestedPropertyHelper NestedPropertyHelper(InLayoutBuilder);
 	
-	if (EditedObject.IsValid() && !EditedObject->IsTemplate(RF_ClassDefaultObject))
+	if (EditedObject.IsValid() && !EditedObject->IsTemplate(RF_ClassDefaultObject) && EditedObject->CurrentConfigData)
 	{
 		// Build ICVFX category.
 		
