ndisplay: exposed nDisplay ICVFX Motion Blur structs into BP's

diff --git a/Engine/Plugins/Runtime/nDisplay/Source/DisplayCluster/Public/Components/DisplayClusterICVFXCameraComponent.h b/Engine/Plugins/Runtime/nDisplay/Source/DisplayCluster/Public/Components/DisplayClusterICVFXCameraComponent.h
--- a/Engine/Plugins/Runtime/nDisplay/Source/DisplayCluster/Public/Components/DisplayClusterICVFXCameraComponent.h
+++ b/Engine/Plugins/Runtime/nDisplay/Source/DisplayCluster/Public/Components/DisplayClusterICVFXCameraComponent.h
@@ -31,10 +31,10 @@ public:
 
 public:
 	// Use external cine camera actor
-	UPROPERTY(EditAnywhere, Category = NDisplay)
+	UPROPERTY(BlueprintReadWrite, EditAnywhere, Category = NDisplay)
 	TSoftObjectPtr<ACineCameraActor> ExternalCameraActor;
 
-	UPROPERTY(EditAnywhere, Category = NDisplay, meta = (ShowInnerProperties))
+	UPROPERTY(BlueprintReadWrite, EditAnywhere, Category = NDisplay, meta = (ShowInnerProperties))
 	FDisplayClusterConfigurationICVFX_CameraSettings CameraSettings;
 
 #if WITH_EDITOR
diff --git a/Engine/Plugins/Runtime/nDisplay/Source/DisplayClusterConfiguration/Public/DisplayClusterConfigurationTypes_ICVFX.h b/Engine/Plugins/Runtime/nDisplay/Source/DisplayClusterConfiguration/Public/DisplayClusterConfigurationTypes_ICVFX.h
--- a/Engine/Plugins/Runtime/nDisplay/Source/DisplayClusterConfiguration/Public/DisplayClusterConfigurationTypes_ICVFX.h
+++ b/Engine/Plugins/Runtime/nDisplay/Source/DisplayClusterConfiguration/Public/DisplayClusterConfigurationTypes_ICVFX.h
@@ -244,7 +244,7 @@ public:
 	int RenderFamilyGroup = -1;
 };
 
-USTRUCT()
+USTRUCT(BlueprintType)
 struct DISPLAYCLUSTERCONFIGURATION_API FDisplayClusterConfigurationICVFX_CameraRenderSettings
 {
 	GENERATED_BODY()
@@ -293,23 +293,23 @@ public:
 
 
 
-USTRUCT()
+USTRUCT(BlueprintType)
 struct DISPLAYCLUSTERCONFIGURATION_API FDisplayClusterConfigurationICVFX_CameraMotionBlur
 {
 	GENERATED_BODY()
 
 public:
-	UPROPERTY(EditAnywhere, Category = NDisplay)
+	UPROPERTY(BlueprintReadWrite, EditAnywhere, Category = NDisplay)
 	EDisplayClusterConfigurationCameraMotionBlurMode MotionBlurMode = EDisplayClusterConfigurationCameraMotionBlurMode::Off;
 
-	UPROPERTY(EditAnywhere, Category = NDisplay)
+	UPROPERTY(BlueprintReadWrite, EditAnywhere, Category = NDisplay)
 	float TranslationScale = 1.f;
 
 	// GUI: Add ext camera refs
 };
 
 
-USTRUCT()
+USTRUCT(BlueprintType)
 struct DISPLAYCLUSTERCONFIGURATION_API FDisplayClusterConfigurationICVFX_CameraSettings
 {
 	GENERATED_BODY()
@@ -338,7 +338,7 @@ public:
 	UPROPERTY(EditAnywhere, Category = NDisplay)
 	FVector FrustumOffset = FVector::ZeroVector;
 
-	UPROPERTY(EditAnywhere, Category = NDisplay)
+	UPROPERTY(BlueprintReadWrite, EditAnywhere, Category = NDisplay)
 	FDisplayClusterConfigurationICVFX_CameraMotionBlur CameraMotionBlur;
 
 	// Configure global render settings for this viewports
