nDisplay: Added Actor Layers properties to root actor ICVFX details panel and Light Cards panel

diff --git a/Engine/Plugins/Runtime/nDisplay/Source/DisplayClusterConfiguration/Public/DisplayClusterConfigurationTypes_ICVFX.h b/Engine/Plugins/Runtime/nDisplay/Source/DisplayClusterConfiguration/Public/DisplayClusterConfigurationTypes_ICVFX.h
--- a/Engine/Plugins/Runtime/nDisplay/Source/DisplayClusterConfiguration/Public/DisplayClusterConfigurationTypes_ICVFX.h
+++ b/Engine/Plugins/Runtime/nDisplay/Source/DisplayClusterConfiguration/Public/DisplayClusterConfigurationTypes_ICVFX.h
@@ -176,7 +176,7 @@ struct DISPLAYCLUSTERCONFIGURATION_API FDisplayClusterConfigurationICVFX_Lightca
 
 public:
 	// Allow lightcard rendering (also require not empty LightCardLayers)
-	UPROPERTY(EditAnywhere, Category = NDisplay)
+	UPROPERTY(EditAnywhere, Category = NDisplay, meta = (DisplayName = "Enable Light Cards"))
 	bool bEnable = true;
 
 	// Global lighcard rendering mode
diff --git a/Engine/Plugins/Runtime/nDisplay/Source/DisplayClusterConfiguration/Public/DisplayClusterConfigurationTypes_Viewport.h b/Engine/Plugins/Runtime/nDisplay/Source/DisplayClusterConfiguration/Public/DisplayClusterConfigurationTypes_Viewport.h
--- a/Engine/Plugins/Runtime/nDisplay/Source/DisplayClusterConfiguration/Public/DisplayClusterConfigurationTypes_Viewport.h
+++ b/Engine/Plugins/Runtime/nDisplay/Source/DisplayClusterConfiguration/Public/DisplayClusterConfigurationTypes_Viewport.h
@@ -377,7 +377,7 @@ public:
 	bool bShouldUseParentViewportRenderFamily = false;
 
 	// Multiply all downscale ratio inside all viewports settings for whole cluster
-	UPROPERTY(EditAnywhere, Category = NDisplay, meta = (ClampMin = "0.05", UIMin = "0.05", ClampMax = "1", UIMax = "1"))
+	UPROPERTY(EditAnywhere, Category = NDisplay, meta = (DisplayName = "Outer Viewport Screen Percentage Multiplier", ClampMin = "0.05", UIMin = "0.05", ClampMax = "1", UIMax = "1"))
 	float ClusterRenderTargetRatioMult = 1.f;
 
 	// Multiply all buffer ratios for whole cluster by this value
diff --git a/Engine/Plugins/Runtime/nDisplay/Source/DisplayClusterConfigurator/Private/Views/Details/DisplayClusterConfiguratorDetailCustomization.cpp b/Engine/Plugins/Runtime/nDisplay/Source/DisplayClusterConfigurator/Private/Views/Details/DisplayClusterConfiguratorDetailCustomization.cpp
--- a/Engine/Plugins/Runtime/nDisplay/Source/DisplayClusterConfigurator/Private/Views/Details/DisplayClusterConfiguratorDetailCustomization.cpp
+++ b/Engine/Plugins/Runtime/nDisplay/Source/DisplayClusterConfigurator/Private/Views/Details/DisplayClusterConfiguratorDetailCustomization.cpp
@@ -174,7 +174,7 @@ void FDisplayClusterConfiguratorClusterDetailCustomization::CustomizeDetails(IDe
 	BEGIN_CATEGORY(DisplayClusterConfigurationStrings::categories::ClusterPostprocessCategory)
 		BEGIN_GROUP(TEXT("GlobalPostProcess"), LOCTEXT("GlobalPostprocessLabel", "All Viewports"))
 			ADD_GROUP_PROPERTY(UDisplayClusterConfigurationCluster, bUseOverallClusterPostProcess);
-			ADD_GROUP_PROPERTY(UDisplayClusterConfigurationCluster, OverallClusterPostProcessSettings);
+			ADD_GROUP_EXPANDED_PROPERTY(UDisplayClusterConfigurationCluster, OverallClusterPostProcessSettings);
 		END_GROUP();
 
 		TArray<FString> ViewportNames;
diff --git a/Engine/Plugins/Runtime/nDisplay/Source/DisplayClusterConfigurator/Private/Views/Details/DisplayClusterConfiguratorDetailCustomizationUtils.h b/Engine/Plugins/Runtime/nDisplay/Source/DisplayClusterConfigurator/Private/Views/Details/DisplayClusterConfiguratorDetailCustomizationUtils.h
--- a/Engine/Plugins/Runtime/nDisplay/Source/DisplayClusterConfigurator/Private/Views/Details/DisplayClusterConfiguratorDetailCustomizationUtils.h
+++ b/Engine/Plugins/Runtime/nDisplay/Source/DisplayClusterConfigurator/Private/Views/Details/DisplayClusterConfiguratorDetailCustomizationUtils.h
@@ -26,6 +26,20 @@ class IPropertyHandle;
 
 #define ADD_CUSTOM_PROPERTY(FilterText) CurrentCategory.AddCustomRow(FilterText)
 
+#define ADD_NESTED_PROPERTY(NestedPropertyHelper, PropertyPath) { \
+	TSharedPtr<IPropertyHandle> PropertyHandle = NestedPropertyHelper.GetNestedProperty(PropertyPath); \
+	check(PropertyHandle.IsValid()); \
+	check(PropertyHandle->IsValidHandle()); \
+	CurrentCategory.AddProperty(PropertyHandle.ToSharedRef()); \
+}
+
+#define ADD_EXPANDED_NESTED_PROPERTY(NestedPropertyHelper, PropertyPath) { \
+	TSharedPtr<IPropertyHandle> PropertyHandle = NestedPropertyHelper.GetNestedProperty(PropertyPath); \
+	check(PropertyHandle.IsValid()); \
+	check(PropertyHandle->IsValidHandle()); \
+	CurrentCategory.AddProperty(PropertyHandle.ToSharedRef()).ShouldAutoExpand(true); \
+}
+
 #define REPLACE_PROPERTY_WITH_CUSTOM(ClassName, PropertyName, Widget) { \
 	TSharedRef<IPropertyHandle> PropertyHandle = InLayoutBuilder.GetProperty(GET_MEMBER_NAME_CHECKED(ClassName, PropertyName)); \
 	check(PropertyHandle->IsValidHandle()); \
@@ -44,6 +58,19 @@ class IPropertyHandle;
 	CurrentGroup.AddPropertyRow(PropertyHandle); \
 }
 
+#define ADD_GROUP_EXPANDED_PROPERTY(ClassName, PropertyName) { \
+	TSharedRef<IPropertyHandle> PropertyHandle = InLayoutBuilder.GetProperty(GET_MEMBER_NAME_CHECKED(ClassName, PropertyName)); \
+	check(PropertyHandle->IsValidHandle()); \
+	CurrentGroup.AddPropertyRow(PropertyHandle).ShouldAutoExpand(true); \
+}
+
+#define ADD_GROUP_NESTED_PROPERTY(NestedPropertyHelper, PropertyPath) { \
+	TSharedPtr<IPropertyHandle> PropertyHandle = NestedPropertyHelper.GetNestedProperty(PropertyPath); \
+	check(PropertyHandle.IsValid()); \
+	check(PropertyHandle->IsValidHandle()); \
+	CurrentGroup.AddPropertyRow(PropertyHandle.ToSharedRef()); \
+}
+
 struct FDisplayClusterConfiguratorNestedPropertyHelper
 {
 public:
diff --git a/Engine/Plugins/Runtime/nDisplay/Source/DisplayClusterConfigurator/Private/Views/Details/DisplayClusterRootActorDetailsCustomization.cpp b/Engine/Plugins/Runtime/nDisplay/Source/DisplayClusterConfigurator/Private/Views/Details/DisplayClusterRootActorDetailsCustomization.cpp
--- a/Engine/Plugins/Runtime/nDisplay/Source/DisplayClusterConfigurator/Private/Views/Details/DisplayClusterRootActorDetailsCustomization.cpp
+++ b/Engine/Plugins/Runtime/nDisplay/Source/DisplayClusterConfigurator/Private/Views/Details/DisplayClusterRootActorDetailsCustomization.cpp
@@ -125,15 +125,7 @@ void FDisplayClusterRootActorDetailsCustomization::BuildLayout(IDetailLayoutBuil
 				END_GROUP();
 
 				ADD_PROPERTY(ADisplayClusterRootActor, InnerFrustumPriority);
-
-				{
-					TArray<TSharedPtr<IPropertyHandle>> GlobalScreenPercentageHandles;
-					NestedPropertyHelper.GetNestedProperties(TEXT("CurrentConfigData.RenderFrameSettings.ClusterRenderTargetRatioMult"), GlobalScreenPercentageHandles);
-					check(GlobalScreenPercentageHandles.Num() == 1);
-
-					GlobalScreenPercentageHandles[0]->SetPropertyDisplayName(LOCTEXT("OuterViewportScreenPercentageMult", "Outer Viewport Screen Percentage Multiplier"));
-					CurrentCategory.AddProperty(GlobalScreenPercentageHandles[0]);
-				}
+				ADD_NESTED_PROPERTY(NestedPropertyHelper, TEXT("CurrentConfigData.RenderFrameSettings.ClusterRenderTargetRatioMult"))
 				
 				TArray<TSharedPtr<IPropertyHandle>> ScreenPercentageHandles;
 				NestedPropertyHelper.GetNestedProperties(TEXT("CurrentConfigData.Cluster.Nodes.Viewports.RenderSettings.BufferRatio"), ScreenPercentageHandles);
@@ -149,33 +141,24 @@ void FDisplayClusterRootActorDetailsCustomization::BuildLayout(IDetailLayoutBuil
 				END_GROUP();
 			}
 
+			BEGIN_GROUP(TEXT("HiddenContentGroup"), LOCTEXT("HiddenContentGroupLabel", "Content Hidden from nDisplay"))
+				ADD_GROUP_NESTED_PROPERTY(NestedPropertyHelper, TEXT("CurrentConfigData.StageSettings.HideList.ActorLayers"))
+				ADD_GROUP_NESTED_PROPERTY(NestedPropertyHelper, TEXT("CurrentConfigData.StageSettings.HideList.Actors"))
+			END_GROUP();
 		END_CATEGORY();
 
 		BEGIN_CATEGORY(DisplayClusterConfigurationStrings::categories::ChromaKeyCategory)
-		{
-			TArray<TSharedPtr<IPropertyHandle>> Handles;
-			NestedPropertyHelper.GetNestedProperties(TEXT("CurrentConfigData.StageSettings.Chromakey"), Handles);
-			
-			check(Handles.Num() == 1);
-			CurrentCategory.AddProperty(Handles[0]).ShouldAutoExpand(true);
-		}
+			ADD_EXPANDED_NESTED_PROPERTY(NestedPropertyHelper, TEXT("CurrentConfigData.StageSettings.Chromakey"))
 		END_CATEGORY();
 		
 		BEGIN_CATEGORY(DisplayClusterConfigurationStrings::categories::LightcardCategory)
-		{
-			TArray<TSharedPtr<IPropertyHandle>> Handles;
-			NestedPropertyHelper.GetNestedProperties(TEXT("CurrentConfigData.StageSettings.Lightcard.bEnable"), Handles);
-			check(Handles.Num() == 1);
-			Handles[0]->SetPropertyDisplayName(LOCTEXT("LightCardDisplayName", "Enable Light Cards"));
-			CurrentCategory.AddProperty(Handles[0]);
-		}
-		{
-			TArray<TSharedPtr<IPropertyHandle>> Handles;
-			NestedPropertyHelper.GetNestedProperties(TEXT("CurrentConfigData.StageSettings.Lightcard.BlendingMode"), Handles);
-			check(Handles.Num() == 1);
+			ADD_NESTED_PROPERTY(NestedPropertyHelper, TEXT("CurrentConfigData.StageSettings.Lightcard.bEnable"))
+			ADD_NESTED_PROPERTY(NestedPropertyHelper, TEXT("CurrentConfigData.StageSettings.Lightcard.BlendingMode"))
 
-			CurrentCategory.AddProperty(Handles[0]);
-		}
+			BEGIN_GROUP(TEXT("LightCardActorsGroup"), LOCTEXT("LightCardActorsGroupLabel", "Content Visibile Only in Outer Viewports"))
+				ADD_GROUP_NESTED_PROPERTY(NestedPropertyHelper, TEXT("CurrentConfigData.StageSettings.Lightcard.ShowOnlyList.ActorLayers"))
+				ADD_GROUP_NESTED_PROPERTY(NestedPropertyHelper, TEXT("CurrentConfigData.StageSettings.Lightcard.ShowOnlyList.Actors"))
+			END_GROUP();
 		END_CATEGORY();
 	}
 	
