nDisplay: Hide ICVFX category on unsupported policies. Fix possible bug with region settings on camera policy.

diff --git a/Engine/Plugins/Runtime/nDisplay/Source/DisplayClusterConfigurator/Private/Views/Details/Cluster/DisplayClusterConfiguratorViewportDetailsCustomization.cpp b/Engine/Plugins/Runtime/nDisplay/Source/DisplayClusterConfigurator/Private/Views/Details/Cluster/DisplayClusterConfiguratorViewportDetailsCustomization.cpp
--- a/Engine/Plugins/Runtime/nDisplay/Source/DisplayClusterConfigurator/Private/Views/Details/Cluster/DisplayClusterConfiguratorViewportDetailsCustomization.cpp
+++ b/Engine/Plugins/Runtime/nDisplay/Source/DisplayClusterConfigurator/Private/Views/Details/Cluster/DisplayClusterConfiguratorViewportDetailsCustomization.cpp
@@ -40,30 +40,41 @@ void FDisplayClusterConfiguratorViewportDetailsCustomization::CustomizeDetails(I
 	InLayoutBuilder.HideProperty(GET_MEMBER_NAME_CHECKED(UDisplayClusterConfigurationViewport, TextureShare));
 #endif
 
+	const bool bDisplayICVFX =
+		ConfigurationViewportPtr->ProjectionPolicy.Type == DisplayClusterProjectionStrings::projection::Mesh ||
+		ConfigurationViewportPtr->ProjectionPolicy.Type == DisplayClusterProjectionStrings::projection::MPCDI;
+
+	if (!bDisplayICVFX)
+	{
+		IDetailCategoryBuilder& Category = InLayoutBuilder.EditCategory(TEXT("In Camera VFX"));
+		Category.SetCategoryVisibility(false);
+	}
+	
 	CameraHandle = InLayoutBuilder.GetProperty(GET_MEMBER_NAME_CHECKED(UDisplayClusterConfigurationViewport, Camera));
 	check(CameraHandle->IsValidHandle());
 
 	if (ConfigurationViewportPtr->ProjectionPolicy.Type == DisplayClusterProjectionStrings::projection::Camera)
 	{
 		CameraHandle->MarkHiddenByCustomization();
-		return;
 	}
-
-	ResetCameraOptions();
-
-	if (IDetailPropertyRow* CameraPropertyRow = InLayoutBuilder.EditDefaultProperty(CameraHandle))
+	else
 	{
-		CameraPropertyRow->CustomWidget()
-			.NameContent()
-			[
-				CameraHandle->CreatePropertyNameWidget()
-			]
-			.ValueContent()
-			[
-				CreateCustomCameraWidget()
-			];
-	}
+		ResetCameraOptions();
 
+		if (IDetailPropertyRow* CameraPropertyRow = InLayoutBuilder.EditDefaultProperty(CameraHandle))
+		{
+			CameraPropertyRow->CustomWidget()
+				.NameContent()
+				[
+					CameraHandle->CreatePropertyNameWidget()
+				]
+				.ValueContent()
+				[
+					CreateCustomCameraWidget()
+				];
+		}
+	}
+	
 	// Update the metadata for the viewport's region. Must set this here instead of in the UPROPERTY specifier because
 	// the Region property is a generic FDisplayClusterConfigurationRectangle struct which is used in lots of places, most of
 	// which don't make sense to have a minimum or maximum limit
