nDisplay: fixed path for *Engine.ini to use

diff --git a/Engine/Plugins/Runtime/nDisplay/Source/DisplayClusterEditor/Private/Settings/DisplayClusterEditorSettings.cpp b/Engine/Plugins/Runtime/nDisplay/Source/DisplayClusterEditor/Private/Settings/DisplayClusterEditorSettings.cpp
--- a/Engine/Plugins/Runtime/nDisplay/Source/DisplayClusterEditor/Private/Settings/DisplayClusterEditorSettings.cpp
+++ b/Engine/Plugins/Runtime/nDisplay/Source/DisplayClusterEditor/Private/Settings/DisplayClusterEditorSettings.cpp
@@ -21,9 +21,8 @@ void UDisplayClusterEditorSettings::PostEditChangeProperty(FPropertyChangedEvent
 	if (PropertyChangedEvent.Property != nullptr)
 	{
 		// Since nDisplay is Windows only, save configs that depends on nDisplay plugin assets in Windows specific config file
-		static const FString PlatformName = FPlatformProperties::PlatformName();
-		static const FString DefaultEnginePath = FString::Printf(TEXT("%s%s/%sEngine.ini"), *FPaths::SourceConfigDir(), *PlatformName, *PlatformName);
-		static const FString DefaultGamePath   = FString::Printf(TEXT("%sDefaultGame.ini"), *FPaths::SourceConfigDir());
+		static const FString DefaultEnginePath = FString::Printf(TEXT("%sDefaultEngine.ini"), *FPaths::SourceConfigDir());
+		static const FString DefaultGamePath   = FString::Printf(TEXT("%sDefaultGame.ini"),   *FPaths::SourceConfigDir());
 
 		FName PropertyName(PropertyChangedEvent.Property->GetFName());
 
