nDisplay: Fix adding components with multiple nDisplay blueprint editors open.

diff --git a/Engine/Plugins/Runtime/nDisplay/Source/DisplayClusterConfigurator/Private/DisplayClusterConfiguratorBlueprintEditor.cpp b/Engine/Plugins/Runtime/nDisplay/Source/DisplayClusterConfigurator/Private/DisplayClusterConfiguratorBlueprintEditor.cpp
--- a/Engine/Plugins/Runtime/nDisplay/Source/DisplayClusterConfigurator/Private/DisplayClusterConfiguratorBlueprintEditor.cpp
+++ b/Engine/Plugins/Runtime/nDisplay/Source/DisplayClusterConfigurator/Private/DisplayClusterConfiguratorBlueprintEditor.cpp
@@ -181,7 +181,7 @@ void FDisplayClusterConfiguratorBlueprintEditor::InitDisplayClusterBlueprintEdit
 		if (!PanelExtensionSubsystem->IsPanelFactoryRegistered(SCSEditorExtensionIdentifier))
 		{
 			FPanelExtensionFactory SCSEditorExtensionWidget;
-			SCSEditorExtensionWidget.CreateExtensionWidget = FPanelExtensionFactory::FCreateExtensionWidget::CreateRaw(this, &FDisplayClusterConfiguratorBlueprintEditor::CreateSCSEditorExtensionWidget);
+			SCSEditorExtensionWidget.CreateExtensionWidget = FPanelExtensionFactory::FCreateExtensionWidget::CreateStatic(&FDisplayClusterConfiguratorBlueprintEditor::CreateSCSEditorExtensionWidget);
 			SCSEditorExtensionWidget.Identifier = SCSEditorExtensionIdentifier;
 
 			PanelExtensionSubsystem->RegisterPanelFactory("SCSEditor.NextToAddComponentButton", SCSEditorExtensionWidget);
@@ -1280,8 +1280,7 @@ void FDisplayClusterConfiguratorBlueprintEditor::ShutdownDCSCSEditors()
 	}
 }
 
-TSharedRef<SWidget> FDisplayClusterConfiguratorBlueprintEditor::CreateSCSEditorExtensionWidget(
-	FWeakObjectPtr ExtensionContext)
+TSharedRef<SWidget> FDisplayClusterConfiguratorBlueprintEditor::CreateSCSEditorExtensionWidget(FWeakObjectPtr ExtensionContext)
 {
 	auto PerformComboAddClass = [=](TSubclassOf<UActorComponent> ComponentClass, EComponentCreateAction::Type ComponentCreateAction,
 		UObject* AssetOverride) -> FSubobjectDataHandle
@@ -1313,7 +1312,7 @@ TSharedRef<SWidget> FDisplayClusterConfiguratorBlueprintEditor::CreateSCSEditorE
 					
 					USubobjectDataSubsystem* SubobjectSystem = USubobjectDataSubsystem::Get();
 					FAddNewSubobjectParams AddSubobjectParams;
-					AddSubobjectParams.BlueprintContext = LoadedBlueprint.Get();
+					AddSubobjectParams.BlueprintContext = SubobjectEditor->GetBlueprint();
 					AddSubobjectParams.NewClass = NewClass;
 										
 					// This adds components according to the type selected in the drop down. If the user
@@ -1419,7 +1418,7 @@ TSharedRef<SWidget> FDisplayClusterConfiguratorBlueprintEditor::CreateSCSEditorE
 		[
 			SNew(SDisplayClusterConfiguratorComponentClassCombo)
 			.OnSubobjectClassSelected(FSubobjectClassSelected::CreateLambda(PerformComboAddClass))
-			.Visibility_Lambda([this, ExtensionContext]
+			.Visibility_Lambda([ExtensionContext]
 			{
 				if(USubobjectEditorExtensionContext* SubobjectEditorExtensionContext = Cast<USubobjectEditorExtensionContext>(ExtensionContext.Get()))
 				{
diff --git a/Engine/Plugins/Runtime/nDisplay/Source/DisplayClusterConfigurator/Private/DisplayClusterConfiguratorBlueprintEditor.h b/Engine/Plugins/Runtime/nDisplay/Source/DisplayClusterConfigurator/Private/DisplayClusterConfiguratorBlueprintEditor.h
--- a/Engine/Plugins/Runtime/nDisplay/Source/DisplayClusterConfigurator/Private/DisplayClusterConfiguratorBlueprintEditor.h
+++ b/Engine/Plugins/Runtime/nDisplay/Source/DisplayClusterConfigurator/Private/DisplayClusterConfiguratorBlueprintEditor.h
@@ -145,7 +145,7 @@ protected:
 	void CreateDCSCSEditors();
 	void ShutdownDCSCSEditors();
 
-	TSharedRef<SWidget> CreateSCSEditorExtensionWidget(FWeakObjectPtr ExtensionContext);
+	static TSharedRef<SWidget> CreateSCSEditorExtensionWidget(FWeakObjectPtr ExtensionContext);
 
 	void CreateSCSEditorWrapper();
 
