nDisplay: Fixed multiviewport issues for CL18291679 refactor.

diff --git a/Engine/Plugins/Runtime/nDisplay/Source/DisplayCluster/Private/Render/Device/DisplayClusterDeviceBase.cpp b/Engine/Plugins/Runtime/nDisplay/Source/DisplayCluster/Private/Render/Device/DisplayClusterDeviceBase.cpp
--- a/Engine/Plugins/Runtime/nDisplay/Source/DisplayCluster/Private/Render/Device/DisplayClusterDeviceBase.cpp
+++ b/Engine/Plugins/Runtime/nDisplay/Source/DisplayCluster/Private/Render/Device/DisplayClusterDeviceBase.cpp
@@ -164,6 +164,41 @@ void FDisplayClusterDeviceBase::InitCanvasFromView(class FSceneView* InView, cla
 	}
 }
 
+EStereoscopicPass FDisplayClusterDeviceBase::GetViewPassForIndex(bool bStereoRequested, int32 ViewIndex) const
+{
+	if (bStereoRequested)
+	{
+		if (IsInRenderingThread())
+		{
+			if (ViewportManagerProxyPtr)
+			{
+				uint32 ViewportContextNum = 0;
+				IDisplayClusterViewportProxy* ViewportProxy = ViewportManagerProxyPtr->FindViewport_RenderThread(ViewIndex, &ViewportContextNum);
+				if (ViewportProxy)
+				{
+					const FDisplayClusterViewport_Context& Context = ViewportProxy->GetContexts_RenderThread()[ViewportContextNum];
+					return Context.StereoscopicPass;
+				}
+			}
+		}
+		else
+		{
+			if (ViewportManagerPtr)
+			{
+				uint32 ViewportContextNum = 0;
+				IDisplayClusterViewport* ViewportPtr = ViewportManagerPtr->FindViewport(ViewIndex, &ViewportContextNum);
+				if (ViewportPtr)
+				{
+					const FDisplayClusterViewport_Context& Context = ViewportPtr->GetContexts()[ViewportContextNum];
+					return Context.StereoscopicPass;
+				}
+			}
+		}
+	}
+
+	return EStereoscopicPass::eSSP_FULL;
+}
+
 void FDisplayClusterDeviceBase::AdjustViewRect(int32 ViewIndex, int32& X, int32& Y, uint32& SizeX, uint32& SizeY) const
 {
 	check(IsInGameThread());
diff --git a/Engine/Plugins/Runtime/nDisplay/Source/DisplayCluster/Private/Render/Device/DisplayClusterDeviceBase.h b/Engine/Plugins/Runtime/nDisplay/Source/DisplayCluster/Private/Render/Device/DisplayClusterDeviceBase.h
--- a/Engine/Plugins/Runtime/nDisplay/Source/DisplayCluster/Private/Render/Device/DisplayClusterDeviceBase.h
+++ b/Engine/Plugins/Runtime/nDisplay/Source/DisplayCluster/Private/Render/Device/DisplayClusterDeviceBase.h
@@ -65,6 +65,8 @@ public:
 		return FMath::Max(1, DesiredNumberOfViews);
 	}
 
+	virtual EStereoscopicPass GetViewPassForIndex(bool bStereoRequested, int32 ViewIndex) const override;
+
 	virtual IStereoRenderTargetManager* GetRenderTargetManager() override
 	{ return this; }
 
