nDisplay: Enable LightCard selection from Virtual Production actor placement menu.

diff --git a/Engine/Plugins/Runtime/nDisplay/Source/DisplayClusterConfigurator/DisplayClusterConfigurator.Build.cs b/Engine/Plugins/Runtime/nDisplay/Source/DisplayClusterConfigurator/DisplayClusterConfigurator.Build.cs
--- a/Engine/Plugins/Runtime/nDisplay/Source/DisplayClusterConfigurator/DisplayClusterConfigurator.Build.cs
+++ b/Engine/Plugins/Runtime/nDisplay/Source/DisplayClusterConfigurator/DisplayClusterConfigurator.Build.cs
@@ -36,6 +36,7 @@ public class DisplayClusterConfigurator : ModuleRules
 				"MainFrame",
 				"MessageLog",
 				"PinnedCommandList",
+				"PlacementMode",
 				"Projects",
 				"PropertyEditor",
 				"Serialization",
@@ -47,6 +48,7 @@ public class DisplayClusterConfigurator : ModuleRules
 				"SubobjectEditor",
 				"SubobjectDataInterface",
 				"ToolWidgets",
+				"VPUtilitiesEditor"
 			});
 	}
 }
diff --git a/Engine/Plugins/Runtime/nDisplay/Source/DisplayClusterConfigurator/Private/DisplayClusterConfiguratorModule.cpp b/Engine/Plugins/Runtime/nDisplay/Source/DisplayClusterConfigurator/Private/DisplayClusterConfiguratorModule.cpp
--- a/Engine/Plugins/Runtime/nDisplay/Source/DisplayClusterConfigurator/Private/DisplayClusterConfiguratorModule.cpp
+++ b/Engine/Plugins/Runtime/nDisplay/Source/DisplayClusterConfigurator/Private/DisplayClusterConfiguratorModule.cpp
@@ -8,6 +8,7 @@
 #include "Settings/DisplayClusterConfiguratorSettings.h"
 #include "Views/Details/DisplayClusterRootActorDetailsCustomization.h"
 #include "Views/OutputMapping/DisplayClusterConfiguratorOutputMappingCommands.h"
+#include "DisplayClusterConfiguratorLog.h"
 
 #include "Views/Details/DisplayClusterConfiguratorBaseTypeCustomization.h"
 #include "Views/Details/DisplayClusterEditorPropertyReferenceTypeCustomization.h"
@@ -37,6 +38,9 @@
 #include "AssetTypeCategories.h"
 #include "HAL/IConsoleManager.h"
 #include "Modules/ModuleManager.h"
+#include "IPlacementModeModule.h"
+#include "IVPUtilitiesEditorModule.h"
+#include "ActorFactories/ActorFactoryBlueprint.h"
 
 #define REGISTER_PROPERTY_LAYOUT(PropertyType, CustomizationType) { \
 	const FName LayoutName = PropertyType::StaticStruct()->GetFName(); \
@@ -83,6 +87,7 @@ void FDisplayClusterConfiguratorModule::StartupModule()
 
 	RegisterCustomLayouts();
 	RegisterSettings();
+	RegisterPlacementModeItems();
 	
 	FDisplayClusterConfiguratorStyle::Initialize();
 
@@ -119,6 +124,7 @@ void FDisplayClusterConfiguratorModule::ShutdownModule()
 
 	UnregisterSettings();
 	UnregisterCustomLayouts();
+	UnregisterPlacementModeItems();
 
 	FDisplayClusterConfiguratorStyle::Shutdown();
 	MenuExtensibilityManager.Reset();
@@ -215,6 +221,74 @@ void FDisplayClusterConfiguratorModule::UnregisterCustomLayouts()
 	RegisteredPropertyLayoutNames.Empty();
 }
 
+void FDisplayClusterConfiguratorModule::RegisterPlacementModeItems()
+{
+	auto RegisterPlaceActors = [&]() -> void
+	{
+		if (!GEditor)
+		{
+			return;
+		}
+		
+		const FPlacementCategoryInfo* Info = IVPUtilitiesEditorModule::Get().GetVirtualProductionPlacementCategoryInfo();
+		if (!Info)
+		{
+			UE_LOG(DisplayClusterConfiguratorLog, Warning, TEXT("Could not find or create VirtualProduction Place Actor Category"));
+			return;
+		}
+
+		FAssetData LightCardAssetData(
+			TEXT("/nDisplay/LightCard/LightCard"),
+			TEXT("/nDisplay/LightCard"),
+			TEXT("LightCard"),
+			TEXT("Blueprint")
+		);
+		check (LightCardAssetData.IsValid());
+	
+		PlaceActors.Add(IPlacementModeModule::Get().RegisterPlaceableItem(Info->UniqueHandle, MakeShared<FPlaceableItem>(
+			*UActorFactoryBlueprint::StaticClass(),
+			LightCardAssetData,
+			NAME_None,
+			NAME_None,
+			TOptional<FLinearColor>(),
+			TOptional<int32>(),
+			NSLOCTEXT("PlacementMode", "LightCard", "LightCard")
+		)));
+	};
+
+	if (GEngine && GEngine->IsInitialized())
+	{
+		RegisterPlaceActors();
+	}
+	else
+	{
+		PostEngineInitHandle = FCoreDelegates::OnPostEngineInit.AddLambda(RegisterPlaceActors);
+	}
+}
+
+void FDisplayClusterConfiguratorModule::UnregisterPlacementModeItems()
+{
+	if (!IsEngineExitRequested() && GEditor && UObjectInitialized())
+	{
+		IPlacementModeModule& PlacementModeModule = IPlacementModeModule::Get();
+
+		for (TOptional<FPlacementModeID>& PlaceActor : PlaceActors)
+		{
+			if (PlaceActor.IsSet())
+			{
+				PlacementModeModule.UnregisterPlaceableItem(*PlaceActor);
+			}
+		}
+	}
+
+	if (PostEngineInitHandle.IsValid())
+	{
+		FCoreDelegates::OnPostEngineInit.Remove(PostEngineInitHandle);
+	}
+	
+	PlaceActors.Empty();
+}
+
 TSharedPtr<FKismetCompilerContext> FDisplayClusterConfiguratorModule::GetCompilerForDisplayClusterBP(UBlueprint* BP,
                                                                                                      FCompilerResultsLog& InMessageLog, const FKismetCompilerOptions& InCompileOptions)
 {
diff --git a/Engine/Plugins/Runtime/nDisplay/Source/DisplayClusterConfigurator/Private/DisplayClusterConfiguratorModule.h b/Engine/Plugins/Runtime/nDisplay/Source/DisplayClusterConfigurator/Private/DisplayClusterConfiguratorModule.h
--- a/Engine/Plugins/Runtime/nDisplay/Source/DisplayClusterConfigurator/Private/DisplayClusterConfiguratorModule.h
+++ b/Engine/Plugins/Runtime/nDisplay/Source/DisplayClusterConfigurator/Private/DisplayClusterConfiguratorModule.h
@@ -6,6 +6,7 @@
 #include "IAssetTypeActions.h"
 #include "Interfaces/IDisplayClusterConfigurator.h"
 
+class FPlacementModeID;
 class IAssetTools;
 class FDisplayClusterConfiguratorAssetTypeActions;
 class FExtensibilityManager;
@@ -36,7 +37,13 @@ private:
 	void UnregisterSettings();
 	void RegisterCustomLayouts();
 	void UnregisterCustomLayouts();
+
+	/** Register items to show up in the Place Actors panel. */
+	void RegisterPlacementModeItems();
 	
+	/** Unregister items in Place Actors panel */
+	void UnregisterPlacementModeItems();
+
 	static TSharedPtr<FKismetCompilerContext> GetCompilerForDisplayClusterBP(UBlueprint* BP, FCompilerResultsLog& InMessageLog, const FKismetCompilerOptions& InCompileOptions);
 
 private:
@@ -46,7 +53,9 @@ private:
 	FDisplayClusterConfiguratorKismetCompiler  BlueprintCompiler;
 	TArray<FName> RegisteredClassLayoutNames;
 	TArray<FName> RegisteredPropertyLayoutNames;
+	TArray<TOptional<FPlacementModeID>> PlaceActors;
 	FDelegateHandle FilesLoadedHandle;
+	FDelegateHandle PostEngineInitHandle;
 
 private:
 	static FOnDisplayClusterConfiguratorReadOnlyChanged OnDisplayClusterConfiguratorReadOnlyChanged;
diff --git a/Engine/Plugins/Runtime/nDisplay/nDisplay.uplugin b/Engine/Plugins/Runtime/nDisplay/nDisplay.uplugin
--- a/Engine/Plugins/Runtime/nDisplay/nDisplay.uplugin
+++ b/Engine/Plugins/Runtime/nDisplay/nDisplay.uplugin
@@ -136,6 +136,10 @@
 		{
 			"Name": "ConcertSyncClient",
 			"Enabled": true
+		},
+		{
+			"Name": "VirtualProductionUtilities",
+			"Enabled": true
 		}
 	]
 }
