nDisplay: updated Blueprint API for topology traversal

diff --git a/Engine/Plugins/Runtime/nDisplay/Source/DisplayCluster/Private/Blueprints/DisplayClusterBlueprintAPIImpl.cpp b/Engine/Plugins/Runtime/nDisplay/Source/DisplayCluster/Private/Blueprints/DisplayClusterBlueprintAPIImpl.cpp
--- a/Engine/Plugins/Runtime/nDisplay/Source/DisplayCluster/Private/Blueprints/DisplayClusterBlueprintAPIImpl.cpp
+++ b/Engine/Plugins/Runtime/nDisplay/Source/DisplayCluster/Private/Blueprints/DisplayClusterBlueprintAPIImpl.cpp
@@ -35,11 +35,8 @@
 #include "Render/Viewport/IDisplayClusterViewportManager.h"
 #include "Render/Viewport/IDisplayClusterViewport.h"
 
+// Cluster runtime API
 
-//////////////////////////////////////////////////////////////////////////////////////////////
-// DisplayCluster module API
-//////////////////////////////////////////////////////////////////////////////////////////////
-/** Return if the module has been initialized. */
 bool UDisplayClusterBlueprintAPIImpl::IsModuleInitialized() const
 {
 	const bool bInitialized = IDisplayCluster::Get().IsModuleInitialized();
@@ -54,10 +51,35 @@ EDisplayClusterOperationMode UDisplayClusterBlueprintAPIImpl::GetOperationMode()
 	return OpMode;
 }
 
+ADisplayClusterRootActor* UDisplayClusterBlueprintAPIImpl::GetRootActor() const
+{
+	ADisplayClusterRootActor* const RootActor = IDisplayCluster::Get().GetGameMgr()->GetRootActor();
+	UE_LOG(LogDisplayClusterBlueprint, Verbose, TEXT("GetRootActor - %s"), RootActor ? *RootActor->GetHumanReadableName() : DisplayClusterStrings::log::NotFound);
+	return RootActor;
+}
+
+// Local node runtime API
+
+FString UDisplayClusterBlueprintAPIImpl::GetNodeId() const
+{
+	const FString NodeId = IDisplayCluster::Get().GetClusterMgr()->GetNodeId();
+	UE_LOG(LogDisplayClusterBlueprint, Verbose, TEXT("GetNodeId - NodeId=%s"), *NodeId);
+	return NodeId;
+}
+
+void UDisplayClusterBlueprintAPIImpl::GetActiveNodeIds(TArray<FString>& OutNodeIds) const
+{
+	UE_LOG(LogDisplayClusterBlueprint, Verbose, TEXT("GetNodeIds"));
+	IDisplayCluster::Get().GetClusterMgr()->GetNodeIds(OutNodeIds);
+}
+
+int32 UDisplayClusterBlueprintAPIImpl::GetActiveNodesAmount() const
+{
+	const int32 NodesAmount = IDisplayCluster::Get().GetClusterMgr()->GetNodesAmount();
+	UE_LOG(LogDisplayClusterBlueprint, Verbose, TEXT("GetNodesAmount - %d"), NodesAmount);
+	return NodesAmount;
+}
 
-//////////////////////////////////////////////////////////////////////////////////////////////
-// Cluster API
-//////////////////////////////////////////////////////////////////////////////////////////////
 bool UDisplayClusterBlueprintAPIImpl::IsMaster() const
 {
 	const bool bIsMaster = IDisplayCluster::Get().GetClusterMgr()->IsMaster();
@@ -86,25 +108,7 @@ EDisplayClusterNodeRole UDisplayClusterBlueprintAPIImpl::GetClusterRole() const
 	return ClusterRole;
 }
 
-void UDisplayClusterBlueprintAPIImpl::GetNodeIds(TArray<FString>& OutNodeIds) const
-{
-	UE_LOG(LogDisplayClusterBlueprint, Verbose, TEXT("GetNodeIds"));
-	IDisplayCluster::Get().GetClusterMgr()->GetNodeIds(OutNodeIds);
-}
-
-FString UDisplayClusterBlueprintAPIImpl::GetNodeId() const
-{
-	const FString NodeId = IDisplayCluster::Get().GetClusterMgr()->GetNodeId();
-	UE_LOG(LogDisplayClusterBlueprint, Verbose, TEXT("GetNodeId - NodeId=%s"), *NodeId);
-	return NodeId;
-}
-
-int32 UDisplayClusterBlueprintAPIImpl::GetNodesAmount() const
-{
-	const int32 NodesAmount = IDisplayCluster::Get().GetClusterMgr()->GetNodesAmount();
-	UE_LOG(LogDisplayClusterBlueprint, Verbose, TEXT("GetNodesAmount - %d"), NodesAmount);
-	return NodesAmount;
-}
+// Cluster events API
 
 void UDisplayClusterBlueprintAPIImpl::AddClusterEventListener(TScriptInterface<IDisplayClusterClusterEventListener> Listener)
 {
@@ -149,198 +153,3 @@ void UDisplayClusterBlueprintAPIImpl::SendClusterEventBinaryTo(const FString& Ad
 	UE_LOG(LogDisplayClusterBlueprint, Verbose, TEXT("SendClusterEventBinaryTo - sending binary event to %s:%d, EventId='%d'"), *Address, Port, Event.EventId);
 	IDisplayCluster::Get().GetClusterMgr()->SendClusterEventTo(Address, static_cast<uint16>(Port), Event, bMasterOnly);
 }
-
-//////////////////////////////////////////////////////////////////////////////////////////////
-// Config API
-//////////////////////////////////////////////////////////////////////////////////////////////
-UFUNCTION(BlueprintCallable, meta = (DisplayName = "Get config"), Category = "NDisplay|Config")
-UDisplayClusterConfigurationData* UDisplayClusterBlueprintAPIImpl::GetConfig() const
-{
-	return const_cast<UDisplayClusterConfigurationData*>(IDisplayCluster::Get().GetConfigMgr()->GetConfig());
-}
-
-
-//////////////////////////////////////////////////////////////////////////////////////////////
-// Game API
-//////////////////////////////////////////////////////////////////////////////////////////////
-// Root
-ADisplayClusterRootActor* UDisplayClusterBlueprintAPIImpl::GetRootActor() const
-{
-	ADisplayClusterRootActor* const RootActor = IDisplayCluster::Get().GetGameMgr()->GetRootActor();
-	UE_LOG(LogDisplayClusterBlueprint, Verbose, TEXT("GetRootActor - %s"), RootActor ? *RootActor->GetHumanReadableName() : DisplayClusterStrings::log::NotFound);
-	return RootActor;
-}
-
-// @todo: implement new BP api
-// @todo: new BP api stuff before release
-#if 0
-
-//////////////////////////////////////////////////////////////////////////////////////////////
-// New Render API protoype
-//////////////////////////////////////////////////////////////////////////////////////////////
-
-IDisplayClusterViewport* ImplFindViewport(const FString& ViewportId)
-{
-	IDisplayClusterRenderManager* DCRenderManager = IDisplayCluster::Get().GetRenderMgr();
-	if (DCRenderManager && DCRenderManager->GetViewportManager())
-	{
-		return DCRenderManager->GetViewportManager()->FindViewport(ViewportId);
-	}
-
-	return nullptr;
-}
-
-const TArrayView<IDisplayClusterViewport*> ImplGetViewports()
-{
-	IDisplayClusterRenderManager* DCRenderManager = IDisplayCluster::Get().GetRenderMgr();
-	if (DCRenderManager && DCRenderManager->GetViewportManager())
-	{
-		return DCRenderManager->GetViewportManager()->GetViewports();
-	}
-
-	return TArrayView<IDisplayClusterViewport*>();
-}
-
-FDisplayClusterViewportContext ImplGetViewportContext(const IDisplayClusterViewport& Viewport)
-{
-	FDisplayClusterViewportContext OutContext;
-
-	OutContext.ViewportID = Viewport.GetId();
-
-	OutContext.RectLocation = Viewport.GetRenderSettings().Rect.Min;
-	OutContext.RectSize     = Viewport.GetRenderSettings().Rect.Size();
-
-	const TArray<FDisplayClusterViewport_Context>& Contexts = Viewport.GetContexts();
-
-	if (Contexts.Num() > 0)
-	{
-		OutContext.ViewLocation     = Contexts[0].ViewLocation;
-		OutContext.ViewRotation     = Contexts[0].ViewRotation;
-		OutContext.ProjectionMatrix = Contexts[0].ProjectionMatrix;
-
-		OutContext.bIsRendering = Contexts[0].bDisableRender == false;
-	}
-	else
-		{
-		OutContext.bIsRendering = false;
-		}
-
-	return OutContext;
-}
-
-FDisplayClusterViewportStereoContext ImplGetViewportStereoContext(const IDisplayClusterViewport& Viewport)
-{
-	FDisplayClusterViewportStereoContext OutContext;
-
-	OutContext.ViewportID = Viewport.GetId();
-
-	OutContext.RectLocation = Viewport.GetRenderSettings().Rect.Min;
-	OutContext.RectSize = Viewport.GetRenderSettings().Rect.Size();
-
-	const TArray<FDisplayClusterViewport_Context>& Contexts = Viewport.GetContexts();
-
-	if (Contexts.Num() > 0)
-	{
-		for (const FDisplayClusterViewport_Context& InContext : Contexts)
-		{
-			OutContext.ViewLocation.Add(InContext.ViewLocation);
-			OutContext.ViewRotation.Add(InContext.ViewRotation);
-			OutContext.ProjectionMatrix.Add(InContext.ProjectionMatrix);
-	}
-
-		OutContext.bIsRendering = Contexts[0].bDisableRender == false;
-	}
-	else
-	{
-		OutContext.bIsRendering = false;
-	}
-
-	return OutContext;
-}
-
-bool UDisplayClusterBlueprintAPIImpl::GetLocalViewportConfiguration(const FString& ViewportID, TWeakObjectPtr<UDisplayClusterConfigurationViewport>& ConfigurationViewport)
-	{
-	IDisplayCluster& DisplayCluster = IDisplayCluster::Get();
-
-	ADisplayClusterRootActor* RootActor = DisplayCluster.GetGameMgr()->GetRootActor();
-	if (RootActor)
-	{
-		return RootActor->GetLocalViewportConfiguration(DisplayCluster.GetConfigMgr()->GetLocalNodeId(), ViewportID, ConfigurationViewport);
-	}
-
-	return false;
-}
-
-/** Return local viewports names */
-void UDisplayClusterBlueprintAPIImpl::GetLocalViewports(TArray<FString>& ViewportIDs) const
-{
-	const TArrayView<IDisplayClusterViewport*> LocalViewports = ImplGetViewports();
-	ViewportIDs.Reserve(LocalViewports.Num());
-
-	for (const IDisplayClusterViewport* Viewport : LocalViewports)
-	{
-		if (Viewport && Viewport->GetRenderSettings().bVisible)
-		{
-			ViewportIDs.Add(Viewport->GetId());
-		}
-	}
-	}
-
-/** Return local viewports runtime contexts */
-void UDisplayClusterBlueprintAPIImpl::GetLocalViewportsContext(TArray<FDisplayClusterViewportContext>& ViewportContexts) const
-{
-	const TArrayView<IDisplayClusterViewport*> LocalViewports = ImplGetViewports();
-	ViewportContexts.Reserve(LocalViewports.Num());
-
-	for (const IDisplayClusterViewport* Viewport : LocalViewports)
-	{
-		if (Viewport && Viewport->GetRenderSettings().bVisible)
-		{
-			ViewportContexts.Add(ImplGetViewportContext(*Viewport));
-		}
-	}
-}
-
-/** Return local viewports runtime stereo contexts */
-void UDisplayClusterBlueprintAPIImpl::GetLocalViewportsStereoContext(TArray<FDisplayClusterViewportStereoContext>& ViewportStereoContexts) const
-{
-	const TArrayView<IDisplayClusterViewport*> LocalViewports = ImplGetViewports();
-	ViewportStereoContexts.Reserve(LocalViewports.Num());
-
-	for (const IDisplayClusterViewport* Viewport : LocalViewports)
-	{
-		if (Viewport && Viewport->GetRenderSettings().bVisible)
-		{
-			ViewportStereoContexts.Add(ImplGetViewportStereoContext(*Viewport));
-		}
-	}
-}
-
-/** Return viewport runtime context (last frame viewport data) */
-bool UDisplayClusterBlueprintAPIImpl::GetLocalViewportContext(const FString& ViewportId, FDisplayClusterViewportContext& ViewportContext) const
-{
-	IDisplayClusterViewport* DesiredViewport = ImplFindViewport(ViewportId);
-	if (DesiredViewport)
-	{
-		ViewportContext = ImplGetViewportContext(*DesiredViewport);
-		return true;
-	}
-
-	UE_LOG(LogDisplayClusterBlueprint, Warning, TEXT("Couldn't GetLocalViewportContext. Viewport '%s' not found"), *ViewportId);
-	return false;
-}
-
-/** Return viewport stereo contexts (last frame viewport data) */
-bool UDisplayClusterBlueprintAPIImpl::GetLocalViewportStereoContext(const FString& ViewportId, FDisplayClusterViewportStereoContext& ViewportStereoContext) const
-{
-	IDisplayClusterViewport* DesiredViewport = ImplFindViewport(ViewportId);
-	if (DesiredViewport)
-	{
-		ViewportStereoContext = ImplGetViewportStereoContext(*DesiredViewport);
-		return true;
-	}
-
-	UE_LOG(LogDisplayClusterBlueprint, Warning, TEXT("Couldn't GetLocalViewportStereoContext. Viewport '%s' not found"), *ViewportId);
-	return false;
-}
-#endif
diff --git a/Engine/Plugins/Runtime/nDisplay/Source/DisplayCluster/Private/Blueprints/DisplayClusterBlueprintAPIImpl.h b/Engine/Plugins/Runtime/nDisplay/Source/DisplayCluster/Private/Blueprints/DisplayClusterBlueprintAPIImpl.h
--- a/Engine/Plugins/Runtime/nDisplay/Source/DisplayCluster/Private/Blueprints/DisplayClusterBlueprintAPIImpl.h
+++ b/Engine/Plugins/Runtime/nDisplay/Source/DisplayCluster/Private/Blueprints/DisplayClusterBlueprintAPIImpl.h
@@ -19,19 +19,34 @@ class DISPLAYCLUSTER_API UDisplayClusterBlueprintAPIImpl
 	GENERATED_BODY()
 
 public:
-	//////////////////////////////////////////////////////////////////////////////////////////////
 	// DisplayCluster module API
-	//////////////////////////////////////////////////////////////////////////////////////////////
 	UFUNCTION(BlueprintCallable, meta = (DisplayName = "Is module initialized"), Category = "NDisplay")
 	virtual bool IsModuleInitialized() const override;
 
+public:
+	// Runtime Cluster API
 	UFUNCTION(BlueprintCallable, meta = (DisplayName = "Get operation mode"), Category = "NDisplay")
 	virtual EDisplayClusterOperationMode GetOperationMode() const override;
 
-public:
-	//////////////////////////////////////////////////////////////////////////////////////////////
-	// Cluster API
-	//////////////////////////////////////////////////////////////////////////////////////////////
+	UFUNCTION(BlueprintCallable, meta = (DisplayName = "Get root actor"), Category = "NDisplay|Game")
+	virtual ADisplayClusterRootActor* GetRootActor() const override;
+
+	// Runtime local node API
+	UFUNCTION(BlueprintCallable, meta = (DisplayName = "Get node ID"), Category = "NDisplay|Cluster")
+	virtual FString GetNodeId() const override;
+	
+	/** Returns List of the active nodes in the runtime cluster node in a cluster. */
+	UFUNCTION(BlueprintCallable, meta = (DisplayName = "Get cluster node IDs"), Category = "NDisplay|Cluster")
+	virtual void GetActiveNodeIds(TArray<FString>& OutNodeIds) const override;
+	
+
+	/** Returns amount of active nodes in a cluster. */
+	UFUNCTION(BlueprintCallable, meta = (DisplayName = "Get nodes amount"), Category = "NDisplay|Cluster")
+	virtual int32 GetActiveNodesAmount() const override;
+
+	UFUNCTION(BlueprintCallable, meta = (DisplayName = "Get cluster role"), Category = "NDisplay|Cluster")
+	virtual EDisplayClusterNodeRole GetClusterRole() const override;
+
 	UFUNCTION(BlueprintCallable, meta = (DisplayName = "Is master node"), Category = "NDisplay|Cluster")
 	virtual bool IsMaster() const override;
 	
@@ -42,18 +57,7 @@ public:
 	UFUNCTION(BlueprintCallable, meta = (DisplayName = "Is backup node"), Category = "NDisplay|Cluster")
 	virtual bool IsBackup() const override;
 
-	UFUNCTION(BlueprintCallable, meta = (DisplayName = "Get cluster role"), Category = "NDisplay|Cluster")
-	virtual EDisplayClusterNodeRole GetClusterRole() const override;
-
-	UFUNCTION(BlueprintCallable, meta = (DisplayName = "Get cluster node IDs"), Category = "NDisplay|Cluster")
-	virtual void GetNodeIds(TArray<FString>& OutNodeIds) const override;
-
-	UFUNCTION(BlueprintCallable, meta = (DisplayName = "Get node ID"), Category = "NDisplay|Cluster")
-	virtual FString GetNodeId() const override;
-
-	UFUNCTION(BlueprintCallable, meta = (DisplayName = "Get nodes amount"), Category = "NDisplay|Cluster")
-	virtual int32 GetNodesAmount() const override;
-
+	// Cluster events API
 	UFUNCTION(BlueprintCallable, meta = (DisplayName = "Add cluster event listener"), Category = "NDisplay|Cluster")
 	virtual void AddClusterEventListener(TScriptInterface<IDisplayClusterClusterEventListener> Listener) override;
 
@@ -71,156 +75,4 @@ public:
 
 	UFUNCTION(BlueprintCallable, meta = (DisplayName = "Emit binary cluster event"), Category = "NDisplay|Cluster")
 	virtual void SendClusterEventBinaryTo(const FString& Address, const int32 Port, const FDisplayClusterClusterEventBinary& Event, bool bMasterOnly) override;
-
-public:
-	//////////////////////////////////////////////////////////////////////////////////////////////
-	// Config API
-	//////////////////////////////////////////////////////////////////////////////////////////////
-	UFUNCTION(BlueprintCallable, meta = (DisplayName = "Get config data"), Category = "NDisplay|Cluster")
-	virtual UDisplayClusterConfigurationData* GetConfig() const override;
-
-public:
-	//////////////////////////////////////////////////////////////////////////////////////////////
-	// Game API
-	//////////////////////////////////////////////////////////////////////////////////////////////
-	// Root
-	UFUNCTION(BlueprintCallable, meta = (DisplayName = "Get root actor"), Category = "NDisplay|Game")
-	virtual ADisplayClusterRootActor* GetRootActor() const override;
-
-public:
-	//////////////////////////////////////////////////////////////////////////////////////////////
-	// Input API
-	//////////////////////////////////////////////////////////////////////////////////////////////
-	// Device information
-	UFUNCTION(BlueprintCallable, meta = (DisplayName = "Get amount of VRPN axis devices", DeprecatedFunction, DeprecationMessage = "VRPN functionality has been moved to LiveLinkVRPN"), Category = "NDisplay|Input")
-	virtual int32 GetAxisDeviceAmount() const override
-	{
-		return 0;
-	}
-
-	UFUNCTION(BlueprintCallable, meta = (DisplayName = "Get amount of VRPN button devices", DeprecatedFunction, DeprecationMessage = "VRPN functionality has been moved to LiveLinkVRPN"), Category = "NDisplay|Input")
-	virtual int32 GetButtonDeviceAmount() const override
-	{
-		return 0;
-	}
-
-	UFUNCTION(BlueprintCallable, meta = (DisplayName = "Get amount of VRPN tracker devices", DeprecatedFunction, DeprecationMessage = "VRPN functionality has been moved to LiveLinkVRPN"), Category = "NDisplay|Input")
-	virtual int32 GetTrackerDeviceAmount() const override
-	{
-		return 0;
-	}
-
-	UFUNCTION(BlueprintCallable, meta = (DisplayName = "Get IDs of VRPN axis devices", DeprecatedFunction, DeprecationMessage = "VRPN functionality has been moved to LiveLinkVRPN"), Category = "NDisplay|Input")
-	virtual void GetAxisDeviceIds(TArray<FString>& DeviceIDs) const override
-	{ }
-
-	UFUNCTION(BlueprintCallable, meta = (DisplayName = "Get IDs of VRPN button devices", DeprecatedFunction, DeprecationMessage = "VRPN functionality has been moved to LiveLinkVRPN"), Category = "NDisplay|Input")
-	virtual void GetButtonDeviceIds(TArray<FString>& DeviceIDs) const override
-	{ }
-
-	UFUNCTION(BlueprintCallable, meta = (DisplayName = "Get IDs of keyboard devices", DeprecatedFunction, DeprecationMessage = "VRPN functionality has been moved to LiveLinkVRPN"), Category = "NDisplay|Input")
-	virtual void GetKeyboardDeviceIds(TArray<FString>& DeviceIDs) const override
-	{ }
-
-	UFUNCTION(BlueprintCallable, meta = (DisplayName = "Get IDs of VRPN tracker devices", DeprecatedFunction, DeprecationMessage = "VRPN functionality has been moved to LiveLinkVRPN"), Category = "NDisplay|Input")
-	virtual void GetTrackerDeviceIds(TArray<FString>& DeviceIDs) const override
-	{ }
-
-	// Buttons
-	UFUNCTION(BlueprintCallable, meta = (DisplayName = "Get VRPN button state", DeprecatedFunction, DeprecationMessage = "VRPN functionality has been moved to LiveLinkVRPN"), Category = "NDisplay|Input")
-	virtual void GetButtonState(const FString& DeviceID, int32 DeviceChannel, bool& CurrentState, bool& IsChannelAvailable) const override
-	{
-		IsChannelAvailable = false;
-	}
-
-	UFUNCTION(BlueprintCallable, meta = (DisplayName = "Is VRPN button pressed", DeprecatedFunction, DeprecationMessage = "VRPN functionality has been moved to LiveLinkVRPN"), Category = "NDisplay|Input")
-	virtual void IsButtonPressed(const FString& DeviceID, int32 DeviceChannel, bool& IsPressedCurrently, bool& IsChannelAvailable) const override
-	{
-		IsChannelAvailable = false;
-	}
-
-	UFUNCTION(BlueprintCallable, meta = (DisplayName = "Is VRPN button released", DeprecatedFunction, DeprecationMessage = "VRPN functionality has been moved to LiveLinkVRPN"), Category = "NDisplay|Input")
-	virtual void IsButtonReleased(const FString& DeviceID, int32 DeviceChannel, bool& IsReleasedCurrently, bool& IsChannelAvailable) const override
-	{
-		IsChannelAvailable = false;
-	}
-
-	UFUNCTION(BlueprintCallable, meta = (DisplayName = "Was VRPN button pressed", DeprecatedFunction, DeprecationMessage = "VRPN functionality has been moved to LiveLinkVRPN"), Category = "NDisplay|Input")
-	virtual void WasButtonPressed(const FString& DeviceID, int32 DeviceChannel, bool& WasPressed, bool& IsChannelAvailable) const override
-	{
-		IsChannelAvailable = false;
-	}
-
-	UFUNCTION(BlueprintCallable, meta = (DisplayName = "Was VRPN button released", DeprecatedFunction, DeprecationMessage = "VRPN functionality has been moved to LiveLinkVRPN"), Category = "NDisplay|Input")
-	virtual void WasButtonReleased(const FString& DeviceID, int32 DeviceChannel, bool& WasReleased, bool& IsChannelAvailable) const override
-	{
-		IsChannelAvailable = false;
-	}
-
-	// Axes
-	UFUNCTION(BlueprintCallable, meta = (DisplayName = "Get VRPN axis value", DeprecatedFunction, DeprecationMessage = "VRPN functionality has been moved to LiveLinkVRPN"), Category = "NDisplay|Input")
-	virtual void GetAxis(const FString& DeviceID, int32 DeviceChannel, float& Value, bool& IsChannelAvailable) const override
-	{
-		IsChannelAvailable = false;
-	}
-
-	// Trackers
-	UFUNCTION(BlueprintCallable, meta = (DisplayName = "Get VRPN tracker location", DeprecatedFunction, DeprecationMessage = "VRPN functionality has been moved to LiveLinkVRPN"), Category = "NDisplay|Input")
-	virtual void GetTrackerLocation(const FString& DeviceID, int32 DeviceChannel, FVector& Location, bool& IsChannelAvailable) const override
-	{
-		IsChannelAvailable = false;
-	}
-
-	UFUNCTION(BlueprintCallable, meta = (DisplayName = "Get VRPN tracker rotation (as quaternion)", DeprecatedFunction, DeprecationMessage = "VRPN functionality has been moved to LiveLinkVRPN"), Category = "NDisplay|Input")
-	virtual void GetTrackerQuat(const FString& DeviceID, int32 DeviceChannel, FQuat& Rotation, bool& IsChannelAvailable) const override
-	{
-		IsChannelAvailable = false;
-	}
-
-public:
-	//////////////////////////////////////////////////////////////////////////////////////////////
-	// Render API
-	//////////////////////////////////////////////////////////////////////////////////////////////
-	UFUNCTION(BlueprintCallable, meta = (DisplayName = "Set viewport camera", DeprecatedFunction, DeprecationMessage = "Use Configuration structures"), Category = "NDisplay|Render")
-	virtual void SetViewportCamera(const FString& CameraId, const FString& ViewportId) override
-	{ }
-
-	UFUNCTION(BlueprintCallable, meta = (DisplayName = "Get viewport's buffer ratio", DeprecatedFunction, DeprecationMessage = "Use Configuration structures"), Category = "NDisplay|Render")
-	virtual bool GetBufferRatio(const FString& ViewportId, float& BufferRatio) const override
-	{
-		return false;
-	}
-
-	UFUNCTION(BlueprintCallable, meta = (DisplayName = "Set viewport's buffer ratio", DeprecatedFunction, DeprecationMessage = "Use Configuration structures"), Category = "NDisplay|Render")
-	virtual bool SetBufferRatio(const FString& ViewportId, float BufferRatio) override
-	{
-		return false;
-	}
-
-	UFUNCTION(BlueprintCallable, meta = (DisplayName = "Set start post processing settings for viewport", DeprecatedFunction, DeprecationMessage = "Use Configuration structures"), Category = "NDisplay|Render")
-	virtual void SetStartPostProcessingSettings(const FString& ViewportId, const FPostProcessSettings& StartPostProcessingSettings) override
-	{ }
-
-	UFUNCTION(BlueprintCallable, meta = (DisplayName = "Set override post processing settings for viewport", DeprecatedFunction, DeprecationMessage = "Use Configuration structures"), Category = "NDisplay|Render")
-	virtual void SetOverridePostProcessingSettings(const FString& ViewportId, const FPostProcessSettings& OverridePostProcessingSettings, float BlendWeight = 1.0f) override
-	{ }
-
-	UFUNCTION(BlueprintCallable, meta = (DisplayName = "Set final post processing settings for viewport", DeprecatedFunction, DeprecationMessage = "Use Configuration structures"), Category = "NDisplay|Render")
-	virtual void SetFinalPostProcessingSettings(const FString& ViewportId, const FPostProcessSettings& FinalPostProcessingSettings) override
-	{ }
-
-	UFUNCTION(BlueprintCallable, meta = (DisplayName = "Get Updated Post Processing", DeprecatedFunction, DeprecationMessage = "Use new api"), Category = "NDisplay|Render")
-		virtual bool GetViewportRect(const FString& ViewportId, FIntPoint& ViewportLoc, FIntPoint& ViewportSize) const override
-	{
-		return false;
-	}
-
-	/** Returns list of local viewports. */
-	UFUNCTION(BlueprintCallable, meta = (DisplayName = "Get local viewports", DeprecatedFunction, DeprecationMessage = "Use new api"), Category = "NDisplay|Config")
-	virtual void GetLocalViewports(TArray<FString>& ViewportIDs, TArray<FString>& ProjectionTypes, TArray<FIntPoint>& ViewportLocations, TArray<FIntPoint>& ViewportSizes) const override
-	{ }
-
-	UFUNCTION(BlueprintCallable, meta = (DisplayName = "Scene View Extension Is Active In Context Function", DeprecatedFunction, DeprecationMessage = "Use Configuration structures"), Category = "NDisplay|Render")
-	virtual void SceneViewExtensionIsActiveInContextFunction(const TArray<FString>& ViewportIDs, FSceneViewExtensionIsActiveFunctor& OutIsActiveFunction) const override
-	{ }
 };
diff --git a/Engine/Plugins/Runtime/nDisplay/Source/DisplayCluster/Private/Config/DisplayClusterConfigManager.cpp b/Engine/Plugins/Runtime/nDisplay/Source/DisplayCluster/Private/Config/DisplayClusterConfigManager.cpp
--- a/Engine/Plugins/Runtime/nDisplay/Source/DisplayCluster/Private/Config/DisplayClusterConfigManager.cpp
+++ b/Engine/Plugins/Runtime/nDisplay/Source/DisplayCluster/Private/Config/DisplayClusterConfigManager.cpp
@@ -54,12 +54,12 @@ FString FDisplayClusterConfigManager::GetMasterNodeId() const
 
 const UDisplayClusterConfigurationClusterNode* FDisplayClusterConfigManager::GetMasterNode() const
 {
-	return ConfigData ? ConfigData->GetClusterNode(GetMasterNodeId()) : nullptr;
+	return ConfigData ? ConfigData->Cluster->GetNode(GetMasterNodeId()) : nullptr;
 }
 
 const UDisplayClusterConfigurationClusterNode* FDisplayClusterConfigManager::GetLocalNode() const
 {
-	return ConfigData ? ConfigData->GetClusterNode(GetLocalNodeId()) : nullptr;
+	return ConfigData ? ConfigData->Cluster->GetNode(GetLocalNodeId()) : nullptr;
 }
 
 const UDisplayClusterConfigurationViewport* FDisplayClusterConfigManager::GetLocalViewport(const FString& ViewportId) const
diff --git a/Engine/Plugins/Runtime/nDisplay/Source/DisplayCluster/Private/Game/EngineClasses/Scene/DisplayClusterRootActor.cpp b/Engine/Plugins/Runtime/nDisplay/Source/DisplayCluster/Private/Game/EngineClasses/Scene/DisplayClusterRootActor.cpp
--- a/Engine/Plugins/Runtime/nDisplay/Source/DisplayCluster/Private/Game/EngineClasses/Scene/DisplayClusterRootActor.cpp
+++ b/Engine/Plugins/Runtime/nDisplay/Source/DisplayCluster/Private/Game/EngineClasses/Scene/DisplayClusterRootActor.cpp
@@ -268,16 +268,6 @@ void ADisplayClusterRootActor::OverrideFromConfig(UDisplayClusterConfigurationDa
 #endif
 }
 
-UDisplayClusterConfigurationViewport* ADisplayClusterRootActor::GetViewportConfiguration(const FString& ClusterNodeID, const FString& ViewportID)
-{
-	if (CurrentConfigData)
-	{
-		return CurrentConfigData->GetViewportConfiguration(ClusterNodeID, ViewportID);
-	}
-
-	return nullptr;
-}
-
 void ADisplayClusterRootActor::UpdateConfigDataInstance(UDisplayClusterConfigurationData* ConfigDataTemplate, bool bForceRecreate)
 {
 	if (ConfigDataTemplate == nullptr)
@@ -918,42 +908,42 @@ bool ADisplayClusterRootActor::SetReplaceTextureFlagForAllViewports(bool bReplac
 		return false;
 	}
 
+	if (!ConfigData->Cluster)
+	{
+		UE_LOG(LogDisplayClusterGame, Warning, TEXT("ConfigData's Cluster was null"));
+		return false;
+	}
+
 	const FString NodeId = Display.GetClusterMgr()->GetNodeId();
-	const UDisplayClusterConfigurationClusterNode* Node = ConfigData->GetClusterNode(NodeId);
+	const UDisplayClusterConfigurationClusterNode* Node = ConfigData->Cluster->GetNode(NodeId);
 
-	if (Node)
+	
+	for (const TPair<FString, UDisplayClusterConfigurationViewport*>& ViewportItem : Node->Viewports)
 	{
-		for (const TPair<FString, UDisplayClusterConfigurationViewport*>& ViewportItem : Node->Viewports)
+		if (ViewportItem.Value)
 		{
-			if (ViewportItem.Value)
-			{
-				ViewportItem.Value->RenderSettings.Replace.bAllowReplace = bReplace;
-			}
+			ViewportItem.Value->RenderSettings.Replace.bAllowReplace = bReplace;
 		}
 	}
-	else if (ConfigData->Cluster)
+	
+	
+	for (const TPair<FString, UDisplayClusterConfigurationClusterNode*>& NodeItem : ConfigData->Cluster->Nodes)
 	{
-		for (const TPair<FString, UDisplayClusterConfigurationClusterNode*>& NodeItem : ConfigData->Cluster->Nodes)
+		if (!NodeItem.Value)
 		{
-			if (!NodeItem.Value)
-			{
-				continue;
-			}
+			continue;
+		}
 
-			for (const TPair<FString, UDisplayClusterConfigurationViewport*>& ViewportItem : NodeItem.Value->Viewports)
+		for (const TPair<FString, UDisplayClusterConfigurationViewport*>& ViewportItem : NodeItem.Value->Viewports)
+		{
+			if (ViewportItem.Value)
 			{
-				if (ViewportItem.Value)
-				{
-					ViewportItem.Value->RenderSettings.Replace.bAllowReplace = bReplace;
-				}
+				ViewportItem.Value->RenderSettings.Replace.bAllowReplace = bReplace;
 			}
 		}
 	}
-	else
-	{
-		UE_LOG(LogDisplayClusterGame, Warning, TEXT("ConfigData's Cluster was null"));
-		return false;
-	}
+	
+	
 
 	return true;
 }
diff --git a/Engine/Plugins/Runtime/nDisplay/Source/DisplayCluster/Private/Game/EngineClasses/Scene/DisplayClusterTestPatternsActor.cpp b/Engine/Plugins/Runtime/nDisplay/Source/DisplayCluster/Private/Game/EngineClasses/Scene/DisplayClusterTestPatternsActor.cpp
--- a/Engine/Plugins/Runtime/nDisplay/Source/DisplayCluster/Private/Game/EngineClasses/Scene/DisplayClusterTestPatternsActor.cpp
+++ b/Engine/Plugins/Runtime/nDisplay/Source/DisplayCluster/Private/Game/EngineClasses/Scene/DisplayClusterTestPatternsActor.cpp
@@ -109,14 +109,17 @@ void ADisplayClusterTestPatternsActor::Tick(float DeltaSeconds)
 
 			for (auto it = ViewportPPSettings.CreateConstIterator(); it; ++it)
 			{
-				UDisplayClusterConfigurationViewport* ViewportCfg = RootActor->GetViewportConfiguration(LocalNodeId, it->Key);
-				if (ViewportCfg)
+				if (RootActor->CurrentConfigData != nullptr)
 				{
-					// Assign current post-process settigns for each viewport
-					FDisplayClusterConfigurationViewport_CustomPostprocessSettings& DstPostProcess = ViewportCfg->RenderSettings.CustomPostprocess.Override;
-					DstPostProcess.bIsEnabled = true;
-					DstPostProcess.bIsOneFrame = true;
-					DstPostProcess.PostProcessSettings = it->Value;
+					UDisplayClusterConfigurationViewport* ViewportCfg = RootActor->CurrentConfigData->GetViewport(LocalNodeId, it->Key);
+					if (ViewportCfg)
+					{
+						// Assign current post-process settigns for each viewport
+						FDisplayClusterConfigurationViewport_CustomPostprocessSettings& DstPostProcess = ViewportCfg->RenderSettings.CustomPostprocess.Override;
+						DstPostProcess.bIsEnabled = true;
+						DstPostProcess.bIsOneFrame = true;
+						DstPostProcess.PostProcessSettings = it->Value;
+					}
 				}
 			}
 		}
diff --git a/Engine/Plugins/Runtime/nDisplay/Source/DisplayCluster/Private/Render/Viewport/Configuration/DisplayClusterViewportConfigurationBase.cpp b/Engine/Plugins/Runtime/nDisplay/Source/DisplayCluster/Private/Render/Viewport/Configuration/DisplayClusterViewportConfigurationBase.cpp
--- a/Engine/Plugins/Runtime/nDisplay/Source/DisplayCluster/Private/Render/Viewport/Configuration/DisplayClusterViewportConfigurationBase.cpp
+++ b/Engine/Plugins/Runtime/nDisplay/Source/DisplayCluster/Private/Render/Viewport/Configuration/DisplayClusterViewportConfigurationBase.cpp
@@ -20,7 +20,7 @@ void FDisplayClusterViewportConfigurationBase::Update(const TArray<FString>& InC
 	// Get render viewports
 	for (const FString& NodeIt : InClusterNodeIds)
 	{
-		const UDisplayClusterConfigurationClusterNode* ClusterNode = ConfigurationData.GetClusterNode(NodeIt);
+		const UDisplayClusterConfigurationClusterNode* ClusterNode = ConfigurationData.Cluster->GetNode(NodeIt);
 		if (ClusterNode)
 		{
 			for (const TPair<FString, UDisplayClusterConfigurationViewport*>& ViewportIt : ClusterNode->Viewports)
@@ -71,7 +71,7 @@ void FDisplayClusterViewportConfigurationBase::Update(const TArray<FString>& InC
 
 void FDisplayClusterViewportConfigurationBase::UpdateClusterNodePostProcess(const FString& InClusterNodeId)
 {
-	const UDisplayClusterConfigurationClusterNode* ClusterNode = ConfigurationData.GetClusterNode(InClusterNodeId);
+	const UDisplayClusterConfigurationClusterNode* ClusterNode = ConfigurationData.Cluster->GetNode(InClusterNodeId);
 	if (ClusterNode)
 	{
 		TSharedPtr<FDisplayClusterViewportPostProcessManager, ESPMode::ThreadSafe> PPManager = ViewportManager.GetPostProcessManager();
diff --git a/Engine/Plugins/Runtime/nDisplay/Source/DisplayCluster/Public/Blueprints/IDisplayClusterBlueprintAPI.h b/Engine/Plugins/Runtime/nDisplay/Source/DisplayCluster/Public/Blueprints/IDisplayClusterBlueprintAPI.h
--- a/Engine/Plugins/Runtime/nDisplay/Source/DisplayCluster/Public/Blueprints/IDisplayClusterBlueprintAPI.h
+++ b/Engine/Plugins/Runtime/nDisplay/Source/DisplayCluster/Public/Blueprints/IDisplayClusterBlueprintAPI.h
@@ -32,7 +32,6 @@ class DISPLAYCLUSTER_API UDisplayClusterBlueprintAPI : public UInterface
 	GENERATED_BODY()
 };
 
-
 /**
  * Blueprint API interface
  */
@@ -41,9 +40,6 @@ class DISPLAYCLUSTER_API IDisplayClusterBlueprintAPI
 	GENERATED_BODY()
 
 public:
-	//////////////////////////////////////////////////////////////////////////////////////////////
-	// DisplayCluster module API
-	//////////////////////////////////////////////////////////////////////////////////////////////
 	/* Returns true if the module has been initialized. */
 	UFUNCTION(BlueprintCallable, meta = (DisplayName = "Is module initialized"), Category = "NDisplay")
 	virtual bool IsModuleInitialized() const = 0;
@@ -53,13 +49,28 @@ public:
 	virtual EDisplayClusterOperationMode GetOperationMode() const = 0;
 
 public:
-	//////////////////////////////////////////////////////////////////////////////////////////////
-	// Cluster API
-	//////////////////////////////////////////////////////////////////////////////////////////////
+	/** Returns DisplayCluster root actor. */
+	UFUNCTION(BlueprintCallable, meta = (DisplayName = "Get root actor"), Category = "NDisplay|Game")
+	virtual ADisplayClusterRootActor* GetRootActor() const = 0;
+
+	// Runtime local node API
+
+	/** Returns Id of the current node in a cluster. */
+	UFUNCTION(BlueprintCallable, meta = (DisplayName = "Get node ID"), Category = "NDisplay|Cluster")
+	virtual FString GetNodeId() const = 0;
+
+	/** Returns List of the active nodes in the runtime cluster node in a cluster. */
+	UFUNCTION(BlueprintCallable, meta = (DisplayName = "Get cluster node IDs"), Category = "NDisplay|Cluster")
+	virtual void GetActiveNodeIds(TArray<FString>& OutNodeIds) const = 0;
+
+	/** Returns amount of active nodes in a cluster. */
+	UFUNCTION(BlueprintCallable, meta = (DisplayName = "Get nodes amount"), Category = "NDisplay|Cluster")
+	virtual int32 GetActiveNodesAmount() const = 0;
+
 	/** Returns true if current node is a master node in a cluster. */
 	UFUNCTION(BlueprintCallable, meta = (DisplayName = "Is master node"), Category = "NDisplay|Cluster")
 	virtual bool IsMaster() const = 0;
-	
+
 	/** Returns true if current node is a slave node in a cluster. */
 	UFUNCTION(BlueprintCallable, meta = (DisplayName = "Is slave node"), Category = "NDisplay|Cluster")
 	virtual bool IsSlave() const = 0;
@@ -72,17 +83,7 @@ public:
 	UFUNCTION(BlueprintCallable, meta = (DisplayName = "Get cluster role"), Category = "NDisplay|Cluster")
 	virtual EDisplayClusterNodeRole GetClusterRole() const = 0;
 
-	/** Returns the role of the current cluster node. */
-	UFUNCTION(BlueprintCallable, meta = (DisplayName = "Get cluster node IDs"), Category = "NDisplay|Cluster")
-	virtual void GetNodeIds(TArray<FString>& OutNodeIds) const = 0;
-
-	/** Returns cluster node name of the current application instance. */
-	UFUNCTION(BlueprintCallable, meta = (DisplayName = "Get node ID"), Category = "NDisplay|Cluster")
-	virtual FString GetNodeId() const = 0;
-
-	/** Returns amount of nodes in a cluster. */
-	UFUNCTION(BlueprintCallable, meta = (DisplayName = "Get nodes amount"), Category = "NDisplay|Cluster")
-	virtual int32 GetNodesAmount() const = 0;
+	// cluster events api
 
 	/** Adds cluster event listener. */
 	UFUNCTION(BlueprintCallable, meta = (DisplayName = "Add cluster event listener"), Category = "NDisplay|Cluster")
@@ -107,184 +108,4 @@ public:
 	/** Sends binary cluster event to a specific target (outside of the cluster). */
 	UFUNCTION(BlueprintCallable, meta = (DisplayName = "Send binary event to a specific host"), Category = "NDisplay|Cluster")
 	virtual void SendClusterEventBinaryTo(const FString& Address, const int32 Port, const FDisplayClusterClusterEventBinary& Event, bool bMasterOnly) = 0;
-
-public:
-	//////////////////////////////////////////////////////////////////////////////////////////////
-	// Config API
-	//////////////////////////////////////////////////////////////////////////////////////////////
-	/** Return current configuration data */
-	UFUNCTION(BlueprintCallable, meta = (DisplayName = "Get config data"), Category = "NDisplay|Cluster")
-	virtual UDisplayClusterConfigurationData* GetConfig() const = 0;
-
-public:
-	//////////////////////////////////////////////////////////////////////////////////////////////
-	// Game API
-	//////////////////////////////////////////////////////////////////////////////////////////////
-	/** Returns DisplayCluster root actor. */
-	UFUNCTION(BlueprintCallable, meta = (DisplayName = "Get root actor"), Category = "NDisplay|Game")
-	virtual ADisplayClusterRootActor* GetRootActor() const = 0;
-
-	// Screens
-	/** Returns screen component by ID. */
-	UFUNCTION(BlueprintCallable, meta = (DisplayName = "Get screen by ID", DeprecatedFunction, DeprecationMessage = "This function has been moved to ADisplayClusterRootActor."), Category = "NDisplay|Game")
-	virtual UDisplayClusterScreenComponent* GetScreenById(const FString& ScreenID)
-	{
-		return nullptr;
-	}
-
-	/** Returns array of all screen components. */
-	UFUNCTION(BlueprintCallable, meta = (DisplayName = "Get all screens", DeprecatedFunction, DeprecationMessage = "This function has been moved to ADisplayClusterRootActor."), Category = "NDisplay|Game")
-	virtual TArray<UDisplayClusterScreenComponent*> GetAllScreens()
-	{
-		return TArray<UDisplayClusterScreenComponent*>();
-	}
-
-	/** Returns amount of screens defined in current configuration file. */
-	UFUNCTION(BlueprintCallable, meta = (DisplayName = "Get amount of screens", DeprecatedFunction, DeprecationMessage = "This function has been moved to ADisplayClusterRootActor."), Category = "NDisplay|Game")
-	virtual int32 GetScreensAmount()
-	{
-		return 0;
-	}
-
-	// Cameras
-	/** Returns array of all available camera components. */
-	UFUNCTION(BlueprintCallable, meta = (DisplayName = "Get all cameras", DeprecatedFunction, DeprecationMessage = "This function has been moved to ADisplayClusterRootActor."), Category = "NDisplay|Game")
-	virtual TArray<UDisplayClusterCameraComponent*> GetAllCameras()
-	{
-		return TArray<UDisplayClusterCameraComponent*>();
-	}
-
-	/** Returns camera component with specified ID. */
-	UFUNCTION(BlueprintCallable, meta = (DisplayName = "Get camera by ID", DeprecatedFunction, DeprecationMessage = "This function has been moved to ADisplayClusterRootActor."), Category = "NDisplay|Game")
-	virtual UDisplayClusterCameraComponent* GetCameraById(const FString& CameraID)
-	{
-		return nullptr;
-	}
-
-	/** Returns amount of cameras. */
-	UFUNCTION(BlueprintCallable, meta = (DisplayName = "Get cameras amount", DeprecatedFunction, DeprecationMessage = "This function has been moved to ADisplayClusterRootActor."), Category = "NDisplay|Game")
-	virtual int32 GetCamerasAmount()
-	{
-		return 0;
-	}
-
-	/** Returns default camera component. */
-	UFUNCTION(BlueprintCallable, meta = (DisplayName = "Get default camera", DeprecatedFunction, DeprecationMessage = "This function has been moved to ADisplayClusterRootActor."), Category = "NDisplay|Game")
-	virtual UDisplayClusterCameraComponent* GetDefaultCamera()
-	{
-		return nullptr;
-	}
-
-	/** Sets default camera component specified by ID. */
-	UFUNCTION(BlueprintCallable, meta = (DisplayName = "Set default camera by ID", DeprecatedFunction, DeprecationMessage = "This function has been moved to ADisplayClusterRootActor."), Category = "NDisplay|Game")
-	virtual void SetDefaultCameraById(const FString& CameraID)
-	{ }
-
-public:
-	//////////////////////////////////////////////////////////////////////////////////////////////
-	// Input API
-	//////////////////////////////////////////////////////////////////////////////////////////////
-	// Device information
-	/** Returns amount of VRPN axis devices. */
-	UFUNCTION(BlueprintCallable, meta = (DisplayName = "Get amount of VRPN axis devices", DeprecatedFunction, DeprecationMessage = "VRPN functionality has been moved to LiveLinkVRPN"), Category = "NDisplay|Input")
-	virtual int32 GetAxisDeviceAmount() const = 0;
-
-	/** Returns amount of VRPN button devices. */
-	UFUNCTION(BlueprintCallable, meta = (DisplayName = "Get amount of VRPN button devices", DeprecatedFunction, DeprecationMessage = "VRPN functionality has been moved to LiveLinkVRPN"), Category = "NDisplay|Input")
-	virtual int32 GetButtonDeviceAmount() const = 0;
-
-	/** Returns amount of VRPN tracker devices. */
-	UFUNCTION(BlueprintCallable, meta = (DisplayName = "Get amount of VRPN tracker devices", DeprecatedFunction, DeprecationMessage = "VRPN functionality has been moved to LiveLinkVRPN"), Category = "NDisplay|Input")
-	virtual int32 GetTrackerDeviceAmount() const = 0;
-
-	/** Returns array of names of all VRPN axis devices. */
-	UFUNCTION(BlueprintCallable, meta = (DisplayName = "Get IDs of VRPN axis devices", DeprecatedFunction, DeprecationMessage = "VRPN functionality has been moved to LiveLinkVRPN"), Category = "NDisplay|Input")
-	virtual void GetAxisDeviceIds(TArray<FString>& DeviceIDs) const = 0;
-
-	/** Returns array of names of all VRPN button devices. */
-	UFUNCTION(BlueprintCallable, meta = (DisplayName = "Get IDs of VRPN button devices", DeprecatedFunction, DeprecationMessage = "VRPN functionality has been moved to LiveLinkVRPN"), Category = "NDisplay|Input")
-	virtual void GetButtonDeviceIds(TArray<FString>& DeviceIDs) const = 0;
-
-	/** Returns array of names of all keyboard devices. */
-	UFUNCTION(BlueprintCallable, meta = (DisplayName = "Get IDs of keyboard devices", DeprecatedFunction, DeprecationMessage = "VRPN functionality has been moved to LiveLinkVRPN"), Category = "NDisplay|Input")
-	virtual void GetKeyboardDeviceIds(TArray<FString>& DeviceIDs) const = 0;
-
-	/** Returns array of names of all VRPN tracker devices. */
-	UFUNCTION(BlueprintCallable, meta = (DisplayName = "Get IDs of VRPN tracker devices", DeprecatedFunction, DeprecationMessage = "VRPN functionality has been moved to LiveLinkVRPN"), Category = "NDisplay|Input")
-	virtual void GetTrackerDeviceIds(TArray<FString>& DeviceIDs) const = 0;
-
-	// Buttons
-	/** Returns state of VRPN button at specified device and channel. */
-	UFUNCTION(BlueprintCallable, meta = (DisplayName = "Get VRPN button state", DeprecatedFunction, DeprecationMessage = "VRPN functionality has been moved to LiveLinkVRPN"), Category = "NDisplay|Input")
-	virtual void GetButtonState(const FString& DeviceID, int32 DeviceChannel, bool& CurrentState, bool& IsChannelAvailable) const = 0;
-
-	/** Returns whether VRPN button is pressed at specified device and channel. */
-	UFUNCTION(BlueprintCallable, meta = (DisplayName = "Is VRPN button pressed", DeprecatedFunction, DeprecationMessage = "VRPN functionality has been moved to LiveLinkVRPN"), Category = "NDisplay|Input")
-	virtual void IsButtonPressed(const FString& DeviceID, int32 DeviceChannel, bool& IsPressedCurrently, bool& IsChannelAvailable) const = 0;
-
-	/** Returns whether VRPN button is released at specified device and channel. */
-	UFUNCTION(BlueprintCallable, meta = (DisplayName = "Is VRPN button released", DeprecatedFunction, DeprecationMessage = "VRPN functionality has been moved to LiveLinkVRPN"), Category = "NDisplay|Input")
-	virtual void IsButtonReleased(const FString& DeviceID, int32 DeviceChannel, bool& IsReleasedCurrently, bool& IsChannelAvailable) const = 0;
-
-	/** Returns whether VRPN button was released at specified device and channel. */
-	UFUNCTION(BlueprintCallable, meta = (DisplayName = "Was VRPN button pressed", DeprecatedFunction, DeprecationMessage = "VRPN functionality has been moved to LiveLinkVRPN"), Category = "NDisplay|Input")
-	virtual void WasButtonPressed(const FString& DeviceID, int32 DeviceChannel, bool& WasPressed, bool& IsChannelAvailable) const = 0;
-
-	/** Returns whether VRPN button was released at specified device and channel. */
-	UFUNCTION(BlueprintCallable, meta = (DisplayName = "Was VRPN button released", DeprecatedFunction, DeprecationMessage = "VRPN functionality has been moved to LiveLinkVRPN"), Category = "NDisplay|Input")
-	virtual void WasButtonReleased(const FString& DeviceID, int32 DeviceChannel, bool& WasReleased, bool& IsChannelAvailable) const = 0;
-
-	// Axes
-	/** Returns axis value at specified device and channel. */
-	UFUNCTION(BlueprintCallable, meta = (DisplayName = "Get VRPN axis value", DeprecatedFunction, DeprecationMessage = "VRPN functionality has been moved to LiveLinkVRPN"), Category = "NDisplay|Input")
-	virtual void GetAxis(const FString& DeviceID, int32 DeviceChannel, float& Value, bool& IsAvailable) const = 0;
-
-	// Trackers
-	/** Returns tracker location values at specified device and channel. */
-	UFUNCTION(BlueprintCallable, meta = (DisplayName = "Get VRPN tracker location", DeprecatedFunction, DeprecationMessage = "VRPN functionality has been moved to LiveLinkVRPN"), Category = "NDisplay|Input")
-	virtual void GetTrackerLocation(const FString& DeviceID, int32 DeviceChannel, FVector& Location, bool& IsChannelAvailable) const = 0;
-
-	/** Returns tracker quaternion values at specified device and channel. */
-	UFUNCTION(BlueprintCallable, meta = (DisplayName = "Get VRPN tracker rotation (as quaternion)", DeprecatedFunction, DeprecationMessage = "VRPN functionality has been moved to LiveLinkVRPN"), Category = "NDisplay|Input")
-	virtual void GetTrackerQuat(const FString& DeviceID, int32 DeviceChannel, FQuat& Rotation, bool& IsChannelAvailable) const = 0;
-
-public:
-	//////////////////////////////////////////////////////////////////////////////////////////////
-	// Render API
-	//////////////////////////////////////////////////////////////////////////////////////////////
-	/** Binds camera to a viewport. */
-	UFUNCTION(BlueprintCallable, meta = (DisplayName = "Set viewport camera", DeprecatedFunction, DeprecationMessage = "Use Configuration structures"), Category = "NDisplay|Render")
-	virtual void SetViewportCamera(const FString& CameraId, const FString& ViewportId) = 0;
-
-	/** Returns current buffer ratio for specified viewport. */
-	UFUNCTION(BlueprintCallable, meta = (DisplayName = "Get viewport's buffer ratio", DeprecatedFunction, DeprecationMessage = "Use Configuration structures"), Category = "NDisplay|Render")
-	virtual bool GetBufferRatio(const FString& ViewportId, float& BufferRatio) const = 0;
-
-	/** Sets buffer ratio for specified viewport. */
-	UFUNCTION(BlueprintCallable, meta = (DisplayName = "Set viewport's buffer ratio", DeprecatedFunction, DeprecationMessage = "Use Configuration structures"), Category = "NDisplay|Render")
-	virtual bool SetBufferRatio(const FString& ViewportId, float BufferRatio) = 0;
-
-	/** Overrides postprocess settings for specified viewport. */
-	UFUNCTION(BlueprintCallable, meta = (DisplayName = "Set start post processing settings for viewport", DeprecatedFunction, DeprecationMessage = "Use Configuration structures"), Category = "NDisplay|Render")
-	virtual void SetStartPostProcessingSettings(const FString& ViewportId, const FPostProcessSettings& StartPostProcessingSettings) = 0;
-
-	/** Overrides postprocess settings for specified viewport. */
-	UFUNCTION(BlueprintCallable, meta = (DisplayName = "Set override post processing settings for viewport", DeprecatedFunction, DeprecationMessage = "Use Configuration structures"), Category = "NDisplay|Render")
-	virtual void SetOverridePostProcessingSettings(const FString& ViewportId, const FPostProcessSettings& OverridePostProcessingSettings, float BlendWeight = 1.0f) = 0;
-
-	/** Overrides postprocess settings for specified viewport. */
-	UFUNCTION(BlueprintCallable, meta = (DisplayName = "Set final post processing settings for viewport", DeprecatedFunction, DeprecationMessage = "Use Configuration structures"), Category = "NDisplay|Render")
-	virtual void SetFinalPostProcessingSettings(const FString& ViewportId, const FPostProcessSettings& FinalPostProcessingSettings) = 0;
-
-	/** Returns location and size of specified viewport. */
-	UFUNCTION(BlueprintCallable, meta = (DisplayName = "Get viewport rectangle", DeprecatedFunction, DeprecationMessage = "Use new api"), Category = "NDisplay|Render")
-	virtual bool GetViewportRect(const FString& ViewportId, FIntPoint& ViewportLoc, FIntPoint& ViewportSize) const = 0;
-
-	/** Returns list of local viewports. */
-	UFUNCTION(BlueprintCallable, meta = (DisplayName = "Get local viewports", DeprecatedFunction, DeprecationMessage = "Use new api"), Category = "NDisplay|Config")
-	virtual void GetLocalViewports(TArray<FString>& ViewportIDs, TArray<FString>& ProjectionTypes, TArray<FIntPoint>& ViewportLocations, TArray<FIntPoint>& ViewportSizes) const = 0;
-
-	/** Returns a functor that determines if any given scene view extension should be active in the given context for the current frame */
-	UFUNCTION(BlueprintCallable, meta = (DisplayName = "Scene View Extension Is Active In Context Function", DeprecatedFunction, DeprecationMessage = "Use Configuration structures"), Category = "NDisplay|Render")
-	virtual void SceneViewExtensionIsActiveInContextFunction(const TArray<FString>& ViewportIDs, FSceneViewExtensionIsActiveFunctor& OutIsActiveFunction) const = 0;
 };
diff --git a/Engine/Plugins/Runtime/nDisplay/Source/DisplayCluster/Public/DisplayClusterRootActor.h b/Engine/Plugins/Runtime/nDisplay/Source/DisplayCluster/Public/DisplayClusterRootActor.h
--- a/Engine/Plugins/Runtime/nDisplay/Source/DisplayCluster/Public/DisplayClusterRootActor.h
+++ b/Engine/Plugins/Runtime/nDisplay/Source/DisplayCluster/Public/DisplayClusterRootActor.h
@@ -95,8 +95,6 @@ public:
 	const FDisplayClusterConfigurationICVFX_StageSettings& GetStageSettings() const;
 	const FDisplayClusterConfigurationRenderFrame& GetRenderFrameSettings() const;
 
-	UDisplayClusterConfigurationViewport* GetViewportConfiguration(const FString& ClusterNodeID, const FString& ViewportID);
-
 protected:
 	//////////////////////////////////////////////////////////////////////////////////////////////
 	// AActor
@@ -216,12 +214,6 @@ private:
 	UPROPERTY()
 	FName ConfigDataName;
 
-	/**
-	 * If set from the DisplayCluster BP Compiler it will be loaded from the class default subobjects in run-time.
-	 */
-	UPROPERTY(VisibleAnywhere, BlueprintReadOnly, Instanced, Category = "NDisplay", meta = (AllowPrivateAccess = "true"))
-	UDisplayClusterConfigurationData* CurrentConfigData;
-
 	/**
 	 * The root component for our hierarchy.
 	 * Must have CPF_Edit(such as VisibleDefaultsOnly) on property for Live Link.
@@ -257,6 +249,12 @@ public:
 	UPROPERTY(EditInstanceOnly, EditFixedSize, Category = "In Camera VFX", meta = (TitleProperty = "Name", DisplayAfter = "ViewportAllowInnerFrustumRef"))
 	TArray<FDisplayClusterComponentRef> InnerFrustumPriority;
 
+	/**
+	 * If set from the DisplayCluster BP Compiler it will be loaded from the class default subobjects in run-time.
+	 */
+	UPROPERTY(VisibleAnywhere, BlueprintReadOnly, Instanced, Category = "NDisplay", meta = (AllowPrivateAccess = "true"))
+	UDisplayClusterConfigurationData* CurrentConfigData;
+
 public:
 	bool IsInnerFrustumEnabled(const FString& InnerFrustumID) const;
 
diff --git a/Engine/Plugins/Runtime/nDisplay/Source/DisplayClusterConfiguration/Private/DisplayClusterConfigurationTypes.cpp b/Engine/Plugins/Runtime/nDisplay/Source/DisplayClusterConfiguration/Private/DisplayClusterConfigurationTypes.cpp
--- a/Engine/Plugins/Runtime/nDisplay/Source/DisplayClusterConfiguration/Private/DisplayClusterConfigurationTypes.cpp
+++ b/Engine/Plugins/Runtime/nDisplay/Source/DisplayClusterConfiguration/Private/DisplayClusterConfigurationTypes.cpp
@@ -41,6 +41,20 @@ UDisplayClusterConfigurationData::UDisplayClusterConfigurationData()
 {
 }
 
+UDisplayClusterConfigurationViewport* UDisplayClusterConfigurationData::GetViewport(const FString& NodeId, const FString& ViewportId) const
+{
+	UDisplayClusterConfigurationClusterNode* Node = Cluster->GetNode(NodeId);
+	if (Node)
+	{
+		UDisplayClusterConfigurationViewport** Viewport = Node->Viewports.Find(ViewportId);
+		if (Viewport)
+		{
+			return *Viewport;
+		}
+	}
+	return nullptr;
+}
+
 bool UDisplayClusterConfigurationData::AssignPostprocess(const FString& NodeId, const FString& PostprocessId, const FString& Type, TMap<FString, FString> Parameters, int32 Order)
 {
 	if (Cluster && Cluster->Nodes.Contains(NodeId))
@@ -75,44 +89,15 @@ bool UDisplayClusterConfigurationData::RemovePostprocess(const FString& NodeId,
 	return false;
 }
 
-UDisplayClusterConfigurationClusterNode* UDisplayClusterConfigurationData::GetClusterNodeConfiguration(const FString& NodeId) const
-{
-	return Cluster->Nodes.Contains(NodeId) ? Cluster->Nodes[NodeId] : nullptr;
-}
-
-const UDisplayClusterConfigurationClusterNode* UDisplayClusterConfigurationData::GetClusterNode(const FString& NodeId) const
-{
-	return Cluster->Nodes.Contains(NodeId) ? Cluster->Nodes[NodeId] : nullptr;
-}
-
-UDisplayClusterConfigurationViewport* UDisplayClusterConfigurationData::GetViewportConfiguration(const FString& NodeId, const FString& ViewportId)
-{
-	const UDisplayClusterConfigurationClusterNode* Node = GetClusterNode(NodeId);
-	if (Node)
-	{
-		if (Node->Viewports.Contains(ViewportId))
-		{
-			return Node->Viewports[ViewportId];
-		}
-	}
-
-	return nullptr;
-}
-
-const UDisplayClusterConfigurationViewport* UDisplayClusterConfigurationData::GetViewport(const FString& NodeId, const FString& ViewportId) const
+bool UDisplayClusterConfigurationData::GetPostprocess(const FString& NodeId, const FString& PostprocessId, FDisplayClusterConfigurationPostprocess& OutPostprocess) const
 {
-	const UDisplayClusterConfigurationClusterNode* Node = GetClusterNode(NodeId);
-	if (Node)
+	if (!Cluster)
 	{
-		return Node->Viewports.Contains(ViewportId) ? Node->Viewports[ViewportId] : nullptr;
+		return false;
 	}
 
-	return nullptr;
-}
+	const UDisplayClusterConfigurationClusterNode* Node = Cluster->GetNode(NodeId);
 
-bool UDisplayClusterConfigurationData::GetPostprocess(const FString& NodeId, const FString& PostprocessId, FDisplayClusterConfigurationPostprocess& OutPostprocess) const
-{
-	const UDisplayClusterConfigurationClusterNode* Node = GetClusterNode(NodeId);
 	if (Node)
 	{
 		const FDisplayClusterConfigurationPostprocess* PostprocessOperation = Node->Postprocess.Find(PostprocessId);
@@ -174,7 +159,7 @@ FDisplayClusterConfigurationProjection::FDisplayClusterConfigurationProjection()
 
 FDisplayClusterConfigurationPostprocess::FDisplayClusterConfigurationPostprocess()
 {
-	
+
 }
 
 const float UDisplayClusterConfigurationViewport::ViewportMinimumSize = 1.0f;
@@ -325,6 +310,17 @@ void UDisplayClusterConfigurationViewport::GetReferencedMeshNames(TArray<FString
 	}
 }
 
+void UDisplayClusterConfigurationClusterNode::GetViewportIds(TArray<FString>& OutViewportIds) const
+{
+	Viewports.GenerateKeyArray(OutViewportIds);
+}
+
+UDisplayClusterConfigurationViewport* UDisplayClusterConfigurationClusterNode::GetViewport(const FString& ViewportId) const
+{
+	UDisplayClusterConfigurationViewport* const* Viewport = Viewports.Find(ViewportId);
+	return Viewport ? *Viewport : nullptr;
+}
+
 void UDisplayClusterConfigurationClusterNode::GetReferencedMeshNames(TArray<FString>& OutMeshNames) const
 {
 	for (const TPair<FString, UDisplayClusterConfigurationViewport*>& It : Viewports)
@@ -338,6 +334,17 @@ void UDisplayClusterConfigurationClusterNode::GetReferencedMeshNames(TArray<FStr
 	}
 }
 
+void UDisplayClusterConfigurationCluster::GetNodeIds(TArray<FString>& OutNodeIds) const
+{
+	Nodes.GenerateKeyArray(OutNodeIds);
+}
+
+UDisplayClusterConfigurationClusterNode* UDisplayClusterConfigurationCluster::GetNode(const FString& NodeId) const
+{
+	UDisplayClusterConfigurationClusterNode* const* Node = Nodes.Find(NodeId);
+	return Node ? *Node : nullptr;
+}
+
 void UDisplayClusterConfigurationCluster::GetReferencedMeshNames(TArray<FString>& OutMeshNames) const
 {
 	for (const TPair<FString, UDisplayClusterConfigurationClusterNode*>& It : Nodes)
diff --git a/Engine/Plugins/Runtime/nDisplay/Source/DisplayClusterConfiguration/Public/DisplayClusterConfigurationTypes.h b/Engine/Plugins/Runtime/nDisplay/Source/DisplayClusterConfiguration/Public/DisplayClusterConfigurationTypes.h
--- a/Engine/Plugins/Runtime/nDisplay/Source/DisplayClusterConfiguration/Public/DisplayClusterConfigurationTypes.h
+++ b/Engine/Plugins/Runtime/nDisplay/Source/DisplayClusterConfiguration/Public/DisplayClusterConfigurationTypes.h
@@ -285,7 +285,14 @@ public:
 public:
 	UDisplayClusterConfigurationClusterNode();
 
+	UFUNCTION(BlueprintPure, Category = "NDisplay|Configuration")
+	void GetViewportIds(TArray<FString>& OutViewportIds) const;
+
+	UFUNCTION(BlueprintPure, Category = "NDisplay|Configuration")
+	UDisplayClusterConfigurationViewport* GetViewport(const FString& ViewportId) const;
+
 	// Return all references to meshes from policy, and other
+	UFUNCTION(BlueprintPure, Category = "NDisplay|Configuration")
 	void GetReferencedMeshNames(TArray<FString>& OutMeshNames) const;
 	
 private:
@@ -429,7 +436,15 @@ protected:
 
 public:
 	// Return all references to meshes from policy, and other
+	UFUNCTION(BlueprintPure, Category = "NDisplay|Configuration")
 	void GetReferencedMeshNames(TArray<FString>& OutMeshNames) const;
+
+	// Nodes API
+	UFUNCTION(BlueprintPure, Category = "NDisplay|Configuration")
+	void GetNodeIds(TArray<FString>& OutNodeIds) const;
+
+	UFUNCTION(BlueprintPure, Category = "NDisplay|Configuration")
+	UDisplayClusterConfigurationClusterNode* GetNode(const FString& NodeId) const;
 };
 
 USTRUCT(Blueprintable)
@@ -466,16 +481,10 @@ public:
 	UDisplayClusterConfigurationData();
 
 public:
-	// Facade API
-	const UDisplayClusterConfigurationClusterNode* GetClusterNode(const FString& NodeId) const;
-	const UDisplayClusterConfigurationViewport*    GetViewport(const FString& NodeId, const FString& ViewportId) const;
-
-	UFUNCTION(BlueprintCallable, Category = "NDisplay|Configuration")
-	UDisplayClusterConfigurationViewport* GetViewportConfiguration(const FString& NodeId, const FString& ViewportId);
-
-	UFUNCTION(BlueprintCallable, Category = "NDisplay|Configuration")
-	UDisplayClusterConfigurationClusterNode* GetClusterNodeConfiguration(const FString& NodeId) const;
-
+	// Viewports API
+	UFUNCTION(BlueprintPure, Category = "NDisplay|Configuration")
+	UDisplayClusterConfigurationViewport* GetViewport(const FString& NodeId, const FString& ViewportId) const;
+	
 	/**
 	* Update\Create node postprocess
 	*
@@ -489,20 +498,17 @@ public:
 	UFUNCTION(BlueprintCallable, Category = "NDisplay|Configuration")
 	bool AssignPostprocess(const FString& NodeId, const FString& PostprocessId, const FString& Type, TMap<FString, FString> Parameters, int32 Order = -1);
 
-	/**
-	* Delet node postprocess
-	*
-	* @param PostprocessId - Unique postprocess name
-	*
-	* @return - true, if success
-	*/
 	UFUNCTION(BlueprintCallable, Category = "NDisplay|Configuration")
 	bool RemovePostprocess(const FString& NodeId, const FString& PostprocessId);
 
+	UFUNCTION(BlueprintCallable, Category = "NDisplay|Configuration")
 	bool GetPostprocess(const FString& NodeId, const FString& PostprocessId, FDisplayClusterConfigurationPostprocess& OutPostprocess) const;
+
+	UFUNCTION(BlueprintCallable, Category = "NDisplay|Configuration")
 	bool GetProjectionPolicy(const FString& NodeId, const FString& ViewportId, FDisplayClusterConfigurationProjection& OutProjection) const;
 
 	// Return all references to meshes from policy, and other
+	UFUNCTION(BlueprintCallable, Category = "NDisplay|Configuration")
 	void GetReferencedMeshNames(TArray<FString>& OutMeshNames) const;
 
 public:
diff --git a/Engine/Plugins/Runtime/nDisplay/Source/DisplayClusterConfiguration/Public/DisplayClusterConfigurationTypes_Base.h b/Engine/Plugins/Runtime/nDisplay/Source/DisplayClusterConfiguration/Public/DisplayClusterConfigurationTypes_Base.h
--- a/Engine/Plugins/Runtime/nDisplay/Source/DisplayClusterConfiguration/Public/DisplayClusterConfigurationTypes_Base.h
+++ b/Engine/Plugins/Runtime/nDisplay/Source/DisplayClusterConfiguration/Public/DisplayClusterConfigurationTypes_Base.h
@@ -159,7 +159,7 @@ public:
 #endif
 };
 
-USTRUCT()
+USTRUCT(Blueprintable)
 struct DISPLAYCLUSTERCONFIGURATION_API FDisplayClusterConfigurationProjection
 	: public FDisplayClusterConfigurationPolymorphicEntity
 {
@@ -169,7 +169,7 @@ public:
 	FDisplayClusterConfigurationProjection();
 };
 
-USTRUCT()
+USTRUCT(Blueprintable)
 struct DISPLAYCLUSTERCONFIGURATION_API FDisplayClusterConfigurationPostprocess
 	: public FDisplayClusterConfigurationPolymorphicEntity
 {
diff --git a/Engine/Plugins/Runtime/nDisplay/Source/DisplayClusterConfiguration/Public/DisplayClusterConfigurationTypes_Viewport.h b/Engine/Plugins/Runtime/nDisplay/Source/DisplayClusterConfiguration/Public/DisplayClusterConfigurationTypes_Viewport.h
--- a/Engine/Plugins/Runtime/nDisplay/Source/DisplayClusterConfiguration/Public/DisplayClusterConfigurationTypes_Viewport.h
+++ b/Engine/Plugins/Runtime/nDisplay/Source/DisplayClusterConfiguration/Public/DisplayClusterConfigurationTypes_Viewport.h
@@ -158,7 +158,7 @@ public:
 	FString Camera;
 
 	/** Specify your Projection Policy Settings */
-	UPROPERTY(EditDefaultsOnly, Category = "Configuration")
+	UPROPERTY(EditDefaultsOnly, BlueprintReadWrite, Category = "Configuration")
 	FDisplayClusterConfigurationProjection ProjectionPolicy;
 
 	/** Enable or disable compatibility with inter process GPU Texture share */
