nDisplay: More gracefully handle current NodeId not found in ConfigData in ADisplayClusterRootActor::SetReplaceTextureFlagForAllViewports.

diff --git a/Engine/Plugins/Runtime/nDisplay/Source/DisplayCluster/Private/Game/EngineClasses/Scene/DisplayClusterRootActor.cpp b/Engine/Plugins/Runtime/nDisplay/Source/DisplayCluster/Private/Game/EngineClasses/Scene/DisplayClusterRootActor.cpp
--- a/Engine/Plugins/Runtime/nDisplay/Source/DisplayCluster/Private/Game/EngineClasses/Scene/DisplayClusterRootActor.cpp
+++ b/Engine/Plugins/Runtime/nDisplay/Source/DisplayCluster/Private/Game/EngineClasses/Scene/DisplayClusterRootActor.cpp
@@ -917,6 +917,11 @@ bool ADisplayClusterRootActor::SetReplaceTextureFlagForAllViewports(bool bReplac
 	const FString NodeId = Display.GetClusterMgr()->GetNodeId();
 	const UDisplayClusterConfigurationClusterNode* Node = ConfigData->Cluster->GetNode(NodeId);
 
+	if (!Node)
+	{
+		UE_LOG(LogDisplayClusterGame, Warning, TEXT("NodeId '%s' not found in ConfigData"), *NodeId);
+		return false;
+	}
 	
 	for (const TPair<FString, UDisplayClusterConfigurationViewport*>& ViewportItem : Node->Viewports)
 	{
@@ -926,7 +931,6 @@ bool ADisplayClusterRootActor::SetReplaceTextureFlagForAllViewports(bool bReplac
 		}
 	}
 	
-	
 	for (const TPair<FString, UDisplayClusterConfigurationClusterNode*>& NodeItem : ConfigData->Cluster->Nodes)
 	{
 		if (!NodeItem.Value)
@@ -942,8 +946,6 @@ bool ADisplayClusterRootActor::SetReplaceTextureFlagForAllViewports(bool bReplac
 			}
 		}
 	}
-	
-	
 
 	return true;
 }
