Use documentation module to link from nDisplay config to help instead of building the final URL in place.

diff --git a/Engine/Plugins/Runtime/nDisplay/Source/DisplayClusterConfigurator/Private/Views/OutputMapping/SDisplayClusterConfiguratorGraphEditor.cpp b/Engine/Plugins/Runtime/nDisplay/Source/DisplayClusterConfigurator/Private/Views/OutputMapping/SDisplayClusterConfiguratorGraphEditor.cpp
--- a/Engine/Plugins/Runtime/nDisplay/Source/DisplayClusterConfigurator/Private/Views/OutputMapping/SDisplayClusterConfiguratorGraphEditor.cpp
+++ b/Engine/Plugins/Runtime/nDisplay/Source/DisplayClusterConfigurator/Private/Views/OutputMapping/SDisplayClusterConfiguratorGraphEditor.cpp
@@ -23,6 +23,7 @@
 #include "GraphEditorActions.h"
 #include "SGraphPanel.h"
 #include "Misc/ConfigCacheIni.h"
+#include "IDocumentation.h"
 
 #include "Interfaces/Views/OutputMapping/IDisplayClusterConfiguratorViewOutputMapping.h"
 
@@ -480,16 +481,8 @@ bool SDisplayClusterConfiguratorGraphEditor::CanAddNewViewport() const
 
 void SDisplayClusterConfiguratorGraphEditor::BrowseDocumentation()
 {
-	const static FString OverrideUrlSection = TEXT("UnrealEd.URLOverrides");
-	const static FString DocumentationURL = TEXT("DocumentationURL");
-	const static FString NDisplayLink = TEXT("en-US/Engine/Rendering/nDisplay/");
-	FString OutURL;
-
-	if (GConfig->GetString(*OverrideUrlSection, *DocumentationURL, OutURL, GEditorIni))
-	{
-		FString URL = OutURL + NDisplayLink;
-		FPlatformProcess::LaunchURL(*URL, nullptr, nullptr);
-	}
+	const static FString NDisplayLink = TEXT("WorkingWithMedia/IntegratingMedia/nDisplay");
+	IDocumentation::Get()->Open(NDisplayLink, FDocumentationSourceInfo(TEXT("ndisplay_config")));
 }
 
 void SDisplayClusterConfiguratorGraphEditor::DeleteSelectedNodes()
