nDisplay: render sync updates
- v-blank logging
- fixed DX memory leak
- fixed policy instantiation for Linux

diff --git a/Engine/Plugins/Runtime/nDisplay/Source/DisplayCluster/Private/Render/DisplayClusterRenderManager.cpp b/Engine/Plugins/Runtime/nDisplay/Source/DisplayCluster/Private/Render/DisplayClusterRenderManager.cpp
--- a/Engine/Plugins/Runtime/nDisplay/Source/DisplayCluster/Private/Render/DisplayClusterRenderManager.cpp
+++ b/Engine/Plugins/Runtime/nDisplay/Source/DisplayCluster/Private/Render/DisplayClusterRenderManager.cpp
@@ -102,6 +102,10 @@ bool FDisplayClusterRenderManager::StartSession(UDisplayClusterConfigurationData
 	// Create synchronization object
 	UE_LOG(LogDisplayClusterRender, Log, TEXT("Instantiating synchronization policy object..."));
 	SyncPolicy = CreateRenderSyncPolicy();
+	if (SyncPolicy)
+	{
+		SyncPolicy->Initialize();
+	}
 
 	// Instantiate render device
 	TSharedPtr<IDisplayClusterRenderDevice, ESPMode::ThreadSafe> NewRenderDevice;
diff --git a/Engine/Plugins/Runtime/nDisplay/Source/DisplayCluster/Private/Render/Synchronization/DisplayClusterRenderSyncHelper.cpp b/Engine/Plugins/Runtime/nDisplay/Source/DisplayCluster/Private/Render/Synchronization/DisplayClusterRenderSyncHelper.cpp
new file mode 100644
--- /dev/null
+++ b/Engine/Plugins/Runtime/nDisplay/Source/DisplayCluster/Private/Render/Synchronization/DisplayClusterRenderSyncHelper.cpp
@@ -0,0 +1,27 @@
+// Copyright Epic Games, Inc. All Rights Reserved.
+
+#include "Render/Synchronization/DisplayClusterRenderSyncHelper.h"
+
+#include "Misc/ScopeLock.h"
+
+
+// Singleton instance
+TUniquePtr<IDisplayClusterRenderSyncHelper> FDisplayClusterRenderSyncHelper::Instance;
+
+// Singleton access mutex
+FCriticalSection FDisplayClusterRenderSyncHelper::CritSecInternals;
+
+
+IDisplayClusterRenderSyncHelper& FDisplayClusterRenderSyncHelper::Get()
+{
+	if (!Instance)
+	{
+		FScopeLock Lock(&FDisplayClusterRenderSyncHelper::CritSecInternals);
+		if (!Instance)
+		{
+			Instance = CreateHelper();
+		}
+	}
+
+	return *Instance.Get();
+}
diff --git a/Engine/Plugins/Runtime/nDisplay/Source/DisplayCluster/Private/Render/Synchronization/DisplayClusterRenderSyncHelper.h b/Engine/Plugins/Runtime/nDisplay/Source/DisplayCluster/Private/Render/Synchronization/DisplayClusterRenderSyncHelper.h
new file mode 100644
--- /dev/null
+++ b/Engine/Plugins/Runtime/nDisplay/Source/DisplayCluster/Private/Render/Synchronization/DisplayClusterRenderSyncHelper.h
@@ -0,0 +1,77 @@
+// Copyright Epic Games, Inc. All Rights Reserved.
+
+#pragma once
+
+#include "CoreMinimal.h"
+
+
+/**
+ * RHI agnostic helper interface
+ */
+class IDisplayClusterRenderSyncHelper
+{
+public:
+	virtual ~IDisplayClusterRenderSyncHelper() = default;
+
+public:
+	virtual bool WaitForVBlank()
+	{
+		return false;
+	}
+
+	virtual bool IsWaitForVBlankSupported()
+	{
+		return false;
+	}
+};
+
+
+/**
+ * Syncrhonization helper implementation. The main purpose of this class
+ * is to solve the combinatoric complexity Platform/RHI/SyncPolicy without
+ * making the whole sync policies hierarchy overcomplicated.
+ */
+class FDisplayClusterRenderSyncHelper
+{
+private:
+	FDisplayClusterRenderSyncHelper();
+
+public:
+	// Singleton access
+	static IDisplayClusterRenderSyncHelper& Get();
+
+protected:
+	// Factory method to instantiate helper instance of proper type
+	static TUniquePtr<IDisplayClusterRenderSyncHelper> CreateHelper();
+
+#if PLATFORM_WINDOWS
+private:
+	// DX helper implementation
+	struct FDCRSHelperDX : public IDisplayClusterRenderSyncHelper
+	{
+		virtual bool WaitForVBlank() override;
+		virtual bool IsWaitForVBlankSupported() override;
+	};
+#endif
+
+private:
+	// Vulkan helper implementation
+	struct FDCRSHelperVulkan : public IDisplayClusterRenderSyncHelper
+	{
+		virtual bool WaitForVBlank() override;
+		virtual bool IsWaitForVBlankSupported() override;
+	};
+
+private:
+	// Stub implementation for unknown RHI
+	struct FDCRSHelperNull : public IDisplayClusterRenderSyncHelper
+	{
+	};
+
+private:
+	// Helper instance
+	static TUniquePtr<IDisplayClusterRenderSyncHelper> Instance;
+
+	// Singleton access mutex
+	static FCriticalSection CritSecInternals;
+};
diff --git a/Engine/Plugins/Runtime/nDisplay/Source/DisplayCluster/Private/Render/Synchronization/DisplayClusterRenderSyncPolicyBase.cpp b/Engine/Plugins/Runtime/nDisplay/Source/DisplayCluster/Private/Render/Synchronization/DisplayClusterRenderSyncPolicyBase.cpp
--- a/Engine/Plugins/Runtime/nDisplay/Source/DisplayCluster/Private/Render/Synchronization/DisplayClusterRenderSyncPolicyBase.cpp
+++ b/Engine/Plugins/Runtime/nDisplay/Source/DisplayCluster/Private/Render/Synchronization/DisplayClusterRenderSyncPolicyBase.cpp
@@ -1,6 +1,7 @@
 // Copyright Epic Games, Inc. All Rights Reserved.
 
 #include "Render/Synchronization/DisplayClusterRenderSyncPolicyBase.h"
+#include "Render/Synchronization/DisplayClusterRenderSyncHelper.h"
 
 #include "Cluster/IPDisplayClusterClusterManager.h"
 #include "Cluster/Controller/IDisplayClusterClusterNodeController.h"
@@ -14,13 +15,31 @@
 #include "RHIResources.h"
 
 
-void FDisplayClusterRenderSyncPolicyBase::SyncBarrierRenderThread()
+static TAutoConsoleVariable<int32> CVarSyncDiagnosticsVBlankMonitoring(
+	TEXT("nDisplay.sync.diag.VBlankMonitoring"),
+	0,
+	TEXT("Sync diagnostics: V-blank monitoring\n")
+	TEXT("0 : disabled\n")
+	TEXT("1 : enabled (if policy supports only)\n")
+	,
+	ECVF_ReadOnly
+);
+
+
+bool FDisplayClusterRenderSyncPolicyBase::Initialize()
 {
-	if (GDisplayCluster->GetOperationMode() == EDisplayClusterOperationMode::Disabled)
+	// In case v-blank monitoring was requested and this feature is supported by
+	// the sync policy implementation, start v-blank monitoring.
+	if (!!CVarSyncDiagnosticsVBlankMonitoring.GetValueOnGameThread() && IsWaitForVBlankFeatureSupported())
 	{
-		return;
+		VBlankMonitor = MakeUnique<FDisplayClusterVBlankMonitor>(*this);
 	}
 
+	return true;
+}
+
+void FDisplayClusterRenderSyncPolicyBase::SyncBarrierRenderThread()
+{
 	{
 		TRACE_CPUPROFILER_EVENT_SCOPE(nDisplay SyncPolicyBase::SyncBarrier);
 		GDisplayCluster->GetPrivateClusterMgr()->GetClusterNodeController()->WaitForSwapSync();
@@ -38,3 +57,13 @@ void FDisplayClusterRenderSyncPolicyBase::WaitForFrameCompletion()
 		Viewport->WaitForFrameEventCompletion();
 	}
 }
+
+bool FDisplayClusterRenderSyncPolicyBase::IsWaitForVBlankFeatureSupported()
+{
+	return FDisplayClusterRenderSyncHelper::Get().IsWaitForVBlankSupported();
+}
+
+bool FDisplayClusterRenderSyncPolicyBase::WaitForVBlank()
+{
+	return FDisplayClusterRenderSyncHelper::Get().WaitForVBlank();
+}
diff --git a/Engine/Plugins/Runtime/nDisplay/Source/DisplayCluster/Private/Render/Synchronization/DisplayClusterRenderSyncPolicyBase.h b/Engine/Plugins/Runtime/nDisplay/Source/DisplayCluster/Private/Render/Synchronization/DisplayClusterRenderSyncPolicyBase.h
--- a/Engine/Plugins/Runtime/nDisplay/Source/DisplayCluster/Private/Render/Synchronization/DisplayClusterRenderSyncPolicyBase.h
+++ b/Engine/Plugins/Runtime/nDisplay/Source/DisplayCluster/Private/Render/Synchronization/DisplayClusterRenderSyncPolicyBase.h
@@ -3,6 +3,7 @@
 #pragma once
 
 #include "Render/Synchronization/IDisplayClusterRenderSyncPolicy.h"
+#include "Render/Synchronization/DisplayClusterVBlankMonitor.h"
 
 
 /**
@@ -19,16 +20,31 @@ public:
 	virtual ~FDisplayClusterRenderSyncPolicyBase() = default;
 
 public:
+	// Non-constructor initialization
+	virtual bool Initialize() override;
+
+	// VBlank related functions
+	virtual bool IsWaitForVBlankFeatureSupported() override;
+	virtual bool WaitForVBlank() override;
+
+public:
+	// Performs Ethernet barrier syncrhonization with other nodes in the cluster
 	void SyncBarrierRenderThread();
 
+	// Returns policy configuration parameters
 	const TMap<FString, FString>& GetParameters() const
 	{
 		return Parameters;
 	}
 
 protected:
+	// Wait for frame rendering completion using a GPU fence
 	virtual void WaitForFrameCompletion();
 
 private:
+	// Sync policy configuration parameters
 	TMap<FString, FString> Parameters;
+
+	// V-blank monitoring thread. Performs V-blank polling on the top level of the sync policies hierarchy.
+	TUniquePtr<FDisplayClusterVBlankMonitor> VBlankMonitor;
 };
diff --git a/Engine/Plugins/Runtime/nDisplay/Source/DisplayCluster/Private/Render/Synchronization/DisplayClusterRenderSyncPolicyEthernet.h b/Engine/Plugins/Runtime/nDisplay/Source/DisplayCluster/Private/Render/Synchronization/DisplayClusterRenderSyncPolicyEthernet.h
--- a/Engine/Plugins/Runtime/nDisplay/Source/DisplayCluster/Private/Render/Synchronization/DisplayClusterRenderSyncPolicyEthernet.h
+++ b/Engine/Plugins/Runtime/nDisplay/Source/DisplayCluster/Private/Render/Synchronization/DisplayClusterRenderSyncPolicyEthernet.h
@@ -34,7 +34,7 @@ protected:
 	virtual void Step_InitializeFrameSynchronization();
 	virtual void Step_WaitForFrameCompletion();
 	virtual void Step_WaitForEthernetBarrierSignal_1();
-	virtual void Step_SkipPresentationOnClosestVblank();
+	virtual void Step_SkipPresentationOnClosestVBlank();
 	virtual void Step_WaitForEthernetBarrierSignal_2();
 	virtual void Step_Present();
 	virtual void Step_FinalizeFrameSynchronization();
diff --git a/Engine/Plugins/Runtime/nDisplay/Source/DisplayCluster/Private/Render/Synchronization/DisplayClusterRenderSyncPolicyFactoryInternal.cpp b/Engine/Plugins/Runtime/nDisplay/Source/DisplayCluster/Private/Render/Synchronization/DisplayClusterRenderSyncPolicyFactoryInternal.cpp
--- a/Engine/Plugins/Runtime/nDisplay/Source/DisplayCluster/Private/Render/Synchronization/DisplayClusterRenderSyncPolicyFactoryInternal.cpp
+++ b/Engine/Plugins/Runtime/nDisplay/Source/DisplayCluster/Private/Render/Synchronization/DisplayClusterRenderSyncPolicyFactoryInternal.cpp
@@ -14,10 +14,17 @@
 
 TSharedPtr<IDisplayClusterRenderSyncPolicy> FDisplayClusterRenderSyncPolicyFactoryInternal::Create(const FString& InPolicyType, const FString& InRHIName, const TMap<FString, FString>& Parameters)
 {
+	// None
 	if (InPolicyType.Equals(DisplayClusterConfigurationStrings::config::cluster::render_sync::None, ESearchCase::IgnoreCase))
 	{
 		return MakeShared<FDisplayClusterRenderSyncPolicyNone>(Parameters);
 	}
+	// EthernetBarrier
+	else if (InPolicyType.Equals(DisplayClusterConfigurationStrings::config::cluster::render_sync::EthernetBarrier, ESearchCase::IgnoreCase))
+	{
+		return MakeShared<FDisplayClusterRenderSyncPolicyEthernetBarrier>(Parameters);
+	}
+	// Ethernet
 	else if (InPolicyType.Equals(DisplayClusterConfigurationStrings::config::cluster::render_sync::Ethernet, ESearchCase::IgnoreCase))
 	{
 #if PLATFORM_WINDOWS
@@ -37,6 +44,7 @@ TSharedPtr<IDisplayClusterRenderSyncPolicy> FDisplayClusterRenderSyncPolicyFacto
 		}
 #endif
 	}
+	// NVIDIA
 	else if (InPolicyType.Equals(DisplayClusterConfigurationStrings::config::cluster::render_sync::Nvidia, ESearchCase::IgnoreCase))
 	{
 #if PLATFORM_WINDOWS
@@ -47,17 +55,22 @@ TSharedPtr<IDisplayClusterRenderSyncPolicy> FDisplayClusterRenderSyncPolicyFacto
 		}
 		else if (InRHIName.Equals(DisplayClusterStrings::rhi::Vulkan, ESearchCase::IgnoreCase))
 		{
-			const FString DefaultPolicy(DisplayClusterConfigurationStrings::config::cluster::render_sync::EthernetBarrier);
+			UE_LOG(LogDisplayClusterRenderSync, Warning, TEXT("Sync policy '%s' has not been implemented for '%s' RHI. Default '%s' will be used."),
+				DisplayClusterConfigurationStrings::config::cluster::render_sync::Nvidia,
+				*InRHIName,
+				DisplayClusterConfigurationStrings::config::cluster::render_sync::EthernetBarrier);
 
-			UE_LOG(LogDisplayClusterRender, Warning, TEXT("Sync policy '%s' has not been implemented for '%s' RHI. Default '%s' will be used."),
-				DisplayClusterConfigurationStrings::config::cluster::render_sync::Nvidia, *InRHIName, *DefaultPolicy);
-
-			return Create(DefaultPolicy, InRHIName, Parameters);
+			return MakeShared<FDisplayClusterRenderSyncPolicyEthernetBarrier>(Parameters);
 		}
 #elif PLATFORM_LINUX
 		if (InRHIName.Equals(DisplayClusterStrings::rhi::Vulkan, ESearchCase::IgnoreCase))
 		{
-			return MakeShared<FDisplayClusterRenderSyncPolicyNvidia>(Parameters);
+			UE_LOG(LogDisplayClusterRenderSync, Warning, TEXT("Sync policy '%s' has not been implemented for '%s' RHI. Default '%s' will be used."),
+				DisplayClusterConfigurationStrings::config::cluster::render_sync::Nvidia,
+				*InRHIName,
+				DisplayClusterConfigurationStrings::config::cluster::render_sync::EthernetBarrier);
+
+			return MakeShared<FDisplayClusterRenderSyncPolicyEthernetBarrier>(Parameters);
 		}
 #endif
 	}
diff --git a/Engine/Plugins/Runtime/nDisplay/Source/DisplayCluster/Private/Render/Synchronization/DisplayClusterRenderSyncPolicyNvidia.h b/Engine/Plugins/Runtime/nDisplay/Source/DisplayCluster/Private/Render/Synchronization/DisplayClusterRenderSyncPolicyNvidia.h
--- a/Engine/Plugins/Runtime/nDisplay/Source/DisplayCluster/Private/Render/Synchronization/DisplayClusterRenderSyncPolicyNvidia.h
+++ b/Engine/Plugins/Runtime/nDisplay/Source/DisplayCluster/Private/Render/Synchronization/DisplayClusterRenderSyncPolicyNvidia.h
@@ -26,6 +26,10 @@ public:
 		return NvidiaPolicy;
 	}
 
+	// Initialization
+	virtual bool Initialize() override;
+
+	// Performs rendering and frame output synchronization
 	virtual bool SynchronizeClusterRendering(int32& InOutSyncInterval) override;
 
 private:
diff --git a/Engine/Plugins/Runtime/nDisplay/Source/DisplayCluster/Private/Render/Synchronization/DisplayClusterVblankMonitor.cpp b/Engine/Plugins/Runtime/nDisplay/Source/DisplayCluster/Private/Render/Synchronization/DisplayClusterVblankMonitor.cpp
new file mode 100644
--- /dev/null
+++ b/Engine/Plugins/Runtime/nDisplay/Source/DisplayCluster/Private/Render/Synchronization/DisplayClusterVblankMonitor.cpp
@@ -0,0 +1,59 @@
+// Copyright Epic Games, Inc. All Rights Reserved.
+
+#include "Render/Synchronization/DisplayClusterVBlankMonitor.h"
+
+#include "Render/Synchronization/IDisplayClusterRenderSyncPolicy.h"
+#include "HAL/RunnableThread.h"
+#include "Misc/DisplayClusterLog.h"
+#include "Misc/ScopeLock.h"
+
+
+FDisplayClusterVBlankMonitor::FDisplayClusterVBlankMonitor(IDisplayClusterRenderSyncPolicy& InSyncPolicy)
+	: SyncPolicy(InSyncPolicy)
+{
+	WorkingThread.Reset(FRunnableThread::Create(this, TEXT("VBLANK_Monitoring_thread"), 32 * 1024, TPri_AboveNormal));
+}
+
+FDisplayClusterVBlankMonitor::~FDisplayClusterVBlankMonitor()
+{
+	// Stop working thread
+	Stop();
+}
+
+uint32 FDisplayClusterVBlankMonitor::Run()
+{
+	UE_LOG(LogDisplayClusterRenderSync, Log, TEXT("V-blank monitoring thread has started"));
+
+	while (!bThreadExitRequested)
+	{
+		if (SyncPolicy.WaitForVBlank())
+		{
+			TRACE_BOOKMARK(TEXT("VBLANK %llu"), VBlankCounter);
+			++VBlankCounter;
+		}
+		else
+		{
+			TRACE_CPUPROFILER_EVENT_SCOPE(VBLANK - not available);
+			FPlatformProcess::Sleep(1.f);
+		}
+	}
+
+	UE_LOG(LogDisplayClusterRenderSync, Log, TEXT("V-blank monitoring thread has finished"));
+
+	return 0;
+}
+
+void FDisplayClusterVBlankMonitor::Stop()
+{
+	FScopeLock Lock(&CritSecInternals);
+
+	// Let the working thread know it has to stop
+	bThreadExitRequested = true;
+
+	// Wait unless working thread is finished
+	if (WorkingThread)
+	{
+		WorkingThread->WaitForCompletion();
+		WorkingThread.Reset();
+	}
+}
diff --git a/Engine/Plugins/Runtime/nDisplay/Source/DisplayCluster/Private/Render/Synchronization/DisplayClusterVblankMonitor.h b/Engine/Plugins/Runtime/nDisplay/Source/DisplayCluster/Private/Render/Synchronization/DisplayClusterVblankMonitor.h
new file mode 100644
--- /dev/null
+++ b/Engine/Plugins/Runtime/nDisplay/Source/DisplayCluster/Private/Render/Synchronization/DisplayClusterVblankMonitor.h
@@ -0,0 +1,44 @@
+// Copyright Epic Games, Inc. All Rights Reserved.
+
+#pragma once
+
+#include "CoreMinimal.h"
+#include "HAL/Runnable.h"
+
+class IDisplayClusterRenderSyncPolicy;
+
+
+/**
+ * V-blank monitor. It continuously polls the current sync policy
+ * for V-blanks and reports the time data to utraces.
+ */
+class FDisplayClusterVBlankMonitor
+	: protected FRunnable
+{
+public:
+	FDisplayClusterVBlankMonitor(IDisplayClusterRenderSyncPolicy& SyncPolicy);
+	~FDisplayClusterVBlankMonitor();
+
+protected:
+	//////////////////////////////////////////////////////////////////////////////////////////////
+	// FRunnable
+	//////////////////////////////////////////////////////////////////////////////////////////////
+	virtual uint32 Run() override;
+	virtual void Stop() override final;
+
+private:
+	// Holds the thread object
+	TUniquePtr<FRunnableThread> WorkingThread;
+
+	// Current sync policy
+	IDisplayClusterRenderSyncPolicy& SyncPolicy;
+
+	// V-blank counter
+	uint64 VBlankCounter = 0;
+
+	// Flag to ask working thread to terminate
+	bool bThreadExitRequested = false;
+
+	// Access control mutex
+	mutable FCriticalSection CritSecInternals;
+};
diff --git a/Engine/Plugins/Runtime/nDisplay/Source/DisplayCluster/Private/Render/Synchronization/Linux/DisplayClusterRenderSyncHelperLinux.cpp b/Engine/Plugins/Runtime/nDisplay/Source/DisplayCluster/Private/Render/Synchronization/Linux/DisplayClusterRenderSyncHelperLinux.cpp
new file mode 100644
--- /dev/null
+++ b/Engine/Plugins/Runtime/nDisplay/Source/DisplayCluster/Private/Render/Synchronization/Linux/DisplayClusterRenderSyncHelperLinux.cpp
@@ -0,0 +1,39 @@
+// Copyright Epic Games, Inc. All Rights Reserved.
+
+#include "Render/Synchronization/DisplayClusterRenderSyncHelper.h"
+
+#include "Misc/DisplayClusterStrings.h"
+
+#include "Engine/Engine.h"
+#include "Engine/GameViewportClient.h"
+#include "DynamicRHI.h"
+
+
+
+TUniquePtr<IDisplayClusterRenderSyncHelper> FDisplayClusterRenderSyncHelper::CreateHelper()
+{
+	if (GDynamicRHI)
+	{
+		const FString RHIName = GDynamicRHI->GetName();
+
+		// Instantiate Vulkan helper
+		if (RHIName.Equals(DisplayClusterStrings::rhi::Vulkan, ESearchCase::IgnoreCase))
+		{
+			return MakeUnique<FDisplayClusterRenderSyncHelper::FDCRSHelperVulkan>();
+		}
+	}
+
+	// Null-helper stub as fallback
+	return MakeUnique<FDisplayClusterRenderSyncHelper::FDCRSHelperNull>();
+}
+
+
+bool FDisplayClusterRenderSyncHelper::FDCRSHelperVulkan::IsWaitForVBlankSupported()
+{
+	return false;
+}
+
+bool FDisplayClusterRenderSyncHelper::FDCRSHelperVulkan::WaitForVBlank()
+{
+	return false;
+}
diff --git a/Engine/Plugins/Runtime/nDisplay/Source/DisplayCluster/Private/Render/Synchronization/Linux/DisplayClusterRenderSyncPolicyEthernetLinux.cpp b/Engine/Plugins/Runtime/nDisplay/Source/DisplayCluster/Private/Render/Synchronization/Linux/DisplayClusterRenderSyncPolicyEthernetLinux.cpp
--- a/Engine/Plugins/Runtime/nDisplay/Source/DisplayCluster/Private/Render/Synchronization/Linux/DisplayClusterRenderSyncPolicyEthernetLinux.cpp
+++ b/Engine/Plugins/Runtime/nDisplay/Source/DisplayCluster/Private/Render/Synchronization/Linux/DisplayClusterRenderSyncPolicyEthernetLinux.cpp
@@ -29,7 +29,7 @@ void FDisplayClusterRenderSyncPolicyEthernet::Step_WaitForEthernetBarrierSignal_
 {
 }
 
-void FDisplayClusterRenderSyncPolicyEthernet::Step_SkipPresentationOnClosestVblank()
+void FDisplayClusterRenderSyncPolicyEthernet::Step_SkipPresentationOnClosestVBlank()
 {
 }
 
diff --git a/Engine/Plugins/Runtime/nDisplay/Source/DisplayCluster/Private/Render/Synchronization/Windows/DisplayClusterRenderSyncHelperWindows.cpp b/Engine/Plugins/Runtime/nDisplay/Source/DisplayCluster/Private/Render/Synchronization/Windows/DisplayClusterRenderSyncHelperWindows.cpp
new file mode 100644
--- /dev/null
+++ b/Engine/Plugins/Runtime/nDisplay/Source/DisplayCluster/Private/Render/Synchronization/Windows/DisplayClusterRenderSyncHelperWindows.cpp
@@ -0,0 +1,84 @@
+// Copyright Epic Games, Inc. All Rights Reserved.
+
+#include "Render/Synchronization/DisplayClusterRenderSyncHelper.h"
+
+#include "Misc/DisplayClusterStrings.h"
+
+#include "Engine/Engine.h"
+#include "Engine/GameViewportClient.h"
+#include "DynamicRHI.h"
+
+#pragma warning(push)
+#pragma warning (disable : 4005)
+#include "Windows/AllowWindowsPlatformTypes.h"
+#include "dxgi.h"
+#include "Windows/HideWindowsPlatformTypes.h"
+#pragma warning(pop)
+
+
+TUniquePtr<IDisplayClusterRenderSyncHelper> FDisplayClusterRenderSyncHelper::CreateHelper()
+{
+	if (GDynamicRHI)
+	{
+		const FString RHIName = GDynamicRHI->GetName();
+
+		// Instantiate DX helper
+		if (RHIName.Equals(DisplayClusterStrings::rhi::D3D11, ESearchCase::IgnoreCase) ||
+			RHIName.Equals(DisplayClusterStrings::rhi::D3D12, ESearchCase::IgnoreCase))
+		{
+			return MakeUnique<FDisplayClusterRenderSyncHelper::FDCRSHelperDX>();
+		}
+		// Instantiate Vulkan helper
+		else if (RHIName.Equals(DisplayClusterStrings::rhi::Vulkan, ESearchCase::IgnoreCase))
+		{
+			return MakeUnique<FDisplayClusterRenderSyncHelper::FDCRSHelperVulkan>();
+		}
+	}
+
+	// Null-helper stub as fallback
+	return MakeUnique<FDisplayClusterRenderSyncHelper::FDCRSHelperNull>();
+}
+
+
+bool FDisplayClusterRenderSyncHelper::FDCRSHelperDX::IsWaitForVBlankSupported()
+{
+	return true;
+}
+
+bool FDisplayClusterRenderSyncHelper::FDCRSHelperDX::WaitForVBlank()
+{
+	if (GEngine && GEngine->GameViewport && GEngine->GameViewport->Viewport)
+	{
+		if (FRHIViewport* RHIViewport = GEngine->GameViewport->Viewport->GetViewportRHI().GetReference())
+		{
+			if (IDXGISwapChain* DXGISwapChain = static_cast<IDXGISwapChain*>(RHIViewport->GetNativeSwapChain()))
+			{
+				IDXGIOutput* DXOutput = nullptr;
+				DXGISwapChain->GetContainingOutput(&DXOutput);
+
+				if (DXOutput)
+				{
+					DXOutput->WaitForVBlank();
+					DXOutput->Release();
+
+					// Return true to notify about successful v-blank awaiting
+					return true;
+				}
+			}
+		}
+	}
+
+	// Something went wrong. We have to let the caller know that
+	// we are not at the v-blank now.
+	return false;
+}
+
+bool FDisplayClusterRenderSyncHelper::FDCRSHelperVulkan::IsWaitForVBlankSupported()
+{
+	return false;
+}
+
+bool FDisplayClusterRenderSyncHelper::FDCRSHelperVulkan::WaitForVBlank()
+{
+	return false;
+}
diff --git a/Engine/Plugins/Runtime/nDisplay/Source/DisplayCluster/Private/Render/Synchronization/Windows/DisplayClusterRenderSyncPolicyEthernetWindows.cpp b/Engine/Plugins/Runtime/nDisplay/Source/DisplayCluster/Private/Render/Synchronization/Windows/DisplayClusterRenderSyncPolicyEthernetWindows.cpp
--- a/Engine/Plugins/Runtime/nDisplay/Source/DisplayCluster/Private/Render/Synchronization/Windows/DisplayClusterRenderSyncPolicyEthernetWindows.cpp
+++ b/Engine/Plugins/Runtime/nDisplay/Source/DisplayCluster/Private/Render/Synchronization/Windows/DisplayClusterRenderSyncPolicyEthernetWindows.cpp
@@ -33,7 +33,6 @@ namespace DisplayClusterRenderSyncPolicyEthernet_Data_Windows
 {
 	FRHIViewport*   Viewport  = nullptr;
 	IDXGISwapChain* SwapChain = nullptr;
-	IDXGIOutput*    DXOutput  = nullptr;
 }
 
 
@@ -61,20 +60,15 @@ bool FDisplayClusterRenderSyncPolicyEthernet::SynchronizeClusterRendering(int32&
 		return true;
 	}
 
-	DXOutput = nullptr;
-	SwapChain->GetContainingOutput(&DXOutput);
-	if (!DXOutput)
-	{
-		UE_LOG(LogDisplayClusterRenderSync, Error, TEXT("nDisplay SYNC: Couldn't get DXOutput"));
-		return true;
-	}
-
 	bool bNeedEnginePresent = false;
 
 	{
 		TRACE_CPUPROFILER_EVENT_SCOPE(nDisplay SYNC);
+
 		if (bSimpleSync)
 		{
+			// Wait unless the frame is rendered
+			WaitForFrameCompletion();
 			// Barrier sync only
 			SyncBarrierRenderThread();
 			// Let the engine present this frame
@@ -159,7 +153,7 @@ void FDisplayClusterRenderSyncPolicyEthernet::Procedure_SynchronizePresent()
 		// Since it's not 100% safe to present frame now, we need to decide either we present now or
 		// postpone presentation until next V-blank. Depending on the timings based math, we might
 		// sleep here for some small time to skip V-blank and present frame after it.
-		Step_SkipPresentationOnClosestVblank();
+		Step_SkipPresentationOnClosestVBlank();
 	}
 
 	{
@@ -242,7 +236,7 @@ void FDisplayClusterRenderSyncPolicyEthernet::Step_WaitForEthernetBarrierSignal_
 	B1A = FPlatformTime::Seconds();
 }
 
-void FDisplayClusterRenderSyncPolicyEthernet::Step_SkipPresentationOnClosestVblank()
+void FDisplayClusterRenderSyncPolicyEthernet::Step_SkipPresentationOnClosestVBlank()
 {
 	// Here we calculate how much time left before the next VBlank
 	const double CurTime = FPlatformTime::Seconds();
@@ -303,14 +297,7 @@ void FDisplayClusterRenderSyncPolicyEthernet::Step_FinalizeFrameSynchronization(
 
 double FDisplayClusterRenderSyncPolicyEthernet::GetVBlankTimestamp()
 {
-	using namespace DisplayClusterRenderSyncPolicyEthernet_Data_Windows;
-
-	if (DXOutput)
-	{
-		// Get timestamp right after the VBlank
-		DXOutput->WaitForVBlank();
-		return FPlatformTime::Seconds();
-	}
+	WaitForVBlank();
 
 	return FPlatformTime::Seconds();
 }
diff --git a/Engine/Plugins/Runtime/nDisplay/Source/DisplayCluster/Private/Render/Synchronization/Windows/DisplayClusterRenderSyncPolicyNvidiaWindows.cpp b/Engine/Plugins/Runtime/nDisplay/Source/DisplayCluster/Private/Render/Synchronization/Windows/DisplayClusterRenderSyncPolicyNvidiaWindows.cpp
--- a/Engine/Plugins/Runtime/nDisplay/Source/DisplayCluster/Private/Render/Synchronization/Windows/DisplayClusterRenderSyncPolicyNvidiaWindows.cpp
+++ b/Engine/Plugins/Runtime/nDisplay/Source/DisplayCluster/Private/Render/Synchronization/Windows/DisplayClusterRenderSyncPolicyNvidiaWindows.cpp
@@ -20,8 +20,6 @@
 #include "Windows/HideWindowsPlatformTypes.h"
 #pragma warning(pop)
 
-#include "Async/Async.h"
-
 
 // TEMPORARY DIAGNOSTICS START
 static TAutoConsoleVariable<int32> CVarNvidiaSyncDiagnosticsInit(
@@ -107,15 +105,6 @@ namespace DisplayClusterRenderSyncPolicyNvidia_Data_Windows
 FDisplayClusterRenderSyncPolicyNvidia::FDisplayClusterRenderSyncPolicyNvidia(const TMap<FString, FString>& Parameters)
 	: FDisplayClusterRenderSyncPolicyBase(Parameters)
 {
-	const NvAPI_Status NvApiResult = NvAPI_Initialize();
-	if (NvApiResult != NVAPI_OK)
-	{
-		UE_LOG(LogDisplayClusterRenderSync, Error, TEXT("NvAPI_Initialize() failed, error 0x%x"), NvApiResult);
-	}
-	else
-	{
-		bNvApiInitialized = true;
-	}
 }
 
 FDisplayClusterRenderSyncPolicyNvidia::~FDisplayClusterRenderSyncPolicyNvidia()
@@ -136,6 +125,21 @@ FDisplayClusterRenderSyncPolicyNvidia::~FDisplayClusterRenderSyncPolicyNvidia()
 	}
 }
 
+bool FDisplayClusterRenderSyncPolicyNvidia::Initialize()
+{
+	const NvAPI_Status NvApiResult = NvAPI_Initialize();
+	if (NvApiResult != NVAPI_OK)
+	{
+		UE_LOG(LogDisplayClusterRenderSync, Error, TEXT("NvAPI_Initialize() failed, error 0x%x"), NvApiResult);
+	}
+	else
+	{
+		bNvApiInitialized = true;
+	}
+
+	return bNvApiInitialized;
+}
+
 bool FDisplayClusterRenderSyncPolicyNvidia::SynchronizeClusterRendering(int32& InOutSyncInterval)
 {
 	using namespace DisplayClusterRenderSyncPolicyNvidia_Data_Windows;
@@ -287,7 +291,7 @@ bool FDisplayClusterRenderSyncPolicyNvidia::InitializeNvidiaSwapLock()
 	// Set frame latency
 	if (DXGISwapChain->SetMaximumFrameLatency(1) != S_OK)
 	{
-		UE_LOG(LogDisplayClusterRenderSync, Error, TEXT("Couldn't set maximum frame latency"));
+		UE_LOG(LogDisplayClusterRenderSync, Warning, TEXT("Couldn't set maximum frame latency"));
 	}
 
 	NvU32 MaxGroups = 0;
@@ -317,34 +321,30 @@ bool FDisplayClusterRenderSyncPolicyNvidia::InitializeNvidiaSwapLock()
 	// Here we initialize NVIDIA sync on the same frame interval on the timescale
 	if (bNvDiagInit)
 	{
-		IDXGIOutput* DXOutput = nullptr;
-		DXGISwapChain->GetContainingOutput(&DXOutput);
-		if (DXOutput)
-		{
-			UE_LOG(LogDisplayClusterRenderSync, VeryVerbose, TEXT("NVAPI DIAG: init start"));
 
-			// Align all the threads on the same timeline. It's very likely all the nodes got free on the same side of a potential
-			// upcoming V-blank. Normally, the barrier synchronization takes up to 500 microseconds which is about 3% of frame time.
-			// So the probability of a situation where some nodes left the barrier before V-blank while others after is really low.
-			// However, it's still possible. And cluster restart should likely be successfull.
-			SyncBarrierRenderThread();
+		UE_LOG(LogDisplayClusterRenderSync, VeryVerbose, TEXT("NVAPI DIAG: init start"));
 
-			// Assuming all the nodes in the right place and have some time before V-blank. Let's wait untill it happend.
-			DXOutput->WaitForVBlank();
+		// Align all the threads on the same timeline. It's very likely all the nodes got free on the same side of a potential
+		// upcoming V-blank. Normally, the barrier synchronization takes up to 500 microseconds which is about 3% of frame time.
+		// So the probability of a situation where some nodes left the barrier before V-blank while others after is really low.
+		// However, it's still possible. And cluster restart should likely be successfull.
+		SyncBarrierRenderThread();
 
-			// We don't really know either we're at the very beginning of V-blank interval or in the end. Let's spend some time
-			// on the barrier to let V-blank interval over.
-			SyncBarrierRenderThread();
+		// Assuming all the nodes in the right place and have some time before V-blank. Let's wait untill it happend.
+		WaitForVBlank();
 
-			// Yes, the task scheduler can break the logic above by allocating CPU resources to this thread like 20ms later which
-			// makes us behind of the next V-blank for 60Hz output. And we probably should play around thread priorities to reduce
-			// probability of that. But hope we're good here. Usually the systems that use NVIDIA sync approach have a lot of CPU
-			// cores and don't run any other applications that compete for CPU.
-			// Moreover, we will initialize NVIDIA sync for all cluster nodes before presenting first frame. So the initialization
-			// process should be fine anyway.
+		// We don't really know either we're at the very beginning of V-blank interval or in the end. Let's spend some time
+		// on the barrier to let V-blank interval over.
+		SyncBarrierRenderThread();
 
-			UE_LOG(LogDisplayClusterRenderSync, VeryVerbose, TEXT("NVAPI DIAG: init end"));
-		}
+		// Yes, the task scheduler can break the logic above by allocating CPU resources to this thread like 20ms later which
+		// makes us behind of the next V-blank for 60Hz output. And we probably should play around thread priorities to reduce
+		// probability of that. But hope we're good here. Usually the systems that use NVIDIA sync approach have a lot of CPU
+		// cores and don't run any other applications that compete for CPU.
+		// Moreover, we will initialize NVIDIA sync for all cluster nodes before presenting first frame. So the initialization
+		// process should be fine anyway.
+
+		UE_LOG(LogDisplayClusterRenderSync, VeryVerbose, TEXT("NVAPI DIAG: init end"));
 	}
 
 	// Join swap group
diff --git a/Engine/Plugins/Runtime/nDisplay/Source/DisplayCluster/Public/Render/Synchronization/IDisplayClusterRenderSyncPolicy.h b/Engine/Plugins/Runtime/nDisplay/Source/DisplayCluster/Public/Render/Synchronization/IDisplayClusterRenderSyncPolicy.h
--- a/Engine/Plugins/Runtime/nDisplay/Source/DisplayCluster/Public/Render/Synchronization/IDisplayClusterRenderSyncPolicy.h
+++ b/Engine/Plugins/Runtime/nDisplay/Source/DisplayCluster/Public/Render/Synchronization/IDisplayClusterRenderSyncPolicy.h
@@ -14,6 +14,20 @@ public:
 	virtual ~IDisplayClusterRenderSyncPolicy() = default;
 
 public:
+	/**
+	* Perform any non-constructor initialization if needed
+	*
+	* @return - true if initialized successfully
+	*/
+	virtual bool Initialize()
+	{
+		return true;
+	}
+
+	/**
+	 * Returns the name of the sync policy
+	 */
+	virtual FName GetName() const = 0;
 
 	/**
 	* Synchronizes rendering threads in a cluster (optionally presents a frame)
@@ -25,7 +39,22 @@ public:
 	virtual bool SynchronizeClusterRendering(int32& InOutSyncInterval) = 0;
 
 	/**
-	 * Returns the name of the sync policy
+	 * Returns true if policy instance support "WaitForVBlank" feature that is
+	 * used on the higher levels
 	 */
-	virtual FName GetName() const = 0;
+	virtual bool IsWaitForVBlankFeatureSupported()
+	{
+		return false;
+	}
+
+	/**
+	 * Wait unless the next V-blank occurs. This function may be called after
+	 * IsWaitForVBlankFeatureSupported returned true. Otherwise it will never be called.
+	*
+	* @return - true if v-blank awaiting was successfull
+	 */
+	virtual bool WaitForVBlank()
+	{
+		return false;
+	}
 };
