nDisplay: fixed build issues (step 2)

diff --git a/Engine/Plugins/Runtime/nDisplay/Source/DisplayCluster/Private/Render/Synchronization/DisplayClusterRenderSyncPolicyBase.h b/Engine/Plugins/Runtime/nDisplay/Source/DisplayCluster/Private/Render/Synchronization/DisplayClusterRenderSyncPolicyBase.h
--- a/Engine/Plugins/Runtime/nDisplay/Source/DisplayCluster/Private/Render/Synchronization/DisplayClusterRenderSyncPolicyBase.h
+++ b/Engine/Plugins/Runtime/nDisplay/Source/DisplayCluster/Private/Render/Synchronization/DisplayClusterRenderSyncPolicyBase.h
@@ -3,7 +3,7 @@
 #pragma once
 
 #include "Render/Synchronization/IDisplayClusterRenderSyncPolicy.h"
-#include "Render/Synchronization/DisplayClusterVBlankMonitor.h"
+#include "Render/Synchronization/DisplayClusterVblankMonitor.h"
 
 
 /**
diff --git a/Engine/Plugins/Runtime/nDisplay/Source/DisplayCluster/Private/Render/Synchronization/DisplayClusterVblankMonitor.cpp b/Engine/Plugins/Runtime/nDisplay/Source/DisplayCluster/Private/Render/Synchronization/DisplayClusterVblankMonitor.cpp
new file mode 100644
--- /dev/null
+++ b/Engine/Plugins/Runtime/nDisplay/Source/DisplayCluster/Private/Render/Synchronization/DisplayClusterVblankMonitor.cpp
@@ -0,0 +1,59 @@
+// Copyright Epic Games, Inc. All Rights Reserved.
+
+#include "Render/Synchronization/DisplayClusterVblankMonitor.h"
+
+#include "Render/Synchronization/IDisplayClusterRenderSyncPolicy.h"
+#include "HAL/RunnableThread.h"
+#include "Misc/DisplayClusterLog.h"
+#include "Misc/ScopeLock.h"
+
+
+FDisplayClusterVBlankMonitor::FDisplayClusterVBlankMonitor(IDisplayClusterRenderSyncPolicy& InSyncPolicy)
+	: SyncPolicy(InSyncPolicy)
+{
+	WorkingThread.Reset(FRunnableThread::Create(this, TEXT("VBLANK_Monitoring_thread"), 32 * 1024, TPri_AboveNormal));
+}
+
+FDisplayClusterVBlankMonitor::~FDisplayClusterVBlankMonitor()
+{
+	// Stop working thread
+	Stop();
+}
+
+uint32 FDisplayClusterVBlankMonitor::Run()
+{
+	UE_LOG(LogDisplayClusterRenderSync, Log, TEXT("V-blank monitoring thread has started"));
+
+	while (!bThreadExitRequested)
+	{
+		if (SyncPolicy.WaitForVBlank())
+		{
+			TRACE_BOOKMARK(TEXT("VBLANK %llu"), VBlankCounter);
+			++VBlankCounter;
+		}
+		else
+		{
+			TRACE_CPUPROFILER_EVENT_SCOPE(VBLANK - not available);
+			FPlatformProcess::Sleep(1.f);
+		}
+	}
+
+	UE_LOG(LogDisplayClusterRenderSync, Log, TEXT("V-blank monitoring thread has finished"));
+
+	return 0;
+}
+
+void FDisplayClusterVBlankMonitor::Stop()
+{
+	FScopeLock Lock(&CritSecInternals);
+
+	// Let the working thread know it has to stop
+	bThreadExitRequested = true;
+
+	// Wait unless working thread is finished
+	if (WorkingThread)
+	{
+		WorkingThread->WaitForCompletion();
+		WorkingThread.Reset();
+	}
+}
diff --git a/Engine/Plugins/Runtime/nDisplay/Source/DisplayCluster/Private/Render/Synchronization/DisplayClusterVblankMonitor.h b/Engine/Plugins/Runtime/nDisplay/Source/DisplayCluster/Private/Render/Synchronization/DisplayClusterVblankMonitor.h
new file mode 100644
--- /dev/null
+++ b/Engine/Plugins/Runtime/nDisplay/Source/DisplayCluster/Private/Render/Synchronization/DisplayClusterVblankMonitor.h
@@ -0,0 +1,44 @@
+// Copyright Epic Games, Inc. All Rights Reserved.
+
+#pragma once
+
+#include "CoreMinimal.h"
+#include "HAL/Runnable.h"
+
+class IDisplayClusterRenderSyncPolicy;
+
+
+/**
+ * V-blank monitor. It continuously polls the current sync policy
+ * for V-blanks and reports the time data to utraces.
+ */
+class FDisplayClusterVBlankMonitor
+	: protected FRunnable
+{
+public:
+	FDisplayClusterVBlankMonitor(IDisplayClusterRenderSyncPolicy& SyncPolicy);
+	~FDisplayClusterVBlankMonitor();
+
+protected:
+	//////////////////////////////////////////////////////////////////////////////////////////////
+	// FRunnable
+	//////////////////////////////////////////////////////////////////////////////////////////////
+	virtual uint32 Run() override;
+	virtual void Stop() override final;
+
+private:
+	// Holds the thread object
+	TUniquePtr<FRunnableThread> WorkingThread;
+
+	// Current sync policy
+	IDisplayClusterRenderSyncPolicy& SyncPolicy;
+
+	// V-blank counter
+	uint64 VBlankCounter = 0;
+
+	// Flag to ask working thread to terminate
+	bool bThreadExitRequested = false;
+
+	// Access control mutex
+	mutable FCriticalSection CritSecInternals;
+};
