nDisplay: fixed build issues (step 3)

diff --git a/Engine/Plugins/Runtime/nDisplay/Source/DisplayCluster/Private/Render/Synchronization/Linux/DisplayClusterRenderSyncPolicyNvidiaLinux.cpp b/Engine/Plugins/Runtime/nDisplay/Source/DisplayCluster/Private/Render/Synchronization/Linux/DisplayClusterRenderSyncPolicyNvidiaLinux.cpp
--- a/Engine/Plugins/Runtime/nDisplay/Source/DisplayCluster/Private/Render/Synchronization/Linux/DisplayClusterRenderSyncPolicyNvidiaLinux.cpp
+++ b/Engine/Plugins/Runtime/nDisplay/Source/DisplayCluster/Private/Render/Synchronization/Linux/DisplayClusterRenderSyncPolicyNvidiaLinux.cpp
@@ -12,6 +12,11 @@ FDisplayClusterRenderSyncPolicyNvidia::~FDisplayClusterRenderSyncPolicyNvidia()
 {
 }
 
+bool FDisplayClusterRenderSyncPolicyNvidia::Initialize()
+{
+	return true;
+}
+
 bool FDisplayClusterRenderSyncPolicyNvidia::InitializeNvidiaSwapLock()
 {
 	return true;
