nDisplay: Added operator panel that can be opened through the Window->Virtual Production menu, and features an extendable tab stack and toolbar that allows other plugins or modules to add their own tabs to the operator. Also added an empty tab and module for the light card editor to be expanded on in later work.

diff --git a/Engine/Plugins/Runtime/nDisplay/Source/DisplayClusterLightCardEditor/DisplayClusterLightCardEditor.Build.cs b/Engine/Plugins/Runtime/nDisplay/Source/DisplayClusterLightCardEditor/DisplayClusterLightCardEditor.Build.cs
new file mode 100644
--- /dev/null
+++ b/Engine/Plugins/Runtime/nDisplay/Source/DisplayClusterLightCardEditor/DisplayClusterLightCardEditor.Build.cs
@@ -0,0 +1,27 @@
+// Copyright Epic Games, Inc. All Rights Reserved.
+
+using UnrealBuildTool;
+using System.IO;
+
+public class DisplayClusterLightCardEditor : ModuleRules
+{
+	public DisplayClusterLightCardEditor(ReadOnlyTargetRules ROTargetRules) : base(ROTargetRules)
+	{
+		PrivateDependencyModuleNames.AddRange(
+			new string[] {
+				"DisplayCluster",
+				"DisplayClusterConfiguration",
+				"DisplayClusterOperator",
+				
+				"Core",
+				"CoreUObject",
+				"Engine",
+				"InputCore",
+				"Slate",
+				"SlateCore",
+				"UnrealEd",
+				"EditorStyle",
+				"WorkspaceMenuStructure",
+			});
+	}
+}
diff --git a/Engine/Plugins/Runtime/nDisplay/Source/DisplayClusterLightCardEditor/Private/DisplayClusterLightCardEditorModule.cpp b/Engine/Plugins/Runtime/nDisplay/Source/DisplayClusterLightCardEditor/Private/DisplayClusterLightCardEditorModule.cpp
new file mode 100644
--- /dev/null
+++ b/Engine/Plugins/Runtime/nDisplay/Source/DisplayClusterLightCardEditor/Private/DisplayClusterLightCardEditorModule.cpp
@@ -0,0 +1,31 @@
+// Copyright Epic Games, Inc. All Rights Reserved.
+
+#include "DisplayClusterLightCardEditorModule.h"
+
+#include "SDisplayClusterLightCardEditor.h"
+
+#define LOCTEXT_NAMESPACE "DisplayClusterLightCardEditor"
+
+void FDisplayClusterLightCardEditorModule::StartupModule()
+{
+	RegisterTabSpawners();
+}
+
+void FDisplayClusterLightCardEditorModule::ShutdownModule()
+{
+	UnregisterTabSpawners();
+}
+
+void FDisplayClusterLightCardEditorModule::RegisterTabSpawners()
+{
+	SDisplayClusterLightCardEditor::RegisterTabSpawner();
+}
+
+void FDisplayClusterLightCardEditorModule::UnregisterTabSpawners()
+{
+	SDisplayClusterLightCardEditor::UnregisterTabSpawner();
+}
+
+IMPLEMENT_MODULE(FDisplayClusterLightCardEditorModule, DisplayClusterLightCardEditor);
+
+#undef LOCTEXT_NAMESPACE
diff --git a/Engine/Plugins/Runtime/nDisplay/Source/DisplayClusterLightCardEditor/Private/DisplayClusterLightCardEditorModule.h b/Engine/Plugins/Runtime/nDisplay/Source/DisplayClusterLightCardEditor/Private/DisplayClusterLightCardEditorModule.h
new file mode 100644
--- /dev/null
+++ b/Engine/Plugins/Runtime/nDisplay/Source/DisplayClusterLightCardEditor/Private/DisplayClusterLightCardEditorModule.h
@@ -0,0 +1,22 @@
+// Copyright Epic Games, Inc. All Rights Reserved.
+
+#pragma once
+
+#include "IDisplayClusterLightCardEditor.h"
+#include "UObject/WeakObjectPtr.h"
+
+/**
+ * Display Cluster editor module
+ */
+class FDisplayClusterLightCardEditorModule :
+	public IDisplayClusterLightCardEditor
+{
+public:
+	//~ IModuleInterface interface
+	virtual void StartupModule() override;
+	virtual void ShutdownModule() override;
+
+private:
+	void RegisterTabSpawners();
+	void UnregisterTabSpawners();
+};
diff --git a/Engine/Plugins/Runtime/nDisplay/Source/DisplayClusterLightCardEditor/Private/SDisplayClusterLightCardEditor.cpp b/Engine/Plugins/Runtime/nDisplay/Source/DisplayClusterLightCardEditor/Private/SDisplayClusterLightCardEditor.cpp
new file mode 100644
--- /dev/null
+++ b/Engine/Plugins/Runtime/nDisplay/Source/DisplayClusterLightCardEditor/Private/SDisplayClusterLightCardEditor.cpp
@@ -0,0 +1,95 @@
+// Copyright Epic Games, Inc. All Rights Reserved.
+
+#include "SDisplayClusterLightCardEditor.h"
+
+#include "IDisplayClusterOperator.h"
+#include "DisplayClusterRootActor.h"
+
+#include "Framework/Docking/TabManager.h"
+#include "Framework/Docking/LayoutExtender.h"
+#include "Framework/MultiBox/MultiBoxBuilder.h"
+#include "Framework/MultiBox/MultiBoxExtender.h"
+#include "Toolkits/AssetEditorToolkit.h"
+#include "Widgets/Docking/SDockTab.h"
+
+#define LOCTEXT_NAMESPACE "SDisplayClusterLightCardEditor"
+
+const FName SDisplayClusterLightCardEditor::TabName = TEXT("DisplayClusterLightCardEditorTab");
+
+void SDisplayClusterLightCardEditor::RegisterTabSpawner()
+{
+	IDisplayClusterOperator::Get().OnRegisterLayoutExtensions().AddStatic(&SDisplayClusterLightCardEditor::RegisterLayoutExtension);
+
+	FGlobalTabmanager::Get()->RegisterNomadTabSpawner(TabName, FOnSpawnTab::CreateStatic(&SDisplayClusterLightCardEditor::SpawnInTab))
+		.SetDisplayName(LOCTEXT("TabDisplayName", "Light Cards Editor"))
+		.SetTooltipText(LOCTEXT("TabTooltip", "Editing tools for nDisplay light cards."));
+
+	TSharedPtr<FExtender> ToolbarExtender = MakeShareable(new FExtender);
+	ToolbarExtender->AddToolBarExtension("General", EExtensionHook::After, nullptr, FToolBarExtensionDelegate::CreateStatic(&SDisplayClusterLightCardEditor::ExtendToolbar));
+	IDisplayClusterOperator::Get().GetOperatorToolBarExtensibilityManager()->AddExtender(ToolbarExtender);
+}
+
+void SDisplayClusterLightCardEditor::UnregisterTabSpawner()
+{
+	FGlobalTabmanager::Get()->UnregisterNomadTabSpawner(TabName);
+}
+
+void SDisplayClusterLightCardEditor::RegisterLayoutExtension(FLayoutExtender& InExtender)
+{
+	FTabManager::FTab NewTab(FTabId(TabName, ETabIdFlags::SaveLayout), ETabState::OpenedTab);
+	InExtender.ExtendStack(IDisplayClusterOperator::Get().GetOperatorExtensionId(), ELayoutExtensionPosition::After, NewTab);
+}
+
+TSharedRef<SDockTab> SDisplayClusterLightCardEditor::SpawnInTab(const FSpawnTabArgs& SpawnTabArgs)
+{
+	const TSharedRef<SDockTab> MajorTab = SNew(SDockTab)
+		.TabRole(ETabRole::NomadTab);
+
+	MajorTab->SetContent(SNew(SDisplayClusterLightCardEditor, MajorTab, SpawnTabArgs.GetOwnerWindow()));
+
+	return MajorTab;
+}
+
+void SDisplayClusterLightCardEditor::ExtendToolbar(FToolBarBuilder& ToolbarBuilder)
+{
+	// TODO: Any toolbar buttons needed for the lightcards editor can be added to the operator panel's toobar using this toolbar extender
+}
+
+SDisplayClusterLightCardEditor::~SDisplayClusterLightCardEditor()
+{
+	IDisplayClusterOperator::Get().OnActiveRootActorChanged().Remove(ActiveRootActorChangedHandle);
+}
+
+void SDisplayClusterLightCardEditor::Construct(const FArguments& InArgs, const TSharedRef<SDockTab>& MajorTabOwner, const TSharedPtr<SWindow>& WindowOwner)
+{
+	ActiveRootActorChangedHandle = IDisplayClusterOperator::Get().OnActiveRootActorChanged().AddSP(this, &SDisplayClusterLightCardEditor::OnActiveRootActorChanged);
+
+    TabManager = FGlobalTabmanager::Get()->NewTabManager(MajorTabOwner);
+
+	// TODO: Register the various tabs for the light card editor here
+
+	const TSharedRef<FTabManager::FLayout> Layout = FTabManager::NewLayout("nDisplayLightCardEditorLayout")
+		->AddArea
+		(
+			FTabManager::NewPrimaryArea()
+			->SetOrientation(EOrientation::Orient_Vertical)
+			->Split
+			(
+				FTabManager::NewStack()
+				->SetHideTabWell(true)
+			)
+		);
+
+	ChildSlot
+	[
+		TabManager->RestoreFrom(Layout, WindowOwner).ToSharedRef()
+	];
+}
+
+void SDisplayClusterLightCardEditor::OnActiveRootActorChanged(ADisplayClusterRootActor* NewRootActor)
+{
+	ActiveRootActor = NewRootActor;
+
+	// TODO: Need to pass a refresh command to any child tabs of this widget to inform them that the active root actor has been changed.
+	// The new root actor pointer could be null, indicating that it was deleted or the user didn't select a valid root actor
+}
\ No newline at end of file
diff --git a/Engine/Plugins/Runtime/nDisplay/Source/DisplayClusterLightCardEditor/Private/SDisplayClusterLightCardEditor.h b/Engine/Plugins/Runtime/nDisplay/Source/DisplayClusterLightCardEditor/Private/SDisplayClusterLightCardEditor.h
new file mode 100644
--- /dev/null
+++ b/Engine/Plugins/Runtime/nDisplay/Source/DisplayClusterLightCardEditor/Private/SDisplayClusterLightCardEditor.h
@@ -0,0 +1,59 @@
+// Copyright Epic Games, Inc. All Rights Reserved.
+
+#pragma once
+
+#include "CoreMinimal.h"
+#include "Widgets/DeclarativeSyntaxSupport.h"
+#include "Widgets/SCompoundWidget.h"
+
+class FTabManager;
+class FLayoutExtender;
+class FSpawnTabArgs;
+class FToolBarBuilder;
+class SDockTab;
+class ADisplayClusterRootActor;
+
+/** A panel that can be spawned in a tab that contains all the UI elements that make up the 2D light cards editor */
+class SDisplayClusterLightCardEditor : public SCompoundWidget
+{
+public:
+	/** The name of the tab that the light card editor lives in */
+	static const FName TabName;
+
+
+	/** Registers the light card editor with the global tab manager and adds it to the operator panel's extension tab stack */
+	static void RegisterTabSpawner();
+
+	/** Unregisters the light card editor from the global tab manager */
+	static void UnregisterTabSpawner();
+
+	/** Registers the light card editor tab with the operator panel using a layout extension */
+	static void RegisterLayoutExtension(FLayoutExtender& InExtender);
+
+	/** Creates a tab with the light card editor inside */
+	static TSharedRef<SDockTab> SpawnInTab(const FSpawnTabArgs& SpawnTabArgs);
+
+	static void ExtendToolbar(FToolBarBuilder& ToolbarBuilder);
+
+
+	SLATE_BEGIN_ARGS(SDisplayClusterLightCardEditor)
+	{}
+	SLATE_END_ARGS()
+
+	~SDisplayClusterLightCardEditor();
+
+	void Construct(const FArguments& args, const TSharedRef<SDockTab>& MajorTabOwner, const TSharedPtr<SWindow>& WindowOwner);
+
+private:
+	/** Raised when the active Display cluster root actor has been changed in the operator panel */
+	void OnActiveRootActorChanged(ADisplayClusterRootActor* NewRootActor);
+
+private:
+	/** Holds the tab manager that manages the light card editor's tabs. */
+	TSharedPtr<FTabManager> TabManager;
+
+	/** A reference to the root actor that is currently being operated on */
+	TWeakObjectPtr<ADisplayClusterRootActor> ActiveRootActor;
+
+	FDelegateHandle ActiveRootActorChangedHandle;
+};
\ No newline at end of file
diff --git a/Engine/Plugins/Runtime/nDisplay/Source/DisplayClusterLightCardEditor/Public/IDisplayClusterLightCardEditor.h b/Engine/Plugins/Runtime/nDisplay/Source/DisplayClusterLightCardEditor/Public/IDisplayClusterLightCardEditor.h
new file mode 100644
--- /dev/null
+++ b/Engine/Plugins/Runtime/nDisplay/Source/DisplayClusterLightCardEditor/Public/IDisplayClusterLightCardEditor.h
@@ -0,0 +1,40 @@
+// Copyright Epic Games, Inc. All Rights Reserved.
+
+#pragma once
+
+#include "CoreMinimal.h"
+#include "Modules/ModuleInterface.h"
+#include "Modules/ModuleManager.h"
+
+/**
+ * Display Cluster Light Card Editor module interface
+ */
+class IDisplayClusterLightCardEditor : public IModuleInterface
+{
+public:
+	static constexpr auto ModuleName = TEXT("DisplayClusterLightCardEditor");
+
+public:
+	virtual ~IDisplayClusterLightCardEditor() = default;
+
+	/**
+	* Singleton-like access to this module's interface.  This is just for convenience!
+	* Beware of calling this during the shutdown phase, though.  Your module might have been unloaded already.
+	*
+	* @return Returns singleton instance, loading the module on demand if needed
+	*/
+	static inline IDisplayClusterLightCardEditor& Get()
+	{
+		return FModuleManager::GetModuleChecked<IDisplayClusterLightCardEditor>(ModuleName);
+	}
+
+	/**
+	* Checks to see if this module is loaded and ready.  It is only valid to call Get() if IsAvailable() returns true.
+	*
+	* @return True if the module is loaded and ready to use
+	*/
+	static inline bool IsAvailable()
+	{
+		return FModuleManager::Get().IsModuleLoaded(ModuleName);
+	}
+};
\ No newline at end of file
diff --git a/Engine/Plugins/Runtime/nDisplay/Source/DisplayClusterOperator/DisplayClusterOperator.Build.cs b/Engine/Plugins/Runtime/nDisplay/Source/DisplayClusterOperator/DisplayClusterOperator.Build.cs
new file mode 100644
--- /dev/null
+++ b/Engine/Plugins/Runtime/nDisplay/Source/DisplayClusterOperator/DisplayClusterOperator.Build.cs
@@ -0,0 +1,25 @@
+// Copyright Epic Games, Inc. All Rights Reserved.
+
+using UnrealBuildTool;
+using System.IO;
+
+public class DisplayClusterOperator : ModuleRules
+{
+	public DisplayClusterOperator(ReadOnlyTargetRules ROTargetRules) : base(ROTargetRules)
+	{
+		PrivateDependencyModuleNames.AddRange(
+			new string[] {
+				"DisplayCluster",
+
+				"Core",
+				"CoreUObject",
+				"Engine",
+				"InputCore",
+				"Slate",
+				"SlateCore",
+				"UnrealEd",
+				"EditorStyle",
+				"WorkspaceMenuStructure",
+			});
+	}
+}
diff --git a/Engine/Plugins/Runtime/nDisplay/Source/DisplayClusterOperator/Private/DisplayClusterOperatorModule.cpp b/Engine/Plugins/Runtime/nDisplay/Source/DisplayClusterOperator/Private/DisplayClusterOperatorModule.cpp
new file mode 100644
--- /dev/null
+++ b/Engine/Plugins/Runtime/nDisplay/Source/DisplayClusterOperator/Private/DisplayClusterOperatorModule.cpp
@@ -0,0 +1,73 @@
+// Copyright Epic Games, Inc. All Rights Reserved.
+
+#include "DisplayClusterOperatorModule.h"
+
+#include "SDisplayClusterOperatorPanel.h"
+#include "DisplayClusterRootActor.h"
+
+#include "Editor.h"
+#include "EngineUtils.h"
+#include "Engine/World.h"
+#include "LevelEditorViewport.h"
+#include "Toolkits/AssetEditorToolkit.h"
+
+#define LOCTEXT_NAMESPACE "DisplayClusterOperator"
+
+void FDisplayClusterOperatorModule::StartupModule()
+{
+	OperatorToolBarExtensibilityManager = MakeShared<FExtensibilityManager>();
+
+	RegisterTabSpawners();
+}
+
+void FDisplayClusterOperatorModule::ShutdownModule()
+{
+	UnregisterTabSpawners();
+}
+
+FName FDisplayClusterOperatorModule::GetOperatorExtensionId()
+{
+	return SDisplayClusterOperatorPanel::TabExtensionId;
+}
+
+void FDisplayClusterOperatorModule::GetRootActorLevelInstances(TArray<ADisplayClusterRootActor*>& OutRootActorInstances)
+{
+	UWorld* World = GCurrentLevelEditingViewportClient ? GCurrentLevelEditingViewportClient->GetWorld() : nullptr;
+	if (World)
+	{
+		for (TActorIterator<ADisplayClusterRootActor> Iter(World); Iter; ++Iter)
+		{
+			ADisplayClusterRootActor* RootActor = *Iter;
+			if (IsValid(RootActor))
+			{
+				OutRootActorInstances.Add(RootActor);
+			}
+		}
+	}
+}
+
+void FDisplayClusterOperatorModule::ShowDetailsForObject(UObject* Object)
+{
+	TArray<UObject*> Objects;
+	Objects.Add(Object);
+	ShowDetailsForObjects(Objects);
+}
+
+void FDisplayClusterOperatorModule::ShowDetailsForObjects(const TArray<UObject*>& Objects)
+{
+	DetailObjectsChanged.Broadcast(Objects);
+}
+
+void FDisplayClusterOperatorModule::RegisterTabSpawners()
+{
+	SDisplayClusterOperatorPanel::RegisterTabSpawner();
+}
+
+void FDisplayClusterOperatorModule::UnregisterTabSpawners()
+{
+	SDisplayClusterOperatorPanel::UnregisterTabSpawner();
+}
+
+IMPLEMENT_MODULE(FDisplayClusterOperatorModule, DisplayClusterOperator);
+
+#undef LOCTEXT_NAMESPACE
diff --git a/Engine/Plugins/Runtime/nDisplay/Source/DisplayClusterOperator/Private/DisplayClusterOperatorModule.h b/Engine/Plugins/Runtime/nDisplay/Source/DisplayClusterOperator/Private/DisplayClusterOperatorModule.h
new file mode 100644
--- /dev/null
+++ b/Engine/Plugins/Runtime/nDisplay/Source/DisplayClusterOperator/Private/DisplayClusterOperatorModule.h
@@ -0,0 +1,42 @@
+// Copyright Epic Games, Inc. All Rights Reserved.
+
+#pragma once
+
+#include "IDisplayClusterOperator.h"
+
+/**
+ * Display Cluster editor module
+ */
+class FDisplayClusterOperatorModule :
+	public IDisplayClusterOperator
+{
+public:
+	//~ IModuleInterface interface
+	virtual void StartupModule() override;
+	virtual void ShutdownModule() override;
+	//~ End IModuleInterface interface
+
+	//~ IDisplayClusterOperator interface
+	virtual FOnRegisterLayoutExtensions& OnRegisterLayoutExtensions() override { return RegisterLayoutExtensions; }
+	virtual FOnActiveRootActorChanged& OnActiveRootActorChanged() override { return RootActorChanged; }
+	virtual FOnDetailObjectsChanged& OnDetailObjectsChanged() override { return DetailObjectsChanged; }
+
+	virtual FName GetOperatorExtensionId() override;
+	virtual TSharedPtr<FExtensibilityManager> GetOperatorToolBarExtensibilityManager() override { return OperatorToolBarExtensibilityManager; }
+	virtual void GetRootActorLevelInstances(TArray<ADisplayClusterRootActor*>& OutRootActorInstances) override;
+	virtual void ShowDetailsForObject(UObject* Object) override;
+	virtual void ShowDetailsForObjects(const TArray<UObject*>& Objects) override;
+	//~ End IDisplayClusterOperator interface
+
+private:
+	void RegisterTabSpawners();
+	void UnregisterTabSpawners();
+
+
+private:
+	FOnRegisterLayoutExtensions RegisterLayoutExtensions;
+	FOnActiveRootActorChanged RootActorChanged;
+	FOnDetailObjectsChanged DetailObjectsChanged;
+
+	TSharedPtr<FExtensibilityManager> OperatorToolBarExtensibilityManager;
+};
diff --git a/Engine/Plugins/Runtime/nDisplay/Source/DisplayClusterOperator/Private/SDisplayClusterOperatorPanel.cpp b/Engine/Plugins/Runtime/nDisplay/Source/DisplayClusterOperator/Private/SDisplayClusterOperatorPanel.cpp
new file mode 100644
--- /dev/null
+++ b/Engine/Plugins/Runtime/nDisplay/Source/DisplayClusterOperator/Private/SDisplayClusterOperatorPanel.cpp
@@ -0,0 +1,148 @@
+// Copyright Epic Games, Inc. All Rights Reserved.
+
+#include "SDisplayClusterOperatorPanel.h"
+
+#include "IDisplayClusterOperator.h"
+#include "SDisplayClusterOperatorToolbar.h"
+
+#include "EditorStyleSet.h"
+#include "Framework/Docking/TabManager.h"
+#include "Framework/Docking/LayoutExtender.h"
+#include "Framework/MultiBox/MultiBoxBuilder.h"
+#include "PropertyEditorModule.h"
+#include "Widgets/Docking/SDockTab.h"
+#include "WorkspaceMenuStructure.h"
+#include "WorkspaceMenuStructureModule.h"
+
+#define LOCTEXT_NAMESPACE "SDisplayClusterOperatorPanel"
+
+const FName SDisplayClusterOperatorPanel::TabName = TEXT("DisplayClusterOperatorTab");
+const FName SDisplayClusterOperatorPanel::ToolbarTabId = TEXT("OperatorToolbar");
+const FName SDisplayClusterOperatorPanel::DetailsTabId = TEXT("OperatorDetails");
+const FName SDisplayClusterOperatorPanel::TabExtensionId = TEXT("OperatorTabStack");
+
+void SDisplayClusterOperatorPanel::RegisterTabSpawner()
+{
+	FGlobalTabmanager::Get()->RegisterNomadTabSpawner(TabName, FOnSpawnTab::CreateStatic(&SDisplayClusterOperatorPanel::SpawnInTab))
+			.SetDisplayName(LOCTEXT("TabDisplayName", "nDisplay Operator"))
+			.SetTooltipText(LOCTEXT("TabTooltip", "Open the nDisplay Operator tab."))
+			.SetGroup(WorkspaceMenu::GetMenuStructure().GetLevelEditorVirtualProductionCategory());
+}
+
+void SDisplayClusterOperatorPanel::UnregisterTabSpawner()
+{
+	FGlobalTabmanager::Get()->UnregisterNomadTabSpawner(TabName);
+}
+
+TSharedRef<SDockTab> SDisplayClusterOperatorPanel::SpawnInTab(const FSpawnTabArgs& SpawnTabArgs)
+{
+	const TSharedRef<SDockTab> MajorTab = SNew(SDockTab)
+		.TabRole(ETabRole::NomadTab);
+
+	MajorTab->SetContent(SNew(SDisplayClusterOperatorPanel, MajorTab, SpawnTabArgs.GetOwnerWindow()));
+
+	return MajorTab;
+}
+
+SDisplayClusterOperatorPanel::~SDisplayClusterOperatorPanel()
+{
+	IDisplayClusterOperator::Get().OnDetailObjectsChanged().Remove(DetailObjectsChangedHandle);
+}
+
+void SDisplayClusterOperatorPanel::Construct(const FArguments& InArgs, const TSharedRef<SDockTab>& MajorTabOwner, const TSharedPtr<SWindow>& WindowOwner)
+{
+	TabManager = FGlobalTabmanager::Get()->NewTabManager(MajorTabOwner);
+	TSharedRef<FWorkspaceItem> AppMenuGroup = TabManager->AddLocalWorkspaceMenuCategory(LOCTEXT("OperatorMenuGroupName", "nDisplay Operator"));
+	TabManager->SetAllowWindowMenuBar(true);
+
+	TabManager->RegisterTabSpawner(ToolbarTabId, FOnSpawnTab::CreateSP(this, &SDisplayClusterOperatorPanel::SpawnToolbarTab))
+		.SetDisplayName(LOCTEXT("ToolbarTabTitle", "Toolbar"))
+		.SetIcon(FSlateIcon(FAppStyle::Get().GetStyleSetName(), "Icons.Settings"))
+		.SetGroup(AppMenuGroup);
+
+	TabManager->RegisterTabSpawner(DetailsTabId, FOnSpawnTab::CreateSP(this, &SDisplayClusterOperatorPanel::SpawnDetailsTab))
+		.SetDisplayName(LOCTEXT("DetailsTabTitle", "Details"))
+		.SetIcon(FSlateIcon(FEditorStyle::GetStyleSetName(), "LevelEditor.Tabs.Details"))
+		.SetGroup(AppMenuGroup);
+
+	const TSharedRef<FTabManager::FLayout> Layout = FTabManager::NewLayout("nDisplayOperatorLayout")
+		->AddArea
+		(
+			FTabManager::NewPrimaryArea()
+			->SetOrientation(EOrientation::Orient_Vertical)
+			->Split
+			(
+				FTabManager::NewStack()
+					->AddTab(ToolbarTabId, ETabState::OpenedTab)
+					->SetHideTabWell(true)
+			)
+			->Split
+			(
+				FTabManager::NewSplitter()
+					->SetOrientation(Orient_Horizontal)
+					->Split
+					(
+						FTabManager::NewStack()
+							->SetExtensionId(TabExtensionId)
+							->SetSizeCoefficient(0.8f)
+					)
+					->Split
+					(
+						FTabManager::NewStack()
+							->AddTab(DetailsTabId, ETabState::OpenedTab)
+							->SetHideTabWell(true)
+							->SetSizeCoefficient(0.2f)
+					)
+			)
+		);
+
+	LayoutExtender = MakeShared<FLayoutExtender>();
+	IDisplayClusterOperator::Get().OnRegisterLayoutExtensions().Broadcast(*LayoutExtender);
+	Layout->ProcessExtensions(*LayoutExtender);
+
+	DetailObjectsChangedHandle = IDisplayClusterOperator::Get().OnDetailObjectsChanged().AddSP(this, &SDisplayClusterOperatorPanel::DisplayObjectsInDetailsPanel);
+
+	ChildSlot
+	[
+		TabManager->RestoreFrom(Layout, WindowOwner).ToSharedRef()
+	];
+}
+
+TSharedRef<SDockTab> SDisplayClusterOperatorPanel::SpawnToolbarTab(const FSpawnTabArgs& Args)
+{
+	return SNew(SDockTab)
+		.ShouldAutosize(true)
+		.TabRole(ETabRole::PanelTab)
+		[
+			SNew(SDisplayClusterOperatorToolbar)
+				.CommandList(TSharedPtr<FUICommandList>())
+		];
+}
+
+TSharedRef<SDockTab> SDisplayClusterOperatorPanel::SpawnDetailsTab(const FSpawnTabArgs& Args)
+{
+	FPropertyEditorModule& EditModule = FModuleManager::Get().GetModuleChecked<FPropertyEditorModule>("PropertyEditor");
+
+	FDetailsViewArgs DetailsViewArgs;
+	DetailsViewArgs.bAllowSearch = false;
+	DetailsViewArgs.NameAreaSettings = FDetailsViewArgs::HideNameArea;
+	DetailsViewArgs.bHideSelectionTip = true;
+
+	DetailsView = EditModule.CreateDetailView(DetailsViewArgs);
+
+	return SNew(SDockTab)
+		.TabRole(ETabRole::PanelTab)
+		[
+			DetailsView.ToSharedRef()
+		];
+}
+
+void SDisplayClusterOperatorPanel::DisplayObjectsInDetailsPanel(const TArray<UObject*>& Objects)
+{
+	if (DetailsView.IsValid())
+	{
+		DetailsView->SetObjects(Objects);
+	}
+}
+
+#undef LOCTEXT_NAMESPACE
diff --git a/Engine/Plugins/Runtime/nDisplay/Source/DisplayClusterOperator/Private/SDisplayClusterOperatorPanel.h b/Engine/Plugins/Runtime/nDisplay/Source/DisplayClusterOperator/Private/SDisplayClusterOperatorPanel.h
new file mode 100644
--- /dev/null
+++ b/Engine/Plugins/Runtime/nDisplay/Source/DisplayClusterOperator/Private/SDisplayClusterOperatorPanel.h
@@ -0,0 +1,71 @@
+// Copyright Epic Games, Inc. All Rights Reserved.
+
+#pragma once
+
+#include "CoreMinimal.h"
+#include "Widgets/SCompoundWidget.h"
+
+class FTabManager;
+class FLayoutExtender;
+class FSpawnTabArgs;
+class FMenuBuilder;
+class SDockTab;
+class SWindow;
+
+/** The nDisplay operator panel that allows users to edit root actor instances with a variety of tools */
+class SDisplayClusterOperatorPanel : public SCompoundWidget
+{
+public:
+	/** The name of the tab that the operator panel lives in */
+	static const FName TabName;
+
+	/** The tab ID that the operator toolbar lives in */
+	static const FName ToolbarTabId;
+
+	/** The tab ID that the details panel lives in */
+	static const FName DetailsTabId;
+
+	/** The ID of the tab stack that can be extended by external tabs */
+	static const FName TabExtensionId;
+
+
+	/** Registers the operator panel with the global tab manager and adds it to the Virtual Projections window menu */
+	static void RegisterTabSpawner();
+
+	/** Unregisters the operator panel from the global tab manager */
+	static void UnregisterTabSpawner();
+
+	/** Creates a tab with the operator panel inside */
+	static TSharedRef<SDockTab> SpawnInTab(const FSpawnTabArgs& SpawnTabArgs);
+
+
+	SLATE_BEGIN_ARGS(SDisplayClusterOperatorPanel) {}
+	SLATE_END_ARGS()
+
+	~SDisplayClusterOperatorPanel();
+
+	void Construct(const FArguments& InArgs, const TSharedRef<SDockTab>& MajorTabOwner, const TSharedPtr<SWindow>& WindowOwner);
+
+private:
+	/** Creates a tab with the operator toolbar in it */
+	TSharedRef<SDockTab> SpawnToolbarTab(const FSpawnTabArgs& Args);
+
+	/** Creates a tab with the details panel in it */
+	TSharedRef<SDockTab> SpawnDetailsTab(const FSpawnTabArgs& Args);
+
+	/** Raised when something wants to display a list of objects in the operator's details panel */
+	void DisplayObjectsInDetailsPanel(const TArray<UObject*>& Objects);
+
+private:
+	/** Holds the tab manager that manages the operator's tabs. */
+	TSharedPtr<FTabManager> TabManager;
+
+	/** The layout extender used by the operator's layout */
+	TSharedPtr<FLayoutExtender> LayoutExtender;
+
+	/** A reference to the operator panel's details view */
+	TSharedPtr<class IDetailsView> DetailsView;
+
+	/** The delegate handle for the operator module's OnDetailObjectsChanged event */
+	FDelegateHandle DetailObjectsChangedHandle;
+};
\ No newline at end of file
diff --git a/Engine/Plugins/Runtime/nDisplay/Source/DisplayClusterOperator/Private/SDisplayClusterOperatorToolbar.cpp b/Engine/Plugins/Runtime/nDisplay/Source/DisplayClusterOperator/Private/SDisplayClusterOperatorToolbar.cpp
new file mode 100644
--- /dev/null
+++ b/Engine/Plugins/Runtime/nDisplay/Source/DisplayClusterOperator/Private/SDisplayClusterOperatorToolbar.cpp
@@ -0,0 +1,159 @@
+// Copyright Epic Games, Inc. All Rights Reserved.
+
+#include "SDisplayClusterOperatorToolbar.h"
+
+#include "IDisplayClusterOperator.h"
+#include "DisplayClusterRootActor.h"
+
+#include "EditorStyleSet.h"
+#include "Framework/MultiBox/MultiBoxBuilder.h"
+#include "Toolkits/AssetEditorToolkit.h"
+#include "Widgets/Layout/SBorder.h"
+#include "Widgets/Input/STextComboBox.h"
+
+#define LOCTEXT_NAMESPACE "SDisplayClusterConfiguratorOperator"
+
+SDisplayClusterOperatorToolbar::~SDisplayClusterOperatorToolbar()
+{
+	if (GEngine != nullptr)
+	{
+		GEngine->OnLevelActorDeleted().Remove(LevelActorDeletedHandle);
+	}
+}
+
+void SDisplayClusterOperatorToolbar::Construct(const FArguments& InArgs)
+{
+	CommandList = InArgs._CommandList;
+
+	FillRootActorList();
+	
+	TSharedPtr<FExtender> ToolBarExtender = IDisplayClusterOperator::Get().GetOperatorToolBarExtensibilityManager()->GetAllExtenders();
+
+    FToolBarBuilder ToolBarBuilder(CommandList, FMultiBoxCustomization::None, ToolBarExtender);
+
+	RootActorComboBox = SNew(SComboBox<TSharedPtr<FString>>)
+		.OptionsSource(&RootActorList)
+		.OnSelectionChanged(this, &SDisplayClusterOperatorToolbar::OnRootActorChanged)
+		.OnComboBoxOpening(this, &SDisplayClusterOperatorToolbar::OnRootActorComboBoxOpening)
+		.OnGenerateWidget(this, &SDisplayClusterOperatorToolbar::GenerateRootActorComboBoxWidget)
+		.Content()
+		[
+			SNew(STextBlock)
+			.Text(this, &SDisplayClusterOperatorToolbar::GetRootActorComboBoxText)
+		];
+
+	if (RootActorList.Num())
+	{
+		RootActorComboBox->SetSelectedItem(RootActorList[0]);
+	}
+
+	ToolBarBuilder.BeginSection("General");
+	{
+		ToolBarBuilder.AddToolBarWidget(RootActorComboBox.ToSharedRef(), LOCTEXT("RootActorPickerLabel", "nDisplay Actor"));
+	}
+	ToolBarBuilder.EndSection();
+
+	ChildSlot
+	[
+		SNew(SBorder)
+		.BorderImage(FEditorStyle::GetBrush("ToolPanel.GroupBorder"))
+		.Padding(0.0f)
+		[
+			ToolBarBuilder.MakeWidget()
+		]
+	];
+
+	if (GEngine != nullptr)
+	{
+		LevelActorDeletedHandle = GEngine->OnLevelActorDeleted().AddSP(this, &SDisplayClusterOperatorToolbar::OnLevelActorDeleted);
+	}
+}
+
+TSharedPtr<FString> SDisplayClusterOperatorToolbar::FillRootActorList(const FString& InitiallySelectedRootActor)
+{
+	RootActorList.Empty();
+
+	TArray<ADisplayClusterRootActor*> RootActors;
+	IDisplayClusterOperator::Get().GetRootActorLevelInstances(RootActors);
+
+	TSharedPtr<FString> SelectedItem = nullptr;
+	for (ADisplayClusterRootActor* RootActor : RootActors)
+	{
+		TSharedPtr<FString> ActorName = MakeShared<FString>(RootActor->GetActorNameOrLabel());
+		if (InitiallySelectedRootActor == *ActorName)
+		{
+			SelectedItem = ActorName;
+		}
+
+		RootActorList.Add(ActorName);
+	}
+
+	return SelectedItem;
+}
+
+void SDisplayClusterOperatorToolbar::OnRootActorChanged(TSharedPtr<FString> ItemSelected, ESelectInfo::Type SelectInfo)
+{
+	if (!ItemSelected.IsValid())
+	{
+		return;
+	}
+
+	TArray<ADisplayClusterRootActor*> RootActors;
+	IDisplayClusterOperator::Get().GetRootActorLevelInstances(RootActors);
+
+	ADisplayClusterRootActor* SelectedRootActor = nullptr;
+	for (ADisplayClusterRootActor* RootActor : RootActors)
+	{
+		if (RootActor->GetActorNameOrLabel() == *ItemSelected)
+		{
+			SelectedRootActor = RootActor;
+			break;
+		}
+	}
+
+	ActiveRootActor = SelectedRootActor;
+	IDisplayClusterOperator::Get().OnActiveRootActorChanged().Broadcast(SelectedRootActor);
+}
+
+void SDisplayClusterOperatorToolbar::OnRootActorComboBoxOpening()
+{
+	FString SelectedRootActor = "";
+	if (TSharedPtr<FString> SelectedItem = RootActorComboBox->GetSelectedItem())
+	{
+		SelectedRootActor = *SelectedItem;
+	}
+
+	TSharedPtr<FString> NewSelectedItem = FillRootActorList(SelectedRootActor);
+
+	RootActorComboBox->RefreshOptions();
+	RootActorComboBox->SetSelectedItem(NewSelectedItem);
+}
+
+TSharedRef<SWidget> SDisplayClusterOperatorToolbar::GenerateRootActorComboBoxWidget(TSharedPtr<FString> InItem) const
+{
+	return SNew(STextBlock).Text(FText::FromString(*InItem.Get()));
+}
+
+FText SDisplayClusterOperatorToolbar::GetRootActorComboBoxText() const
+{
+	TSharedPtr<FString> SelectedRootActor = RootActorComboBox->GetSelectedItem();
+
+	if (!SelectedRootActor.IsValid())
+	{
+		return LOCTEXT("NoRootActorSelectedLabel", "No nDisplay Actor Selected");
+	}
+	
+	return FText::FromString(*RootActorComboBox->GetSelectedItem());
+}
+
+void SDisplayClusterOperatorToolbar::OnLevelActorDeleted(AActor* Actor)
+{
+	if (Actor == ActiveRootActor)
+	{
+		ActiveRootActor = nullptr;
+		IDisplayClusterOperator::Get().OnActiveRootActorChanged().Broadcast(nullptr);
+		RootActorComboBox->SetSelectedItem(nullptr);
+	}
+}
+
+#undef LOCTEXT_NAMESPACE
\ No newline at end of file
diff --git a/Engine/Plugins/Runtime/nDisplay/Source/DisplayClusterOperator/Private/SDisplayClusterOperatorToolbar.h b/Engine/Plugins/Runtime/nDisplay/Source/DisplayClusterOperator/Private/SDisplayClusterOperatorToolbar.h
new file mode 100644
--- /dev/null
+++ b/Engine/Plugins/Runtime/nDisplay/Source/DisplayClusterOperator/Private/SDisplayClusterOperatorToolbar.h
@@ -0,0 +1,68 @@
+// Copyright Epic Games, Inc. All Rights Reserved.
+
+#pragma once
+
+#include "CoreMinimal.h"
+#include "Widgets/DeclarativeSyntaxSupport.h"
+#include "Widgets/SCompoundWidget.h"
+#include "UObject/WeakObjectPtr.h"
+
+class FUICommandList;
+class ADisplayClusterRootActor;
+
+template<class T>
+class SComboBox;
+
+/** A toolbar widget used by the nDisplay operator panel */
+class SDisplayClusterOperatorToolbar : public SCompoundWidget
+{
+public:
+	SLATE_BEGIN_ARGS(SDisplayClusterOperatorToolbar)
+	{}
+		SLATE_ARGUMENT(TSharedPtr<FUICommandList>, CommandList)
+	SLATE_END_ARGS()
+
+	~SDisplayClusterOperatorToolbar();
+
+	void Construct(const FArguments& InArgs);
+
+private:
+	/** 
+	 * Constructs the list of root actors that exist on the current level to use for the root actor picker dropdown.
+	 * 
+	 * @param InitiallySelectedRootActor - The actor name of the root actor that should be initially selected in the combo box
+	 * @eturns The pointer to the item in the root actor list to pass to the combo box as the initially selected item
+	 */
+	TSharedPtr<FString> FillRootActorList(const FString& InitiallySelectedRootActor = TEXT(""));
+
+	/** Raised when the user selects a new root actor from the root actor picker dropdown */
+	void OnRootActorChanged(TSharedPtr<FString> ItemSelected, ESelectInfo::Type SelectInfo);
+
+	/** Raised when the root actor picker dropdown is being opened */
+	void OnRootActorComboBoxOpening();
+
+	/** Creates the widget to display for the specified dropdown item */
+	TSharedRef<SWidget> GenerateRootActorComboBoxWidget(TSharedPtr<FString> InItem) const;
+
+	/** Gets the text to display in the combo box for the selected root actor */
+	FText GetRootActorComboBoxText() const;
+
+	/** Raised when the user deletes an actor from the level */
+	void OnLevelActorDeleted(AActor* Actor);
+
+private:
+	/** The command list used by the toolbar */
+	TSharedPtr<FUICommandList> CommandList;
+	
+	/** The list of root actor names on the current level to display in the root actor picker dropdown */
+	TArray<TSharedPtr<FString>> RootActorList;
+
+	/** The combo box widget that allows the uesr to pick the active root actor from */
+	TSharedPtr<SComboBox<TSharedPtr<FString>>> RootActorComboBox;
+
+	/** The active root actor that the operator panel is currently editing */
+	TWeakObjectPtr<ADisplayClusterRootActor> ActiveRootActor;
+
+	/** The OnLevelActorDeleted delegate handle */
+	FDelegateHandle LevelActorDeletedHandle;
+};
\ No newline at end of file
diff --git a/Engine/Plugins/Runtime/nDisplay/Source/DisplayClusterOperator/Public/IDisplayClusterOperator.h b/Engine/Plugins/Runtime/nDisplay/Source/DisplayClusterOperator/Public/IDisplayClusterOperator.h
new file mode 100644
--- /dev/null
+++ b/Engine/Plugins/Runtime/nDisplay/Source/DisplayClusterOperator/Public/IDisplayClusterOperator.h
@@ -0,0 +1,72 @@
+// Copyright Epic Games, Inc. All Rights Reserved.
+
+#pragma once
+
+#include "CoreMinimal.h"
+#include "Modules/ModuleInterface.h"
+#include "Modules/ModuleManager.h"
+
+class FLayoutExtender;
+class FExtensibilityManager;
+class ADisplayClusterRootActor;
+
+DECLARE_EVENT_OneParam(IDisplayClusterOperator, FOnRegisterLayoutExtensions, FLayoutExtender&);
+DECLARE_EVENT_OneParam(IDisplayClusterOperator, FOnActiveRootActorChanged, ADisplayClusterRootActor*);
+DECLARE_EVENT_OneParam(IDisplayClusterOperator, FOnDetailObjectsChanged, const TArray<UObject*>&);
+
+/**
+ * Display Cluster Operator module interface
+ */
+class IDisplayClusterOperator : public IModuleInterface
+{
+public:
+	static constexpr auto ModuleName = TEXT("DisplayClusterOperator");
+
+public:
+	virtual ~IDisplayClusterOperator() = default;
+
+	/**
+	* Singleton-like access to this module's interface.  This is just for convenience!
+	* Beware of calling this during the shutdown phase, though.  Your module might have been unloaded already.
+	*
+	* @return Returns singleton instance, loading the module on demand if needed
+	*/
+	static inline IDisplayClusterOperator& Get()
+	{
+		return FModuleManager::GetModuleChecked<IDisplayClusterOperator>(ModuleName);
+	}
+
+	/**
+	* Checks to see if this module is loaded and ready.  It is only valid to call Get() if IsAvailable() returns true.
+	*
+	* @return True if the module is loaded and ready to use
+	*/
+	static inline bool IsAvailable()
+	{
+		return FModuleManager::Get().IsModuleLoaded(ModuleName);
+	}
+
+	/** Gets the event handler that is raised when the operator panel processes extensions to its layout */
+	virtual FOnRegisterLayoutExtensions& OnRegisterLayoutExtensions() = 0;
+
+	/** Gets the event handler that is raised when the operator panel changes the root actor being operated on */
+	virtual FOnActiveRootActorChanged& OnActiveRootActorChanged() = 0;
+
+	/** Gets the event handler that is raised when the objects being displayed in the operator's details panel are changed */
+	virtual FOnDetailObjectsChanged& OnDetailObjectsChanged() = 0;
+
+	/** Gets the extension ID that can be used to add tabs to the operator panel */
+	virtual FName GetOperatorExtensionId() = 0;
+
+	/** Gets the extensibility manager for the operator panel's toolbar */
+	virtual TSharedPtr<FExtensibilityManager> GetOperatorToolBarExtensibilityManager() = 0;
+
+	/** Gets a list of all nDisplay root actor instances that are on the currently loaded level */
+	virtual void GetRootActorLevelInstances(TArray<ADisplayClusterRootActor*>& OutRootActorInstances) = 0;
+
+	/** Displays the properties of the specified object in the operator's details panel */
+	virtual void ShowDetailsForObject(UObject* Object) = 0;
+
+	/** Displays the properties of the specified object in the operator's details panel */
+	virtual void ShowDetailsForObjects(const TArray<UObject*>& Objects) = 0;
+};
\ No newline at end of file
diff --git a/Engine/Plugins/Runtime/nDisplay/nDisplay.uplugin b/Engine/Plugins/Runtime/nDisplay/nDisplay.uplugin
--- a/Engine/Plugins/Runtime/nDisplay/nDisplay.uplugin
+++ b/Engine/Plugins/Runtime/nDisplay/nDisplay.uplugin
@@ -81,6 +81,15 @@
 				"Linux"
 			]
 		},
+		{
+			"Name": "DisplayClusterOperator",
+			"Type": "Editor",
+			"LoadingPhase": "Default",
+			"PlatformAllowList": [
+				"Win64",
+				"Linux"
+			]
+		},
 		{
 			"Name": "DisplayClusterConfigurator",
 			"Type": "Editor",
@@ -90,6 +99,15 @@
 				"Linux"
 			]
 		},
+		{
+			"Name": "DisplayClusterLightCardEditor",
+			"Type": "Editor",
+			"LoadingPhase": "Default",
+			"PlatformAllowList": [
+				"Win64",
+				"Linux"
+			]
+		},
 		{
 			"Name": "DisplayClusterStageMonitoring",
 			"Type": "Runtime",
