[Editor,nDisplay] - Ability to disable actor preview in the level editor

diff --git a/Engine/Plugins/Runtime/nDisplay/Source/DisplayCluster/Private/Game/EngineClasses/Scene/DisplayClusterICVFXCameraComponent.cpp b/Engine/Plugins/Runtime/nDisplay/Source/DisplayCluster/Private/Game/EngineClasses/Scene/DisplayClusterICVFXCameraComponent.cpp
--- a/Engine/Plugins/Runtime/nDisplay/Source/DisplayCluster/Private/Game/EngineClasses/Scene/DisplayClusterICVFXCameraComponent.cpp
+++ b/Engine/Plugins/Runtime/nDisplay/Source/DisplayCluster/Private/Game/EngineClasses/Scene/DisplayClusterICVFXCameraComponent.cpp
@@ -6,6 +6,7 @@
 #include "Components/DisplayClusterCameraComponent.h"
 #include "DisplayClusterRootActor.h"
 
+
 void UDisplayClusterICVFXCameraComponent::GetDesiredView(FMinimalViewInfo& DesiredView)
 {
 	if (ADisplayClusterRootActor* RootActor = Cast<ADisplayClusterRootActor>(GetOwner()))
@@ -19,7 +20,7 @@ void UDisplayClusterICVFXCameraComponent::GetDesiredView(FMinimalViewInfo& Desir
 
 UCameraComponent* UDisplayClusterICVFXCameraComponent::GetCameraComponent()
 {
-	return CameraSettings.ExternalCameraActor.IsValid() ? CameraSettings.ExternalCameraActor->GetCameraComponent() : this;
+	return CameraSettings.ExternalCameraActor.IsValid() ? CameraSettings.ExternalCameraActor->GetCineCameraComponent() : this;
 }
 
 FString UDisplayClusterICVFXCameraComponent::GetCameraUniqueId() const
@@ -30,7 +31,16 @@ FString UDisplayClusterICVFXCameraComponent::GetCameraUniqueId() const
 #if WITH_EDITOR
 bool UDisplayClusterICVFXCameraComponent::GetEditorPreviewInfo(float DeltaTime, FMinimalViewInfo& ViewOut)
 {
-	return false;
+	return CameraSettings.ExternalCameraActor.IsValid() ?
+		CameraSettings.ExternalCameraActor->GetCineCameraComponent()->GetEditorPreviewInfo(DeltaTime, ViewOut) :
+		UCameraComponent::GetEditorPreviewInfo(DeltaTime, ViewOut);
+}
+
+TSharedPtr<SWidget> UDisplayClusterICVFXCameraComponent::GetCustomEditorPreviewWidget()
+{
+	return CameraSettings.ExternalCameraActor.IsValid() ?
+		CameraSettings.ExternalCameraActor->GetCineCameraComponent()->GetCustomEditorPreviewWidget() :
+		UCameraComponent::GetCustomEditorPreviewWidget();
 }
 #endif
 
diff --git a/Engine/Plugins/Runtime/nDisplay/Source/DisplayCluster/Public/Components/DisplayClusterICVFXCameraComponent.h b/Engine/Plugins/Runtime/nDisplay/Source/DisplayCluster/Public/Components/DisplayClusterICVFXCameraComponent.h
--- a/Engine/Plugins/Runtime/nDisplay/Source/DisplayCluster/Public/Components/DisplayClusterICVFXCameraComponent.h
+++ b/Engine/Plugins/Runtime/nDisplay/Source/DisplayCluster/Public/Components/DisplayClusterICVFXCameraComponent.h
@@ -14,8 +14,10 @@
 #include "DisplayClusterICVFXCameraComponent.generated.h"
 
 struct FMinimalViewInfo;
+class SWidget;
 class UCameraComponent;
 
+
 /**
  * nDisplay in-camera VFX camera representation
  */
@@ -35,6 +37,7 @@ public:
 
 #if WITH_EDITOR
 	virtual bool GetEditorPreviewInfo(float DeltaTime, FMinimalViewInfo& ViewOut) override;
+	virtual TSharedPtr<SWidget> GetCustomEditorPreviewWidget() override;
 #endif
 
 public:
diff --git a/Engine/Plugins/Runtime/nDisplay/Source/DisplayCluster/Public/DisplayClusterRootActor.h b/Engine/Plugins/Runtime/nDisplay/Source/DisplayCluster/Public/DisplayClusterRootActor.h
--- a/Engine/Plugins/Runtime/nDisplay/Source/DisplayCluster/Public/DisplayClusterRootActor.h
+++ b/Engine/Plugins/Runtime/nDisplay/Source/DisplayCluster/Public/DisplayClusterRootActor.h
@@ -342,6 +342,12 @@ public:
 	virtual bool IsSelectedInEditor() const override;
 	void SetIsSelectedInEditor(bool bValue);
 
+	// Don't show actor preview in the level viewport when DCRA actor is selected, but none of its children are.
+	virtual bool IsDefaultPreviewEnabled() const override
+	{
+		return false;
+	}
+
 protected:
 	FString GeneratePreviewComponentName(const FString& NodeId, const FString& ViewportId) const;
 	void RenderPreview_Editor();
diff --git a/Engine/Source/Editor/LevelEditor/Private/SLevelViewport.cpp b/Engine/Source/Editor/LevelEditor/Private/SLevelViewport.cpp
--- a/Engine/Source/Editor/LevelEditor/Private/SLevelViewport.cpp
+++ b/Engine/Source/Editor/LevelEditor/Private/SLevelViewport.cpp
@@ -237,6 +237,9 @@ void SLevelViewport::Construct(const FArguments& InArgs, const FAssetEditorViewp
 	// Tell the level editor we want to be notified when selection changes
 	LevelEditor.OnActorSelectionChanged().AddRaw( this, &SLevelViewport::OnActorSelectionChanged );
 
+	// Tell the level editor we want to be notified when selection changes
+	LevelEditor.OnElementSelectionChanged().AddRaw( this, &SLevelViewport::OnElementSelectionChanged );
+
 	// Tell the level editor we want to be notified when selection changes
 	LevelEditor.OnMapChanged().AddRaw( this, &SLevelViewport::OnMapChanged );
 
@@ -2760,6 +2763,12 @@ void SLevelViewport::OnActorSelectionChanged(const TArray<UObject*>& NewSelectio
 	bNeedToUpdatePreviews = true;
 }
 
+void SLevelViewport::OnElementSelectionChanged(const UTypedElementSelectionSet* SelectionSet, bool bForceRefresh)
+{
+	// Request preview update. It's possible that the actor currently selected has forbidden default preview.
+	// However, we need to show the preview widget when some of its child components selected.
+	bNeedToUpdatePreviews = true;
+}
 
 void SLevelViewport::PreviewSelectedCameraActors(const bool bPreviewInDesktopViewport)
 {
diff --git a/Engine/Source/Editor/LevelEditor/Public/SLevelViewport.h b/Engine/Source/Editor/LevelEditor/Public/SLevelViewport.h
--- a/Engine/Source/Editor/LevelEditor/Public/SLevelViewport.h
+++ b/Engine/Source/Editor/LevelEditor/Public/SLevelViewport.h
@@ -179,6 +179,14 @@ public:
 	 */
 	void OnActorSelectionChanged(const TArray<UObject*>& NewSelection, bool bForceRefresh=false);
 
+	/**
+	 * Called on all viewports, when element selection changes.
+	 * 
+	 * @param SelectionSet  New selection
+	 * @param bForceRefresh Force refresh
+	 */
+	void OnElementSelectionChanged(const UTypedElementSelectionSet* SelectionSet, bool bForceRefresh);
+
 	/**
 	 * Called when game view should be toggled
 	 */
diff --git a/Engine/Source/Editor/UnrealEd/Private/LevelEditorViewport.cpp b/Engine/Source/Editor/UnrealEd/Private/LevelEditorViewport.cpp
--- a/Engine/Source/Editor/UnrealEd/Private/LevelEditorViewport.cpp
+++ b/Engine/Source/Editor/UnrealEd/Private/LevelEditorViewport.cpp
@@ -3638,6 +3638,9 @@ UActorComponent* FLevelEditorViewportClient::FindViewComponentForActor(AActor co
 		// see if actor has a component with preview capabilities (prioritize camera components)
 		const TSet<UActorComponent*>& Comps = Actor->GetComponents();
 
+		// We need to know if any child component with preview info is selected
+		bool bFoundSelectedComp = false;
+
 		for (UActorComponent* Comp : Comps)
 		{
 			FMinimalViewInfo DummyViewInfo;
@@ -3646,6 +3649,7 @@ UActorComponent* FLevelEditorViewportClient::FindViewComponentForActor(AActor co
 				if (Comp->IsSelected())
 				{
 					PreviewComponent = Comp;
+					bFoundSelectedComp = true;
 					break;
 				}
 				else if (PreviewComponent)
@@ -3661,6 +3665,12 @@ UActorComponent* FLevelEditorViewportClient::FindViewComponentForActor(AActor co
 			}
 		}
 
+		// No preview if default preview is forbidden and no children selection found
+		if (!Actor->IsDefaultPreviewEnabled() && !bFoundSelectedComp)
+		{
+			return nullptr;
+		}
+
 		// now see if any actors are attached to us, directly or indirectly, that have an active camera component we might want to use
 		// we will just return the first one.
 		if (PreviewComponent == nullptr)
diff --git a/Engine/Source/Runtime/Engine/Classes/GameFramework/Actor.h b/Engine/Source/Runtime/Engine/Classes/GameFramework/Actor.h
--- a/Engine/Source/Runtime/Engine/Classes/GameFramework/Actor.h
+++ b/Engine/Source/Runtime/Engine/Classes/GameFramework/Actor.h
@@ -1938,6 +1938,12 @@ public:
 	virtual void PostTransacted(const FTransactionObjectEvent& TransactionEvent) override;
 	virtual bool IsSelectedInEditor() const override;
 
+	/** Defines if preview should be shown when you select an actor, but none of its children */
+	virtual bool IsDefaultPreviewEnabled() const
+	{
+		return true;
+	}
+
 	/** Used to know if actor supports some editor operations. (Delete, Replace) */
 	virtual bool IsUserManaged() const { return true; }
 
