nDisplay: Allow vblank monitoring in nvidia policy for windows (when enabled).

diff --git a/Engine/Plugins/Runtime/nDisplay/Source/DisplayCluster/Private/Render/Synchronization/DisplayClusterRenderSyncPolicyNvidia.h b/Engine/Plugins/Runtime/nDisplay/Source/DisplayCluster/Private/Render/Synchronization/DisplayClusterRenderSyncPolicyNvidia.h
--- a/Engine/Plugins/Runtime/nDisplay/Source/DisplayCluster/Private/Render/Synchronization/DisplayClusterRenderSyncPolicyNvidia.h
+++ b/Engine/Plugins/Runtime/nDisplay/Source/DisplayCluster/Private/Render/Synchronization/DisplayClusterRenderSyncPolicyNvidia.h
@@ -13,6 +13,9 @@ class FDisplayClusterRenderSyncPolicyNvidia
 	: public FDisplayClusterRenderSyncPolicyBase
 {
 public:
+
+	using Super = FDisplayClusterRenderSyncPolicyBase;
+
 	FDisplayClusterRenderSyncPolicyNvidia(const TMap<FString, FString>& Parameters);
 	virtual ~FDisplayClusterRenderSyncPolicyNvidia();
 
diff --git a/Engine/Plugins/Runtime/nDisplay/Source/DisplayCluster/Private/Render/Synchronization/Windows/DisplayClusterRenderSyncPolicyNvidiaWindows.cpp b/Engine/Plugins/Runtime/nDisplay/Source/DisplayCluster/Private/Render/Synchronization/Windows/DisplayClusterRenderSyncPolicyNvidiaWindows.cpp
--- a/Engine/Plugins/Runtime/nDisplay/Source/DisplayCluster/Private/Render/Synchronization/Windows/DisplayClusterRenderSyncPolicyNvidiaWindows.cpp
+++ b/Engine/Plugins/Runtime/nDisplay/Source/DisplayCluster/Private/Render/Synchronization/Windows/DisplayClusterRenderSyncPolicyNvidiaWindows.cpp
@@ -127,6 +127,11 @@ FDisplayClusterRenderSyncPolicyNvidia::~FDisplayClusterRenderSyncPolicyNvidia()
 
 bool FDisplayClusterRenderSyncPolicyNvidia::Initialize()
 {
+	if (!Super::Initialize())
+	{
+		return false;
+	}
+
 	const NvAPI_Status NvApiResult = NvAPI_Initialize();
 	if (NvApiResult != NVAPI_OK)
 	{
