nDisplay - added support for ProceduralMeshComponent as an input for mesh projection policy

diff --git a/Engine/Plugins/Runtime/nDisplay/Source/DisplayCluster/DisplayCluster.Build.cs b/Engine/Plugins/Runtime/nDisplay/Source/DisplayCluster/DisplayCluster.Build.cs
--- a/Engine/Plugins/Runtime/nDisplay/Source/DisplayCluster/DisplayCluster.Build.cs
+++ b/Engine/Plugins/Runtime/nDisplay/Source/DisplayCluster/DisplayCluster.Build.cs
@@ -45,6 +45,7 @@ public class DisplayCluster : ModuleRules
 			"Slate",
 			"SlateCore",
 			"Sockets",
+			"ProceduralMeshComponent",
 		});
 
 		if (Target.bBuildEditor == true)
diff --git a/Engine/Plugins/Runtime/nDisplay/Source/DisplayCluster/Private/Game/EngineClasses/Scene/DisplayClusterPreviewComponent.cpp b/Engine/Plugins/Runtime/nDisplay/Source/DisplayCluster/Private/Game/EngineClasses/Scene/DisplayClusterPreviewComponent.cpp
--- a/Engine/Plugins/Runtime/nDisplay/Source/DisplayCluster/Private/Game/EngineClasses/Scene/DisplayClusterPreviewComponent.cpp
+++ b/Engine/Plugins/Runtime/nDisplay/Source/DisplayCluster/Private/Game/EngineClasses/Scene/DisplayClusterPreviewComponent.cpp
@@ -382,7 +382,7 @@ void UDisplayClusterPreviewComponent::HandleRenderTargetTextureDeferredUpdate()
 	check(IsInGameThread());
 
 	//! @todo: integrate to configurator logic
-	//! deffered update flag
+	//! deferred update flag
 	RenderTargetSurfaceChangedCnt = 2;
 }
 
diff --git a/Engine/Plugins/Runtime/nDisplay/Source/DisplayCluster/Private/Game/EngineClasses/Scene/DisplayClusterRootActor.cpp b/Engine/Plugins/Runtime/nDisplay/Source/DisplayCluster/Private/Game/EngineClasses/Scene/DisplayClusterRootActor.cpp
--- a/Engine/Plugins/Runtime/nDisplay/Source/DisplayCluster/Private/Game/EngineClasses/Scene/DisplayClusterRootActor.cpp
+++ b/Engine/Plugins/Runtime/nDisplay/Source/DisplayCluster/Private/Game/EngineClasses/Scene/DisplayClusterRootActor.cpp
@@ -12,6 +12,7 @@
 #include "Components/DisplayClusterICVFXCameraComponent.h"
 #include "Components/DisplayClusterSceneComponentSyncThis.h"
 #include "CineCameraComponent.h"
+#include "ProceduralMeshComponent.h"
 
 #include "Config/IPDisplayClusterConfigManager.h"
 #include "DisplayClusterConfigurationStrings.h"
@@ -36,8 +37,11 @@
 #include "Misc/DisplayClusterStrings.h"
 
 #include "Render/Viewport/DisplayClusterViewportManager.h"
-#include "Game/EngineClasses/Scene/DisplayClusterRootActorInitializer.h"
+#include "Render/Viewport/IDisplayClusterViewport.h"
+#include "Render/Projection/IDisplayClusterProjectionPolicy.h"
+#include "WarpBlend/IDisplayClusterWarpBlend.h"
 
+#include "Game/EngineClasses/Scene/DisplayClusterRootActorInitializer.h"
 
 ADisplayClusterRootActor::ADisplayClusterRootActor(const FObjectInitializer& ObjectInitializer)
 	: Super(ObjectInitializer)
@@ -464,6 +468,7 @@ bool ADisplayClusterRootActor::GetHiddenInGamePrimitives(TSet<FPrimitiveComponen
 		if (WarpMeshNames.Num() > 0)
 		{
 			GetTypedPrimitives<UStaticMeshComponent>(OutPrimitives, &WarpMeshNames);
+			GetTypedPrimitives<UProceduralMeshComponent>(OutPrimitives, &WarpMeshNames);
 		}
 	}
 
@@ -564,6 +569,38 @@ void ADisplayClusterRootActor::InitializeRootActor()
 #endif
 }
 
+void ADisplayClusterRootActor::UpdateProceduralMeshComponentData(const UProceduralMeshComponent* InProceduralMeshComponent)
+{
+	check(IsInGameThread());
+
+	if (ViewportManager.IsValid())
+	{
+		FName ProceduralComponentName = (InProceduralMeshComponent==nullptr) ? NAME_None : InProceduralMeshComponent->GetFName();
+
+		const TArrayView<IDisplayClusterViewport*> Viewports = ViewportManager->GetViewports();
+		for (IDisplayClusterViewport* ViewportIt : Viewports)
+		{
+			if (ViewportIt != nullptr)
+			{
+				const TSharedPtr<IDisplayClusterProjectionPolicy, ESPMode::ThreadSafe>& ProjectionPolicy = ViewportIt->GetProjectionPolicy();
+				if (ProjectionPolicy.IsValid())
+				{
+					TSharedPtr<IDisplayClusterWarpBlend, ESPMode::ThreadSafe> WarpBlendInterface;
+					if (ProjectionPolicy->GetWarpBlendInterface(WarpBlendInterface) && WarpBlendInterface.IsValid())
+					{
+						// Update only interfaces with ProceduralMesh as geometry source 
+						if (WarpBlendInterface->GetWarpGeometryType() == EDisplayClusterWarpGeometryType::WarpProceduralMesh)
+						{
+							// Set the ProceduralMeshComponent geometry dirty for all valid WarpBlendInterface
+							WarpBlendInterface->MarkWarpGeometryComponentDirty(ProceduralComponentName);
+						}
+					}
+				}
+			}
+		}
+	}
+}
+
 bool ADisplayClusterRootActor::BuildHierarchy()
 {
 	check(CurrentConfigData);
diff --git a/Engine/Plugins/Runtime/nDisplay/Source/DisplayCluster/Private/Game/EngineClasses/Scene/DisplayClusterRootActorEditor.cpp b/Engine/Plugins/Runtime/nDisplay/Source/DisplayCluster/Private/Game/EngineClasses/Scene/DisplayClusterRootActorEditor.cpp
--- a/Engine/Plugins/Runtime/nDisplay/Source/DisplayCluster/Private/Game/EngineClasses/Scene/DisplayClusterRootActorEditor.cpp
+++ b/Engine/Plugins/Runtime/nDisplay/Source/DisplayCluster/Private/Game/EngineClasses/Scene/DisplayClusterRootActorEditor.cpp
@@ -150,7 +150,7 @@ void ADisplayClusterRootActor::GetPreviewRenderTargetableTextures(const TArray<F
 
 							if(bPerformDeferredUpdate)
 							{
-								// handle configurator logic: raise deffered update for preview component RTT
+								// handle configurator logic: raise deferred update for preview component RTT
 								PreviewComponentIt.Value->HandleRenderTargetTextureDeferredUpdate();
 							}
 						}
diff --git a/Engine/Plugins/Runtime/nDisplay/Source/DisplayCluster/Private/Misc/DisplayClusterObjectRef.cpp b/Engine/Plugins/Runtime/nDisplay/Source/DisplayCluster/Private/Misc/DisplayClusterObjectRef.cpp
--- a/Engine/Plugins/Runtime/nDisplay/Source/DisplayCluster/Private/Misc/DisplayClusterObjectRef.cpp
+++ b/Engine/Plugins/Runtime/nDisplay/Source/DisplayCluster/Private/Misc/DisplayClusterObjectRef.cpp
@@ -11,7 +11,7 @@ USceneComponent* FDisplayClusterSceneComponentRef::GetOrFindSceneComponent() con
 		return nullptr;
 	}
 
-	if (!ComponentPtr.IsValid() || ComponentPtr.IsStale())
+	if (!IsSceneComponentValid())
 	{
 		ComponentPtr.Reset();
 
diff --git a/Engine/Plugins/Runtime/nDisplay/Source/DisplayCluster/Private/Render/Containers/DisplayClusterRender_MeshComponent.cpp b/Engine/Plugins/Runtime/nDisplay/Source/DisplayCluster/Private/Render/Containers/DisplayClusterRender_MeshComponent.cpp
--- a/Engine/Plugins/Runtime/nDisplay/Source/DisplayCluster/Private/Render/Containers/DisplayClusterRender_MeshComponent.cpp
+++ b/Engine/Plugins/Runtime/nDisplay/Source/DisplayCluster/Private/Render/Containers/DisplayClusterRender_MeshComponent.cpp
@@ -3,20 +3,22 @@
 #include "Render/Containers/DisplayClusterRender_MeshComponent.h"
 #include "Render/Containers/DisplayClusterRender_MeshComponentProxy.h"
 #include "Render/Containers/DisplayClusterRender_MeshComponentProxyData.h"
+#include "Misc/DisplayClusterLog.h"
 
 #include "Engine/StaticMesh.h"
+#include "ProceduralMeshComponent.h"
 
-#include "Misc/DisplayClusterLog.h"
-
+//-------------------------------------------------------------------------------
+//
 static void ImplUpdateMeshComponentProxyData(FDisplayClusterRender_MeshComponentProxy* MeshComponentProxy, FDisplayClusterRender_MeshComponentProxyData* NewProxyData)
 {
 	if (MeshComponentProxy)
 	{
-		ENQUEUE_RENDER_COMMAND(DisplayClusterRender_MeshComponent_UpdateProxyData)(
+		ENQUEUE_RENDER_COMMAND(DisplayClusterRender_MeshComponentProxy_Update)(
 			[MeshComponentProxy, NewProxyData](FRHICommandListImmediate& RHICmdList)
 		{
 			// Update RHI
-			if (NewProxyData)
+			if (NewProxyData != nullptr)
 			{
 				MeshComponentProxy->UpdateRHI_RenderThread(RHICmdList, NewProxyData);
 				delete NewProxyData;
@@ -31,93 +33,357 @@ static void ImplUpdateMeshComponentProxyData(FDisplayClusterRender_MeshComponent
 
 static void ImplDeleteMeshComponentProxy(FDisplayClusterRender_MeshComponentProxy* MeshComponentProxy)
 {
-	ENQUEUE_RENDER_COMMAND(DisplayClusterRender_MeshComponent_Delete)(
+	ENQUEUE_RENDER_COMMAND(DisplayClusterRender_MeshComponentProxy_Delete)(
 		[MeshComponentProxy](FRHICommandListImmediate& RHICmdList)
 	{
 		delete MeshComponentProxy;
 	});
 }
 
-//*************************************************************************
-//* FDisplayClusterRender_MeshComponent
-//*************************************************************************
+static const FStaticMeshLODResources* ImplGetStaticMeshLODResources(const UStaticMesh* InStaticMesh, int32 LODIndex = 0)
+{
+	if (InStaticMesh != nullptr)
+	{
+		if (InStaticMesh->bAllowCPUAccess)
+		{
+			UE_LOG(LogDisplayClusterRender, Warning, TEXT("If packaging this project, static mesh '%s' requires its AllowCPUAccess flag to be enabled."), *InStaticMesh->GetName());
+#if !WITH_EDITOR
+			// Can't access to cooked data from CPU without this flag
+			return nullptr;
+#endif
+		}
+
+		const FStaticMeshLODResources& StaticMeshLODResources = InStaticMesh->GetLODForExport(LODIndex);
+
+		return &StaticMeshLODResources;
+	}
+
+	return nullptr;
+}
+
+static const FStaticMeshLODResources* ImplGetStaticMeshComponentLODResources(const UStaticMeshComponent* InStaticMeshComponent, int32 LODIndex = 0)
+{
+	if (InStaticMeshComponent != nullptr)
+	{
+		return ImplGetStaticMeshLODResources(InStaticMeshComponent->GetStaticMesh(), LODIndex);
+	}
+
+	return nullptr;
+}
+
+static const FProcMeshSection* ImplGetProceduralMeshComponentSection(UProceduralMeshComponent* InProceduralMeshComponent, const int32 InSectionIndex)
+{
+	if (InProceduralMeshComponent != nullptr && InSectionIndex >= 0)
+	{
+		return InProceduralMeshComponent->GetProcMeshSection(InSectionIndex);
+	}
+
+	return nullptr;
+}
+
+//-------------------------------------------------------------------------------
+//        FDisplayClusterRender_MeshComponent
+//-------------------------------------------------------------------------------
 FDisplayClusterRender_MeshComponent::FDisplayClusterRender_MeshComponent()
-	: MeshComponentProxy(new FDisplayClusterRender_MeshComponentProxy())
-{ }
+{
+	// Create render proxy object
+	MeshComponentProxy = new FDisplayClusterRender_MeshComponentProxy();
+}
 
 FDisplayClusterRender_MeshComponent::~FDisplayClusterRender_MeshComponent()
 {
+	// Delete proxy on RenderThread
 	ImplDeleteMeshComponentProxy(MeshComponentProxy);
+
+	// Forget Ptr on GameThread
 	MeshComponentProxy = nullptr;
 }
 
-void FDisplayClusterRender_MeshComponent::AssignMeshRefs(UStaticMeshComponent* MeshComponent, USceneComponent* OriginComponent)
+USceneComponent* FDisplayClusterRender_MeshComponent::GetOriginComponent()
 {
 	check(IsInGameThread());
 
-	// Update Origin component ref
-	OriginComponentRef.SetSceneComponent(OriginComponent);
-	MeshComponentRef.SetComponentRef(MeshComponent);
+	return OriginComponentRef.GetOrFindSceneComponent();
+}
 
-	// Mesh geometry changed, update related RHI data
-	UpdateDefferedRef();
+void FDisplayClusterRender_MeshComponent::SetGeometryFunc(const EDisplayClusterRender_MeshComponentProxyDataFunc InDataFunc)
+{
+	check(IsInGameThread());
+
+	DataFunc = InDataFunc;
 }
 
-const FStaticMeshLODResources* FDisplayClusterRender_MeshComponent::GetStaticMeshLODResource(int32 LodIndex) const
+FDisplayClusterRender_MeshComponentProxy* FDisplayClusterRender_MeshComponent::GetMeshComponentProxy_RenderThread() const
+{
+	check(IsInRenderingThread());
+
+	return MeshComponentProxy;
+}
+
+UStaticMeshComponent* FDisplayClusterRender_MeshComponent::GetStaticMeshComponent()
 {
 	check(IsInGameThread());
 
-	UStaticMeshComponent* MeshComponent = MeshComponentRef.GetOrFindMeshComponent();
-	if (MeshComponent != nullptr)
+	if(GeometrySource == EDisplayClusterRender_MeshComponentGeometrySource::StaticMeshComponentRef)
 	{
-		UStaticMesh* StaticMesh = MeshComponent->GetStaticMesh();
-		if (StaticMesh != nullptr)
-		{
-			return &(StaticMesh->GetLODForExport(LodIndex));
-		}
+		return StaticMeshComponentRef.GetOrFindStaticMeshComponent();
+	}
+
+	return nullptr;
+}
+
+const FStaticMeshLODResources* FDisplayClusterRender_MeshComponent::GetStaticMeshComponentLODResources(int32 InLODIndex) const
+{
+	check(IsInGameThread());
+
+	if(GeometrySource == EDisplayClusterRender_MeshComponentGeometrySource::StaticMeshComponentRef)
+	{
+		return ImplGetStaticMeshComponentLODResources(StaticMeshComponentRef.GetOrFindStaticMeshComponent(), InLODIndex);
+	}
+
+	return nullptr;
+}
+
+void FDisplayClusterRender_MeshComponent::AssignStaticMeshComponentRefs(UStaticMeshComponent* InStaticMeshComponent, USceneComponent* InOriginComponent, int32 InLODIndex)
+{
+	check(IsInGameThread());
+
+	// Set component refs
+	OriginComponentRef.SetSceneComponent(InOriginComponent);
+	StaticMeshComponentRef.SetStaticMeshComponentRef(InStaticMeshComponent);
+	ProceduralMeshComponentRef.ResetProceduralMeshComponentRef();
+
+	// Set source geometry type
+	GeometrySource = EDisplayClusterRender_MeshComponentGeometrySource::StaticMeshComponentRef;
+
+	// Get geometry data
+	const FStaticMeshLODResources* StaticMeshLODResources = ImplGetStaticMeshComponentLODResources(InStaticMeshComponent, InLODIndex);
+	if (StaticMeshLODResources != nullptr)
+	{
+		// Send geometry to proxy
+		ImplUpdateMeshComponentProxyData(MeshComponentProxy, new FDisplayClusterRender_MeshComponentProxyData(DataFunc, *StaticMeshLODResources));
+	}
+	else
+	{
+		// no StaticMesh - release proxy geometry
+		ImplUpdateMeshComponentProxyData(MeshComponentProxy, nullptr);
+	}
+}
+
+UProceduralMeshComponent* FDisplayClusterRender_MeshComponent::GetProceduralMeshComponent()
+{
+	check(IsInGameThread());
+
+	if(GeometrySource == EDisplayClusterRender_MeshComponentGeometrySource::ProceduralMeshComponentRef)
+	{
+		return ProceduralMeshComponentRef.GetOrFindProceduralMeshComponent();
 	}
 
 	return nullptr;
 }
 
-void FDisplayClusterRender_MeshComponent::Release()
+const FProcMeshSection* FDisplayClusterRender_MeshComponent::GetProceduralMeshComponentSection(const int32 InSectionIndex) const
 {
 	check(IsInGameThread());
 
+	if(GeometrySource == EDisplayClusterRender_MeshComponentGeometrySource::ProceduralMeshComponentRef)
+	{
+		return ImplGetProceduralMeshComponentSection(ProceduralMeshComponentRef.GetOrFindProceduralMeshComponent(), InSectionIndex);
+	}
+
+	return nullptr;
+}
+
+void FDisplayClusterRender_MeshComponent::AssignProceduralMeshComponentRefs(UProceduralMeshComponent* InProceduralMeshComponent, USceneComponent* InOriginComponent, const int32 InSectionIndex)
+{
+	check(IsInGameThread());
+
+	// Set component refs
+	OriginComponentRef.SetSceneComponent(InOriginComponent);
+	StaticMeshComponentRef.ResetStaticMeshComponentRef();
+	ProceduralMeshComponentRef.SetProceduralMeshComponentRef(InProceduralMeshComponent);
+
+	// Set source geometry type
+	GeometrySource = EDisplayClusterRender_MeshComponentGeometrySource::ProceduralMeshComponentRef;
+
+	// Get geometry data
+	const FProcMeshSection* ProcMeshSection = ImplGetProceduralMeshComponentSection(InProceduralMeshComponent, InSectionIndex);
+	if (ProcMeshSection != nullptr)
+	{
+		// Send geometry to proxy
+		ImplUpdateMeshComponentProxyData(MeshComponentProxy, new FDisplayClusterRender_MeshComponentProxyData(DataFunc, *ProcMeshSection));
+	}
+	else
+	{
+		// no ProceduralMesh - release proxy geometry
+		ImplUpdateMeshComponentProxyData(MeshComponentProxy, nullptr);
+	}
+}
+
+void FDisplayClusterRender_MeshComponent::AssignProceduralMeshSection(const FProcMeshSection& InProcMeshSection)
+{
+	check(IsInGameThread());
+
+	// Reset component refs
+	OriginComponentRef.ResetSceneComponent();
+	StaticMeshComponentRef.ResetStaticMeshComponentRef();
+	ProceduralMeshComponentRef.ResetProceduralMeshComponentRef();
+
+	// Set source geometry type
+	GeometrySource = EDisplayClusterRender_MeshComponentGeometrySource::ProceduralMeshSection;
+
+	// Send geometry to proxy
+	ImplUpdateMeshComponentProxyData(MeshComponentProxy, new FDisplayClusterRender_MeshComponentProxyData(DataFunc, InProcMeshSection));
+}
+
+void FDisplayClusterRender_MeshComponent::AssignStaticMesh(const UStaticMesh* InStaticMesh, int32 InLODIndex)
+{
+	check(IsInGameThread());
+
+	// Reset component refs
+	OriginComponentRef.ResetSceneComponent();
+	StaticMeshComponentRef.ResetStaticMeshComponentRef();
+	ProceduralMeshComponentRef.ResetProceduralMeshComponentRef();
+
+	// Set source geometry type
+	GeometrySource = EDisplayClusterRender_MeshComponentGeometrySource::StaticMeshAsset;
+
+	// Get geometry data
+	const FStaticMeshLODResources* StaticMeshLODResources = ImplGetStaticMeshLODResources(InStaticMesh, InLODIndex);
+	if (StaticMeshLODResources != nullptr)
+	{
+		// Send geometry to proxy
+		ImplUpdateMeshComponentProxyData(MeshComponentProxy, new FDisplayClusterRender_MeshComponentProxyData(DataFunc, *StaticMeshLODResources));
+	}
+	else
+	{
+		// no StaticMesh - release proxy geometry
+		ImplUpdateMeshComponentProxyData(MeshComponentProxy, nullptr);
+	}
+}
+
+void FDisplayClusterRender_MeshComponent::AssignMeshGeometry(const FDisplayClusterRender_MeshGeometry* InMeshGeometry)
+{
+	check(IsInGameThread());
+
+	// Reset component refs
+	OriginComponentRef.ResetSceneComponent();
+	StaticMeshComponentRef.ResetStaticMeshComponentRef();
+	ProceduralMeshComponentRef.ResetProceduralMeshComponentRef();
+
+	// Set source geometry type
+	GeometrySource = EDisplayClusterRender_MeshComponentGeometrySource::MeshGeometry;
+
+	if (InMeshGeometry != nullptr)
+	{
+		// Send geometry to proxy
+		ImplUpdateMeshComponentProxyData(MeshComponentProxy, new FDisplayClusterRender_MeshComponentProxyData(DataFunc, *InMeshGeometry));
+	}
+	else
+	{
+		// no InMeshGeometry- release proxy geometry
+		ImplUpdateMeshComponentProxyData(MeshComponentProxy, nullptr);
+	}
+}
+
+void FDisplayClusterRender_MeshComponent::ReleaseMeshComponent()
+{
+	check(IsInGameThread());
+
+	// Reset component refs
 	OriginComponentRef.ResetSceneComponent();
-	MeshComponentRef.ResetSceneComponent();
+	StaticMeshComponentRef.ResetStaticMeshComponentRef();
+	ProceduralMeshComponentRef.ResetProceduralMeshComponentRef();
+
+	// Set source geometry type
+	GeometrySource = EDisplayClusterRender_MeshComponentGeometrySource::Disabled;
+
+	ImplUpdateMeshComponentProxyData(MeshComponentProxy, nullptr);
+}
+
+void FDisplayClusterRender_MeshComponent::ReleaseProxyGeometry()
+{
+	check(IsInGameThread());
 
+	// just release  geometry on proxy
 	ImplUpdateMeshComponentProxyData(MeshComponentProxy, nullptr);
 }
 
-void FDisplayClusterRender_MeshComponent::UpdateDefferedRef()
+bool FDisplayClusterRender_MeshComponent::EqualsMeshComponentName(const FName& InMeshComponentName) const
+{
+	check(IsInGameThread());
+
+	switch (GeometrySource)
+	{
+	case EDisplayClusterRender_MeshComponentGeometrySource::StaticMeshComponentRef:
+		return StaticMeshComponentRef.IsEqualsComponentName(InMeshComponentName);
+		break;
+	case EDisplayClusterRender_MeshComponentGeometrySource::ProceduralMeshComponentRef:
+		return ProceduralMeshComponentRef.IsEqualsComponentName(InMeshComponentName);
+		break;
+	default:
+		break;
+	}
+
+	return false;
+}
+
+bool FDisplayClusterRender_MeshComponent::IsMeshComponentRefGeometryDirty() const
 {
 	check(IsInGameThread());
 
-	FDisplayClusterRender_MeshComponentProxyData* NewProxyData = nullptr;
-	UStaticMeshComponent* StaticMeshComponent = MeshComponentRef.GetOrFindMeshComponent();
-	if (StaticMeshComponent)
+	switch (GeometrySource)
 	{
-		NewProxyData = new FDisplayClusterRender_MeshComponentProxyData(DataFunc, *StaticMeshComponent);
+	case EDisplayClusterRender_MeshComponentGeometrySource::StaticMeshComponentRef:
+		return StaticMeshComponentRef.IsStaticMeshGeometryDirty();
+		break;
+	case EDisplayClusterRender_MeshComponentGeometrySource::ProceduralMeshComponentRef:
+		return ProceduralMeshComponentRef.IsProceduralMeshGeometryDirty();
+		break;
+	default:
+		break;
 	}
 
-	ImplUpdateMeshComponentProxyData(MeshComponentProxy, NewProxyData);
+	return false;
 }
 
-void FDisplayClusterRender_MeshComponent::UpdateDeffered(const UStaticMesh* InStaticMesh)
+void FDisplayClusterRender_MeshComponent::MarkMeshComponentRefGeometryDirty() const
 {
-	FDisplayClusterRender_MeshComponentProxyData* NewProxyData = nullptr;
-	if (InStaticMesh)
+	check(IsInGameThread());
+
+	switch (GeometrySource)
 	{
-		NewProxyData = new FDisplayClusterRender_MeshComponentProxyData(DataFunc, *InStaticMesh);
+	case EDisplayClusterRender_MeshComponentGeometrySource::StaticMeshComponentRef:
+		StaticMeshComponentRef.MarkStaticMeshGeometryDirty();
+		break;
+	case EDisplayClusterRender_MeshComponentGeometrySource::ProceduralMeshComponentRef:
+		ProceduralMeshComponentRef.MarkProceduralMeshGeometryDirty();
+		break;
+	default:
+		break;
 	}
+}
 
-	ImplUpdateMeshComponentProxyData(MeshComponentProxy, NewProxyData);
+void FDisplayClusterRender_MeshComponent::ResetMeshComponentRefGeometryDirty() const
+{
+	check(IsInGameThread());
+
+	switch (GeometrySource)
+	{
+	case EDisplayClusterRender_MeshComponentGeometrySource::StaticMeshComponentRef:
+		StaticMeshComponentRef.ResetStaticMeshGeometryDirty();
+		break;
+	case EDisplayClusterRender_MeshComponentGeometrySource::ProceduralMeshComponentRef:
+		ProceduralMeshComponentRef.ResetProceduralMeshGeometryDirty();
+		break;
+	default:
+		break;
+	}
 }
 
-void FDisplayClusterRender_MeshComponent::UpdateDeffered(const FDisplayClusterRender_MeshGeometry& InMeshGeometry)
+EDisplayClusterRender_MeshComponentGeometrySource FDisplayClusterRender_MeshComponent::GetGeometrySource() const
 {
 	check(IsInGameThread());
 
-	ImplUpdateMeshComponentProxyData(MeshComponentProxy, new FDisplayClusterRender_MeshComponentProxyData(DataFunc, InMeshGeometry));
+	return GeometrySource;
 }
diff --git a/Engine/Plugins/Runtime/nDisplay/Source/DisplayCluster/Private/Render/Containers/DisplayClusterRender_MeshComponentProxyData.cpp b/Engine/Plugins/Runtime/nDisplay/Source/DisplayCluster/Private/Render/Containers/DisplayClusterRender_MeshComponentProxyData.cpp
--- a/Engine/Plugins/Runtime/nDisplay/Source/DisplayCluster/Private/Render/Containers/DisplayClusterRender_MeshComponentProxyData.cpp
+++ b/Engine/Plugins/Runtime/nDisplay/Source/DisplayCluster/Private/Render/Containers/DisplayClusterRender_MeshComponentProxyData.cpp
@@ -6,30 +6,22 @@
 #include "Misc/DisplayClusterLog.h"
 
 #include "Engine/StaticMesh.h"
+#include "ProceduralMeshComponent.h"
 #include "StaticMeshResources.h"
 
 //*************************************************************************
 //* FDisaplyClusterMeshComponentProxyData
 //*************************************************************************
-void FDisplayClusterRender_MeshComponentProxyData::ImplInitializeMesh(const FDisplayClusterRender_MeshComponentProxyDataFunc InDataFunc, const UStaticMesh& InStaticMesh, const int32 InUVChromakeyIndex)
+FDisplayClusterRender_MeshComponentProxyData::FDisplayClusterRender_MeshComponentProxyData(const EDisplayClusterRender_MeshComponentProxyDataFunc InDataFunc, const FStaticMeshLODResources& InStaticMeshLODResource, const int32 InUVChromakeyIndex)
 {
-	if (InStaticMesh.bAllowCPUAccess == false)
-	{
-		UE_LOG(LogDisplayClusterRender, Warning, TEXT("If packaging this project, static mesh '%s' requires its AllowCPUAccess flag to be enabled."), *InStaticMesh.GetName());
-#if !WITH_EDITOR
-		// Can't access to cooked data from CPU without this flag
-		return;
-#endif
-	}
-
-	const FStaticMeshLODResources& SrcMeshResource = InStaticMesh.GetLODForExport(0);
+	check(IsInGameThread());
 
-	const FPositionVertexBuffer& PositionBuffer = SrcMeshResource.VertexBuffers.PositionVertexBuffer;
-	const FStaticMeshVertexBuffer& VertexBuffer = SrcMeshResource.VertexBuffers.StaticMeshVertexBuffer;
-	const FRawStaticIndexBuffer& IndexBuffer = SrcMeshResource.IndexBuffer;
+	const FPositionVertexBuffer& PositionBuffer = InStaticMeshLODResource.VertexBuffers.PositionVertexBuffer;
+	const FStaticMeshVertexBuffer& VertexBuffer = InStaticMeshLODResource.VertexBuffers.StaticMeshVertexBuffer;
+	const FRawStaticIndexBuffer& IndexBuffer = InStaticMeshLODResource.IndexBuffer;
 
 	uint32 UVBaseIndex = 0;
-	uint32 UVChromakeyIndex = (InUVChromakeyIndex < SrcMeshResource.GetNumTexCoords()) ? InUVChromakeyIndex : 0;
+	uint32 UVChromakeyIndex = (InUVChromakeyIndex < InStaticMeshLODResource.GetNumTexCoords()) ? InUVChromakeyIndex : 0;
 
 	IndexBuffer.GetCopy(IndexData);
 
@@ -44,25 +36,44 @@ void FDisplayClusterRender_MeshComponentProxyData::ImplInitializeMesh(const FDis
 	UpdateData(InDataFunc);
 }
 
-FDisplayClusterRender_MeshComponentProxyData::FDisplayClusterRender_MeshComponentProxyData(const FDisplayClusterRender_MeshComponentProxyDataFunc InDataFunc, const UStaticMeshComponent& InMeshComponent, const int32 InUVChromakeyIndex)
+FDisplayClusterRender_MeshComponentProxyData::FDisplayClusterRender_MeshComponentProxyData(const EDisplayClusterRender_MeshComponentProxyDataFunc InDataFunc, const FProcMeshSection& InProcMeshSection, const int32 InUVChromakeyIndex)
 {
-	check(IsInGameThread());
-
-	UStaticMesh* StaticMesh = InMeshComponent.GetStaticMesh();
-	if (StaticMesh)
+	VertexData.AddZeroed(InProcMeshSection.ProcVertexBuffer.Num());
+	for (int32 VertexIdx = 0; VertexIdx < VertexData.Num(); VertexIdx++)
 	{
-		ImplInitializeMesh(InDataFunc, *StaticMesh, InUVChromakeyIndex);
+		const FProcMeshVertex& InVertex = InProcMeshSection.ProcVertexBuffer[VertexIdx];
+
+		VertexData[VertexIdx].Position = InVertex.Position;
+		VertexData[VertexIdx].UV = InVertex.UV0;
+
+		switch (InUVChromakeyIndex)
+		{
+		case 1:
+			VertexData[VertexIdx].UV_Chromakey = InVertex.UV1;
+			break;
+		case 2:
+			VertexData[VertexIdx].UV_Chromakey = InVertex.UV2;
+			break;
+		case 3:
+			VertexData[VertexIdx].UV_Chromakey = InVertex.UV3;
+			break;
+		default:
+			VertexData[VertexIdx].UV_Chromakey = InVertex.UV0;
+			break;
+		}
 	}
-}
 
-FDisplayClusterRender_MeshComponentProxyData::FDisplayClusterRender_MeshComponentProxyData(const FDisplayClusterRender_MeshComponentProxyDataFunc InDataFunc, const UStaticMesh& InStaticMesh, const int32 InUVChromakeyIndex)
-{
-	check(IsInGameThread());
+	int32 IndexDataIt = 0;
+	IndexData.AddZeroed(InProcMeshSection.ProcIndexBuffer.Num());
+	for (const uint32& ProcIndexBufferIt : InProcMeshSection.ProcIndexBuffer)
+	{
+		IndexData[IndexDataIt++] = ProcIndexBufferIt;
+	}
 
-	ImplInitializeMesh(InDataFunc, InStaticMesh, InUVChromakeyIndex);
+	UpdateData(InDataFunc);
 }
 
-FDisplayClusterRender_MeshComponentProxyData::FDisplayClusterRender_MeshComponentProxyData(const FDisplayClusterRender_MeshComponentProxyDataFunc InDataFunc, const FDisplayClusterRender_MeshGeometry& InMeshGeometry)
+FDisplayClusterRender_MeshComponentProxyData::FDisplayClusterRender_MeshComponentProxyData(const EDisplayClusterRender_MeshComponentProxyDataFunc InDataFunc, const FDisplayClusterRender_MeshGeometry& InMeshGeometry)
 {
 	bool bUseChromakeyUV = InMeshGeometry.ChromakeyUV.Num() > 0;
 
@@ -84,20 +95,20 @@ FDisplayClusterRender_MeshComponentProxyData::FDisplayClusterRender_MeshComponen
 	UpdateData(InDataFunc);
 }
 
-void FDisplayClusterRender_MeshComponentProxyData::UpdateData(const FDisplayClusterRender_MeshComponentProxyDataFunc InDataFunc)
+void FDisplayClusterRender_MeshComponentProxyData::UpdateData(const EDisplayClusterRender_MeshComponentProxyDataFunc InDataFunc)
 {
 	if (IndexData.Num() > 0 && VertexData.Num() > 0)
 	{
 		switch (InDataFunc)
 		{
-		case FDisplayClusterRender_MeshComponentProxyDataFunc::OutputRemapScreenSpace:
+		case EDisplayClusterRender_MeshComponentProxyDataFunc::OutputRemapScreenSpace:
 			// Output remap require normalize mesh to screen space coords:
 			NormalizeToScreenSpace();
 			RemoveInvisibleFaces();
 			break;
 
 		default:
-		case FDisplayClusterRender_MeshComponentProxyDataFunc::None:
+		case EDisplayClusterRender_MeshComponentProxyDataFunc::Disabled:
 			break;
 		}
 	}
diff --git a/Engine/Plugins/Runtime/nDisplay/Source/DisplayCluster/Private/Render/Containers/DisplayClusterRender_MeshComponentRef.cpp b/Engine/Plugins/Runtime/nDisplay/Source/DisplayCluster/Private/Render/Containers/DisplayClusterRender_MeshComponentRef.cpp
deleted file mode 100644
diff --git a/Engine/Plugins/Runtime/nDisplay/Source/DisplayCluster/Private/Render/Containers/DisplayClusterRender_ProceduralMeshComponentRef.cpp b/Engine/Plugins/Runtime/nDisplay/Source/DisplayCluster/Private/Render/Containers/DisplayClusterRender_ProceduralMeshComponentRef.cpp
new file mode 100644
--- /dev/null
+++ b/Engine/Plugins/Runtime/nDisplay/Source/DisplayCluster/Private/Render/Containers/DisplayClusterRender_ProceduralMeshComponentRef.cpp
@@ -0,0 +1,67 @@
+// Copyright Epic Games, Inc. All Rights Reserved.
+
+#include "Render/Containers/DisplayClusterRender_ProceduralMeshComponentRef.h"
+
+#include "ProceduralMeshComponent.h"
+
+//-------------------------------------------------------------------------------
+//        FDisplayClusterRender_ProceduralMeshComponentRef
+//-------------------------------------------------------------------------------
+UProceduralMeshComponent* FDisplayClusterRender_ProceduralMeshComponentRef::GetOrFindProceduralMeshComponent() const
+{
+	FScopeLock lock(&DataGuard);
+
+	// When scene component of ProceduralMesh is re-created, etc - mark as dirty for update the geometry
+	if (!IsSceneComponentValid())
+	{
+		MarkProceduralMeshGeometryDirty();
+	}
+
+	USceneComponent* SceneComponent = GetOrFindSceneComponent();
+	if (SceneComponent != nullptr)
+	{
+		UProceduralMeshComponent* ProceduralMeshComponent = Cast<UProceduralMeshComponent>(SceneComponent);
+
+		return ProceduralMeshComponent;
+	}
+
+	return nullptr;
+}
+
+bool FDisplayClusterRender_ProceduralMeshComponentRef::SetProceduralMeshComponentRef(UProceduralMeshComponent* InProceduralMeshComponent)
+{
+	FScopeLock lock(&DataGuard);
+
+	ResetProceduralMeshGeometryDirty();
+
+	return SetSceneComponent(InProceduralMeshComponent);
+}
+
+void FDisplayClusterRender_ProceduralMeshComponentRef::ResetProceduralMeshComponentRef()
+{
+	FScopeLock lock(&DataGuard);
+
+	ResetProceduralMeshGeometryDirty();
+	ResetSceneComponent();
+}
+
+bool FDisplayClusterRender_ProceduralMeshComponentRef::IsProceduralMeshGeometryDirty() const
+{
+	FScopeLock lock(&DataGuard);
+
+	return bIsProceduralMeshGeometryDirty;
+}
+
+void FDisplayClusterRender_ProceduralMeshComponentRef::MarkProceduralMeshGeometryDirty() const
+{
+	FScopeLock lock(&DataGuard);
+
+	bIsProceduralMeshGeometryDirty = true;
+}
+
+void FDisplayClusterRender_ProceduralMeshComponentRef::ResetProceduralMeshGeometryDirty() const
+{
+	FScopeLock lock(&DataGuard);
+
+	bIsProceduralMeshGeometryDirty = false;
+}
diff --git a/Engine/Plugins/Runtime/nDisplay/Source/DisplayCluster/Private/Render/Containers/DisplayClusterRender_StaticMeshComponentRef.cpp b/Engine/Plugins/Runtime/nDisplay/Source/DisplayCluster/Private/Render/Containers/DisplayClusterRender_StaticMeshComponentRef.cpp
new file mode 100644
--- /dev/null
+++ b/Engine/Plugins/Runtime/nDisplay/Source/DisplayCluster/Private/Render/Containers/DisplayClusterRender_StaticMeshComponentRef.cpp
@@ -0,0 +1,108 @@
+// Copyright Epic Games, Inc. All Rights Reserved.
+
+#include "Render/Containers/DisplayClusterRender_StaticMeshComponentRef.h"
+#include "Engine/StaticMesh.h"
+
+//----------------------------------------------------------------------------------------
+//
+inline FName ImplGetStaticMeshGeometryName(const UStaticMeshComponent* StaticMeshComponent)
+{
+	if (StaticMeshComponent != nullptr)
+	{
+		const UStaticMesh* StaticMesh = StaticMeshComponent->GetStaticMesh();
+
+		// When mesh geometry asset is not assigned to a component, return custom text
+		if (StaticMesh == nullptr)
+		{
+			return FName(TEXT("nullptr"));
+		}
+
+		// Return assigned geometry asset name
+		return StaticMesh->GetFName();
+	}
+
+	// Return empty FName for null component ptr
+	return NAME_None;
+}
+
+//----------------------------------------------------------------------------------------
+//      FDisplayClusterRender_StaticMeshComponentRef
+//----------------------------------------------------------------------------------------
+UStaticMeshComponent* FDisplayClusterRender_StaticMeshComponentRef::GetOrFindStaticMeshComponent() const
+{
+	FScopeLock lock(&DataGuard);
+
+	USceneComponent* SceneComponent = GetOrFindSceneComponent();
+	if (SceneComponent != nullptr)
+	{
+		UStaticMeshComponent* StaticMeshComponent = Cast<UStaticMeshComponent>(SceneComponent);
+		if (StaticMeshComponent != nullptr)
+		{
+			// Handle static mesh component geometry asset reference changes runtime:
+			if (!IsStaticMeshGeometryDirty())
+			{
+				const FName NewMeshGeometryName = ImplGetStaticMeshGeometryName(StaticMeshComponent);
+				if(NewMeshGeometryName != MeshGeometryName)
+				{
+					MarkStaticMeshGeometryDirty();
+				}
+			}
+
+			return StaticMeshComponent;
+		}
+	}
+
+	return nullptr;
+}
+
+bool FDisplayClusterRender_StaticMeshComponentRef::SetStaticMeshComponentRef(UStaticMeshComponent* InStaticMeshComponent)
+{
+	FScopeLock lock(&DataGuard);
+
+	ResetStaticMeshGeometryDirty();
+
+	if (SetSceneComponent(InStaticMeshComponent))
+	{
+		// Update the name of the current mesh geometry asset
+		MeshGeometryName = ImplGetStaticMeshGeometryName(InStaticMeshComponent);
+
+		return true;
+	}
+
+	// Clear the name of the current mesh geometry asset
+	MeshGeometryName = FName();
+
+	return false;
+}
+
+void FDisplayClusterRender_StaticMeshComponentRef::ResetStaticMeshComponentRef()
+{
+	FScopeLock lock(&DataGuard);
+
+	// Clear the name of the current mesh geometry asset
+	MeshGeometryName = NAME_None;
+
+	ResetStaticMeshGeometryDirty();
+	ResetSceneComponent();
+}
+
+bool FDisplayClusterRender_StaticMeshComponentRef::IsStaticMeshGeometryDirty() const
+{
+	FScopeLock lock(&DataGuard);
+
+	return bIsStaticMeshGeometryDirty;
+}
+
+void FDisplayClusterRender_StaticMeshComponentRef::MarkStaticMeshGeometryDirty() const
+{
+	FScopeLock lock(&DataGuard);
+
+	bIsStaticMeshGeometryDirty = true;
+}
+
+void FDisplayClusterRender_StaticMeshComponentRef::ResetStaticMeshGeometryDirty() const
+{
+	FScopeLock lock(&DataGuard);
+
+	bIsStaticMeshGeometryDirty = false;
+}
diff --git a/Engine/Plugins/Runtime/nDisplay/Source/DisplayCluster/Private/Render/Viewport/Configuration/DisplayClusterViewportConfigurationHelpers.cpp b/Engine/Plugins/Runtime/nDisplay/Source/DisplayCluster/Private/Render/Viewport/Configuration/DisplayClusterViewportConfigurationHelpers.cpp
--- a/Engine/Plugins/Runtime/nDisplay/Source/DisplayCluster/Private/Render/Viewport/Configuration/DisplayClusterViewportConfigurationHelpers.cpp
+++ b/Engine/Plugins/Runtime/nDisplay/Source/DisplayCluster/Private/Render/Viewport/Configuration/DisplayClusterViewportConfigurationHelpers.cpp
@@ -303,7 +303,7 @@ void FDisplayClusterViewportConfigurationHelpers::UpdateProjectionPolicy(FDispla
 		{
 			if (DstViewport.Owner.IsSceneOpened())
 			{
-				// Try initialize proj policy every tick (mesh deffered load, etc)
+				// Try initialize proj policy every tick (mesh deferred load, etc)
 				DstViewport.HandleStartScene();
 			}
 		}
diff --git a/Engine/Plugins/Runtime/nDisplay/Source/DisplayCluster/Private/Render/Viewport/Containers/DisplayClusterViewportRemap.cpp b/Engine/Plugins/Runtime/nDisplay/Source/DisplayCluster/Private/Render/Viewport/Containers/DisplayClusterViewportRemap.cpp
--- a/Engine/Plugins/Runtime/nDisplay/Source/DisplayCluster/Private/Render/Viewport/Containers/DisplayClusterViewportRemap.cpp
+++ b/Engine/Plugins/Runtime/nDisplay/Source/DisplayCluster/Private/Render/Viewport/Containers/DisplayClusterViewportRemap.cpp
@@ -124,7 +124,7 @@ void FDisplayClusterViewportRemap::ImplUpdateRemapMesh(const FDisplayClusterView
 			RemapMesh = MakeShared<FDisplayClusterRender_MeshComponent, ESPMode::ThreadSafe>();
 		}
 
-		RemapMesh->UpdateDeffered(RemapGeometry);
+		RemapMesh->AssignMeshGeometry(&RemapGeometry);
 	}
 	else
 	{
diff --git a/Engine/Plugins/Runtime/nDisplay/Source/DisplayCluster/Private/Render/Viewport/DisplayClusterViewportManagerProxy.cpp b/Engine/Plugins/Runtime/nDisplay/Source/DisplayCluster/Private/Render/Viewport/DisplayClusterViewportManagerProxy.cpp
--- a/Engine/Plugins/Runtime/nDisplay/Source/DisplayCluster/Private/Render/Viewport/DisplayClusterViewportManagerProxy.cpp
+++ b/Engine/Plugins/Runtime/nDisplay/Source/DisplayCluster/Private/Render/Viewport/DisplayClusterViewportManagerProxy.cpp
@@ -295,7 +295,7 @@ void FDisplayClusterViewportManagerProxy::UpdateFrameResources_RenderThread(FRHI
 	// Handle warped viewport projection policy logic:
 	for (uint8 WarpPass = 0; WarpPass < (uint8)EWarpPass::COUNT; WarpPass++)
 	{
-		// Update deffered resources for viewports
+		// Update deferred resources for viewports
 		for (FDisplayClusterViewportProxy* ViewportProxy : SortedViewportProxy)
 		{
 			if (ViewportProxy)
@@ -351,7 +351,7 @@ void FDisplayClusterViewportManagerProxy::UpdateFrameResources_RenderThread(FRHI
 						{
 						case EWarpPass::Render:
 							// just resolve not warped viewports to frame target texture
-							ViewportProxy->ResolveResources(RHICmdList, EDisplayClusterViewportResourceType::InputShaderResource, ViewportProxy->GetOutputResourceType());
+							ViewportProxy->ResolveResources_RenderThread(RHICmdList, EDisplayClusterViewportResourceType::InputShaderResource, ViewportProxy->GetOutputResourceType_RenderThread());
 							
 							// Implement ViewportRemap feature after resolve
 							ViewportProxy->ImplViewportRemap_RenderThread(RHICmdList);
diff --git a/Engine/Plugins/Runtime/nDisplay/Source/DisplayCluster/Private/Render/Viewport/DisplayClusterViewportProxy.cpp b/Engine/Plugins/Runtime/nDisplay/Source/DisplayCluster/Private/Render/Viewport/DisplayClusterViewportProxy.cpp
--- a/Engine/Plugins/Runtime/nDisplay/Source/DisplayCluster/Private/Render/Viewport/DisplayClusterViewportProxy.cpp
+++ b/Engine/Plugins/Runtime/nDisplay/Source/DisplayCluster/Private/Render/Viewport/DisplayClusterViewportProxy.cpp
@@ -57,8 +57,10 @@ FDisplayClusterViewportProxy::~FDisplayClusterViewportProxy()
 {
 }
 
-const IDisplayClusterViewportManagerProxy& FDisplayClusterViewportProxy::GetOwner() const
+const IDisplayClusterViewportManagerProxy& FDisplayClusterViewportProxy::GetOwner_RenderThread() const
 {
+	check(IsInRenderingThread());
+
 	return Owner;
 }
 
@@ -188,6 +190,8 @@ bool FDisplayClusterViewportProxy::GetResources_RenderThread(const EDisplayClust
 
 void FDisplayClusterViewportProxy::ImplViewportRemap_RenderThread(FRHICommandListImmediate& RHICmdList) const
 {
+	check(IsInRenderingThread());
+
 #if WITH_EDITOR
 	if (Owner.GetRenderFrameSettings_RenderThread().RenderMode == EDisplayClusterRenderFrameMode::PreviewMono)
 	{
@@ -198,7 +202,7 @@ void FDisplayClusterViewportProxy::ImplViewportRemap_RenderThread(FRHICommandLis
 
 	if (RemapMesh.IsValid())
 	{
-		const FDisplayClusterRender_MeshComponentProxy* MeshProxy = RemapMesh->GetProxy();
+		const FDisplayClusterRender_MeshComponentProxy* MeshProxy = RemapMesh->GetMeshComponentProxy_RenderThread();
 		if (MeshProxy && MeshProxy->IsValid_RenderThread())
 		{
 			if (AdditionalFrameTargetableResources.Num() != OutputFrameTargetableResources.Num())
@@ -224,7 +228,7 @@ void FDisplayClusterViewportProxy::ImplViewportRemap_RenderThread(FRHICommandLis
 	}
 }
 
-EDisplayClusterViewportResourceType FDisplayClusterViewportProxy::GetOutputResourceType() const
+EDisplayClusterViewportResourceType FDisplayClusterViewportProxy::GetOutputResourceType_RenderThread() const
 {
 	check(IsInRenderingThread());
 
@@ -237,7 +241,7 @@ EDisplayClusterViewportResourceType FDisplayClusterViewportProxy::GetOutputResou
 
 	if (RemapMesh.IsValid())
 	{
-		const FDisplayClusterRender_MeshComponentProxy* MeshProxy = RemapMesh->GetProxy();
+		const FDisplayClusterRender_MeshComponentProxy* MeshProxy = RemapMesh->GetMeshComponentProxy_RenderThread();
 		if (MeshProxy && MeshProxy->IsValid_RenderThread())
 		{
 			// In this case render to additional frame targetable
@@ -325,7 +329,7 @@ void FDisplayClusterViewportProxy::UpdateDeferredResources(FRHICommandListImmedi
 	}
 
 	// Pass 0: Resolve from RTT region to separated viewport context resource:
-	ImplResolveResources(RHICmdList, SourceProxy, SourceType, EDisplayClusterViewportResourceType::InputShaderResource);
+	ImplResolveResources_RenderThread(RHICmdList, SourceProxy, SourceType, EDisplayClusterViewportResourceType::InputShaderResource);
 
 	// Pass 1: Generate blur postprocess effect for render target texture rect for all contexts
 	if (PostRenderSettings.PostprocessBlur.IsEnabled())
@@ -341,7 +345,7 @@ void FDisplayClusterViewportProxy::UpdateDeferredResources(FRHICommandListImmedi
 			}
 
 			// Copy result back to input
-			ResolveResources(RHICmdList, EDisplayClusterViewportResourceType::AdditionalTargetableResource, EDisplayClusterViewportResourceType::InputShaderResource);
+			ResolveResources_RenderThread(RHICmdList, EDisplayClusterViewportResourceType::AdditionalTargetableResource, EDisplayClusterViewportResourceType::InputShaderResource);
 		}
 	}
 
@@ -352,7 +356,7 @@ void FDisplayClusterViewportProxy::UpdateDeferredResources(FRHICommandListImmedi
 		if (GetResources_RenderThread(EDisplayClusterViewportResourceType::MipsShaderResource, InOutMipsResources))
 		{
 			// Copy input image to layer0 on mips texture
-			ResolveResources(RHICmdList, EDisplayClusterViewportResourceType::InputShaderResource, EDisplayClusterViewportResourceType::MipsShaderResource);
+			ResolveResources_RenderThread(RHICmdList, EDisplayClusterViewportResourceType::InputShaderResource, EDisplayClusterViewportResourceType::MipsShaderResource);
 
 			// Generate mips
 			for (FRHITexture2D*& ResourceIt : InOutMipsResources)
@@ -478,12 +482,12 @@ void ImplResolveResource(FRHICommandListImmediate& RHICmdList, FRHITexture2D* In
 }
 
 // Resolve resource contexts
-bool FDisplayClusterViewportProxy::ResolveResources(FRHICommandListImmediate& RHICmdList, const EDisplayClusterViewportResourceType InputResourceType, const EDisplayClusterViewportResourceType OutputResourceType) const
+bool FDisplayClusterViewportProxy::ResolveResources_RenderThread(FRHICommandListImmediate& RHICmdList, const EDisplayClusterViewportResourceType InputResourceType, const EDisplayClusterViewportResourceType OutputResourceType) const
 {
-	return ImplResolveResources(RHICmdList, this, InputResourceType, OutputResourceType);
+	return ImplResolveResources_RenderThread(RHICmdList, this, InputResourceType, OutputResourceType);
 }
 
-bool FDisplayClusterViewportProxy::ImplResolveResources(FRHICommandListImmediate& RHICmdList, FDisplayClusterViewportProxy const* SourceProxy, const EDisplayClusterViewportResourceType InputResourceType, const EDisplayClusterViewportResourceType OutputResourceType) const
+bool FDisplayClusterViewportProxy::ImplResolveResources_RenderThread(FRHICommandListImmediate& RHICmdList, FDisplayClusterViewportProxy const* SourceProxy, const EDisplayClusterViewportResourceType InputResourceType, const EDisplayClusterViewportResourceType OutputResourceType) const
 {
 	check(IsInRenderingThread());
 
diff --git a/Engine/Plugins/Runtime/nDisplay/Source/DisplayCluster/Private/Render/Viewport/DisplayClusterViewportProxy.h b/Engine/Plugins/Runtime/nDisplay/Source/DisplayCluster/Private/Render/Viewport/DisplayClusterViewportProxy.h
--- a/Engine/Plugins/Runtime/nDisplay/Source/DisplayCluster/Private/Render/Viewport/DisplayClusterViewportProxy.h
+++ b/Engine/Plugins/Runtime/nDisplay/Source/DisplayCluster/Private/Render/Viewport/DisplayClusterViewportProxy.h
@@ -70,17 +70,18 @@ public:
 	virtual bool GetResourcesWithRects_RenderThread(const EDisplayClusterViewportResourceType InResourceType, TArray<FRHITexture2D*>& OutResources, TArray<FIntRect>& OutRects) const override;
 
 	// Resolve resource contexts
-	virtual bool ResolveResources(FRHICommandListImmediate& RHICmdList, const EDisplayClusterViewportResourceType InputResourceType, const EDisplayClusterViewportResourceType OutputResourceType) const override;
+	virtual bool ResolveResources_RenderThread(FRHICommandListImmediate& RHICmdList, const EDisplayClusterViewportResourceType InputResourceType, const EDisplayClusterViewportResourceType OutputResourceType) const override;
 
-	virtual const IDisplayClusterViewportManagerProxy& GetOwner() const override;
-	virtual EDisplayClusterViewportResourceType GetOutputResourceType() const override;
+	virtual EDisplayClusterViewportResourceType GetOutputResourceType_RenderThread() const override;
+
+	virtual const IDisplayClusterViewportManagerProxy& GetOwner_RenderThread() const override;
 
 
 	///////////////////////////////
 	// ~IDisplayClusterViewportProxy
 	///////////////////////////////
 
-	bool ImplResolveResources(FRHICommandListImmediate& RHICmdList, FDisplayClusterViewportProxy const* SourceProxy, const EDisplayClusterViewportResourceType InputResourceType, const EDisplayClusterViewportResourceType OutputResourceType) const;
+	bool ImplResolveResources_RenderThread(FRHICommandListImmediate& RHICmdList, FDisplayClusterViewportProxy const* SourceProxy, const EDisplayClusterViewportResourceType InputResourceType, const EDisplayClusterViewportResourceType OutputResourceType) const;
 
 	void ImplViewportRemap_RenderThread(FRHICommandListImmediate& RHICmdList) const;
 
diff --git a/Engine/Plugins/Runtime/nDisplay/Source/DisplayCluster/Private/Render/Viewport/Postprocess/DisplayClusterViewportPostProcessOutputRemap.cpp b/Engine/Plugins/Runtime/nDisplay/Source/DisplayCluster/Private/Render/Viewport/Postprocess/DisplayClusterViewportPostProcessOutputRemap.cpp
--- a/Engine/Plugins/Runtime/nDisplay/Source/DisplayCluster/Private/Render/Viewport/Postprocess/DisplayClusterViewportPostProcessOutputRemap.cpp
+++ b/Engine/Plugins/Runtime/nDisplay/Source/DisplayCluster/Private/Render/Viewport/Postprocess/DisplayClusterViewportPostProcessOutputRemap.cpp
@@ -93,8 +93,8 @@ bool FDisplayClusterViewportPostProcessOutputRemap::UpdateConfiguration_External
 	bErrorFailLoadFromFileOnce = false;
 
 
-	OutputRemapMesh.DataFunc = FDisplayClusterRender_MeshComponentProxyDataFunc::OutputRemapScreenSpace;
-	OutputRemapMesh.UpdateDeffered(MeshGeometry);
+	OutputRemapMesh.SetGeometryFunc(EDisplayClusterRender_MeshComponentProxyDataFunc::OutputRemapScreenSpace);
+	OutputRemapMesh.AssignMeshGeometry(&MeshGeometry);
 
 	bIsEnabled = true;
 	ExternalFile = InExternalFile;
@@ -116,8 +116,8 @@ bool FDisplayClusterViewportPostProcessOutputRemap::UpdateConfiguration_StaticMe
 		return true;
 	}
 
-	OutputRemapMesh.DataFunc = FDisplayClusterRender_MeshComponentProxyDataFunc::OutputRemapScreenSpace;
-	OutputRemapMesh.UpdateDeffered(InStaticMesh);
+	OutputRemapMesh.SetGeometryFunc(EDisplayClusterRender_MeshComponentProxyDataFunc::OutputRemapScreenSpace);
+	OutputRemapMesh.AssignStaticMesh(InStaticMesh);
 
 	bIsEnabled = true;
 	StaticMesh = InStaticMesh;
@@ -131,7 +131,7 @@ void FDisplayClusterViewportPostProcessOutputRemap::UpdateConfiguration_Disabled
 	ExternalFile.Empty();
 	bIsEnabled = false;
 
-	OutputRemapMesh.UpdateDeffered(nullptr);
+	OutputRemapMesh.ReleaseMeshComponent();
 }
 
 void FDisplayClusterViewportPostProcessOutputRemap::PerformPostProcessFrame_RenderThread(FRHICommandListImmediate& RHICmdList, const TArray<FRHITexture2D*>* InFrameTargets, const TArray<FRHITexture2D*>* InAdditionalFrameTargets) const
@@ -140,7 +140,7 @@ void FDisplayClusterViewportPostProcessOutputRemap::PerformPostProcessFrame_Rend
 
 	if (InFrameTargets && InAdditionalFrameTargets)
 	{
-		const FDisplayClusterRender_MeshComponentProxy* MeshProxy = OutputRemapMesh.GetProxy();
+		const FDisplayClusterRender_MeshComponentProxy* MeshProxy = OutputRemapMesh.GetMeshComponentProxy_RenderThread();
 		if (MeshProxy && MeshProxy->IsValid_RenderThread())
 		{
 			const bool bClearOutputRemapFrameTextureEnabled = CVarClearOutputRemapFrameTextureEnabled.GetValueOnRenderThread() != 0;
diff --git a/Engine/Plugins/Runtime/nDisplay/Source/DisplayCluster/Public/DisplayClusterRootActor.h b/Engine/Plugins/Runtime/nDisplay/Source/DisplayCluster/Public/DisplayClusterRootActor.h
--- a/Engine/Plugins/Runtime/nDisplay/Source/DisplayCluster/Public/DisplayClusterRootActor.h
+++ b/Engine/Plugins/Runtime/nDisplay/Source/DisplayCluster/Public/DisplayClusterRootActor.h
@@ -34,6 +34,7 @@ class UDisplayClusterCameraComponent;
 class UDisplayClusterOriginComponent;
 class UDisplayClusterPreviewComponent;
 class UDisplayClusterSyncTickComponent;
+class UProceduralMeshComponent;
 
 
 /**
@@ -139,6 +140,14 @@ public:
 		return nullptr;
 	}
 
+	/**
+	* Update the geometry of the procedural mesh component(s) referenced inside nDisplay
+	*
+	* @param InProceduralMeshComponent - (optional) Mark the specified procedural mesh component, not all
+	*/
+	UFUNCTION(BlueprintCallable, meta = (DisplayName = "Update ProceduralMeshComponent data"), Category = "NDisplay|Components")
+	void UpdateProceduralMeshComponentData(const UProceduralMeshComponent* InProceduralMeshComponent = nullptr);
+
 public:
 	IDisplayClusterViewportManager* GetViewportManager() const
 	{
diff --git a/Engine/Plugins/Runtime/nDisplay/Source/DisplayCluster/Public/Misc/DisplayClusterObjectRef.h b/Engine/Plugins/Runtime/nDisplay/Source/DisplayCluster/Public/Misc/DisplayClusterObjectRef.h
--- a/Engine/Plugins/Runtime/nDisplay/Source/DisplayCluster/Public/Misc/DisplayClusterObjectRef.h
+++ b/Engine/Plugins/Runtime/nDisplay/Source/DisplayCluster/Public/Misc/DisplayClusterObjectRef.h
@@ -197,6 +197,11 @@ public:
 		return !ComponentName.IsNone() && IsDefinedSceneActor();
 	}
 
+	bool IsSceneComponentValid() const
+	{
+		return ComponentPtr.IsValid() && !ComponentPtr.IsStale();
+	}
+
 	// Return component object ptr.
 	// For killed object ptr, reset and find component new object ptr by name and save to [mutable] ComponentPtr
 	USceneComponent* GetOrFindSceneComponent() const;
@@ -207,7 +212,7 @@ public:
 
 		ResetSceneComponent();
 
-		if (InComponent && SetSceneActor(InComponent->GetOwner()))
+		if (InComponent != nullptr && SetSceneActor(InComponent->GetOwner()))
 		{
 			ComponentName = InComponent->GetFName();
 			ComponentPtr = TWeakObjectPtr<USceneComponent>(InComponent);
@@ -227,6 +232,11 @@ public:
 		ComponentPtr.Reset();
 	}
 
+	bool IsEqualsComponentName(const FName& InComponentName) const
+	{
+		return ComponentName == InComponentName;
+	}
+
 private:
 	// Find component new object ptr by name and save to ComponentPtr
 	FName   ComponentName;
diff --git a/Engine/Plugins/Runtime/nDisplay/Source/DisplayCluster/Public/Render/Containers/DisplayClusterRender_MeshComponent.h b/Engine/Plugins/Runtime/nDisplay/Source/DisplayCluster/Public/Render/Containers/DisplayClusterRender_MeshComponent.h
--- a/Engine/Plugins/Runtime/nDisplay/Source/DisplayCluster/Public/Render/Containers/DisplayClusterRender_MeshComponent.h
+++ b/Engine/Plugins/Runtime/nDisplay/Source/DisplayCluster/Public/Render/Containers/DisplayClusterRender_MeshComponent.h
@@ -4,13 +4,19 @@
 
 #include "CoreMinimal.h"
 
-#include "Render/Containers/DisplayClusterRender_MeshComponentRef.h"
 #include "Render/Containers/DisplayClusterRender_MeshComponentTypes.h"
+#include "Render/Containers/DisplayClusterRender_StaticMeshComponentRef.h"
+#include "Render/Containers/DisplayClusterRender_ProceduralMeshComponentRef.h"
 
-struct FStaticMeshLODResources;
+class USceneComponent;
+class UStaticMeshComponent;
+class UProceduralMeshComponent;
 class FDisplayClusterRender_MeshGeometry;
 class FDisplayClusterRender_MeshComponentProxy;
 
+struct FStaticMeshLODResources;
+struct FProcMeshSection;
+
 class DISPLAYCLUSTER_API FDisplayClusterRender_MeshComponent
 {
 public:
@@ -18,26 +24,143 @@ public:
 	~FDisplayClusterRender_MeshComponent();
 
 public:
-	void AssignMeshRefs(UStaticMeshComponent* MeshComponent, USceneComponent* OriginComponent = nullptr);
-	void UpdateDefferedRef();
-	
-	void UpdateDeffered(const UStaticMesh* InStaticMesh);
-	void UpdateDeffered(const FDisplayClusterRender_MeshGeometry& InMeshGeometry);
+	/**
+	* Assign static mesh component and origin to this container
+	* Then send geometry from the StaticMesh component to proxy
+	*
+	* @param InStaticMeshComponent - Pointer to the StaticMesh component with geometry for warp
+	* @param InOriginComponent     - (optional) stage origin point component
+	* @param InLODIndex            - (optional) LOD index
+	*/
+	void AssignStaticMeshComponentRefs(UStaticMeshComponent* InStaticMeshComponent, USceneComponent* InOriginComponent = nullptr, int32 InLODIndex = 0);
 
-	const FStaticMeshLODResources* GetStaticMeshLODResource(int32 LodIndex = 0) const;
+	/**
+	* Assign procedural mesh component, section index and origin to this container
+	* Then send geometry from the ProceduralMesh(SectionIndex) to proxy
+	*
+	* @param InProceduralMeshComponent - Pointer to the ProceduralMesh component with geometry for warp
+	* @param InSectionIndex            - (optional) The section index in ProceduralMesh
+	* @param InOriginComponent         - (optional) stage origin point component
+	*/
+	void AssignProceduralMeshComponentRefs(UProceduralMeshComponent* InProceduralMeshComponent, USceneComponent* InOriginComponent = nullptr, const int32 InSectionIndex = 0);
 
-	void Release();
+	/**
+	* Assign procedural mesh section geometry
+	* Then send geometry from the ProcMeshSection to proxy
+	*
+	* @param InStaticMesh - Pointer to the StaticMesh geometry
+	*/
+	void AssignProceduralMeshSection(const FProcMeshSection& InProcMeshSection);
 
-	FDisplayClusterRender_MeshComponentProxy* GetProxy() const
-	{ return MeshComponentProxy; }
+	/**
+	* Assign static mesh geometry
+	* Then send geometry from the StaticMesh to proxy
+	*
+	* @param InStaticMesh - Pointer to the StaticMesh geometry
+	*/
+	void AssignStaticMesh(const UStaticMesh* InStaticMesh, int32 InLODIndex = 0);
 
-public:
-	FDisplayClusterRender_MeshComponentProxyDataFunc DataFunc = FDisplayClusterRender_MeshComponentProxyDataFunc::None;
+	/**
+	* Assign mesh geometry
+	* Then send geometry to proxy
+	*
+	* @param InMeshGeometry - Pointer to the geometry data
+	*/
+	void AssignMeshGeometry(const FDisplayClusterRender_MeshGeometry* InMeshGeometry);
+
+	/**
+	* Release refs, proxy, RHI.
+	*
+	*/
+	void ReleaseMeshComponent();
+
+	/**
+	* Release proxy
+	*
+	*/
+	void ReleaseProxyGeometry();
+
+	/**
+	* Get referenced Origin component object
+	*
+	* @return - pointer to the Origin component
+	*/
+	USceneComponent*          GetOriginComponent();
+
+	/**
+	* Get referenced StaticMesh component object
+	*
+	* @return - pointer to the StaticMesh component
+	*/
+	UStaticMeshComponent*     GetStaticMeshComponent();
+
+	/**
+	* Get referenced ProceduralMesh component object
+	*
+	* @return - pointer to the ProceduralMesh component
+	*/
+	UProceduralMeshComponent* GetProceduralMeshComponent();
+
+	/**
+	* Get geometry from assigned StaticMesh component
+	*
+	* @param LodIndex - (optional) lod index
+	*
+	* @return - pointer to the static mesh geometry
+	*/
+	const FStaticMeshLODResources* GetStaticMeshComponentLODResources(int32 InLODIndex = 0) const;
+
+	/**
+	* Get geometry from assigned ProceduralMesh section by index
+	*
+	* @param SectionIndex - Geometry source section index
+	*
+	* @return - pointer to the geometry section data
+	*/
+	const FProcMeshSection* GetProceduralMeshComponentSection(const int32 SectionIndex = 0) const;
+
+	/**
+	* Get nDisplay mesh component proxy object
+	*
+	* @return - pointer to the nDisplay mesh component proxy
+	*/
+	FDisplayClusterRender_MeshComponentProxy* GetMeshComponentProxy_RenderThread() const;
+
+	/**
+	* Set geometry function. This function change geometry prepared for proxy
+	*
+	* @param InDataFunc - geometry function type
+	*/
+	void SetGeometryFunc(const EDisplayClusterRender_MeshComponentProxyDataFunc InDataFunc);
+
+	/**
+	* Return true if referenced component geometry changed
+	*/
+	bool IsMeshComponentRefGeometryDirty() const;
+
+	/**
+	* Mark referenced component geometry as changed
+	*/
+	void MarkMeshComponentRefGeometryDirty() const;
+
+	/**
+	* Clear referenced component geometry changed flag
+	*/
+	void ResetMeshComponentRefGeometryDirty() const;
+
+	bool EqualsMeshComponentName(const FName& InMeshComponentName) const;
+
+	EDisplayClusterRender_MeshComponentGeometrySource GetGeometrySource() const;
+
+private:
+	EDisplayClusterRender_MeshComponentGeometrySource GeometrySource = EDisplayClusterRender_MeshComponentGeometrySource::Disabled;
+	EDisplayClusterRender_MeshComponentProxyDataFunc  DataFunc = EDisplayClusterRender_MeshComponentProxyDataFunc::Disabled;
 
 	// Reference containers:
-	FDisplayClusterRender_MeshComponentRef    MeshComponentRef;
-	FDisplayClusterSceneComponentRef          OriginComponentRef;
+	FDisplayClusterSceneComponentRef                 OriginComponentRef;
+	FDisplayClusterRender_StaticMeshComponentRef     StaticMeshComponentRef;
+	FDisplayClusterRender_ProceduralMeshComponentRef ProceduralMeshComponentRef;
 
 private:
-	FDisplayClusterRender_MeshComponentProxy* MeshComponentProxy;
+	FDisplayClusterRender_MeshComponentProxy* MeshComponentProxy = nullptr;
 };
diff --git a/Engine/Plugins/Runtime/nDisplay/Source/DisplayCluster/Public/Render/Containers/DisplayClusterRender_MeshComponentProxyData.h b/Engine/Plugins/Runtime/nDisplay/Source/DisplayCluster/Public/Render/Containers/DisplayClusterRender_MeshComponentProxyData.h
--- a/Engine/Plugins/Runtime/nDisplay/Source/DisplayCluster/Public/Render/Containers/DisplayClusterRender_MeshComponentProxyData.h
+++ b/Engine/Plugins/Runtime/nDisplay/Source/DisplayCluster/Public/Render/Containers/DisplayClusterRender_MeshComponentProxyData.h
@@ -5,17 +5,18 @@
 #include "CoreMinimal.h"
 #include "Render/Containers/DisplayClusterRender_MeshComponentTypes.h"
 
-class UStaticMesh;
-class UStaticMeshComponent;
+
 class FDisplayClusterRender_MeshGeometry;
+struct FStaticMeshLODResources;
+struct FProcMeshSection;
 
 class DISPLAYCLUSTER_API FDisplayClusterRender_MeshComponentProxyData
 {
 public:
 	// Use channel 1 as default source for chromakey custom markers UV
-	FDisplayClusterRender_MeshComponentProxyData(const FDisplayClusterRender_MeshComponentProxyDataFunc InDataFunc, const UStaticMeshComponent& InMeshComponent, const int32 InUVChromakeyIndex = 1);
-	FDisplayClusterRender_MeshComponentProxyData(const FDisplayClusterRender_MeshComponentProxyDataFunc InDataFunc, const UStaticMesh& InStaticMesh, const int32 InUVChromakeyIndex = 1);
-	FDisplayClusterRender_MeshComponentProxyData(const FDisplayClusterRender_MeshComponentProxyDataFunc InDataFunc, const FDisplayClusterRender_MeshGeometry& InMeshGeometry);
+	FDisplayClusterRender_MeshComponentProxyData(const EDisplayClusterRender_MeshComponentProxyDataFunc InDataFunc, const FStaticMeshLODResources&            InStaticMeshLODResource, const int32 InUVChromakeyIndex = 1);
+	FDisplayClusterRender_MeshComponentProxyData(const EDisplayClusterRender_MeshComponentProxyDataFunc InDataFunc, const FProcMeshSection&                   InProcMeshSection,       const int32 InUVChromakeyIndex = 1);
+	FDisplayClusterRender_MeshComponentProxyData(const EDisplayClusterRender_MeshComponentProxyDataFunc InDataFunc, const FDisplayClusterRender_MeshGeometry& InMeshGeometry);
 
 	const TArray<uint32>& GetIndexData() const
 	{ return IndexData; }
@@ -35,7 +36,7 @@ public:
 	}
 
 private:
-	void UpdateData(const FDisplayClusterRender_MeshComponentProxyDataFunc InDataFunc);
+	void UpdateData(const EDisplayClusterRender_MeshComponentProxyDataFunc InDataFunc);
 
 	void NormalizeToScreenSpace();
 
@@ -46,7 +47,7 @@ private:
 	bool IsFaceVisible(int32 Face);
 	bool IsUVVisible(int32 UVIndex);
 
-	void ImplInitializeMesh(const FDisplayClusterRender_MeshComponentProxyDataFunc InDataFunc, const UStaticMesh& InStaticMesh, const int32 InUVChromakeyIndex = 1);
+	void ImplInitializeMesh(const EDisplayClusterRender_MeshComponentProxyDataFunc InDataFunc, const UStaticMesh& InStaticMesh, const int32 InUVChromakeyIndex = 1);
 
 private:
 	TArray<uint32> IndexData;
diff --git a/Engine/Plugins/Runtime/nDisplay/Source/DisplayCluster/Public/Render/Containers/DisplayClusterRender_MeshComponentRef.h b/Engine/Plugins/Runtime/nDisplay/Source/DisplayCluster/Public/Render/Containers/DisplayClusterRender_MeshComponentRef.h
deleted file mode 100644
diff --git a/Engine/Plugins/Runtime/nDisplay/Source/DisplayCluster/Public/Render/Containers/DisplayClusterRender_MeshComponentTypes.h b/Engine/Plugins/Runtime/nDisplay/Source/DisplayCluster/Public/Render/Containers/DisplayClusterRender_MeshComponentTypes.h
--- a/Engine/Plugins/Runtime/nDisplay/Source/DisplayCluster/Public/Render/Containers/DisplayClusterRender_MeshComponentTypes.h
+++ b/Engine/Plugins/Runtime/nDisplay/Source/DisplayCluster/Public/Render/Containers/DisplayClusterRender_MeshComponentTypes.h
@@ -5,12 +5,23 @@
 #include "CoreMinimal.h"
 
 // Support special mesh proxy geometry funcs
-enum class FDisplayClusterRender_MeshComponentProxyDataFunc : uint8
+enum class EDisplayClusterRender_MeshComponentProxyDataFunc : uint8
 {
-	None = 0,
+	Disabled = 0,
 	OutputRemapScreenSpace
 };
 
+// Assigned geometry type
+enum class EDisplayClusterRender_MeshComponentGeometrySource: uint8
+{
+	Disabled = 0,
+	StaticMeshComponentRef,
+	ProceduralMeshComponentRef,
+	ProceduralMeshSection,
+	StaticMeshAsset,
+	MeshGeometry,
+};
+
 /** The vertex data used to filter a texture. */
 struct FDisplayClusterMeshVertex
 {
diff --git a/Engine/Plugins/Runtime/nDisplay/Source/DisplayCluster/Public/Render/Containers/DisplayClusterRender_ProceduralMeshComponentRef.h b/Engine/Plugins/Runtime/nDisplay/Source/DisplayCluster/Public/Render/Containers/DisplayClusterRender_ProceduralMeshComponentRef.h
new file mode 100644
--- /dev/null
+++ b/Engine/Plugins/Runtime/nDisplay/Source/DisplayCluster/Public/Render/Containers/DisplayClusterRender_ProceduralMeshComponentRef.h
@@ -0,0 +1,36 @@
+// Copyright Epic Games, Inc. All Rights Reserved.
+
+#pragma once
+
+#include "CoreMinimal.h"
+
+#include "Misc/DisplayClusterObjectRef.h"
+
+class UProceduralMeshComponent;
+
+class DISPLAYCLUSTER_API FDisplayClusterRender_ProceduralMeshComponentRef
+	: public FDisplayClusterSceneComponentRef
+{
+public:
+	FDisplayClusterRender_ProceduralMeshComponentRef()
+		: FDisplayClusterSceneComponentRef()
+	{ }
+
+	FDisplayClusterRender_ProceduralMeshComponentRef(const FDisplayClusterRender_ProceduralMeshComponentRef& In)
+		: FDisplayClusterSceneComponentRef(In)
+		, bIsProceduralMeshGeometryDirty(In.bIsProceduralMeshGeometryDirty)
+	{ }
+
+	// Get or find the referenced ProceduralMesh component
+	UProceduralMeshComponent* GetOrFindProceduralMeshComponent() const;
+
+	bool SetProceduralMeshComponentRef(UProceduralMeshComponent* InProceduralMeshComponent);
+	void ResetProceduralMeshComponentRef();
+
+	bool IsProceduralMeshGeometryDirty() const;
+	void MarkProceduralMeshGeometryDirty() const;
+	void ResetProceduralMeshGeometryDirty() const;
+
+private:
+	mutable bool bIsProceduralMeshGeometryDirty = false;
+};
diff --git a/Engine/Plugins/Runtime/nDisplay/Source/DisplayCluster/Public/Render/Containers/DisplayClusterRender_StaticMeshComponentRef.h b/Engine/Plugins/Runtime/nDisplay/Source/DisplayCluster/Public/Render/Containers/DisplayClusterRender_StaticMeshComponentRef.h
new file mode 100644
--- /dev/null
+++ b/Engine/Plugins/Runtime/nDisplay/Source/DisplayCluster/Public/Render/Containers/DisplayClusterRender_StaticMeshComponentRef.h
@@ -0,0 +1,39 @@
+// Copyright Epic Games, Inc. All Rights Reserved.
+
+#pragma once
+
+#include "CoreMinimal.h"
+
+#include "Misc/DisplayClusterObjectRef.h"
+
+class UStaticMeshComponent;
+
+class DISPLAYCLUSTER_API FDisplayClusterRender_StaticMeshComponentRef
+	: public FDisplayClusterSceneComponentRef
+{
+public:
+	FDisplayClusterRender_StaticMeshComponentRef()
+		: FDisplayClusterSceneComponentRef()
+	{ }
+
+	FDisplayClusterRender_StaticMeshComponentRef(const FDisplayClusterRender_StaticMeshComponentRef& In)
+		: FDisplayClusterSceneComponentRef(In)
+		, bIsStaticMeshGeometryDirty(In.bIsStaticMeshGeometryDirty)
+		, MeshGeometryName(In.MeshGeometryName)
+	{ }
+
+	// Get or find the referenced StaticMesh component
+	UStaticMeshComponent* GetOrFindStaticMeshComponent() const;
+
+	bool SetStaticMeshComponentRef(UStaticMeshComponent* InStaticMeshComponent);
+	void ResetStaticMeshComponentRef();
+
+	bool IsStaticMeshGeometryDirty() const;
+	void MarkStaticMeshGeometryDirty() const;
+	void ResetStaticMeshGeometryDirty() const;
+
+private:
+	mutable bool bIsStaticMeshGeometryDirty = false;
+	// Compares the assigned static mesh geometry by name, and raise for changed [mutable] bIsStaticMeshGeometryDirty
+	FName MeshGeometryName;
+};
diff --git a/Engine/Plugins/Runtime/nDisplay/Source/DisplayCluster/Public/Render/Projection/IDisplayClusterProjectionPolicy.h b/Engine/Plugins/Runtime/nDisplay/Source/DisplayCluster/Public/Render/Projection/IDisplayClusterProjectionPolicy.h
--- a/Engine/Plugins/Runtime/nDisplay/Source/DisplayCluster/Public/Render/Projection/IDisplayClusterProjectionPolicy.h
+++ b/Engine/Plugins/Runtime/nDisplay/Source/DisplayCluster/Public/Render/Projection/IDisplayClusterProjectionPolicy.h
@@ -167,6 +167,30 @@ public:
 	virtual void EndWarpBlend_RenderThread(FRHICommandListImmediate& RHICmdList, const class IDisplayClusterViewportProxy* InViewportProxy)
 	{ }
 
+	/**
+	* Return warpblend interface of this policy
+	*
+	* @param OutWarpBlendInterface - output interface ref
+	* 
+	* @return - true if output interface valid
+	*/
+	virtual bool GetWarpBlendInterface(TSharedPtr<class IDisplayClusterWarpBlend, ESPMode::ThreadSafe>& OutWarpBlendInterface) const
+	{
+		return false;
+	}
+
+	/**
+	* Return warpblend interface proxy of this policy
+	*
+	* @param OutWarpBlendInterfaceProxy - output interface ref
+	*
+	* @return - true if output interface valid
+	*/
+	virtual bool GetWarpBlendInterface_RenderThread(TSharedPtr<class IDisplayClusterWarpBlend, ESPMode::ThreadSafe>& OutWarpBlendInterfaceProxy) const
+	{
+		return false;
+	}
+
 #if WITH_EDITOR
 	/**
 	* Ask projection policy instance if it has any mesh based preview
diff --git a/Engine/Plugins/Runtime/nDisplay/Source/DisplayCluster/Public/Render/Viewport/IDisplayClusterViewportProxy.h b/Engine/Plugins/Runtime/nDisplay/Source/DisplayCluster/Public/Render/Viewport/IDisplayClusterViewportProxy.h
--- a/Engine/Plugins/Runtime/nDisplay/Source/DisplayCluster/Public/Render/Viewport/IDisplayClusterViewportProxy.h
+++ b/Engine/Plugins/Runtime/nDisplay/Source/DisplayCluster/Public/Render/Viewport/IDisplayClusterViewportProxy.h
@@ -36,8 +36,9 @@ public:
 	virtual bool GetResourcesWithRects_RenderThread(const EDisplayClusterViewportResourceType InResourceType, TArray<FRHITexture2D*>& OutResources, TArray<FIntRect>& OutRects) const = 0;
 
 	// Resolve resource contexts
-	virtual bool ResolveResources(FRHICommandListImmediate& RHICmdList, const EDisplayClusterViewportResourceType InputResourceType, const EDisplayClusterViewportResourceType OutputResourceType) const = 0;
+	virtual bool ResolveResources_RenderThread(FRHICommandListImmediate& RHICmdList, const EDisplayClusterViewportResourceType InputResourceType, const EDisplayClusterViewportResourceType OutputResourceType) const = 0;
 
-	virtual const class IDisplayClusterViewportManagerProxy& GetOwner() const = 0;
-	virtual EDisplayClusterViewportResourceType   GetOutputResourceType() const = 0;
+	virtual EDisplayClusterViewportResourceType   GetOutputResourceType_RenderThread() const = 0;
+
+	virtual const class IDisplayClusterViewportManagerProxy& GetOwner_RenderThread() const = 0;
 };
diff --git a/Engine/Plugins/Runtime/nDisplay/Source/DisplayClusterConfigurator/DisplayClusterConfigurator.Build.cs b/Engine/Plugins/Runtime/nDisplay/Source/DisplayClusterConfigurator/DisplayClusterConfigurator.Build.cs
--- a/Engine/Plugins/Runtime/nDisplay/Source/DisplayClusterConfigurator/DisplayClusterConfigurator.Build.cs
+++ b/Engine/Plugins/Runtime/nDisplay/Source/DisplayClusterConfigurator/DisplayClusterConfigurator.Build.cs
@@ -48,7 +48,8 @@ public class DisplayClusterConfigurator : ModuleRules
 				"SubobjectEditor",
 				"SubobjectDataInterface",
 				"ToolWidgets",
-				"VPUtilitiesEditor"
+				"VPUtilitiesEditor",
+				"ProceduralMeshComponent",
 			});
 	}
 }
diff --git a/Engine/Plugins/Runtime/nDisplay/Source/DisplayClusterConfigurator/Private/DisplayClusterConfiguratorBlueprintEditor.cpp b/Engine/Plugins/Runtime/nDisplay/Source/DisplayClusterConfigurator/Private/DisplayClusterConfiguratorBlueprintEditor.cpp
--- a/Engine/Plugins/Runtime/nDisplay/Source/DisplayClusterConfigurator/Private/DisplayClusterConfiguratorBlueprintEditor.cpp
+++ b/Engine/Plugins/Runtime/nDisplay/Source/DisplayClusterConfigurator/Private/DisplayClusterConfiguratorBlueprintEditor.cpp
@@ -405,15 +405,18 @@ void FDisplayClusterConfiguratorBlueprintEditor::SelectAncillaryViewports(const
 				UDisplayClusterConfigurationViewport* Viewport = ViewportPair.Value;
 
 				FString ComponentName;
-				if (Viewport->ProjectionPolicy.Parameters.Contains(DisplayClusterProjectionStrings::cfg::simple::Screen))
+				if (Viewport->ProjectionPolicy.Type.Equals(DisplayClusterProjectionStrings::projection::Simple, ESearchCase::IgnoreCase)
+					&& Viewport->ProjectionPolicy.Parameters.Contains(DisplayClusterProjectionStrings::cfg::simple::Screen))
 				{
 					ComponentName = Viewport->ProjectionPolicy.Parameters[DisplayClusterProjectionStrings::cfg::simple::Screen];
 				}
-				else if (Viewport->ProjectionPolicy.Parameters.Contains(DisplayClusterProjectionStrings::cfg::mesh::Component))
+				else if (Viewport->ProjectionPolicy.Type.Equals(DisplayClusterProjectionStrings::projection::Mesh, ESearchCase::IgnoreCase)
+					&& Viewport->ProjectionPolicy.Parameters.Contains(DisplayClusterProjectionStrings::cfg::mesh::Component))
 				{
 					ComponentName = Viewport->ProjectionPolicy.Parameters[DisplayClusterProjectionStrings::cfg::mesh::Component];
 				}
-				else if (Viewport->ProjectionPolicy.Parameters.Contains(DisplayClusterProjectionStrings::cfg::camera::Component))
+				else if (Viewport->ProjectionPolicy.Type.Equals(DisplayClusterProjectionStrings::projection::Camera, ESearchCase::IgnoreCase)
+					&& Viewport->ProjectionPolicy.Parameters.Contains(DisplayClusterProjectionStrings::cfg::camera::Component))
 				{
 					ComponentName = Viewport->ProjectionPolicy.Parameters[DisplayClusterProjectionStrings::cfg::camera::Component];
 				}
diff --git a/Engine/Plugins/Runtime/nDisplay/Source/DisplayClusterConfigurator/Private/Views/Details/Cluster/DisplayClusterConfiguratorViewportDetailsCustomization.cpp b/Engine/Plugins/Runtime/nDisplay/Source/DisplayClusterConfigurator/Private/Views/Details/Cluster/DisplayClusterConfiguratorViewportDetailsCustomization.cpp
--- a/Engine/Plugins/Runtime/nDisplay/Source/DisplayClusterConfigurator/Private/Views/Details/Cluster/DisplayClusterConfiguratorViewportDetailsCustomization.cpp
+++ b/Engine/Plugins/Runtime/nDisplay/Source/DisplayClusterConfigurator/Private/Views/Details/Cluster/DisplayClusterConfiguratorViewportDetailsCustomization.cpp
@@ -41,8 +41,8 @@ void FDisplayClusterConfiguratorViewportDetailsCustomization::CustomizeDetails(I
 #endif
 
 	const bool bDisplayICVFX =
-		ConfigurationViewportPtr->ProjectionPolicy.Type == DisplayClusterProjectionStrings::projection::Mesh ||
-		ConfigurationViewportPtr->ProjectionPolicy.Type == DisplayClusterProjectionStrings::projection::MPCDI;
+			ConfigurationViewportPtr->ProjectionPolicy.Type.Equals(DisplayClusterProjectionStrings::projection::Mesh, ESearchCase::IgnoreCase)
+		||  ConfigurationViewportPtr->ProjectionPolicy.Type.Equals(DisplayClusterProjectionStrings::projection::MPCDI, ESearchCase::IgnoreCase);
 
 	if (!bDisplayICVFX)
 	{
@@ -53,7 +53,7 @@ void FDisplayClusterConfiguratorViewportDetailsCustomization::CustomizeDetails(I
 	CameraHandle = InLayoutBuilder.GetProperty(GET_MEMBER_NAME_CHECKED(UDisplayClusterConfigurationViewport, Camera));
 	check(CameraHandle->IsValidHandle());
 
-	if (ConfigurationViewportPtr->ProjectionPolicy.Type == DisplayClusterProjectionStrings::projection::Camera)
+	if (ConfigurationViewportPtr->ProjectionPolicy.Type.Equals(DisplayClusterProjectionStrings::projection::Camera, ESearchCase::IgnoreCase))
 	{
 		CameraHandle->MarkHiddenByCustomization();
 	}
diff --git a/Engine/Plugins/Runtime/nDisplay/Source/DisplayClusterConfigurator/Private/Views/Details/Policies/DisplayClusterConfiguratorPolicyDetailCustomization.cpp b/Engine/Plugins/Runtime/nDisplay/Source/DisplayClusterConfigurator/Private/Views/Details/Policies/DisplayClusterConfiguratorPolicyDetailCustomization.cpp
--- a/Engine/Plugins/Runtime/nDisplay/Source/DisplayClusterConfigurator/Private/Views/Details/Policies/DisplayClusterConfiguratorPolicyDetailCustomization.cpp
+++ b/Engine/Plugins/Runtime/nDisplay/Source/DisplayClusterConfigurator/Private/Views/Details/Policies/DisplayClusterConfiguratorPolicyDetailCustomization.cpp
@@ -14,6 +14,8 @@
 #include "Components/DisplayClusterScreenComponent.h"
 #include "Components/DisplayClusterICVFXCameraComponent.h"
 
+#include "ProceduralMeshComponent.h"
+
 #include "Camera/CameraComponent.h"
 
 #include "DisplayClusterProjectionStrings.h"
@@ -444,7 +446,14 @@ void FDisplayClusterConfiguratorProjectionCustomization::CreateMeshPolicy(UDispl
 		DisplayClusterProjectionStrings::cfg::mesh::Component,
 		Blueprint,
 		ConfigurationViewports,
-		TArray<TSubclassOf<UActorComponent>>{ UStaticMeshComponent::StaticClass() }));
+		TArray<TSubclassOf<UActorComponent>>{ UStaticMeshComponent::StaticClass(), UProceduralMeshComponent::StaticClass()}));
+
+	CustomPolicyParameters.Add(MakeShared<FPolicyParameterInfoNumber<int32>>(
+		"SectionIndex",
+		DisplayClusterProjectionStrings::cfg::mesh::SectionIndex,
+		Blueprint,
+		ConfigurationViewports));
+
 }
 
 void FDisplayClusterConfiguratorProjectionCustomization::CreateDomePolicy(UDisplayClusterBlueprint* Blueprint)
diff --git a/Engine/Plugins/Runtime/nDisplay/Source/DisplayClusterConfigurator/Private/Views/TreeViews/Cluster/TreeItems/DisplayClusterConfiguratorTreeItemViewport.cpp b/Engine/Plugins/Runtime/nDisplay/Source/DisplayClusterConfigurator/Private/Views/TreeViews/Cluster/TreeItems/DisplayClusterConfiguratorTreeItemViewport.cpp
--- a/Engine/Plugins/Runtime/nDisplay/Source/DisplayClusterConfigurator/Private/Views/TreeViews/Cluster/TreeItems/DisplayClusterConfiguratorTreeItemViewport.cpp
+++ b/Engine/Plugins/Runtime/nDisplay/Source/DisplayClusterConfigurator/Private/Views/TreeViews/Cluster/TreeItems/DisplayClusterConfiguratorTreeItemViewport.cpp
@@ -61,15 +61,18 @@ void FDisplayClusterConfiguratorTreeItemViewport::OnSelection()
 	UDisplayClusterConfigurationViewport* Viewport = GetObjectChecked<UDisplayClusterConfigurationViewport>();
 
 	TArray<FString> AssociatedComponents;
-	if (Viewport->ProjectionPolicy.Parameters.Contains(DisplayClusterProjectionStrings::cfg::simple::Screen))
+	if (Viewport->ProjectionPolicy.Type.Equals(DisplayClusterProjectionStrings::projection::Simple, ESearchCase::IgnoreCase)
+		&& Viewport->ProjectionPolicy.Parameters.Contains(DisplayClusterProjectionStrings::cfg::simple::Screen))
 	{
 		AssociatedComponents.Add(Viewport->ProjectionPolicy.Parameters[DisplayClusterProjectionStrings::cfg::simple::Screen]);
 	}
-	else if (Viewport->ProjectionPolicy.Parameters.Contains(DisplayClusterProjectionStrings::cfg::mesh::Component))
+	else if (Viewport->ProjectionPolicy.Type.Equals(DisplayClusterProjectionStrings::projection::Mesh, ESearchCase::IgnoreCase)
+		&& Viewport->ProjectionPolicy.Parameters.Contains(DisplayClusterProjectionStrings::cfg::mesh::Component))
 	{
 		AssociatedComponents.Add(Viewport->ProjectionPolicy.Parameters[DisplayClusterProjectionStrings::cfg::mesh::Component]);
 	}
-	else if (Viewport->ProjectionPolicy.Parameters.Contains(DisplayClusterProjectionStrings::cfg::camera::Component))
+	else if (Viewport->ProjectionPolicy.Type.Equals(DisplayClusterProjectionStrings::projection::Camera, ESearchCase::IgnoreCase)
+		&& Viewport->ProjectionPolicy.Parameters.Contains(DisplayClusterProjectionStrings::cfg::camera::Component))
 	{
 		AssociatedComponents.Add(Viewport->ProjectionPolicy.Parameters[DisplayClusterProjectionStrings::cfg::camera::Component]);
 	}
diff --git a/Engine/Plugins/Runtime/nDisplay/Source/DisplayClusterProjection/DisplayClusterProjection.Build.cs b/Engine/Plugins/Runtime/nDisplay/Source/DisplayClusterProjection/DisplayClusterProjection.Build.cs
--- a/Engine/Plugins/Runtime/nDisplay/Source/DisplayClusterProjection/DisplayClusterProjection.Build.cs
+++ b/Engine/Plugins/Runtime/nDisplay/Source/DisplayClusterProjection/DisplayClusterProjection.Build.cs
@@ -29,12 +29,12 @@ public class DisplayClusterProjection : ModuleRules
 				"Projects",
 				"RenderCore",
 				"RHI",
-				"RHICore"
+				"RHICore",
+				"ProceduralMeshComponent",
 			});
 
 		if (Target.bBuildEditor == true)
 		{
-			PrivateDependencyModuleNames.Add("ProceduralMeshComponent");
 			PrivateDependencyModuleNames.Add("UnrealEd");
 		}
 
diff --git a/Engine/Plugins/Runtime/nDisplay/Source/DisplayClusterProjection/Private/Policy/Domeprojection/Windows/DX11/DisplayClusterProjectionDomeprojectionViewAdapterDX11.cpp b/Engine/Plugins/Runtime/nDisplay/Source/DisplayClusterProjection/Private/Policy/Domeprojection/Windows/DX11/DisplayClusterProjectionDomeprojectionViewAdapterDX11.cpp
--- a/Engine/Plugins/Runtime/nDisplay/Source/DisplayClusterProjection/Private/Policy/Domeprojection/Windows/DX11/DisplayClusterProjectionDomeprojectionViewAdapterDX11.cpp
+++ b/Engine/Plugins/Runtime/nDisplay/Source/DisplayClusterProjection/Private/Policy/Domeprojection/Windows/DX11/DisplayClusterProjectionDomeprojectionViewAdapterDX11.cpp
@@ -223,7 +223,7 @@ bool FDisplayClusterProjectionDomeprojectionViewAdapterDX11::ApplyWarpBlend_Rend
 	}
 
 	// resolve warp result images from temp targetable to FrameTarget
-	return InViewportProxy->ResolveResources(RHICmdList, EDisplayClusterViewportResourceType::AdditionalTargetableResource, InViewportProxy->GetOutputResourceType());
+	return InViewportProxy->ResolveResources_RenderThread(RHICmdList, EDisplayClusterViewportResourceType::AdditionalTargetableResource, InViewportProxy->GetOutputResourceType_RenderThread());
 }
 
 void FDisplayClusterProjectionDomeprojectionViewAdapterDX11::FViewData::Release(FCriticalSection& DllAccessCS)
diff --git a/Engine/Plugins/Runtime/nDisplay/Source/DisplayClusterProjection/Private/Policy/EasyBlend/DisplayClusterProjectionEasyBlendPolicyBase.cpp b/Engine/Plugins/Runtime/nDisplay/Source/DisplayClusterProjection/Private/Policy/EasyBlend/DisplayClusterProjectionEasyBlendPolicyBase.cpp
--- a/Engine/Plugins/Runtime/nDisplay/Source/DisplayClusterProjection/Private/Policy/EasyBlend/DisplayClusterProjectionEasyBlendPolicyBase.cpp
+++ b/Engine/Plugins/Runtime/nDisplay/Source/DisplayClusterProjection/Private/Policy/EasyBlend/DisplayClusterProjectionEasyBlendPolicyBase.cpp
@@ -111,7 +111,7 @@ void FDisplayClusterProjectionEasyBlendPolicyBase::ApplyWarpBlend_RenderThread(F
 		if (!ViewAdapter->ApplyWarpBlend_RenderThread(RHICmdList, InViewportProxy))
 		{
 			// Warpbled failed, just copy unwarped rtt for frame output
-			InViewportProxy->ResolveResources(RHICmdList, EDisplayClusterViewportResourceType::InputShaderResource, InViewportProxy->GetOutputResourceType());
+			InViewportProxy->ResolveResources_RenderThread(RHICmdList, EDisplayClusterViewportResourceType::InputShaderResource, InViewportProxy->GetOutputResourceType_RenderThread());
 		}
 	}
 }
diff --git a/Engine/Plugins/Runtime/nDisplay/Source/DisplayClusterProjection/Private/Policy/EasyBlend/Windows/DX11/DisplayClusterProjectionEasyBlendViewAdapterDX11.cpp b/Engine/Plugins/Runtime/nDisplay/Source/DisplayClusterProjection/Private/Policy/EasyBlend/Windows/DX11/DisplayClusterProjectionEasyBlendViewAdapterDX11.cpp
--- a/Engine/Plugins/Runtime/nDisplay/Source/DisplayClusterProjection/Private/Policy/EasyBlend/Windows/DX11/DisplayClusterProjectionEasyBlendViewAdapterDX11.cpp
+++ b/Engine/Plugins/Runtime/nDisplay/Source/DisplayClusterProjection/Private/Policy/EasyBlend/Windows/DX11/DisplayClusterProjectionEasyBlendViewAdapterDX11.cpp
@@ -261,7 +261,7 @@ bool FDisplayClusterProjectionEasyBlendViewAdapterDX11::ApplyWarpBlend_RenderThr
 	}
 
 	// resolve warp result images from temp targetable to FrameTarget
-	return InViewportProxy->ResolveResources(RHICmdList, EDisplayClusterViewportResourceType::AdditionalTargetableResource, InViewportProxy->GetOutputResourceType());
+	return InViewportProxy->ResolveResources_RenderThread(RHICmdList, EDisplayClusterViewportResourceType::AdditionalTargetableResource, InViewportProxy->GetOutputResourceType_RenderThread());
 }
 
 bool FDisplayClusterProjectionEasyBlendViewAdapterDX11::ImplApplyWarpBlend_RenderThread(FRHICommandListImmediate& RHICmdList, uint32 ContextNum, FRHITexture2D* InputTexture, FRHITexture2D* OutputTexture)
diff --git a/Engine/Plugins/Runtime/nDisplay/Source/DisplayClusterProjection/Private/Policy/MPCDI/DisplayClusterProjectionMPCDIPolicy.cpp b/Engine/Plugins/Runtime/nDisplay/Source/DisplayClusterProjection/Private/Policy/MPCDI/DisplayClusterProjectionMPCDIPolicy.cpp
--- a/Engine/Plugins/Runtime/nDisplay/Source/DisplayClusterProjection/Private/Policy/MPCDI/DisplayClusterProjectionMPCDIPolicy.cpp
+++ b/Engine/Plugins/Runtime/nDisplay/Source/DisplayClusterProjection/Private/Policy/MPCDI/DisplayClusterProjectionMPCDIPolicy.cpp
@@ -122,6 +122,34 @@ void FDisplayClusterProjectionMPCDIPolicy::ImplRelease()
 	WarpBlendContexts.Empty();
 }
 
+bool FDisplayClusterProjectionMPCDIPolicy::GetWarpBlendInterface(TSharedPtr<IDisplayClusterWarpBlend, ESPMode::ThreadSafe>& OutWarpBlendInterface) const
+{
+	check(IsInGameThread());
+
+	if (WarpBlendInterface.IsValid())
+	{
+		OutWarpBlendInterface = WarpBlendInterface;
+
+		return true;
+	}
+
+	return false;
+}
+
+bool FDisplayClusterProjectionMPCDIPolicy::GetWarpBlendInterface_RenderThread(TSharedPtr<class IDisplayClusterWarpBlend, ESPMode::ThreadSafe>& OutWarpBlendInterfaceProxy) const
+{
+	check(IsInRenderingThread());
+
+	if (WarpBlendInterface.IsValid())
+	{
+		OutWarpBlendInterfaceProxy = WarpBlendInterface_Proxy;
+
+		return true;
+	}
+
+	return false;
+}
+
 bool FDisplayClusterProjectionMPCDIPolicy::CalculateView(IDisplayClusterViewport* InViewport, const uint32 InContextNum, FVector& InOutViewLocation, FRotator& InOutViewRotation, const FVector& ViewOffset, const float WorldToMeters, const float NCP, const float FCP)
 {
 	check(IsInGameThread());
@@ -197,7 +225,7 @@ void FDisplayClusterProjectionMPCDIPolicy::ApplyWarpBlend_RenderThread(FRHIComma
 {
 	check(IsInRenderingThread());
 
-	if (WarpBlendInterface_Proxy.IsValid() == false || WarpBlendContexts_Proxy.Num() == 0)
+	if (InViewportProxy==nullptr || WarpBlendInterface_Proxy.IsValid() == false || WarpBlendContexts_Proxy.Num() == 0)
 	{
 		return;
 	}
@@ -217,7 +245,7 @@ void FDisplayClusterProjectionMPCDIPolicy::ApplyWarpBlend_RenderThread(FRHIComma
 	}
 
 	// Get output resources with rects
-	if (!InViewportProxy->GetResourcesWithRects_RenderThread(InViewportProxy->GetOutputResourceType(), OutputTextures, OutputRects))
+	if (!InViewportProxy->GetResourcesWithRects_RenderThread(InViewportProxy->GetOutputResourceType_RenderThread(), OutputTextures, OutputRects))
 	{
 		return;
 	}
@@ -231,6 +259,8 @@ void FDisplayClusterProjectionMPCDIPolicy::ApplyWarpBlend_RenderThread(FRHIComma
 		TArray<FDisplayClusterShaderParametersICVFX_ViewportResource*> ViewportResources;
 		FDisplayClusterShaderParameters_ICVFX ShaderICVFX(SettingsICVFX.ICVFX);
 
+		const IDisplayClusterViewportManagerProxy& RefViewportManagerProxy = InViewportProxy->GetOwner_RenderThread();
+
 		// Collect refs
 		ShaderICVFX.CollectRefViewports(ViewportResources);
 		{
@@ -243,7 +273,7 @@ void FDisplayClusterProjectionMPCDIPolicy::ApplyWarpBlend_RenderThread(FRHIComma
 					// reset prev resource ref
 					It->Texture = nullptr;
 
-					const IDisplayClusterViewportProxy* RefViewportProxy = InViewportProxy->GetOwner().FindViewport_RenderThread(It->ViewportId);
+					const IDisplayClusterViewportProxy* RefViewportProxy = RefViewportManagerProxy.FindViewport_RenderThread(It->ViewportId);
 					if (RefViewportProxy)
 					{
 						TArray<FRHITexture2D*> RefTextures;
diff --git a/Engine/Plugins/Runtime/nDisplay/Source/DisplayClusterProjection/Private/Policy/MPCDI/DisplayClusterProjectionMPCDIPolicy.h b/Engine/Plugins/Runtime/nDisplay/Source/DisplayClusterProjection/Private/Policy/MPCDI/DisplayClusterProjectionMPCDIPolicy.h
--- a/Engine/Plugins/Runtime/nDisplay/Source/DisplayClusterProjection/Private/Policy/MPCDI/DisplayClusterProjectionMPCDIPolicy.h
+++ b/Engine/Plugins/Runtime/nDisplay/Source/DisplayClusterProjection/Private/Policy/MPCDI/DisplayClusterProjectionMPCDIPolicy.h
@@ -52,6 +52,9 @@ public:
 	virtual bool IsWarpBlendSupported() override;
 	virtual void ApplyWarpBlend_RenderThread(FRHICommandListImmediate& RHICmdList, const class IDisplayClusterViewportProxy* InViewportProxy) override;
 
+	virtual bool GetWarpBlendInterface(TSharedPtr<class IDisplayClusterWarpBlend, ESPMode::ThreadSafe>& OutWarpBlendInterface) const override;
+	virtual bool GetWarpBlendInterface_RenderThread(TSharedPtr<class IDisplayClusterWarpBlend, ESPMode::ThreadSafe>& OutWarpBlendInterfaceProxy) const override;
+
 	virtual bool ShouldUseSourceTextureWithMips() const override
 	{
 		// Support input texture with mips
diff --git a/Engine/Plugins/Runtime/nDisplay/Source/DisplayClusterProjection/Private/Policy/Mesh/DisplayClusterProjectionMeshPolicy.cpp b/Engine/Plugins/Runtime/nDisplay/Source/DisplayClusterProjection/Private/Policy/Mesh/DisplayClusterProjectionMeshPolicy.cpp
--- a/Engine/Plugins/Runtime/nDisplay/Source/DisplayClusterProjection/Private/Policy/Mesh/DisplayClusterProjectionMeshPolicy.cpp
+++ b/Engine/Plugins/Runtime/nDisplay/Source/DisplayClusterProjection/Private/Policy/Mesh/DisplayClusterProjectionMeshPolicy.cpp
@@ -15,7 +15,7 @@
 #include "WarpBlend/IDisplayClusterWarpBlendManager.h"
 
 #include "Components/StaticMeshComponent.h"
-
+#include "ProceduralMeshComponent.h"
 
 FDisplayClusterProjectionMeshPolicy::FDisplayClusterProjectionMeshPolicy(const FString& ProjectionPolicyId, const struct FDisplayClusterConfigurationProjection* InConfigurationProjectionPolicy)
 	: FDisplayClusterProjectionMPCDIPolicy(ProjectionPolicyId, InConfigurationProjectionPolicy)
@@ -28,14 +28,35 @@ bool FDisplayClusterProjectionMeshPolicy::CreateWarpMeshInterface(class IDisplay
 
 	if (WarpBlendInterface.IsValid() == false)
 	{
-		FDisplayClusterWarpBlendConstruct::FAssignWarpMesh CreateParameters;
-		if (!GetWarpMeshAndOrigin(InViewport, CreateParameters.MeshComponent, CreateParameters.OriginComponent))
+		FWarpMeshConfiguration WarpCfg;
+		if (!GetWarpMeshConfiguration(InViewport, WarpCfg))
 		{
 			return false;
 		}
 
 		IDisplayClusterShaders& ShadersAPI = IDisplayClusterShaders::Get();
-		if (!ShadersAPI.GetWarpBlendManager().Create(CreateParameters, WarpBlendInterface))
+		bool bResult = false;
+
+		if (WarpCfg.StaticMeshComponent != nullptr)
+		{
+			FDisplayClusterWarpBlendConstruct::FAssignWarpMesh CreateParameters;
+			CreateParameters.MeshComponent   = WarpCfg.StaticMeshComponent;
+			CreateParameters.OriginComponent = WarpCfg.OriginComponent;
+
+			bResult = ShadersAPI.GetWarpBlendManager().Create(CreateParameters, WarpBlendInterface);
+		}
+		else
+		{
+			FDisplayClusterWarpBlendConstruct::FAssignWarpProceduralMesh CreateParameters;
+			CreateParameters.OriginComponent = WarpCfg.OriginComponent;
+
+			CreateParameters.ProceduralMeshComponent = WarpCfg.ProceduralMeshComponent;
+			CreateParameters.ProceduralMeshComponentSectionIndex = WarpCfg.SectionIndex;
+
+			bResult = ShadersAPI.GetWarpBlendManager().Create(CreateParameters, WarpBlendInterface);
+		}
+
+		if (!bResult)
 		{
 			if (!IsEditorOperationMode())
 			{
@@ -77,7 +98,7 @@ bool FDisplayClusterProjectionMeshPolicy::HandleStartScene(class IDisplayCluster
 	return true;
 }
 
-bool FDisplayClusterProjectionMeshPolicy::GetWarpMeshAndOrigin(class IDisplayClusterViewport* InViewport, class UStaticMeshComponent*& OutMeshComponent, class USceneComponent*& OutOriginComponent)
+bool FDisplayClusterProjectionMeshPolicy::GetWarpMeshConfiguration(IDisplayClusterViewport* InViewport, FWarpMeshConfiguration& OutWarpCfg)
 {
 	check(InViewport);
 
@@ -109,20 +130,39 @@ bool FDisplayClusterProjectionMeshPolicy::GetWarpMeshAndOrigin(class IDisplayClu
 		return false;
 	}
 
-	// Get mesh component
-	UStaticMeshComponent* MeshComponent = Root->GetComponentByName<UStaticMeshComponent>(ComponentId);
-	if (!MeshComponent)
+	// Get the StaticMeshComponent
+	OutWarpCfg.StaticMeshComponent = Root->GetComponentByName<UStaticMeshComponent>(ComponentId);
+	if (OutWarpCfg.StaticMeshComponent == nullptr)
 	{
-		if (!IsEditorOperationMode())
+		// Get the ProceduralMeshComponent
+		OutWarpCfg.ProceduralMeshComponent = Root->GetComponentByName<UProceduralMeshComponent>(ComponentId);
+		if (OutWarpCfg.ProceduralMeshComponent == nullptr)
 		{
-			UE_LOG(LogDisplayClusterProjectionMesh, Warning, TEXT("Couldn't initialize mesh component '%s'"), *ComponentId);
+			if (!IsEditorOperationMode())
+			{
+				UE_LOG(LogDisplayClusterProjectionMesh, Warning, TEXT("Couldn't initialize mesh component '%s'"), *ComponentId);
+			}
+
+			return false;
 		}
 
-		return false;
+		int CfgSectionIndex;
+		if (DisplayClusterHelpers::map::template ExtractValueFromString(GetParameters(), DisplayClusterProjectionStrings::cfg::mesh::SectionIndex, CfgSectionIndex))
+		{
+			if (CfgSectionIndex >= 0)
+			{
+				UE_LOG(LogDisplayClusterProjectionMesh, Verbose, TEXT("Found SectionIndex value - '%d'"), CfgSectionIndex);
+
+				OutWarpCfg.SectionIndex = CfgSectionIndex;
+			}
+			else
+			{
+				UE_LOG(LogDisplayClusterProjectionMesh, Error, TEXT("Invalid SectionIndex value - '%d'"), CfgSectionIndex);
+			}
+		}
 	}
 
-	OutMeshComponent = MeshComponent;
-	OutOriginComponent = Root->GetRootComponent();
+	OutWarpCfg.OriginComponent = Root->GetRootComponent();
 
 	return true;
 }
@@ -135,12 +175,21 @@ UMeshComponent* FDisplayClusterProjectionMeshPolicy::GetOrCreatePreviewMeshCompo
 {
 	bOutIsRootActorComponent = true;
 
-	UStaticMeshComponent* MeshComponent = nullptr;
-	USceneComponent* OriginComponent = nullptr;
+	FWarpMeshConfiguration WarpCfg;
+	if (GetWarpMeshConfiguration(InViewport, WarpCfg))
+	{
+		if(WarpCfg.StaticMeshComponent != nullptr)
+		{
+			return WarpCfg.StaticMeshComponent;
+		}
 
-	GetWarpMeshAndOrigin(InViewport, MeshComponent, OriginComponent);
+		if (WarpCfg.ProceduralMeshComponent != nullptr)
+		{
+			return WarpCfg.ProceduralMeshComponent;
+		}
+	}
 
-	return MeshComponent;
+	return nullptr;
 }
 #endif
 
diff --git a/Engine/Plugins/Runtime/nDisplay/Source/DisplayClusterProjection/Private/Policy/Mesh/DisplayClusterProjectionMeshPolicy.h b/Engine/Plugins/Runtime/nDisplay/Source/DisplayClusterProjection/Private/Policy/Mesh/DisplayClusterProjectionMeshPolicy.h
--- a/Engine/Plugins/Runtime/nDisplay/Source/DisplayClusterProjection/Private/Policy/Mesh/DisplayClusterProjectionMeshPolicy.h
+++ b/Engine/Plugins/Runtime/nDisplay/Source/DisplayClusterProjection/Private/Policy/Mesh/DisplayClusterProjectionMeshPolicy.h
@@ -3,11 +3,14 @@
 #pragma once
 
 #include "Policy/MPCDI/DisplayClusterProjectionMPCDIPolicy.h"
-#include "Engine/Classes/Components/StaticMeshComponent.h"
 
+class IDisplayClusterViewport;
+class UStaticMeshComponent;
+class UProceduralMeshComponent;
+class USceneComponent;
 
 /**
- * Adapter for the Mesh warp
+ * Adapter for the Mesh and ProceduralMesh warp
  */
 class FDisplayClusterProjectionMeshPolicy
 	: public FDisplayClusterProjectionMPCDIPolicy
@@ -22,16 +25,24 @@ public:
 	//////////////////////////////////////////////////////////////////////////////////////////////
 	// IDisplayClusterProjectionPolicy
 	//////////////////////////////////////////////////////////////////////////////////////////////
-	virtual bool HandleStartScene(class IDisplayClusterViewport* InViewport) override;
+	virtual bool HandleStartScene(IDisplayClusterViewport* InViewport) override;
 	
 	virtual EWarpType GetWarpType() const override
 	{ return EWarpType::mesh; }
 
 	/** Parse the config data for a mesh id and try to retrieve it from the root actor. */
-	bool CreateWarpMeshInterface(class IDisplayClusterViewport* InViewport);
+	bool CreateWarpMeshInterface(IDisplayClusterViewport* InViewport);
 
 private:
-	bool GetWarpMeshAndOrigin(class IDisplayClusterViewport* InViewport, class UStaticMeshComponent* &OutMeshComponent, class USceneComponent* & OutOriginComponent);
+	struct FWarpMeshConfiguration
+	{
+		USceneComponent*          OriginComponent = nullptr;
+		UStaticMeshComponent*     StaticMeshComponent = nullptr;
+		UProceduralMeshComponent* ProceduralMeshComponent = nullptr;
+		int32 SectionIndex = 0;
+	};
+
+	bool GetWarpMeshConfiguration(IDisplayClusterViewport* InViewport, FWarpMeshConfiguration& OutWarpCfg);
 
 #if WITH_EDITOR
 protected:
@@ -43,6 +54,6 @@ protected:
 		return true;
 	}
 
-	virtual class UMeshComponent* GetOrCreatePreviewMeshComponent(class IDisplayClusterViewport* InViewport, bool& bOutIsRootActorComponent) override;
+	virtual class UMeshComponent* GetOrCreatePreviewMeshComponent(IDisplayClusterViewport* InViewport, bool& bOutIsRootActorComponent) override;
 #endif
 };
diff --git a/Engine/Plugins/Runtime/nDisplay/Source/DisplayClusterProjection/Private/Policy/VIOSO/Windows/DisplayClusterProjectionVIOSOPolicy.cpp b/Engine/Plugins/Runtime/nDisplay/Source/DisplayClusterProjection/Private/Policy/VIOSO/Windows/DisplayClusterProjectionVIOSOPolicy.cpp
--- a/Engine/Plugins/Runtime/nDisplay/Source/DisplayClusterProjection/Private/Policy/VIOSO/Windows/DisplayClusterProjectionVIOSOPolicy.cpp
+++ b/Engine/Plugins/Runtime/nDisplay/Source/DisplayClusterProjection/Private/Policy/VIOSO/Windows/DisplayClusterProjectionVIOSOPolicy.cpp
@@ -268,15 +268,19 @@ bool FDisplayClusterProjectionVIOSOPolicy::IsWarpBlendSupported()
 
 void FDisplayClusterProjectionVIOSOPolicy::ApplyWarpBlend_RenderThread(FRHICommandListImmediate& RHICmdList, const class IDisplayClusterViewportProxy* InViewportProxy)
 {
+	check(IsInRenderingThread());
+
 	if (!ImplApplyWarpBlend_RenderThread(RHICmdList, InViewportProxy))
 	{
 		// warp failed, just resolve texture to frame
-		InViewportProxy->ResolveResources(RHICmdList, EDisplayClusterViewportResourceType::InputShaderResource, InViewportProxy->GetOutputResourceType());
+		InViewportProxy->ResolveResources_RenderThread(RHICmdList, EDisplayClusterViewportResourceType::InputShaderResource, InViewportProxy->GetOutputResourceType_RenderThread());
 	}
 }
 
 bool FDisplayClusterProjectionVIOSOPolicy::ImplApplyWarpBlend_RenderThread(FRHICommandListImmediate& RHICmdList, const class IDisplayClusterViewportProxy* InViewportProxy)
 {
+	check(IsInRenderingThread());
+
 	// Get in\out remp resources ref from viewport
 	TArray<FRHITexture2D*> InputTextures, OutputTextures;
 
@@ -318,7 +322,7 @@ bool FDisplayClusterProjectionVIOSOPolicy::ImplApplyWarpBlend_RenderThread(FRHIC
 	}
 
 	// resolve warp result images from temp targetable to FrameTarget
-	return InViewportProxy->ResolveResources(RHICmdList, EDisplayClusterViewportResourceType::AdditionalTargetableResource, InViewportProxy->GetOutputResourceType());
+	return InViewportProxy->ResolveResources_RenderThread(RHICmdList, EDisplayClusterViewportResourceType::AdditionalTargetableResource, InViewportProxy->GetOutputResourceType_RenderThread());
 }
 
 //////////////////////////////////////////////////////////////////////////////////////////////
diff --git a/Engine/Plugins/Runtime/nDisplay/Source/DisplayClusterProjection/Public/DisplayClusterProjectionStrings.h b/Engine/Plugins/Runtime/nDisplay/Source/DisplayClusterProjection/Public/DisplayClusterProjectionStrings.h
--- a/Engine/Plugins/Runtime/nDisplay/Source/DisplayClusterProjection/Public/DisplayClusterProjectionStrings.h
+++ b/Engine/Plugins/Runtime/nDisplay/Source/DisplayClusterProjection/Public/DisplayClusterProjectionStrings.h
@@ -76,9 +76,8 @@ namespace DisplayClusterProjectionStrings
 		
 		namespace mesh
 		{
-			static constexpr auto FileID    = TEXT("@UESM");
-			static constexpr auto BufferID  = TEXT("@@Buf");
-			static constexpr auto Component = TEXT("mesh_component");
+			static constexpr auto Component    = TEXT("mesh_component");
+			static constexpr auto SectionIndex = TEXT("section_index");
 		}
 
 		namespace mpcdi
diff --git a/Engine/Plugins/Runtime/nDisplay/Source/DisplayClusterShaders/DisplayClusterShaders.Build.cs b/Engine/Plugins/Runtime/nDisplay/Source/DisplayClusterShaders/DisplayClusterShaders.Build.cs
--- a/Engine/Plugins/Runtime/nDisplay/Source/DisplayClusterShaders/DisplayClusterShaders.Build.cs
+++ b/Engine/Plugins/Runtime/nDisplay/Source/DisplayClusterShaders/DisplayClusterShaders.Build.cs
@@ -24,6 +24,7 @@ public class DisplayClusterShaders : ModuleRules
 				"RenderCore",
 				"RHI",
 				"Renderer",
+				"ProceduralMeshComponent",
 			});
 
 		if (Target.bBuildEditor == true)
diff --git a/Engine/Plugins/Runtime/nDisplay/Source/DisplayClusterShaders/Private/Render/DisplayClusterRenderTexture.h b/Engine/Plugins/Runtime/nDisplay/Source/DisplayClusterShaders/Private/Render/DisplayClusterRenderTexture.h
--- a/Engine/Plugins/Runtime/nDisplay/Source/DisplayClusterShaders/Private/Render/DisplayClusterRenderTexture.h
+++ b/Engine/Plugins/Runtime/nDisplay/Source/DisplayClusterShaders/Private/Render/DisplayClusterRenderTexture.h
@@ -90,7 +90,7 @@ public:
 		return TextureResource ? TextureResource->GetPixelFormat() : PF_Unknown;
 	}
 
-	virtual bool IsValid() const override
+	virtual bool IsEnabled() const override
 	{
 		const FDisplayClusterRenderTextureResource* TextureResource = GetResource();
 		if (TextureResource)
diff --git a/Engine/Plugins/Runtime/nDisplay/Source/DisplayClusterShaders/Private/Shaders/DisplayClusterShadersPostprocess_OutputRemap.cpp b/Engine/Plugins/Runtime/nDisplay/Source/DisplayClusterShaders/Private/Shaders/DisplayClusterShadersPostprocess_OutputRemap.cpp
--- a/Engine/Plugins/Runtime/nDisplay/Source/DisplayClusterShaders/Private/Shaders/DisplayClusterShadersPostprocess_OutputRemap.cpp
+++ b/Engine/Plugins/Runtime/nDisplay/Source/DisplayClusterShaders/Private/Shaders/DisplayClusterShadersPostprocess_OutputRemap.cpp
@@ -90,6 +90,7 @@ DECLARE_GPU_STAT_NAMED(nDisplay_PostProcess_OutputRemap, TEXT("nDisplay PostProc
 bool FDisplayClusterShadersPostprocess_OutputRemap::RenderPostprocess_OutputRemap(FRHICommandListImmediate& RHICmdList, FRHITexture2D* InSourceTexture, FRHITexture2D* InRenderTargetableDestTexture, const FDisplayClusterRender_MeshComponentProxy& InMeshProxy)
 {
 	check(IsInRenderingThread());
+
 	if (InSourceTexture == nullptr || InRenderTargetableDestTexture == nullptr)
 	{
 		return false;
@@ -104,16 +105,20 @@ bool FDisplayClusterShadersPostprocess_OutputRemap::RenderPostprocess_OutputRema
 		{
 			// Use simple 1:1 test mesh for shader forwarding
 			static FDisplayClusterRender_MeshComponent TestMesh_Passthrough;
-			if (TestMesh_Passthrough.GetProxy())
+
+			FDisplayClusterRender_MeshComponentProxy* TestMeshMeshComponentProxy = TestMesh_Passthrough.GetMeshComponentProxy_RenderThread();
+			if (TestMeshMeshComponentProxy != nullptr)
 			{
-				if (!TestMesh_Passthrough.GetProxy()->IsValid_RenderThread())
+				if (!TestMeshMeshComponentProxy->IsValid_RenderThread())
 				{
+					// Initialize once:
 					FDisplayClusterRender_MeshGeometry PassthroughMeshGeometry(EDisplayClusterRender_MeshGeometryCreateType::Passthrough);
-					FDisplayClusterRender_MeshComponentProxyData ProxyData(FDisplayClusterRender_MeshComponentProxyDataFunc::OutputRemapScreenSpace, PassthroughMeshGeometry);
-					TestMesh_Passthrough.GetProxy()->UpdateRHI_RenderThread(RHICmdList, &ProxyData);
+					FDisplayClusterRender_MeshComponentProxyData ProxyData(EDisplayClusterRender_MeshComponentProxyDataFunc::OutputRemapScreenSpace, PassthroughMeshGeometry);
+
+					TestMeshMeshComponentProxy->UpdateRHI_RenderThread(RHICmdList, &ProxyData);
 				}
 
-				MeshProxy = TestMesh_Passthrough.GetProxy();
+				MeshProxy = TestMeshMeshComponentProxy;
 			}
 			break;
 		}
diff --git a/Engine/Plugins/Runtime/nDisplay/Source/DisplayClusterShaders/Private/Shaders/DisplayClusterShadersWarpblend_ICVFX.cpp b/Engine/Plugins/Runtime/nDisplay/Source/DisplayClusterShaders/Private/Shaders/DisplayClusterShadersWarpblend_ICVFX.cpp
--- a/Engine/Plugins/Runtime/nDisplay/Source/DisplayClusterShaders/Private/Shaders/DisplayClusterShadersWarpblend_ICVFX.cpp
+++ b/Engine/Plugins/Runtime/nDisplay/Source/DisplayClusterShaders/Private/Shaders/DisplayClusterShadersWarpblend_ICVFX.cpp
@@ -430,6 +430,7 @@ public:
 			switch (WarpBlendParameters.WarpInterface->GetWarpGeometryType())
 			{
 				case EDisplayClusterWarpGeometryType::WarpMesh:
+				case EDisplayClusterWarpGeometryType::WarpProceduralMesh:
 				{
 					// Use mesh inseat of warp texture
 					RenderPassData.PSPermutationVector.Set<IcvfxShaderPermutation::FIcvfxShaderMeshWarp>(true);
@@ -544,8 +545,6 @@ public:
 			FPlane(0, 1, 0, 0),
 			FPlane(0, 0, 0, 1));
 
-		static const FMatrix Render2Game = Game2Render.Inverse();
-
 		FRotationTranslationMatrix CameraTranslationMatrix(Camera.CameraViewRotation, Camera.CameraViewLocation);
 
 		FMatrix WorldToCamera = CameraTranslationMatrix.Inverse();
@@ -661,15 +660,12 @@ public:
 				return true;
 
 			case EDisplayClusterWarpGeometryType::WarpMesh:
+			case EDisplayClusterWarpGeometryType::WarpProceduralMesh:
 			{
-				const FDisplayClusterRender_MeshComponent* DCWarpMeshComponent = WarpBlendParameters.WarpInterface->GetWarpMesh();
-				if (DCWarpMeshComponent)
-				{
-					const FDisplayClusterRender_MeshComponentProxy* WarpMesh = DCWarpMeshComponent->GetProxy();
-				if (WarpMesh)
+				const FDisplayClusterRender_MeshComponentProxy* WarpMeshProxy = WarpBlendParameters.WarpInterface->GetWarpMeshProxy_RenderThread();
+				if (WarpMeshProxy != nullptr)
 				{
-					return WarpMesh->BeginRender_RenderThread(RHICmdList, GraphicsPSOInit);
-				}
+					return WarpMeshProxy->BeginRender_RenderThread(RHICmdList, GraphicsPSOInit);
 				}
 				break;
 			}
@@ -693,15 +689,12 @@ public:
 				return true;
 
 			case EDisplayClusterWarpGeometryType::WarpMesh:
+			case EDisplayClusterWarpGeometryType::WarpProceduralMesh:
 			{
-				const FDisplayClusterRender_MeshComponent* DCWarpMeshComponent = WarpBlendParameters.WarpInterface->GetWarpMesh();
-				if (DCWarpMeshComponent)
+				const FDisplayClusterRender_MeshComponentProxy* WarpMeshProxy = WarpBlendParameters.WarpInterface->GetWarpMeshProxy_RenderThread();
+				if (WarpMeshProxy != nullptr)
 				{
-					const FDisplayClusterRender_MeshComponentProxy* WarpMesh = DCWarpMeshComponent->GetProxy();
-				if (WarpMesh)
-				{
-					return WarpMesh->FinishRender_RenderThread(RHICmdList);
-				}
+					return WarpMeshProxy->FinishRender_RenderThread(RHICmdList);
 				}
 				break;
 			}
diff --git a/Engine/Plugins/Runtime/nDisplay/Source/DisplayClusterShaders/Private/Shaders/DisplayClusterShadersWarpblend_MPCDI.cpp b/Engine/Plugins/Runtime/nDisplay/Source/DisplayClusterShaders/Private/Shaders/DisplayClusterShadersWarpblend_MPCDI.cpp
--- a/Engine/Plugins/Runtime/nDisplay/Source/DisplayClusterShaders/Private/Shaders/DisplayClusterShadersWarpblend_MPCDI.cpp
+++ b/Engine/Plugins/Runtime/nDisplay/Source/DisplayClusterShaders/Private/Shaders/DisplayClusterShadersWarpblend_MPCDI.cpp
@@ -258,6 +258,7 @@ public:
 			switch (WarpBlendParameters.WarpInterface->GetWarpGeometryType())
 			{
 				case EDisplayClusterWarpGeometryType::WarpMesh:
+				case EDisplayClusterWarpGeometryType::WarpProceduralMesh:
 				{
 					// Use mesh inseat of warp texture
 					RenderPassData.PSPermutationVector.Set<MpcdiShaderPermutation::FMpcdiShaderMeshWarp>(true);
@@ -345,15 +346,12 @@ public:
 				return true;
 
 			case EDisplayClusterWarpGeometryType::WarpMesh:
+			case EDisplayClusterWarpGeometryType::WarpProceduralMesh:
 			{
-				const FDisplayClusterRender_MeshComponent* DCWarpMeshComponent = WarpBlendParameters.WarpInterface->GetWarpMesh();
-				if (DCWarpMeshComponent)
+				const FDisplayClusterRender_MeshComponentProxy* WarpMeshProxy = WarpBlendParameters.WarpInterface->GetWarpMeshProxy_RenderThread();
+				if (WarpMeshProxy != nullptr)
 				{
-					const FDisplayClusterRender_MeshComponentProxy* WarpMesh = DCWarpMeshComponent->GetProxy();
-				if (WarpMesh)
-				{
-					return WarpMesh->BeginRender_RenderThread(RHICmdList, GraphicsPSOInit);
-				}
+					return WarpMeshProxy->BeginRender_RenderThread(RHICmdList, GraphicsPSOInit);
 				}
 				break;
 			}
@@ -377,15 +375,12 @@ public:
 				return true;
 
 			case EDisplayClusterWarpGeometryType::WarpMesh:
+			case EDisplayClusterWarpGeometryType::WarpProceduralMesh:
 			{
-				const FDisplayClusterRender_MeshComponent* DCWarpMeshComponent = WarpBlendParameters.WarpInterface->GetWarpMesh();
-				if (DCWarpMeshComponent)
+				const FDisplayClusterRender_MeshComponentProxy* WarpMeshProxy = WarpBlendParameters.WarpInterface->GetWarpMeshProxy_RenderThread();
+				if (WarpMeshProxy != nullptr)
 				{
-					const FDisplayClusterRender_MeshComponentProxy* WarpMesh = DCWarpMeshComponent->GetProxy();
-				if (WarpMesh)
-				{
-					return WarpMesh->FinishRender_RenderThread(RHICmdList);
-				}
+					return WarpMeshProxy->FinishRender_RenderThread(RHICmdList);
 				}
 				break;
 			}
diff --git a/Engine/Plugins/Runtime/nDisplay/Source/DisplayClusterShaders/Private/WarpBlend/DisplayClusterWarpBlend.cpp b/Engine/Plugins/Runtime/nDisplay/Source/DisplayClusterShaders/Private/WarpBlend/DisplayClusterWarpBlend.cpp
--- a/Engine/Plugins/Runtime/nDisplay/Source/DisplayClusterShaders/Private/WarpBlend/DisplayClusterWarpBlend.cpp
+++ b/Engine/Plugins/Runtime/nDisplay/Source/DisplayClusterShaders/Private/WarpBlend/DisplayClusterWarpBlend.cpp
@@ -63,6 +63,11 @@ static TAutoConsoleVariable<float> CVarMPCDIFrustumCachePrecision(
 	ECVF_RenderThreadSafe
 );
 
+bool FDisplayClusterWarpBlend::MarkWarpGeometryComponentDirty(const FName& InComponentName)
+{
+	return GeometryContext.GeometryProxy.MarkWarpGeometryComponentDirty(InComponentName);
+}
+
 bool FDisplayClusterWarpBlend::CalcFrustumContext(class IDisplayClusterViewport* InViewport, const uint32 InContextNum, const FDisplayClusterWarpEye& InEye, FDisplayClusterWarpContext& OutWarpContext)
 {
 	if (!GeometryContext.Update(InEye.WorldScale))
@@ -107,10 +112,12 @@ bool FDisplayClusterWarpBlend::CalcFrustumContext(class IDisplayClusterViewport*
 bool FDisplayClusterWarpBlend::ExportWarpMapGeometry(FMPCDIGeometryExportData* OutMeshData, uint32 InMaxDimension) const
 {
 #if WITH_EDITOR
-	return OutMeshData ? FDisplayClusterWarpBlendExporter_WarpMap::ExportWarpMap(GeometryContext.GeometryProxy.WarpMap, *OutMeshData, InMaxDimension) : false;
-#else
-	return false;
+	if (GeometryContext.GeometryProxy.WarpMapTexture.IsValid())
+	{
+		return OutMeshData ? FDisplayClusterWarpBlendExporter_WarpMap::ExportWarpMap(GeometryContext.GeometryProxy.WarpMapTexture.Get(), *OutMeshData, InMaxDimension) : false;
+	}
 #endif
+	return false;
 }
 
 
diff --git a/Engine/Plugins/Runtime/nDisplay/Source/DisplayClusterShaders/Private/WarpBlend/DisplayClusterWarpBlend.h b/Engine/Plugins/Runtime/nDisplay/Source/DisplayClusterShaders/Private/WarpBlend/DisplayClusterWarpBlend.h
--- a/Engine/Plugins/Runtime/nDisplay/Source/DisplayClusterShaders/Private/WarpBlend/DisplayClusterWarpBlend.h
+++ b/Engine/Plugins/Runtime/nDisplay/Source/DisplayClusterShaders/Private/WarpBlend/DisplayClusterWarpBlend.h
@@ -23,6 +23,8 @@ public:
 	virtual ~FDisplayClusterWarpBlend() override
 	{ }
 
+	virtual bool MarkWarpGeometryComponentDirty(const FName& InComponentName) override;
+
 	virtual bool CalcFrustumContext(class IDisplayClusterViewport* InViewport, const uint32 InContextNum, const FDisplayClusterWarpEye& InEye, FDisplayClusterWarpContext& OutWarpContext) override;
 
 	// Access to resources
@@ -37,9 +39,9 @@ public:
 		return GeometryContext.GeometryProxy.AlphaMapEmbeddedGamma; 
 	}
 
-	virtual const class FDisplayClusterRender_MeshComponent* GetWarpMesh() const override
+	virtual const FDisplayClusterRender_MeshComponentProxy* GetWarpMeshProxy_RenderThread() const override
 	{
-		return GeometryContext.GeometryProxy.WarpMesh;
+		return GeometryContext.GeometryProxy.GetWarpMeshProxy_RenderThread();
 	}
 
 	// Return current warp profile type
diff --git a/Engine/Plugins/Runtime/nDisplay/Source/DisplayClusterShaders/Private/WarpBlend/DisplayClusterWarpBlendManager.cpp b/Engine/Plugins/Runtime/nDisplay/Source/DisplayClusterShaders/Private/WarpBlend/DisplayClusterWarpBlendManager.cpp
--- a/Engine/Plugins/Runtime/nDisplay/Source/DisplayClusterShaders/Private/WarpBlend/DisplayClusterWarpBlendManager.cpp
+++ b/Engine/Plugins/Runtime/nDisplay/Source/DisplayClusterShaders/Private/WarpBlend/DisplayClusterWarpBlendManager.cpp
@@ -4,6 +4,7 @@
 
 #include "WarpBlend/Loader/DisplayClusterWarpBlendLoader_MPCDI.h"
 #include "WarpBlend/Loader/DisplayClusterWarpBlendLoader_MeshComponent.h"
+#include "WarpBlend/Loader/DisplayClusterWarpBlendLoader_ProceduralMeshComponent.h"
 
 bool FDisplayClusterWarpBlendManager::Create(const FDisplayClusterWarpBlendConstruct::FLoadMPCDIFile& InConstructParameters, TSharedPtr<IDisplayClusterWarpBlend, ESPMode::ThreadSafe>& OutWarpBlend) const
 {
@@ -20,3 +21,8 @@ bool  FDisplayClusterWarpBlendManager::Create(const FDisplayClusterWarpBlendCons
 	return FDisplayClusterWarpBlendLoader_MeshComponent::Load(InConstructParameters, OutWarpBlend);
 };
 
+bool  FDisplayClusterWarpBlendManager::Create(const FDisplayClusterWarpBlendConstruct::FAssignWarpProceduralMesh& InConstructParameters, TSharedPtr<IDisplayClusterWarpBlend, ESPMode::ThreadSafe>& OutWarpBlend) const
+{
+	return FDisplayClusterWarpBlendLoader_ProceduralMeshComponent::Load(InConstructParameters, OutWarpBlend);
+};
+
diff --git a/Engine/Plugins/Runtime/nDisplay/Source/DisplayClusterShaders/Private/WarpBlend/DisplayClusterWarpBlendManager.h b/Engine/Plugins/Runtime/nDisplay/Source/DisplayClusterShaders/Private/WarpBlend/DisplayClusterWarpBlendManager.h
--- a/Engine/Plugins/Runtime/nDisplay/Source/DisplayClusterShaders/Private/WarpBlend/DisplayClusterWarpBlendManager.h
+++ b/Engine/Plugins/Runtime/nDisplay/Source/DisplayClusterShaders/Private/WarpBlend/DisplayClusterWarpBlendManager.h
@@ -8,7 +8,8 @@ class FDisplayClusterWarpBlendManager
 	: public IDisplayClusterWarpBlendManager
 {
 public:
-	virtual bool Create(const FDisplayClusterWarpBlendConstruct::FLoadMPCDIFile& InConstructParameters,  TSharedPtr<IDisplayClusterWarpBlend, ESPMode::ThreadSafe>& OutWarpBlend) const override;
-	virtual bool Create(const FDisplayClusterWarpBlendConstruct::FLoadPFMFile& InConstructParameters,    TSharedPtr<IDisplayClusterWarpBlend, ESPMode::ThreadSafe>& OutWarpBlend) const override;
-	virtual bool Create(const FDisplayClusterWarpBlendConstruct::FAssignWarpMesh& InConstructParameters, TSharedPtr<IDisplayClusterWarpBlend, ESPMode::ThreadSafe>& OutWarpBlend) const override;
+	virtual bool Create(const FDisplayClusterWarpBlendConstruct::FLoadMPCDIFile&            InConstructParameters, TSharedPtr<IDisplayClusterWarpBlend, ESPMode::ThreadSafe>& OutWarpBlend) const override;
+	virtual bool Create(const FDisplayClusterWarpBlendConstruct::FLoadPFMFile&              InConstructParameters, TSharedPtr<IDisplayClusterWarpBlend, ESPMode::ThreadSafe>& OutWarpBlend) const override;
+	virtual bool Create(const FDisplayClusterWarpBlendConstruct::FAssignWarpMesh&           InConstructParameters, TSharedPtr<IDisplayClusterWarpBlend, ESPMode::ThreadSafe>& OutWarpBlend) const override;
+	virtual bool Create(const FDisplayClusterWarpBlendConstruct::FAssignWarpProceduralMesh& InConstructParameters, TSharedPtr<IDisplayClusterWarpBlend, ESPMode::ThreadSafe>& OutWarpBlend) const override;
 };
diff --git a/Engine/Plugins/Runtime/nDisplay/Source/DisplayClusterShaders/Private/WarpBlend/DisplayClusterWarpBlend_GeometryProxy.cpp b/Engine/Plugins/Runtime/nDisplay/Source/DisplayClusterShaders/Private/WarpBlend/DisplayClusterWarpBlend_GeometryProxy.cpp
--- a/Engine/Plugins/Runtime/nDisplay/Source/DisplayClusterShaders/Private/WarpBlend/DisplayClusterWarpBlend_GeometryProxy.cpp
+++ b/Engine/Plugins/Runtime/nDisplay/Source/DisplayClusterShaders/Private/WarpBlend/DisplayClusterWarpBlend_GeometryProxy.cpp
@@ -6,6 +6,7 @@
 #include "Render/Containers/DisplayClusterRender_MeshComponent.h"
 
 #include "WarpBlend/Math/DisplayClusterWarpBlendMath_WarpMesh.h"
+#include "WarpBlend/Math/DisplayClusterWarpBlendMath_WarpProceduralMesh.h"
 #include "WarpBlend/Math/DisplayClusterWarpBlendMath_WarpMap.h"
 
 /////////////////////////////////////////////////////////////////////////////////
@@ -18,8 +19,13 @@ bool FDisplayClusterWarpBlend_GeometryProxy::UpdateGeometry()
 	{
 	case EDisplayClusterWarpGeometryType::WarpMesh:
 		return ImplUpdateGeometry_WarpMesh();
+
+	case EDisplayClusterWarpGeometryType::WarpProceduralMesh:
+		return ImplUpdateGeometry_WarpProceduralMesh();
+
 	case EDisplayClusterWarpGeometryType::WarpMap:
 		return ImplUpdateGeometry_WarpMap();
+
 	default:
 		break;
 	}
@@ -28,11 +34,50 @@ bool FDisplayClusterWarpBlend_GeometryProxy::UpdateGeometry()
 	return false;
 }
 
+const FDisplayClusterRender_MeshComponentProxy* FDisplayClusterWarpBlend_GeometryProxy::GetWarpMeshProxy_RenderThread() const
+{
+	check(IsInRenderingThread());
+
+	switch (GeometryType)
+	{
+	case EDisplayClusterWarpGeometryType::WarpMesh:
+	case EDisplayClusterWarpGeometryType::WarpProceduralMesh:
+		return MeshComponent.IsValid() ? MeshComponent->GetMeshComponentProxy_RenderThread() : nullptr;
+
+	default:
+		break;
+	}
+
+	return nullptr;
+}
+
+bool FDisplayClusterWarpBlend_GeometryProxy::MarkWarpGeometryComponentDirty(const FName& InComponentName)
+{
+	switch (GeometryType)
+	{
+	case EDisplayClusterWarpGeometryType::WarpMesh:
+	case EDisplayClusterWarpGeometryType::WarpProceduralMesh:
+		if (MeshComponent.IsValid())
+		{
+			if (InComponentName == NAME_None || MeshComponent->EqualsMeshComponentName(InComponentName))
+			{
+				MeshComponent->MarkMeshComponentRefGeometryDirty();
+				return true;
+			}
+		}
+		break;
+	default:
+		break;
+	}
+
+	return false;
+}
+
 bool FDisplayClusterWarpBlend_GeometryProxy::ImplUpdateGeometry_WarpMap()
 {
 	bIsGeometryValid = false;
 
-	if (WarpMap && WarpMap->IsValid())
+	if (WarpMapTexture.IsValid() && WarpMapTexture->IsEnabled())
 	{
 		// Update caches
 		if ((bIsGeometryValid = bIsGeometryCacheValid) == false)
@@ -50,34 +95,28 @@ bool FDisplayClusterWarpBlend_GeometryProxy::ImplUpdateGeometry_WarpMesh()
 {
 	bIsGeometryValid = false;
 
-	if (WarpMesh == nullptr)
+	if (!MeshComponent.IsValid())
 	{
 		return false;
 	}
 
-	UStaticMeshComponent* MeshComponent = WarpMesh->MeshComponentRef.GetOrFindMeshComponent();
-	USceneComponent*    OriginComponent = WarpMesh->OriginComponentRef.GetOrFindSceneComponent();
+	UStaticMeshComponent* StaticMeshComponent = MeshComponent->GetStaticMeshComponent();
+	USceneComponent*      OriginComponent     = MeshComponent->GetOriginComponent();
 
-	if (MeshComponent == nullptr)
+	const FStaticMeshLODResources* StaticMeshLODResources = (StaticMeshComponent!=nullptr) ? MeshComponent->GetStaticMeshComponentLODResources(StaticMeshComponentLODIndex) : nullptr;
+	if (StaticMeshLODResources == nullptr)
 	{
 		// mesh deleted?
-		WarpMesh->UpdateDefferedRef();
-		bIsWarpMeshComponentLost = true;
+		MeshComponent->ReleaseProxyGeometry();
+		bIsMeshComponentLost = true;
 		return false;
 	};
 
-	UStaticMesh* StaticMesh = MeshComponent->GetStaticMesh();
-	if (StaticMesh == nullptr)
-	{
-		// mesh geometry not assigned?
-		return false;
-	}
-
 	// If StaticMesh geometry changed, update mpcdi math and RHI resources
-	if (WarpMesh->MeshComponentRef.IsMeshComponentChanged() || bIsWarpMeshComponentLost)
+	if (MeshComponent->IsMeshComponentRefGeometryDirty() || bIsMeshComponentLost)
 	{
-		AssignWarpMesh(MeshComponent, OriginComponent);
-		bIsWarpMeshComponentLost = false;
+		MeshComponent->AssignStaticMeshComponentRefs(StaticMeshComponent, OriginComponent, StaticMeshComponentLODIndex);
+		bIsMeshComponentLost = false;
 	}
 	
 	// Update caches
@@ -88,45 +127,93 @@ bool FDisplayClusterWarpBlend_GeometryProxy::ImplUpdateGeometry_WarpMesh()
 
 	if (OriginComponent)
 	{
-		FMatrix MeshToWorldMatrix = MeshComponent->GetComponentTransform().ToMatrixWithScale();
+		FMatrix MeshToWorldMatrix = StaticMeshComponent->GetComponentTransform().ToMatrixWithScale();
 		FMatrix WorldToOriginMatrix = OriginComponent->GetComponentTransform().ToInverseMatrixWithScale();
 
 		GeometryCache.GeometryToOrigin.SetFromMatrix(MeshToWorldMatrix * WorldToOriginMatrix);
 	}
 	else
 	{
-		GeometryCache.GeometryToOrigin = MeshComponent->GetRelativeTransform();
+		GeometryCache.GeometryToOrigin = StaticMeshComponent->GetRelativeTransform();
 	}
 
 	return bIsGeometryValid;
 }
 
-bool FDisplayClusterWarpBlend_GeometryProxy::AssignWarpMesh(UStaticMeshComponent* MeshComponent, USceneComponent* OriginComponent)
+bool FDisplayClusterWarpBlend_GeometryProxy::ImplUpdateGeometry_WarpProceduralMesh()
 {
-	check(IsInGameThread());
+	bIsGeometryValid = false;
+
+	if (!MeshComponent.IsValid())
+	{
+		return false;
+	}
 
-	GeometryType = EDisplayClusterWarpGeometryType::WarpMesh;
+	UProceduralMeshComponent* ProceduralMeshComponent = MeshComponent->GetProceduralMeshComponent();
+	USceneComponent*          OriginComponent         = MeshComponent->GetOriginComponent();
 
-	if (WarpMesh)
+	const FProcMeshSection* ProcMeshSection = (ProceduralMeshComponent != nullptr) ? MeshComponent->GetProceduralMeshComponentSection(ProceduralMeshComponentSectionIndex) : nullptr;
+	if (ProcMeshSection == nullptr)
 	{
-		WarpMesh->AssignMeshRefs(MeshComponent, OriginComponent);
+		// mesh deleted, lost or section not defined
+		MeshComponent->ReleaseProxyGeometry();
+		bIsMeshComponentLost = true;
+		return false;
+	};
 
-		return ImplUpdateGeometryCache_WarpMesh();
+	// If ProceduralMesh geometry changed, update mpcdi math and RHI resources
+	if (MeshComponent->IsMeshComponentRefGeometryDirty() || bIsMeshComponentLost)
+	{
+		MeshComponent->AssignProceduralMeshComponentRefs(ProceduralMeshComponent, OriginComponent, ProceduralMeshComponentSectionIndex);
+		bIsMeshComponentLost = false;
 	}
 
-	return false;
+	// Update caches
+	if ((bIsGeometryValid = bIsGeometryCacheValid) == false)
+	{
+		bIsGeometryValid = ImplUpdateGeometryCache_WarpProceduralMesh();
+	}
+
+	// These matrices were copied from LocalPlayer.cpp.
+	// They change the coordinate system from the Unreal "Game" coordinate system to the Unreal "Render" coordinate system
+	static const FMatrix Game2Render(
+		FPlane(0, 0, 1, 0),
+		FPlane(1, 0, 0, 0),
+		FPlane(0, 1, 0, 0),
+		FPlane(0, 0, 0, 1));
+
+	static const FMatrix Render2Game = Game2Render.Inverse();
+
+	// Procedural mesh coord is render, so convert to game at first
+	FMatrix GeometryToOriginMatrix;// = Render2Game;
+
+	if (OriginComponent)
+	{
+		FMatrix MeshToWorldMatrix = ProceduralMeshComponent->GetComponentTransform().ToMatrixWithScale();
+		FMatrix WorldToOriginMatrix = OriginComponent->GetComponentTransform().ToInverseMatrixWithScale();
+
+		GeometryToOriginMatrix = MeshToWorldMatrix * WorldToOriginMatrix;
+	}
+	else
+	{
+		GeometryToOriginMatrix = ProceduralMeshComponent->GetRelativeTransform().ToMatrixWithScale();
+	}
+
+	GeometryCache.GeometryToOrigin.SetFromMatrix(GeometryToOriginMatrix);
+
+	return bIsGeometryValid;
 }
 
 bool FDisplayClusterWarpBlend_GeometryProxy::ImplUpdateGeometryCache_WarpMesh()
 {
 	bIsGeometryValid = false;
 
-	if (WarpMesh)
+	if (MeshComponent.IsValid())
 	{
-		const FStaticMeshLODResources* WarpMeshResource = WarpMesh->GetStaticMeshLODResource();
-		if (WarpMeshResource != nullptr)
+		const FStaticMeshLODResources* StaticMeshLODResources = MeshComponent->GetStaticMeshComponentLODResources(StaticMeshComponentLODIndex);
+		if (StaticMeshLODResources != nullptr)
 		{
-			FDisplayClusterWarpBlendMath_WarpMesh MeshHelper(*WarpMeshResource);
+			FDisplayClusterWarpBlendMath_WarpMesh MeshHelper(*StaticMeshLODResources);
 
 			GeometryCache.AABBox = MeshHelper.CalcAABBox();
 			MeshHelper.CalcSurfaceVectors(GeometryCache.SurfaceViewNormal, GeometryCache.SurfaceViewPlane);
@@ -139,23 +226,50 @@ bool FDisplayClusterWarpBlend_GeometryProxy::ImplUpdateGeometryCache_WarpMesh()
 	return bIsGeometryValid;
 }
 
+bool FDisplayClusterWarpBlend_GeometryProxy::ImplUpdateGeometryCache_WarpProceduralMesh()
+{
+	bIsGeometryValid = false;
+
+	if (MeshComponent.IsValid())
+	{
+		const FProcMeshSection* ProcMeshSection = MeshComponent->GetProceduralMeshComponentSection(ProceduralMeshComponentSectionIndex);
+		if (ProcMeshSection != nullptr)
+		{
+			FDisplayClusterWarpBlendMath_WarpProceduralMesh ProceduralMeshHelper(*ProcMeshSection);
+
+			GeometryCache.AABBox = ProceduralMeshHelper.CalcAABBox();
+			ProceduralMeshHelper.CalcSurfaceVectors(GeometryCache.SurfaceViewNormal, GeometryCache.SurfaceViewPlane);
+
+			bIsGeometryCacheValid = true;
+			bIsGeometryValid = true;
+
+			return true;
+		}
+	}
+
+	return false;
+}
+
 bool FDisplayClusterWarpBlend_GeometryProxy::ImplUpdateGeometryCache_WarpMap()
 {
 	// WarpMap must be initialized before, outside
 	bIsGeometryValid = false;
 	
-	if (WarpMap && WarpMap->IsValid())
+	if (WarpMapTexture.IsValid() && WarpMapTexture->IsEnabled())
 	{
-		FDisplayClusterWarpBlendMath_WarpMap DataHelper(*WarpMap);
+		FDisplayClusterWarpBlendMath_WarpMap DataHelper(*(WarpMapTexture.Get()));
+
 		GeometryCache.AABBox = DataHelper.GetAABBox();
 		GeometryCache.SurfaceViewNormal = DataHelper.GetSurfaceViewNormal();
-		GeometryCache.SurfaceViewPlane = DataHelper.GetSurfaceViewPlane();
+		GeometryCache.SurfaceViewPlane  = DataHelper.GetSurfaceViewPlane();
 
 		bIsGeometryCacheValid = true;
 		bIsGeometryValid = true;
+
+		return true;
 	}
 
-	return bIsGeometryValid;
+	return false;
 }
 
 bool FDisplayClusterWarpBlend_GeometryProxy::UpdateGeometryLOD(const FIntPoint& InSizeLOD)
@@ -164,14 +278,14 @@ bool FDisplayClusterWarpBlend_GeometryProxy::UpdateGeometryLOD(const FIntPoint&
 
 	GeometryCache.IndexLOD.Empty();
 
-	if (WarpMap && WarpMap->IsValid())
+	if (WarpMapTexture.IsValid() && WarpMapTexture->IsEnabled())
 	{
 		switch (GeometryType)
 		{
 		case EDisplayClusterWarpGeometryType::WarpMap:
 		{
 			// Generate valid points for texturebox method:
-			FDisplayClusterWarpBlendMath_WarpMap DataHelper(*WarpMap);
+			FDisplayClusterWarpBlendMath_WarpMap DataHelper(*(WarpMapTexture.Get()));
 			DataHelper.BuildIndexLOD(InSizeLOD.X, InSizeLOD.Y, GeometryCache.IndexLOD);
 
 			return true;
@@ -184,30 +298,20 @@ bool FDisplayClusterWarpBlend_GeometryProxy::UpdateGeometryLOD(const FIntPoint&
 	return false;
 }
 
-void FDisplayClusterWarpBlend_GeometryProxy::ImplReleaseResources()
+const FStaticMeshLODResources* FDisplayClusterWarpBlend_GeometryProxy::GetStaticMeshComponentLODResources() const
 {
-	if (WarpMap)
-	{
-		delete WarpMap;
-		WarpMap = nullptr;
-	}
-
-	if (AlphaMap)
-	{
-		delete AlphaMap;
-		AlphaMap = nullptr;
-	}
-
-	if (BetaMap)
-	{
-		delete BetaMap;
-		BetaMap = nullptr;
-	}
+	return MeshComponent.IsValid() ? MeshComponent->GetStaticMeshComponentLODResources(StaticMeshComponentLODIndex) : nullptr;
+}
 
-	if (WarpMesh)
-	{
-		delete WarpMesh;
-		WarpMesh = nullptr;
-	}
+const FProcMeshSection* FDisplayClusterWarpBlend_GeometryProxy::GetProceduralMeshComponentSection() const
+{
+	return MeshComponent.IsValid() ? MeshComponent->GetProceduralMeshComponentSection(ProceduralMeshComponentSectionIndex) : nullptr;
 }
 
+void FDisplayClusterWarpBlend_GeometryProxy::ImplReleaseResources()
+{
+	WarpMapTexture.Reset();
+	AlphaMapTexture.Reset();
+	BetaMapTexture.Reset();
+	MeshComponent.Reset();
+}
diff --git a/Engine/Plugins/Runtime/nDisplay/Source/DisplayClusterShaders/Private/WarpBlend/DisplayClusterWarpBlend_GeometryProxy.h b/Engine/Plugins/Runtime/nDisplay/Source/DisplayClusterShaders/Private/WarpBlend/DisplayClusterWarpBlend_GeometryProxy.h
--- a/Engine/Plugins/Runtime/nDisplay/Source/DisplayClusterShaders/Private/WarpBlend/DisplayClusterWarpBlend_GeometryProxy.h
+++ b/Engine/Plugins/Runtime/nDisplay/Source/DisplayClusterShaders/Private/WarpBlend/DisplayClusterWarpBlend_GeometryProxy.h
@@ -3,12 +3,15 @@
 #pragma once
 
 #include "CoreMinimal.h"
-#include "WarpBlend/DisplayClusterWarpEnums.h"
+#include "WarpBlend\DisplayClusterWarpEnums.h"
+
+#include "Render\Containers\DisplayClusterRender_MeshComponent.h"
 
 class IDisplayClusterRenderTexture;
-class FDisplayClusterRender_MeshComponent;
-class UStaticMeshComponent;
-class USceneComponent;
+class FDisplayClusterRender_MeshComponentProxy;
+
+struct FStaticMeshLODResources;
+struct FProcMeshSection;
 
 class FDisplayClusterWarpBlend_GeometryProxy
 {
@@ -25,13 +28,13 @@ public:
 		switch (TextureType)
 		{
 		case EDisplayClusterWarpBlendTextureType::WarpMap:
-			return WarpMap;
+			return WarpMapTexture.Get();
 
 		case EDisplayClusterWarpBlendTextureType::AlphaMap:
-			return AlphaMap;
+			return AlphaMapTexture.Get();
 
 		case EDisplayClusterWarpBlendTextureType::BetaMap:
-			return BetaMap;
+			return BetaMapTexture.Get();
 
 		default:
 			break;
@@ -39,19 +42,22 @@ public:
 		return nullptr;
 	}
 
+	bool MarkWarpGeometryComponentDirty(const FName& InComponentName);
+
 	bool UpdateGeometry();
 	bool UpdateGeometryLOD(const FIntPoint& InSizeLOD);
 
-	bool AssignWarpMesh(UStaticMeshComponent* MeshComponent, USceneComponent* OriginComponent);
+	const FDisplayClusterRender_MeshComponentProxy* GetWarpMeshProxy_RenderThread() const;
+
+	const FStaticMeshLODResources* GetStaticMeshComponentLODResources() const;
+	const FProcMeshSection* GetProceduralMeshComponentSection() const;
 
 private:
 	bool ImplUpdateGeometry_WarpMesh();
+	bool ImplUpdateGeometry_WarpProceduralMesh();
 	bool ImplUpdateGeometry_WarpMap();
 
 public:
-	// for huge warp mesh geometry, change this value, and use LOD geometry in frustum math
-	int32 WarpMeshLODIndex = 0;
-
 	bool bIsGeometryCacheValid = false;
 	bool bIsGeometryValid = false;
 	EDisplayClusterWarpGeometryType GeometryType = EDisplayClusterWarpGeometryType::Invalid;
@@ -61,16 +67,24 @@ private:
 
 public:
 	// Render resources:
-	IDisplayClusterRenderTexture* WarpMap  = nullptr;
-	IDisplayClusterRenderTexture* AlphaMap = nullptr;
-	IDisplayClusterRenderTexture* BetaMap  = nullptr;
+	TUniquePtr<IDisplayClusterRenderTexture> WarpMapTexture;
+	TUniquePtr<IDisplayClusterRenderTexture> AlphaMapTexture;
+	TUniquePtr<IDisplayClusterRenderTexture> BetaMapTexture;
 
-	FDisplayClusterRender_MeshComponent* WarpMesh = nullptr;
+	TUniquePtr<FDisplayClusterRender_MeshComponent> MeshComponent;
+
+	// for huge warp mesh geometry, change this value, and use LOD geometry in frustum math
+	int32 StaticMeshComponentLODIndex = 0;
+	int32 ProceduralMeshComponentSectionIndex = 0;
 
 	float AlphaMapEmbeddedGamma = 1.f;
 
+private:
+	bool bIsMeshComponentLost = false;
+
 private:
 	bool ImplUpdateGeometryCache_WarpMesh();
+	bool ImplUpdateGeometryCache_WarpProceduralMesh();
 	bool ImplUpdateGeometryCache_WarpMap();
 
 public:
@@ -88,8 +102,5 @@ public:
 	};
 
 	FDisplayClusterWarpBlend_GeometryCache GeometryCache;
-
-private:
-	bool bIsWarpMeshComponentLost = false;
 };
 
diff --git a/Engine/Plugins/Runtime/nDisplay/Source/DisplayClusterShaders/Private/WarpBlend/Exporter/DisplayClusterWarpBlendExporter_WarpMap.cpp b/Engine/Plugins/Runtime/nDisplay/Source/DisplayClusterShaders/Private/WarpBlend/Exporter/DisplayClusterWarpBlendExporter_WarpMap.cpp
--- a/Engine/Plugins/Runtime/nDisplay/Source/DisplayClusterShaders/Private/WarpBlend/Exporter/DisplayClusterWarpBlendExporter_WarpMap.cpp
+++ b/Engine/Plugins/Runtime/nDisplay/Source/DisplayClusterShaders/Private/WarpBlend/Exporter/DisplayClusterWarpBlendExporter_WarpMap.cpp
@@ -7,7 +7,7 @@
 
 bool FDisplayClusterWarpBlendExporter_WarpMap::ExportWarpMap(IDisplayClusterRenderTexture* InWarpMap, struct FMPCDIGeometryExportData& Dst, uint32 InMaxDimension)
 {
-	if (InWarpMap == nullptr || InWarpMap->IsValid() == false)
+	if (InWarpMap == nullptr || !InWarpMap->IsEnabled())
 	{
 		return false;
 	}
diff --git a/Engine/Plugins/Runtime/nDisplay/Source/DisplayClusterShaders/Private/WarpBlend/Loader/DisplayClusterWarpBlendLoader_MPCDI.cpp b/Engine/Plugins/Runtime/nDisplay/Source/DisplayClusterShaders/Private/WarpBlend/Loader/DisplayClusterWarpBlendLoader_MPCDI.cpp
--- a/Engine/Plugins/Runtime/nDisplay/Source/DisplayClusterShaders/Private/WarpBlend/Loader/DisplayClusterWarpBlendLoader_MPCDI.cpp
+++ b/Engine/Plugins/Runtime/nDisplay/Source/DisplayClusterShaders/Private/WarpBlend/Loader/DisplayClusterWarpBlendLoader_MPCDI.cpp
@@ -54,33 +54,43 @@ struct FMPCDIRegionLoader
 		WarpBlend->GeometryContext.RegionMatrix = RegionMatrix;
 
 		FDisplayClusterWarpBlend_GeometryProxy& Proxy = WarpBlend->GeometryContext.GeometryProxy;
-		Proxy.AlphaMap = AlphaMap;
-		AlphaMap = nullptr;
 
-		Proxy.AlphaMapEmbeddedGamma = AlphaMapGammaEmbedded;
+		if (AlphaMap != nullptr)
+		{
+			Proxy.AlphaMapEmbeddedGamma = AlphaMapGammaEmbedded;
 
-		Proxy.BetaMap = BetaMap;
-		BetaMap = nullptr;
+			Proxy.AlphaMapTexture = TUniquePtr<IDisplayClusterRenderTexture>(AlphaMap);
+			AlphaMap = nullptr;
+		}
+
+		if (BetaMap != nullptr)
+		{
+			Proxy.BetaMapTexture = TUniquePtr<IDisplayClusterRenderTexture>(BetaMap);
+			BetaMap = nullptr;
+		}
 
-		Proxy.WarpMap = WarpMap;
-		WarpMap = nullptr;
+		if (WarpMap != nullptr)
+		{
+			Proxy.WarpMapTexture = TUniquePtr<IDisplayClusterRenderTexture>(WarpMap);
+			WarpMap = nullptr;
+		}
 
 		return WarpBlend;
 	}
 
 	void ReleaseResources()
 	{
-		if (WarpMap)
+		if (WarpMap != nullptr)
 		{
 			delete WarpMap;
 			WarpMap = nullptr;
 		}
-		if (AlphaMap)
+		if (AlphaMap != nullptr)
 		{
 			delete AlphaMap;
 			AlphaMap = nullptr;
 		}
-		if (BetaMap)
+		if (BetaMap != nullptr)
 		{
 			delete BetaMap;
 			BetaMap = nullptr;
@@ -125,7 +135,7 @@ struct FMPCDIRegionLoader
 			{
 				WarpMap = FDisplayClusterWarpBlendLoader_Texture::CreateWarpMap(ProfileType, WarpMapSource);
 
-				if (WarpMap)
+				if (WarpMap != nullptr)
 				{
 					return true;
 				}
diff --git a/Engine/Plugins/Runtime/nDisplay/Source/DisplayClusterShaders/Private/WarpBlend/Loader/DisplayClusterWarpBlendLoader_MeshComponent.cpp b/Engine/Plugins/Runtime/nDisplay/Source/DisplayClusterShaders/Private/WarpBlend/Loader/DisplayClusterWarpBlendLoader_MeshComponent.cpp
--- a/Engine/Plugins/Runtime/nDisplay/Source/DisplayClusterShaders/Private/WarpBlend/Loader/DisplayClusterWarpBlendLoader_MeshComponent.cpp
+++ b/Engine/Plugins/Runtime/nDisplay/Source/DisplayClusterShaders/Private/WarpBlend/Loader/DisplayClusterWarpBlendLoader_MeshComponent.cpp
@@ -26,12 +26,14 @@ bool FDisplayClusterWarpBlendLoader_MeshComponent::Load(const FDisplayClusterWar
 		//ok, Create and initialize warpblend interface:
 		TSharedPtr<FDisplayClusterWarpBlend, ESPMode::ThreadSafe> WarpBlend = MakeShared<FDisplayClusterWarpBlend, ESPMode::ThreadSafe>();
 
-		WarpBlend->GeometryContext.GeometryProxy.GeometryType = EDisplayClusterWarpGeometryType::WarpMesh;
 		WarpBlend->GeometryContext.ProfileType = EDisplayClusterWarpProfileType::warp_A3D;
 
 		FDisplayClusterWarpBlend_GeometryProxy& Proxy = WarpBlend->GeometryContext.GeometryProxy;
-		Proxy.WarpMesh = new FDisplayClusterRender_MeshComponent();
-		Proxy.WarpMesh->AssignMeshRefs(InParameters.MeshComponent, InParameters.OriginComponent);
+		Proxy.MeshComponent = MakeUnique<FDisplayClusterRender_MeshComponent>();
+		Proxy.MeshComponent->AssignStaticMeshComponentRefs(InParameters.MeshComponent, InParameters.OriginComponent, InParameters.StaticMeshComponentLODIndex);
+		Proxy.StaticMeshComponentLODIndex = InParameters.StaticMeshComponentLODIndex;
+
+		Proxy.GeometryType = EDisplayClusterWarpGeometryType::WarpMesh;
 
 		OutWarpBlend = WarpBlend;
 		return true;
diff --git a/Engine/Plugins/Runtime/nDisplay/Source/DisplayClusterShaders/Private/WarpBlend/Loader/DisplayClusterWarpBlendLoader_ProceduralMeshComponent.cpp b/Engine/Plugins/Runtime/nDisplay/Source/DisplayClusterShaders/Private/WarpBlend/Loader/DisplayClusterWarpBlendLoader_ProceduralMeshComponent.cpp
new file mode 100644
--- /dev/null
+++ b/Engine/Plugins/Runtime/nDisplay/Source/DisplayClusterShaders/Private/WarpBlend/Loader/DisplayClusterWarpBlendLoader_ProceduralMeshComponent.cpp
@@ -0,0 +1,49 @@
+// Copyright Epic Games, Inc. All Rights Reserved.
+
+#include "DisplayClusterWarpBlendLoader_ProceduralMeshComponent.h"
+
+#include "DisplayClusterShadersLog.h"
+
+#include "Misc/DisplayClusterHelpers.h"
+#include "Misc/FileHelper.h"
+
+#include "Stats/Stats.h"
+#include "Engine/Engine.h"
+#include "Misc/ConfigCacheIni.h"
+#include "Misc/Paths.h"
+
+#include "WarpBlend/Loader/DisplayClusterWarpBlendLoader_Texture.h"
+#include "WarpBlend/DisplayClusterWarpBlend.h"
+#include "WarpBlend/DisplayClusterWarpBlend_GeometryContext.h"
+#include "WarpBlend/DisplayClusterWarpBlend_GeometryProxy.h"
+
+#include "Render/Containers/DisplayClusterRender_MeshComponent.h"
+
+bool FDisplayClusterWarpBlendLoader_ProceduralMeshComponent::Load(const FDisplayClusterWarpBlendConstruct::FAssignWarpProceduralMesh& InParameters, TSharedPtr<IDisplayClusterWarpBlend, ESPMode::ThreadSafe>& OutWarpBlend)
+{
+	if (InParameters.ProceduralMeshComponent != nullptr)
+	{
+		//ok, Create and initialize warpblend interface:
+		TSharedPtr<FDisplayClusterWarpBlend, ESPMode::ThreadSafe> WarpBlend = MakeShared<FDisplayClusterWarpBlend, ESPMode::ThreadSafe>();
+
+		WarpBlend->GeometryContext.ProfileType = EDisplayClusterWarpProfileType::warp_A3D;
+
+		FDisplayClusterWarpBlend_GeometryProxy& Proxy = WarpBlend->GeometryContext.GeometryProxy;
+
+
+		// Assign procedural mesh
+		Proxy.MeshComponent = MakeUnique<FDisplayClusterRender_MeshComponent>();
+		Proxy.MeshComponent->AssignProceduralMeshComponentRefs(InParameters.ProceduralMeshComponent, InParameters.OriginComponent, InParameters.ProceduralMeshComponentSectionIndex);
+		Proxy.ProceduralMeshComponentSectionIndex = InParameters.ProceduralMeshComponentSectionIndex;
+
+		Proxy.GeometryType = EDisplayClusterWarpGeometryType::WarpProceduralMesh;
+
+		OutWarpBlend = WarpBlend;
+		return true;
+	}
+
+	return false;
+}
+
+
+
diff --git a/Engine/Plugins/Runtime/nDisplay/Source/DisplayClusterShaders/Private/WarpBlend/Loader/DisplayClusterWarpBlendLoader_ProceduralMeshComponent.h b/Engine/Plugins/Runtime/nDisplay/Source/DisplayClusterShaders/Private/WarpBlend/Loader/DisplayClusterWarpBlendLoader_ProceduralMeshComponent.h
new file mode 100644
--- /dev/null
+++ b/Engine/Plugins/Runtime/nDisplay/Source/DisplayClusterShaders/Private/WarpBlend/Loader/DisplayClusterWarpBlendLoader_ProceduralMeshComponent.h
@@ -0,0 +1,15 @@
+// Copyright Epic Games, Inc. All Rights Reserved.
+
+#pragma once
+
+#include "CoreMinimal.h"
+
+#include "WarpBlend/IDisplayClusterWarpBlendManager.h"
+
+class IDisplayClusterWarpBlend;
+
+class FDisplayClusterWarpBlendLoader_ProceduralMeshComponent
+{
+public:
+	static bool Load(const FDisplayClusterWarpBlendConstruct::FAssignWarpProceduralMesh& InParameters, TSharedPtr<IDisplayClusterWarpBlend, ESPMode::ThreadSafe>& OutWarpBlend);
+};
diff --git a/Engine/Plugins/Runtime/nDisplay/Source/DisplayClusterShaders/Private/WarpBlend/Math/DisplayClusterWarpBlendMath_Frustum.cpp b/Engine/Plugins/Runtime/nDisplay/Source/DisplayClusterShaders/Private/WarpBlend/Math/DisplayClusterWarpBlendMath_Frustum.cpp
--- a/Engine/Plugins/Runtime/nDisplay/Source/DisplayClusterShaders/Private/WarpBlend/Math/DisplayClusterWarpBlendMath_Frustum.cpp
+++ b/Engine/Plugins/Runtime/nDisplay/Source/DisplayClusterShaders/Private/WarpBlend/Math/DisplayClusterWarpBlendMath_Frustum.cpp
@@ -10,6 +10,7 @@
 #include "Render/Viewport/IDisplayClusterViewport.h"
 
 #include "StaticMeshResources.h"
+#include "ProceduralMeshComponent.h"
 
 FMatrix GetProjectionMatrixAssymetric(IDisplayClusterViewport* InViewport, const uint32 InContextNum, const FDisplayClusterWarpEye& InEye, const FDisplayClusterWarpContext& InContext)
 {
@@ -33,8 +34,6 @@ void FDisplayClusterWarpBlendMath_Frustum::ImplBuildFrustum(IDisplayClusterViewp
 		FPlane(0, 1, 0, 0),
 		FPlane(0, 0, 0, 1));
 
-	static const FMatrix Render2Game = Game2Render.Inverse();
-
 	// Compute warp projection matrix
 	Frustum.OutCameraRotation = Local2World.Rotator();
 	Frustum.OutCameraOrigin = Local2World.GetOrigin();
@@ -52,15 +51,15 @@ void FDisplayClusterWarpBlendMath_Frustum::ImplBuildFrustum(IDisplayClusterViewp
 
 bool FDisplayClusterWarpBlendMath_Frustum::ImplCalcFrustum_FULL_WarpMap()
 {
-	if(GeometryContext.GeometryProxy.WarpMap == nullptr)
+	if(!GeometryContext.GeometryProxy.WarpMapTexture.IsValid())
 	{
 		return false;
 	}
 
 	bool bAllPointsInFrustum = true;
 
-	const int32 PointsAmmount = GeometryContext.GeometryProxy.WarpMap->GetTotalPoints();
-	const FVector4f* SourcePts = (FVector4f*)(GeometryContext.GeometryProxy.WarpMap->GetData());
+	const int32 PointsAmmount = GeometryContext.GeometryProxy.WarpMapTexture->GetTotalPoints();
+	const FVector4f* SourcePts = (FVector4f*)(GeometryContext.GeometryProxy.WarpMapTexture->GetData());
 
 	// Search a camera space frustum
 	for (int32 PointIndex = 0; PointIndex < PointsAmmount; ++PointIndex)
@@ -77,20 +76,20 @@ bool FDisplayClusterWarpBlendMath_Frustum::ImplCalcFrustum_FULL_WarpMap()
 
 bool FDisplayClusterWarpBlendMath_Frustum::ImplCalcFrustum_FULL_WarpMesh()
 {
-	if (!GeometryContext.GeometryProxy.WarpMesh)
+	if (!GeometryContext.GeometryProxy.MeshComponent.IsValid())
 	{
 		return false;
 	}
 
 	bool bAllPointsInFrustum = true;
 
-	const FStaticMeshLODResources* WarpMeshResource = GeometryContext.GeometryProxy.WarpMesh->GetStaticMeshLODResource();
-	if (WarpMeshResource == nullptr)
+	const FStaticMeshLODResources* StaticMeshLODResources = GeometryContext.GeometryProxy.GetStaticMeshComponentLODResources();
+	if (StaticMeshLODResources == nullptr)
 	{
 		return false;
 	}
 
-	const FPositionVertexBuffer& VertexPosition = WarpMeshResource->VertexBuffers.PositionVertexBuffer;
+	const FPositionVertexBuffer& VertexPosition = StaticMeshLODResources->VertexBuffers.PositionVertexBuffer;
 	for (uint32 i = 0; i < VertexPosition.GetNumVertices(); i++)
 	{
 		if (GetProjectionClip(FVector4(VertexPosition.VertexPosition(i), 1)) == false)
@@ -102,10 +101,36 @@ bool FDisplayClusterWarpBlendMath_Frustum::ImplCalcFrustum_FULL_WarpMesh()
 	return bAllPointsInFrustum;
 }
 
+bool FDisplayClusterWarpBlendMath_Frustum::ImplCalcFrustum_FULL_WarpProceduralMesh()
+{
+	if (!GeometryContext.GeometryProxy.MeshComponent.IsValid())
+	{
+		return false;
+	}
+
+	bool bAllPointsInFrustum = true;
+
+	const FProcMeshSection* ProcMeshSection = GeometryContext.GeometryProxy.MeshComponent->GetProceduralMeshComponentSection();
+	if (ProcMeshSection == nullptr)
+	{
+		return false;
+	}
+
+	for (const FProcMeshVertex& VertexIt : ProcMeshSection->ProcVertexBuffer)
+	{
+		if (GetProjectionClip(FVector4(VertexIt.Position, 1)) == false)
+		{
+			bAllPointsInFrustum = false;
+		}
+	}
+
+	return bAllPointsInFrustum;
+}
+
 // Calculate better performance for PFM
 bool FDisplayClusterWarpBlendMath_Frustum::ImplCalcFrustum_LOD_WarpMap()
 {
-	if (GeometryContext.GeometryProxy.WarpMap == nullptr)
+	if (!GeometryContext.GeometryProxy.WarpMapTexture.IsValid())
 	{
 		return false;
 	}
@@ -121,8 +146,8 @@ bool FDisplayClusterWarpBlendMath_Frustum::ImplCalcFrustum_LOD_WarpMap()
 		GeometryContext.GeometryProxy.UpdateGeometryLOD(FIntPoint(WarpMapLODRatio));
 	}
 
-	const int32 PointsAmmount = GeometryContext.GeometryProxy.WarpMap->GetTotalPoints();
-	const FVector4f* SourcePts = (FVector4f*)(GeometryContext.GeometryProxy.WarpMap->GetData());
+	const int32 PointsAmmount = GeometryContext.GeometryProxy.WarpMapTexture->GetTotalPoints();
+	const FVector4f* SourcePts = (FVector4f*)(GeometryContext.GeometryProxy.WarpMapTexture->GetData());
 
 	// Search a camera space frustum
 	for (const int32& It: IndexLOD)
diff --git a/Engine/Plugins/Runtime/nDisplay/Source/DisplayClusterShaders/Private/WarpBlend/Math/DisplayClusterWarpBlendMath_Frustum.h b/Engine/Plugins/Runtime/nDisplay/Source/DisplayClusterShaders/Private/WarpBlend/Math/DisplayClusterWarpBlendMath_Frustum.h
--- a/Engine/Plugins/Runtime/nDisplay/Source/DisplayClusterShaders/Private/WarpBlend/Math/DisplayClusterWarpBlendMath_Frustum.h
+++ b/Engine/Plugins/Runtime/nDisplay/Source/DisplayClusterShaders/Private/WarpBlend/Math/DisplayClusterWarpBlendMath_Frustum.h
@@ -355,6 +355,9 @@ private:
 		case EDisplayClusterWarpGeometryType::WarpMesh:
 			return ImplCalcFrustum_FULL_WarpMesh();
 
+		case EDisplayClusterWarpGeometryType::WarpProceduralMesh:
+			return ImplCalcFrustum_FULL_WarpProceduralMesh();
+
 		default:
 			break;
 		}
@@ -369,8 +372,12 @@ private:
 			return ImplCalcFrustum_LOD_WarpMap();
 
 		case EDisplayClusterWarpGeometryType::WarpMesh:
-			// LOD not implemened for WarpMesh, use full
+			// LOD not implemented for WarpMesh, use full
 			return ImplCalcFrustum_FULL_WarpMesh();
+
+		case EDisplayClusterWarpGeometryType::WarpProceduralMesh:
+			// LOD not implemented for WarpProceduralMesh, use full
+			return ImplCalcFrustum_FULL_WarpProceduralMesh();
 		default:
 			break;
 		}
@@ -381,6 +388,7 @@ private:
 	// Frustum build logics:
 	bool ImplCalcFrustum_FULL_WarpMap();
 	bool ImplCalcFrustum_FULL_WarpMesh();
+	bool ImplCalcFrustum_FULL_WarpProceduralMesh();
 
 	bool ImplCalcFrustum_LOD_WarpMap();
 
diff --git a/Engine/Plugins/Runtime/nDisplay/Source/DisplayClusterShaders/Private/WarpBlend/Math/DisplayClusterWarpBlendMath_WarpProceduralMesh.h b/Engine/Plugins/Runtime/nDisplay/Source/DisplayClusterShaders/Private/WarpBlend/Math/DisplayClusterWarpBlendMath_WarpProceduralMesh.h
new file mode 100644
--- /dev/null
+++ b/Engine/Plugins/Runtime/nDisplay/Source/DisplayClusterShaders/Private/WarpBlend/Math/DisplayClusterWarpBlendMath_WarpProceduralMesh.h
@@ -0,0 +1,80 @@
+// Copyright Epic Games, Inc. All Rights Reserved.
+
+#pragma once
+
+#include "CoreMinimal.h"
+
+#include "ProceduralMeshComponent.h"
+
+class FDisplayClusterWarpBlendMath_WarpProceduralMesh
+{
+public:
+	FDisplayClusterWarpBlendMath_WarpProceduralMesh(const FProcMeshSection& InProcMeshSection)
+		: ProcMeshSection(InProcMeshSection)
+	{ }
+
+public:
+	FBox CalcAABBox()
+	{
+		FBox AABBox = FBox(FVector(FLT_MAX, FLT_MAX, FLT_MAX), FVector(-FLT_MAX, -FLT_MAX, -FLT_MAX));
+
+		for (const FProcMeshVertex& VertexIt : ProcMeshSection.ProcVertexBuffer)
+		{
+			const FVector& Pts = VertexIt.Position;
+
+			AABBox.Min.X = FMath::Min(AABBox.Min.X, Pts.X);
+			AABBox.Min.Y = FMath::Min(AABBox.Min.Y, Pts.Y);
+			AABBox.Min.Z = FMath::Min(AABBox.Min.Z, Pts.Z);
+
+			AABBox.Max.X = FMath::Max(AABBox.Max.X, Pts.X);
+			AABBox.Max.Y = FMath::Max(AABBox.Max.Y, Pts.Y);
+			AABBox.Max.Z = FMath::Max(AABBox.Max.Z, Pts.Z);
+		}
+
+		return AABBox;
+	}
+
+	void CalcSurfaceVectors(FVector& OutSurfaceViewNormal, FVector& OutSurfaceViewPlane)
+	{
+		// Calc static normal and plane
+		const int32 IdxNum = ProcMeshSection.ProcIndexBuffer.Num();
+		const int32 TriNum = IdxNum / 3;
+
+		if (TriNum <= 0)
+		{
+			return;
+		}
+
+		// Calc static normal and plane
+		bool bAverageNormalValid = false;
+		FVector AverageNormal;
+
+		for (int32 TriIdx = 0; TriIdx < TriNum; ++TriIdx)
+		{
+			const uint32& Index0 = ProcMeshSection.ProcIndexBuffer[TriIdx * 3 + 0];
+			const uint32& Index1 = ProcMeshSection.ProcIndexBuffer[TriIdx * 3 + 1];
+			const uint32& Index2 = ProcMeshSection.ProcIndexBuffer[TriIdx * 3 + 2];
+
+			const FVector& Pts0 = ProcMeshSection.ProcVertexBuffer[Index0].Position;
+			const FVector& Pts1 = ProcMeshSection.ProcVertexBuffer[Index1].Position;
+			const FVector& Pts2 = ProcMeshSection.ProcVertexBuffer[Index2].Position;
+
+			const FVector N1 = Pts0 - Pts1;
+			const FVector N2 = Pts2 - Pts1;
+
+			const FVector N = FVector::CrossProduct(N2, N1).GetSafeNormal();
+
+			AverageNormal = bAverageNormalValid ? (N + AverageNormal).GetSafeNormal() : N;
+
+			bAverageNormalValid = true;
+		}
+
+		OutSurfaceViewNormal = AverageNormal;
+
+		//@todo: MeshSurfaceViewPlane not implemented, use MeshSurfaceViewNormal
+		OutSurfaceViewPlane = OutSurfaceViewNormal;
+	}
+
+private:
+	const FProcMeshSection& ProcMeshSection;
+};
diff --git a/Engine/Plugins/Runtime/nDisplay/Source/DisplayClusterShaders/Public/Render/IDisplayClusterRenderTexture.h b/Engine/Plugins/Runtime/nDisplay/Source/DisplayClusterShaders/Public/Render/IDisplayClusterRenderTexture.h
--- a/Engine/Plugins/Runtime/nDisplay/Source/DisplayClusterShaders/Public/Render/IDisplayClusterRenderTexture.h
+++ b/Engine/Plugins/Runtime/nDisplay/Source/DisplayClusterShaders/Public/Render/IDisplayClusterRenderTexture.h
@@ -32,7 +32,7 @@ public:
 	virtual uint32_t GetHeight() const = 0;
 	virtual uint32_t GetTotalPoints() const = 0;
 	virtual EPixelFormat GetPixelFormat() const = 0;
-	virtual bool IsValid() const = 0;
+	virtual bool IsEnabled() const = 0;
 
 	virtual FRHITexture* GetRHITexture() const = 0;
 };
diff --git a/Engine/Plugins/Runtime/nDisplay/Source/DisplayClusterShaders/Public/WarpBlend/DisplayClusterWarpEnums.h b/Engine/Plugins/Runtime/nDisplay/Source/DisplayClusterShaders/Public/WarpBlend/DisplayClusterWarpEnums.h
--- a/Engine/Plugins/Runtime/nDisplay/Source/DisplayClusterShaders/Public/WarpBlend/DisplayClusterWarpEnums.h
+++ b/Engine/Plugins/Runtime/nDisplay/Source/DisplayClusterShaders/Public/WarpBlend/DisplayClusterWarpEnums.h
@@ -26,6 +26,9 @@ enum class EDisplayClusterWarpGeometryType : uint8
 	// Use mesh component as geometry source
 	WarpMesh,
 
+	// Use procedural mesh component as geometry source
+	WarpProceduralMesh,
+
 	// Undefined source
 	Invalid
 };
diff --git a/Engine/Plugins/Runtime/nDisplay/Source/DisplayClusterShaders/Public/WarpBlend/IDisplayClusterWarpBlend.h b/Engine/Plugins/Runtime/nDisplay/Source/DisplayClusterShaders/Public/WarpBlend/IDisplayClusterWarpBlend.h
--- a/Engine/Plugins/Runtime/nDisplay/Source/DisplayClusterShaders/Public/WarpBlend/IDisplayClusterWarpBlend.h
+++ b/Engine/Plugins/Runtime/nDisplay/Source/DisplayClusterShaders/Public/WarpBlend/IDisplayClusterWarpBlend.h
@@ -32,7 +32,7 @@ public:
 	virtual class FRHITexture* GetTexture(EDisplayClusterWarpBlendTextureType TextureType) const = 0;
 	virtual float              GetAlphaMapEmbeddedGamma() const = 0;
 
-	virtual const class FDisplayClusterRender_MeshComponent* GetWarpMesh() const = 0;
+	virtual const class FDisplayClusterRender_MeshComponentProxy* GetWarpMeshProxy_RenderThread() const = 0;
 
 	// Return current warp profile type
 	virtual EDisplayClusterWarpProfileType  GetWarpProfileType() const = 0;
@@ -41,4 +41,13 @@ public:
 	virtual EDisplayClusterWarpGeometryType GetWarpGeometryType() const = 0;
 
 	virtual bool ExportWarpMapGeometry(FMPCDIGeometryExportData* OutMeshData, uint32 InMaxDimension = 0) const = 0;
+
+	/**
+	* Mark internal component ref as dirty for geometry update
+	*
+	* @param InComponentName - (optional) the name of the internal geometry ref component. Empty string for any component name
+	*
+	* @return - true, if there is a marked.
+	*/
+	virtual bool MarkWarpGeometryComponentDirty(const FName& InComponentName) = 0;
 };
diff --git a/Engine/Plugins/Runtime/nDisplay/Source/DisplayClusterShaders/Public/WarpBlend/IDisplayClusterWarpBlendManager.h b/Engine/Plugins/Runtime/nDisplay/Source/DisplayClusterShaders/Public/WarpBlend/IDisplayClusterWarpBlendManager.h
--- a/Engine/Plugins/Runtime/nDisplay/Source/DisplayClusterShaders/Public/WarpBlend/IDisplayClusterWarpBlendManager.h
+++ b/Engine/Plugins/Runtime/nDisplay/Source/DisplayClusterShaders/Public/WarpBlend/IDisplayClusterWarpBlendManager.h
@@ -37,6 +37,14 @@ struct FDisplayClusterWarpBlendConstruct
 	{
 		class UStaticMeshComponent* MeshComponent   = nullptr;
 		class USceneComponent*      OriginComponent = nullptr;
+		int32 StaticMeshComponentLODIndex = 0;
+	};
+
+	struct FAssignWarpProceduralMesh
+	{
+		class UProceduralMeshComponent* ProceduralMeshComponent = nullptr;
+		class USceneComponent* OriginComponent = nullptr;
+		int32 ProceduralMeshComponentSectionIndex = 0;
 	};
 };
 
@@ -47,7 +55,8 @@ public:
 	virtual ~IDisplayClusterWarpBlendManager() = default;
 
 public:
-	virtual bool Create(const FDisplayClusterWarpBlendConstruct::FLoadMPCDIFile&  InConstructParameters, TSharedPtr<IDisplayClusterWarpBlend, ESPMode::ThreadSafe>& OutWarpBlend) const = 0;
-	virtual bool Create(const FDisplayClusterWarpBlendConstruct::FLoadPFMFile&    InConstructParameters, TSharedPtr<IDisplayClusterWarpBlend, ESPMode::ThreadSafe>& OutWarpBlend) const = 0;
-	virtual bool Create(const FDisplayClusterWarpBlendConstruct::FAssignWarpMesh& InConstructParameters, TSharedPtr<IDisplayClusterWarpBlend, ESPMode::ThreadSafe>& OutWarpBlend) const = 0;
+	virtual bool Create(const FDisplayClusterWarpBlendConstruct::FLoadMPCDIFile&            InConstructParameters, TSharedPtr<IDisplayClusterWarpBlend, ESPMode::ThreadSafe>& OutWarpBlend) const = 0;
+	virtual bool Create(const FDisplayClusterWarpBlendConstruct::FLoadPFMFile&              InConstructParameters, TSharedPtr<IDisplayClusterWarpBlend, ESPMode::ThreadSafe>& OutWarpBlend) const = 0;
+	virtual bool Create(const FDisplayClusterWarpBlendConstruct::FAssignWarpMesh&           InConstructParameters, TSharedPtr<IDisplayClusterWarpBlend, ESPMode::ThreadSafe>& OutWarpBlend) const = 0;
+	virtual bool Create(const FDisplayClusterWarpBlendConstruct::FAssignWarpProceduralMesh& InConstructParameters, TSharedPtr<IDisplayClusterWarpBlend, ESPMode::ThreadSafe>& OutWarpBlend) const = 0;
 };
diff --git a/Engine/Plugins/Runtime/nDisplay/nDisplay.uplugin b/Engine/Plugins/Runtime/nDisplay/nDisplay.uplugin
--- a/Engine/Plugins/Runtime/nDisplay/nDisplay.uplugin
+++ b/Engine/Plugins/Runtime/nDisplay/nDisplay.uplugin
@@ -129,6 +129,10 @@
 			"Name": "TextureShare",
 			"Enabled": true
 		},
+		{
+			"Name": "ProceduralMeshComponent",
+			"Enabled": true
+		},
 		{
 			"Name": "RemoteControlInterception",
 			"Enabled": true
