nDisplay: fixed non-Unity build

diff --git a/Engine/Plugins/Runtime/nDisplay/Source/DisplayCluster/Public/Render/Containers/DisplayClusterRender_MeshComponentProxyData.h b/Engine/Plugins/Runtime/nDisplay/Source/DisplayCluster/Public/Render/Containers/DisplayClusterRender_MeshComponentProxyData.h
--- a/Engine/Plugins/Runtime/nDisplay/Source/DisplayCluster/Public/Render/Containers/DisplayClusterRender_MeshComponentProxyData.h
+++ b/Engine/Plugins/Runtime/nDisplay/Source/DisplayCluster/Public/Render/Containers/DisplayClusterRender_MeshComponentProxyData.h
@@ -7,6 +7,7 @@
 
 
 class FDisplayClusterRender_MeshGeometry;
+class UStaticMesh;
 struct FStaticMeshLODResources;
 struct FProcMeshSection;
 
