[nDisplay] Workaround for Value32==value64 assertion

diff --git a/Engine/Plugins/Runtime/nDisplay/Source/DisplayCluster/Private/Render/Viewport/RenderTarget/DisplayClusterRenderTargetResource.cpp b/Engine/Plugins/Runtime/nDisplay/Source/DisplayCluster/Private/Render/Viewport/RenderTarget/DisplayClusterRenderTargetResource.cpp
--- a/Engine/Plugins/Runtime/nDisplay/Source/DisplayCluster/Private/Render/Viewport/RenderTarget/DisplayClusterRenderTargetResource.cpp
+++ b/Engine/Plugins/Runtime/nDisplay/Source/DisplayCluster/Private/Render/Viewport/RenderTarget/DisplayClusterRenderTargetResource.cpp
@@ -25,7 +25,7 @@ void FDisplayClusterTextureResource::InitDynamicRHI()
 
 		TRefCountPtr<FRHITexture2D> DummyTexture2DRHI;
 
-		FRHIResourceCreateInfo CreateInfo(TEXT("DisplayClusterTextureResource"), FClearValueBinding::None);
+		FRHIResourceCreateInfo CreateInfo(TEXT("DisplayClusterTextureMipsResource"), FClearValueBinding::None);
 
 		RHICreateTargetableShaderResource2D(
 			InitSettings.Size.X, 
@@ -39,8 +39,6 @@ void FDisplayClusterTextureResource::InitDynamicRHI()
 			TextureRHI,
 			DummyTexture2DRHI
 		);
-
-		TextureRHI->SetName(TEXT("DisplayClusterTextureMipsResource"));
 	}
 	else
 	{
