nDisplay: added frustum visualization for icvfx camera and viewports. Improved ICVFX camera mesh behaviour(hide) when ExternalCameraReference is set.

diff --git a/Engine/Plugins/Runtime/nDisplay/Source/DisplayCluster/Private/Game/EngineClasses/Scene/DisplayClusterICVFXCameraComponent.cpp b/Engine/Plugins/Runtime/nDisplay/Source/DisplayCluster/Private/Game/EngineClasses/Scene/DisplayClusterICVFXCameraComponent.cpp
--- a/Engine/Plugins/Runtime/nDisplay/Source/DisplayCluster/Private/Game/EngineClasses/Scene/DisplayClusterICVFXCameraComponent.cpp
+++ b/Engine/Plugins/Runtime/nDisplay/Source/DisplayCluster/Private/Game/EngineClasses/Scene/DisplayClusterICVFXCameraComponent.cpp
@@ -1,6 +1,7 @@
 // Copyright Epic Games, Inc. All Rights Reserved.
 
 #include "Components/DisplayClusterICVFXCameraComponent.h"
+#include "Components/DrawFrustumComponent.h"
 
 #include "Render/Viewport/Containers/DisplayClusterViewport_CameraMotionBlur.h"
 #include "Components/DisplayClusterCameraComponent.h"
@@ -77,3 +78,68 @@ FDisplayClusterViewport_CameraMotionBlur UDisplayClusterICVFXCameraComponent::Ge
 
 	return OutParameters;
 }
+
+void UDisplayClusterICVFXCameraComponent::OnRegister()
+{
+	Super::OnRegister();
+
+#if WITH_EDITORONLY_DATA
+	// disable frustum for icvfx camera component
+	DrawFrustum->bFrustumEnabled = false;
+
+	// Update ExternalCineactor behaviour
+	UpdateICVFXPreviewState();
+#endif
+}
+
+#if WITH_EDITORONLY_DATA
+void UDisplayClusterICVFXCameraComponent::PreEditChange(FProperty* PropertyThatWillChange)
+{
+	Super::PreEditChange(PropertyThatWillChange);
+
+	// save the current value
+	ExternalCameraCachedValue = CameraSettings.ExternalCameraActor;
+}
+
+void UDisplayClusterICVFXCameraComponent::PostEditChangeProperty(FPropertyChangedEvent& PropertyChangedEvent)
+{
+	Super::PostEditChangeProperty(PropertyChangedEvent);
+
+	UpdateICVFXPreviewState();
+}
+
+void UDisplayClusterICVFXCameraComponent::UpdateICVFXPreviewState()
+{
+	// handle frustum visibility
+	if (CameraSettings.ExternalCameraActor.IsValid())
+	{
+		ACineCameraActor* CineCamera = CameraSettings.ExternalCameraActor.Get();
+		CineCamera->GetCineCameraComponent()->bDrawFrustumAllowed = false;
+
+		UDrawFrustumComponent* DrawFustumComponent = Cast<UDrawFrustumComponent>(CineCamera->GetComponentByClass(UDrawFrustumComponent::StaticClass()));
+		if (DrawFustumComponent != nullptr)
+		{
+			DrawFustumComponent->bFrustumEnabled = false;
+			DrawFustumComponent->MarkRenderStateDirty();
+		}
+
+		if (ProxyMeshComponent)
+		{
+			ProxyMeshComponent->DestroyComponent();
+			ProxyMeshComponent = nullptr;
+		}
+	}
+
+
+	// restore frustum visibility if reference was changed
+	if (ExternalCameraCachedValue.IsValid())
+	{
+		ACineCameraActor* CineCamera = ExternalCameraCachedValue.Get();
+		UDrawFrustumComponent* DrawFustumComponent = Cast<UDrawFrustumComponent>(CineCamera->GetComponentByClass(UDrawFrustumComponent::StaticClass()));
+		DrawFustumComponent->bFrustumEnabled = true;
+		DrawFustumComponent->MarkRenderStateDirty();
+
+		ExternalCameraCachedValue.Reset();
+	}
+}
+#endif
\ No newline at end of file
diff --git a/Engine/Plugins/Runtime/nDisplay/Source/DisplayCluster/Private/Game/EngineClasses/Scene/DisplayClusterRootActorEditor.cpp b/Engine/Plugins/Runtime/nDisplay/Source/DisplayCluster/Private/Game/EngineClasses/Scene/DisplayClusterRootActorEditor.cpp
--- a/Engine/Plugins/Runtime/nDisplay/Source/DisplayCluster/Private/Game/EngineClasses/Scene/DisplayClusterRootActorEditor.cpp
+++ b/Engine/Plugins/Runtime/nDisplay/Source/DisplayCluster/Private/Game/EngineClasses/Scene/DisplayClusterRootActorEditor.cpp
@@ -9,6 +9,7 @@
 #include "Components/DisplayClusterSceneComponentSyncParent.h"
 #include "Components/DisplayClusterPreviewComponent.h"
 #include "Components/DisplayClusterICVFXCameraComponent.h"
+#include "Components/LineBatchComponent.h"
 
 #include "Config/IPDisplayClusterConfigManager.h"
 #include "DisplayClusterConfigurationStrings.h"
@@ -27,6 +28,8 @@
 #include "Misc/DisplayClusterLog.h"
 #include "Misc/DisplayClusterStrings.h"
 
+#include "Render/Viewport/Configuration/DisplayClusterViewportConfigurationHelpers_ICVFX.h"
+#include "Render/Viewport/DisplayClusterViewportStrings.h"
 #include "Render/Viewport/IDisplayClusterViewportManager.h"
 #include "Render/Viewport/IDisplayClusterViewport.h"
 #include "Render/Viewport/IDisplayClusterViewportProxy.h"
@@ -278,7 +281,6 @@ void ADisplayClusterRootActor::RenderPreview_Editor()
 			}
 		}
 
-		//@todo: add GUI to change preview scene
 		// Now always use RootActor world to preview. 
 		UWorld* PreviewWorld = GetWorld();
 		if (PreviewWorld)
@@ -292,7 +294,160 @@ void ADisplayClusterRootActor::RenderPreview_Editor()
 				OnPreviewGenerated.ExecuteIfBound();
 			}
 		}
+
+		// preview frustums
+		RenderPreviewFrustums();
+	}
+}
+
+void ADisplayClusterRootActor::RenderPreviewFrustums()
+{
+	if (!ViewportManager.IsValid())
+	{
+		return;
+	}
+
+	// frustum preview viewports
+	TArray<IDisplayClusterViewport*> FrustumPreviewViewports;
+
+	for (const TPair<FString, UDisplayClusterConfigurationClusterNode*>& Node : CurrentConfigData->Cluster->Nodes)
+	{
+		if (Node.Value == nullptr)
+		{
+			continue;
+		}
+
+		// collect node viewports
+		for (const TPair<FString, UDisplayClusterConfigurationViewport*>& ViewportConfig : Node.Value->Viewports)
+		{
+			if (ViewportConfig.Value->bAllowPreviewFrustumRendering == false || ViewportConfig.Value->bIsVisible == false || ViewportConfig.Value == nullptr)
+			{
+				continue;
+			}
+
+			IDisplayClusterViewport* Viewport = ViewportManager->FindViewport(ViewportConfig.Key);
+			if (Viewport != nullptr)
+			{
+				FrustumPreviewViewports.Add(Viewport);
+			}
+		}
+	}
+
+	// collect incameras
+	if (bPreviewICVFXFrustums)
+	{
+		for (UActorComponent* ActorComponentIt : GetComponents())
+		{
+			if (ActorComponentIt)
+			{
+				UDisplayClusterICVFXCameraComponent* CineCameraComponent = Cast<UDisplayClusterICVFXCameraComponent>(ActorComponentIt);
+				if (CineCameraComponent && CineCameraComponent->IsICVFXEnabled())
+				{
+					const FString InnerFrustumID = CineCameraComponent->GetCameraUniqueId();
+					FDisplayClusterViewport* IncameraViewport = FDisplayClusterViewportConfigurationHelpers_ICVFX::ImplFindViewport(*this, InnerFrustumID, DisplayClusterViewportStrings::icvfx::camera);
+					if (IncameraViewport != nullptr)
+					{
+						FrustumPreviewViewports.Add(IncameraViewport);
+					}
+				}
+			}
+		}
 	}
+
+	for (IDisplayClusterViewport* Viewport : FrustumPreviewViewports)
+	{
+		const TArray<FDisplayClusterViewport_Context>& Contexts = Viewport->GetContexts();
+		for (const FDisplayClusterViewport_Context& Context : Contexts)
+		{
+			FMatrix ProjectionMatrix = Context.ProjectionMatrix;
+
+			uint32 ViewportContextNum = 0;
+			Viewport->GetProjectionMatrix(ViewportContextNum, ProjectionMatrix);
+
+			FMatrix ViewRotationMatrix = FInverseRotationMatrix(Context.ViewRotation) * FMatrix(
+				FPlane(0, 0, 1, 0),
+				FPlane(1, 0, 0, 0),
+				FPlane(0, 1, 0, 0),
+				FPlane(0, 0, 0, 1));
+
+			const FMatrix ViewMatrix = FTranslationMatrix(-Context.ViewLocation) * ViewRotationMatrix;
+
+			RenderPreviewFrustum(ProjectionMatrix, ViewMatrix, Context.ViewLocation);
+		}
+	}
+}
+
+void ADisplayClusterRootActor::RenderPreviewFrustum(const FMatrix ProjectionMatrix, const FMatrix ViewMatrix, const FVector ViewOrigin)
+{
+	const float FarPlane = PreviewICVFXFrustumsFarDistance;
+	const float NearPlane = GNearClippingPlane;
+	const FColor Color = FColor::Green;
+	const float Thickness = 1.0f;
+
+	const UWorld* World = GetWorld();
+	ULineBatchComponent* LineBatcher = World ? World->LineBatcher : nullptr;
+	if (!LineBatcher)
+	{
+		return;
+	}
+
+	// Get FOV and AspectRatio from the view's projection matrix.
+	const float AspectRatio = ProjectionMatrix.M[1][1] / ProjectionMatrix.M[0][0];
+	const bool bIsPerspectiveProjection = true;
+
+	// Build the camera frustum for this cascade
+	const float HalfHorizontalFOV = bIsPerspectiveProjection ? FMath::Atan(1.0f / ProjectionMatrix.M[0][0]) : PI / 4.0f;
+	const float HalfVerticalFOV = bIsPerspectiveProjection ? FMath::Atan(1.0f / ProjectionMatrix.M[1][1]) : FMath::Atan((FMath::Tan(PI / 4.0f) / AspectRatio));
+	const float AsymmetricFOVScaleX = ProjectionMatrix.M[2][0];
+	const float AsymmetricFOVScaleY = ProjectionMatrix.M[2][1];
+	
+	// Near plane
+	const float StartHorizontalTotalLength = NearPlane * FMath::Tan(HalfHorizontalFOV);
+	const float StartVerticalTotalLength = NearPlane * FMath::Tan(HalfVerticalFOV);
+	const FVector StartCameraLeftOffset = ViewMatrix.GetColumn(0) * -StartHorizontalTotalLength * (1 + AsymmetricFOVScaleX);
+	const FVector StartCameraRightOffset = ViewMatrix.GetColumn(0) * StartHorizontalTotalLength * (1 - AsymmetricFOVScaleX);
+	const FVector StartCameraBottomOffset = ViewMatrix.GetColumn(1) * -StartVerticalTotalLength * (1 + AsymmetricFOVScaleY);
+	const FVector StartCameraTopOffset = ViewMatrix.GetColumn(1) * StartVerticalTotalLength * (1 - AsymmetricFOVScaleY);
+	
+	// Far plane
+	const float EndHorizontalTotalLength = FarPlane * FMath::Tan(HalfHorizontalFOV);
+	const float EndVerticalTotalLength = FarPlane * FMath::Tan(HalfVerticalFOV);
+	const FVector EndCameraLeftOffset = ViewMatrix.GetColumn(0) * -EndHorizontalTotalLength * (1 + AsymmetricFOVScaleX);
+	const FVector EndCameraRightOffset = ViewMatrix.GetColumn(0) * EndHorizontalTotalLength * (1 - AsymmetricFOVScaleX);
+	const FVector EndCameraBottomOffset = ViewMatrix.GetColumn(1) * -EndVerticalTotalLength * (1 + AsymmetricFOVScaleY);
+	const FVector EndCameraTopOffset = ViewMatrix.GetColumn(1) * EndVerticalTotalLength * (1 - AsymmetricFOVScaleY);
+	
+	const FVector CameraDirection = ViewMatrix.GetColumn(2);
+
+	// Get the 4 points of the camera frustum near plane, in world space
+	PreviewFrustumVerts[0] = ViewOrigin + CameraDirection * NearPlane + StartCameraRightOffset + StartCameraTopOffset;         // 0 Near  Top    Right
+	PreviewFrustumVerts[1] = ViewOrigin + CameraDirection * NearPlane + StartCameraRightOffset + StartCameraBottomOffset;      // 1 Near  Bottom Right
+	PreviewFrustumVerts[2] = ViewOrigin + CameraDirection * NearPlane + StartCameraLeftOffset + StartCameraTopOffset;          // 2 Near  Top    Left
+	PreviewFrustumVerts[3] = ViewOrigin + CameraDirection * NearPlane + StartCameraLeftOffset + StartCameraBottomOffset;       // 3 Near  Bottom Left
+
+	// Get the 4 points of the camera frustum far plane, in world space
+	PreviewFrustumVerts[4] = ViewOrigin + CameraDirection * FarPlane + EndCameraRightOffset + EndCameraTopOffset;         // 4 Far  Top    Right
+	PreviewFrustumVerts[5] = ViewOrigin + CameraDirection * FarPlane + EndCameraRightOffset + EndCameraBottomOffset;      // 5 Far  Bottom Right
+	PreviewFrustumVerts[6] = ViewOrigin + CameraDirection * FarPlane + EndCameraLeftOffset + EndCameraTopOffset;          // 6 Far  Top    Left
+	PreviewFrustumVerts[7] = ViewOrigin + CameraDirection * FarPlane + EndCameraLeftOffset + EndCameraBottomOffset;       // 7 Far  Bottom Left	
+	
+	// frustum lines
+	LineBatcher->DrawLine(PreviewFrustumVerts[0], PreviewFrustumVerts[4], Color, SDPG_World, Thickness, 0.f); // right top
+	LineBatcher->DrawLine(PreviewFrustumVerts[1], PreviewFrustumVerts[5], Color, SDPG_World, Thickness, 0.f); // right bottom
+	LineBatcher->DrawLine(PreviewFrustumVerts[2], PreviewFrustumVerts[6], Color, SDPG_World, Thickness, 0.f); // left top
+	LineBatcher->DrawLine(PreviewFrustumVerts[3], PreviewFrustumVerts[7], Color, SDPG_World, Thickness, 0.f); // left bottom
+
+	// near plane square
+	LineBatcher->DrawLine(PreviewFrustumVerts[0], PreviewFrustumVerts[1], Color, SDPG_World, Thickness, 0.f); // right top to right bottom
+	LineBatcher->DrawLine(PreviewFrustumVerts[1], PreviewFrustumVerts[3], Color, SDPG_World, Thickness, 0.f); // right bottom to left bottom
+	LineBatcher->DrawLine(PreviewFrustumVerts[3], PreviewFrustumVerts[2], Color, SDPG_World, Thickness, 0.f); // left bottom to left top
+	LineBatcher->DrawLine(PreviewFrustumVerts[2], PreviewFrustumVerts[0], Color, SDPG_World, Thickness, 0.f); // left top to right top
+
+	// far plane square
+	LineBatcher->DrawLine(PreviewFrustumVerts[4], PreviewFrustumVerts[5], Color, SDPG_World, Thickness, 0.f); // right top to right bottom
+	LineBatcher->DrawLine(PreviewFrustumVerts[5], PreviewFrustumVerts[7], Color, SDPG_World, Thickness, 0.f); // right bottom to left bottom
+	LineBatcher->DrawLine(PreviewFrustumVerts[7], PreviewFrustumVerts[6], Color, SDPG_World, Thickness, 0.f); // left bottom to left top
+	LineBatcher->DrawLine(PreviewFrustumVerts[6], PreviewFrustumVerts[4], Color, SDPG_World, Thickness, 0.f); // left top to right top
 }
 
 void ADisplayClusterRootActor::PostEditChangeProperty(FPropertyChangedEvent& PropertyChangedEvent)
diff --git a/Engine/Plugins/Runtime/nDisplay/Source/DisplayCluster/Private/Render/Viewport/Configuration/DisplayClusterViewportConfigurationHelpers_ICVFX.cpp b/Engine/Plugins/Runtime/nDisplay/Source/DisplayCluster/Private/Render/Viewport/Configuration/DisplayClusterViewportConfigurationHelpers_ICVFX.cpp
--- a/Engine/Plugins/Runtime/nDisplay/Source/DisplayCluster/Private/Render/Viewport/Configuration/DisplayClusterViewportConfigurationHelpers_ICVFX.cpp
+++ b/Engine/Plugins/Runtime/nDisplay/Source/DisplayCluster/Private/Render/Viewport/Configuration/DisplayClusterViewportConfigurationHelpers_ICVFX.cpp
@@ -55,7 +55,7 @@ FString ImplGetNameICVFX(const FString& InViewportId, const FString& InResourceI
 	return FString::Printf(TEXT("%s_%s_%s"), DisplayClusterViewportStrings::icvfx::prefix, *InViewportId, *InResourceId);
 }
 
-static FDisplayClusterViewport* ImplFindViewport(ADisplayClusterRootActor& RootActor, const FString& InViewportId, const FString& InResourceId)
+FDisplayClusterViewport* FDisplayClusterViewportConfigurationHelpers_ICVFX::ImplFindViewport(ADisplayClusterRootActor& RootActor, const FString& InViewportId, const FString& InResourceId)
 {
 	FDisplayClusterViewportManager* ViewportManager = FDisplayClusterViewportConfigurationHelpers_ICVFX::GetViewportManager(RootActor);
 	if (ViewportManager)
@@ -291,7 +291,7 @@ bool FDisplayClusterViewportConfigurationHelpers_ICVFX::GetCameraContext(UDispla
 
 	// Get camera pos-rot-prj from policy
 	const float WorldToMeters = 100.f;
-	const float CfgNCP = 1.f;
+	const float CfgNCP = 1.0f;
 	const FVector ViewOffset = FVector::ZeroVector;
 
 	if (CameraProjectionPolicy->CalculateView(nullptr, 0, OutCameraContext.ViewLocation, OutCameraContext.ViewRotation, ViewOffset, WorldToMeters, CfgNCP, CfgNCP) &&
diff --git a/Engine/Plugins/Runtime/nDisplay/Source/DisplayCluster/Private/Render/Viewport/Configuration/DisplayClusterViewportConfigurationHelpers_ICVFX.h b/Engine/Plugins/Runtime/nDisplay/Source/DisplayCluster/Private/Render/Viewport/Configuration/DisplayClusterViewportConfigurationHelpers_ICVFX.h
--- a/Engine/Plugins/Runtime/nDisplay/Source/DisplayCluster/Private/Render/Viewport/Configuration/DisplayClusterViewportConfigurationHelpers_ICVFX.h
+++ b/Engine/Plugins/Runtime/nDisplay/Source/DisplayCluster/Private/Render/Viewport/Configuration/DisplayClusterViewportConfigurationHelpers_ICVFX.h
@@ -56,6 +56,7 @@ public:
 	static void UpdateCameraCustomFrustum(FDisplayClusterViewport& DstViewport, const FDisplayClusterConfigurationICVFX_CameraCustomFrustum& InCameraCustomFrustum);
 	static void UpdateCameraViewportBufferRatio(FDisplayClusterViewport& DstViewport, const FDisplayClusterConfigurationICVFX_CameraSettings& CameraSettings);
 
+	static FDisplayClusterViewport* ImplFindViewport(ADisplayClusterRootActor& RootActor, const FString& InViewportId, const FString& InResourceId);
 private:
 	static FDisplayClusterViewport* ImplCreateViewport(ADisplayClusterRootActor& RootActor, const FString& InViewportId, const FString& InResourceId, TSharedPtr<IDisplayClusterProjectionPolicy, ESPMode::ThreadSafe>& InProjectionPolicy);
 };
diff --git a/Engine/Plugins/Runtime/nDisplay/Source/DisplayCluster/Public/Components/DisplayClusterICVFXCameraComponent.h b/Engine/Plugins/Runtime/nDisplay/Source/DisplayCluster/Public/Components/DisplayClusterICVFXCameraComponent.h
--- a/Engine/Plugins/Runtime/nDisplay/Source/DisplayCluster/Public/Components/DisplayClusterICVFXCameraComponent.h
+++ b/Engine/Plugins/Runtime/nDisplay/Source/DisplayCluster/Public/Components/DisplayClusterICVFXCameraComponent.h
@@ -60,11 +60,23 @@ public:
 	UCameraComponent* GetCameraComponent();
 	void GetDesiredView(FMinimalViewInfo& DesiredView);
 
+	// UActorComponent interface
+	virtual void OnRegister() override;
 	
 //////////////////////////////////////////////////////////////////////////////////////////////
 // Details Panel Property Referencers
 //////////////////////////////////////////////////////////////////////////////////////////////
 #if WITH_EDITORONLY_DATA
+public:
+	// update the status of the Editor's ICVFX preview elements
+	void UpdateICVFXPreviewState();
+
+	// saves the value of external camera reference
+	TSoftObjectPtr<ACineCameraActor> ExternalCameraCachedValue;
+
+	virtual void PreEditChange(FProperty* PropertyThatWillChange) override;
+	virtual void PostEditChangeProperty(FPropertyChangedEvent& PropertyChangedEvent) override;
+
 private:
 	friend class FDisplayClusterICVFXCameraComponentDetailsCustomization;
 	
diff --git a/Engine/Plugins/Runtime/nDisplay/Source/DisplayCluster/Public/DisplayClusterRootActor.h b/Engine/Plugins/Runtime/nDisplay/Source/DisplayCluster/Public/DisplayClusterRootActor.h
--- a/Engine/Plugins/Runtime/nDisplay/Source/DisplayCluster/Public/DisplayClusterRootActor.h
+++ b/Engine/Plugins/Runtime/nDisplay/Source/DisplayCluster/Public/DisplayClusterRootActor.h
@@ -301,6 +301,14 @@ public:
 	UPROPERTY(EditAnywhere, Category = "Editor Preview", meta = (DisplayName = "Preview Screen Percentage", ClampMin = "0.05", UIMin = "0.05", ClampMax = "1", UIMax = "1", EditCondition = "bPreviewEnable"))
 	float PreviewRenderTargetRatioMult = 0.25;
 
+	/** Render ICVFX Frustums */
+	UPROPERTY(EditAnywhere, Category = "Editor Preview", meta = (DisplayName = "ICVFX Camera Frustums", EditCondition = "bPreviewEnable"))
+	bool bPreviewICVFXFrustums = false;
+
+	/** Render ICVFX Frustums */
+	UPROPERTY(EditAnywhere, Category = "Editor Preview", meta = (DisplayName = "ICVFX Camera Frustums Distance", EditCondition = "bPreviewEnable"))
+	float PreviewICVFXFrustumsFarDistance = 1000.0f;
+
 private:
 	UPROPERTY(Transient)
 	TMap<FString, UDisplayClusterPreviewComponent*> PreviewComponents;
@@ -361,6 +369,9 @@ protected:
 	FString GeneratePreviewComponentName(const FString& NodeId, const FString& ViewportId) const;
 	void RenderPreview_Editor();
 	bool UpdatePreviewConfiguration_Editor(bool bUpdateAllViewports);
+	
+	void RenderPreviewFrustums();
+	void RenderPreviewFrustum(const FMatrix ProjectionMatrix, const FMatrix ViewMatrix, const FVector ViewOrigin);
 
 	virtual void PostEditChangeProperty(FPropertyChangedEvent& PropertyChangedEvent) override;
 	virtual void PostEditMove(bool bFinished) override;
@@ -369,6 +380,9 @@ private:
 	/** The number of times to update the render target via deferred update. */
 	int32 PreviewRenderTargetUpdatesRequired = 0;
 	bool bIsSelectedInEditor = false;
+
+	// Preview frustum vertices
+	FVector PreviewFrustumVerts[8];
 	
 private:
 	TWeakPtr<IDisplayClusterConfiguratorBlueprintEditor> ToolkitPtr;
diff --git a/Engine/Plugins/Runtime/nDisplay/Source/DisplayClusterConfiguration/Public/DisplayClusterConfigurationTypes_Viewport.h b/Engine/Plugins/Runtime/nDisplay/Source/DisplayClusterConfiguration/Public/DisplayClusterConfigurationTypes_Viewport.h
--- a/Engine/Plugins/Runtime/nDisplay/Source/DisplayClusterConfiguration/Public/DisplayClusterConfigurationTypes_Viewport.h
+++ b/Engine/Plugins/Runtime/nDisplay/Source/DisplayClusterConfiguration/Public/DisplayClusterConfigurationTypes_Viewport.h
@@ -169,6 +169,9 @@ public:
 	/** Locks the Viewport aspect ratio for easier resizing */
 	UPROPERTY(EditAnywhere, Category = "Configuration")
 	bool bFixedAspectRatio;
+
+	UPROPERTY(EditAnywhere, BlueprintReadWrite, Category = "Preview", meta = (DisplayName = "Preview Frustum"))
+	bool bAllowPreviewFrustumRendering = false;
 #endif
 	
 	/** Define the Viewport 2D coordinates */
diff --git a/Engine/Source/Runtime/Engine/Classes/Camera/CameraComponent.h b/Engine/Source/Runtime/Engine/Classes/Camera/CameraComponent.h
--- a/Engine/Source/Runtime/Engine/Classes/Camera/CameraComponent.h
+++ b/Engine/Source/Runtime/Engine/Classes/Camera/CameraComponent.h
@@ -77,6 +77,10 @@ class ENGINE_API UCameraComponent : public USceneComponent
 #endif
 
 #if WITH_EDITORONLY_DATA
+	// The Frustum visibility flag for draw frustum component initialization
+	UPROPERTY(EditAnywhere, BlueprintReadWrite, Category = Camera)
+	bool bDrawFrustumAllowed = true;
+
 	/** If the camera mesh is visible in game. Only relevant when running editor builds. */
 	UPROPERTY(EditAnywhere, BlueprintReadWrite, Category = Camera)
 	uint8 bCameraMeshHiddenInGame : 1;
diff --git a/Engine/Source/Runtime/Engine/Classes/Components/DrawFrustumComponent.h b/Engine/Source/Runtime/Engine/Classes/Components/DrawFrustumComponent.h
--- a/Engine/Source/Runtime/Engine/Classes/Components/DrawFrustumComponent.h
+++ b/Engine/Source/Runtime/Engine/Classes/Components/DrawFrustumComponent.h
@@ -18,7 +18,11 @@ UCLASS(collapsecategories, hidecategories=Object, editinlinenew, MinimalAPI)
 class UDrawFrustumComponent : public UPrimitiveComponent
 {
 	GENERATED_UCLASS_BODY()
-	
+
+	// Enable or disable frustum visualization for this camera
+	UPROPERTY(EditAnywhere, BlueprintReadWrite, Category = DrawFrustumComponent)
+	bool bFrustumEnabled = true;
+
 	/** Color to draw the wireframe frustum. */
 	UPROPERTY(EditAnywhere, BlueprintReadWrite, Category=DrawFrustumComponent)
 	FColor FrustumColor;
diff --git a/Engine/Source/Runtime/Engine/Private/Camera/CameraComponent.cpp b/Engine/Source/Runtime/Engine/Private/Camera/CameraComponent.cpp
--- a/Engine/Source/Runtime/Engine/Private/Camera/CameraComponent.cpp
+++ b/Engine/Source/Runtime/Engine/Private/Camera/CameraComponent.cpp
@@ -133,6 +133,7 @@ void UCameraComponent::OnRegister()
 			DrawFrustum->SetupAttachment(this);
 			DrawFrustum->SetIsVisualizationComponent(true);
 			DrawFrustum->CreationMethod = CreationMethod;
+			DrawFrustum->bFrustumEnabled = bDrawFrustumAllowed;
 			DrawFrustum->RegisterComponentWithWorld(GetWorld());
 		}
 	}
diff --git a/Engine/Source/Runtime/Engine/Private/Components/DrawFrustumComponent.cpp b/Engine/Source/Runtime/Engine/Private/Components/DrawFrustumComponent.cpp
--- a/Engine/Source/Runtime/Engine/Private/Components/DrawFrustumComponent.cpp
+++ b/Engine/Source/Runtime/Engine/Private/Components/DrawFrustumComponent.cpp
@@ -27,6 +27,7 @@ public:
 	*/
 	FDrawFrustumSceneProxy(const UDrawFrustumComponent* InComponent)
 	:	FPrimitiveSceneProxy(InComponent)
+	,	bFrustumEnabled(InComponent->bFrustumEnabled)
 	,	FrustumColor(InComponent->FrustumColor)
 	,	FrustumAngle(InComponent->FrustumAngle)
 	,	FrustumAspectRatio(InComponent->FrustumAspectRatio)
@@ -118,7 +119,7 @@ public:
 	virtual FPrimitiveViewRelevance GetViewRelevance(const FSceneView* View) const override
 	{
 		FPrimitiveViewRelevance Result;
-		Result.bDrawRelevance = IsShown(View) && View->Family->EngineShowFlags.CameraFrustums;
+		Result.bDrawRelevance = IsShown(View) && View->Family->EngineShowFlags.CameraFrustums && bFrustumEnabled;
 		Result.bDynamicRelevance = true;
 		Result.bShadowRelevance = IsShadowCast(View);
 		Result.bEditorPrimitiveRelevance = UseEditorCompositing(View);
@@ -129,6 +130,7 @@ public:
 	uint32 GetAllocatedSize( void ) const { return( FPrimitiveSceneProxy::GetAllocatedSize() ); }
 
 private:
+	bool bFrustumEnabled;
 	FColor FrustumColor;
 	float FrustumAngle;
 	float FrustumAspectRatio;
