nDisplay: Rework compiler to better support subobject transactions. Fixes crash and potential invalid instance container data when adding new nodes. Fix transaction crashes when a node is selected and it is being deleted.

diff --git a/Engine/Plugins/Runtime/nDisplay/Source/DisplayCluster/DisplayCluster.Build.cs b/Engine/Plugins/Runtime/nDisplay/Source/DisplayCluster/DisplayCluster.Build.cs
--- a/Engine/Plugins/Runtime/nDisplay/Source/DisplayCluster/DisplayCluster.Build.cs
+++ b/Engine/Plugins/Runtime/nDisplay/Source/DisplayCluster/DisplayCluster.Build.cs
@@ -56,6 +56,15 @@ public class DisplayCluster : ModuleRules
 			PrivateDependencyModuleNames.Add("LevelEditor");
 		}
 
+		if (Target.Type == TargetType.Editor)
+		{
+			PrivateDependencyModuleNames.AddRange(
+				new string[]
+				{
+					"ConcertSyncClient"
+				});
+		}
+
 		if (Target.Platform == UnrealTargetPlatform.Win64)
 		{
 			PrivateDependencyModuleNames.AddRange(
diff --git a/Engine/Plugins/Runtime/nDisplay/Source/DisplayCluster/Private/Game/EngineClasses/Scene/DisplayClusterRootActor.cpp b/Engine/Plugins/Runtime/nDisplay/Source/DisplayCluster/Private/Game/EngineClasses/Scene/DisplayClusterRootActor.cpp
--- a/Engine/Plugins/Runtime/nDisplay/Source/DisplayCluster/Private/Game/EngineClasses/Scene/DisplayClusterRootActor.cpp
+++ b/Engine/Plugins/Runtime/nDisplay/Source/DisplayCluster/Private/Game/EngineClasses/Scene/DisplayClusterRootActor.cpp
@@ -43,6 +43,12 @@
 
 #include "Game/EngineClasses/Scene/DisplayClusterRootActorInitializer.h"
 
+#if WITH_EDITOR
+#include "IConcertSyncClientModule.h"
+#include "IConcertClientWorkspace.h"
+#include "IConcertSyncClient.h"
+#endif
+
 ADisplayClusterRootActor::ADisplayClusterRootActor(const FObjectInitializer& ObjectInitializer)
 	: Super(ObjectInitializer)
 	, OperationMode(EDisplayClusterOperationMode::Disabled)
@@ -827,10 +833,33 @@ static void PropagateDefaultMapToInstancedMap(const FMapProperty* MapProperty, U
 		{
 			if (const UObject* InstancedObject = *InstancedObjectPtr)
 			{
-				if (ensure(DefaultObject == InstancedObject->GetArchetype()))
+				const bool bArchetypeCorrect = DefaultObject == InstancedObject->GetArchetype();
+				
+				// The archetype should always match the new default. There are edge cases with MU
+				// where the instance may be updated prior to the CDO and should be corrected once
+				// the BP is saved. If this occurs outside of MU there could be a serious problem.
+				if (!bArchetypeCorrect)
 				{
-					continue;
+					const ADisplayClusterRootActor* RootActor =
+					CastChecked<ADisplayClusterRootActor>(InstancedObject->GetTypedOuter(ADisplayClusterRootActor::StaticClass()));
+#if WITH_EDITOR
+					if (GEditor)
+					{
+						bool bIsMultiUserSession = false;
+						TSharedPtr<IConcertSyncClient> ConcertSyncClient = IConcertSyncClientModule::Get().GetClient(TEXT("MultiUser"));
+						if (ConcertSyncClient.IsValid())
+						{
+							TSharedPtr<IConcertClientWorkspace> Workspace = ConcertSyncClient->GetWorkspace();
+							bIsMultiUserSession = Workspace.IsValid();
+						}
+
+						ensure(bArchetypeCorrect || bIsMultiUserSession);
+					}
+#endif
+					UE_LOG(LogDisplayClusterBlueprint, Warning, TEXT("Archetype mismatch on nDisplay config %s. Make sure the config is compiled and saved. Property: %s, Instance: %s, Archetype: %s, Default %s."),
+						*RootActor->GetName(), *MapProperty->GetName(), *InstancedObject->GetName(), *InstancedObject->GetArchetype()->GetName(), *DefaultObject->GetName());
 				}
+				continue;
 			}
 		}
 		
@@ -848,18 +877,17 @@ static void PropagateDefaultMapToInstancedMap(const FMapProperty* MapProperty, U
 		
 		if (!MapDefaultsHelper.FindValueFromHash(Key))
 		{
-			UObject** InstanceObjectPtr = MapProperty->ValueProp->ContainerPtrToValuePtr<UObject*>(InstancePairPtr);
-			if (InstanceObjectPtr)
+			if (UObject** InstanceObjectPtr = MapProperty->ValueProp->ContainerPtrToValuePtr<UObject*>(InstancePairPtr))
 			{
-				UObject* InstanceObject = *InstanceObjectPtr;
-				check(InstanceObject);
-				
-				// Trash the object -- default propagation won't handle this any more.
-				// RemoveAt below will remove the reference to it. This transaction can still be undone.
-				// Rename to transient package now so the same name is available immediately.
-				// It's possible a new object needs to be created with this outer using the same name.
-				InstanceObject->Rename(nullptr, GetTransientPackage(), REN_DontCreateRedirectors);
-				InstanceObject->SetFlags(RF_Transient);
+				if (UObject* InstanceObject = *InstanceObjectPtr)
+				{
+					// Trash the object -- default propagation won't handle this any more.
+					// RemoveAt below will remove the reference to it. This transaction can still be undone.
+					// Rename to transient package now so the same name is available immediately.
+					// It's possible a new object needs to be created with this outer using the same name.
+					InstanceObject->Rename(nullptr, GetTransientPackage(), REN_DontCreateRedirectors);
+					InstanceObject->SetFlags(RF_Transient);
+				}
 			}
 			
 			MapInstanceHelper.RemoveAt(*InstanceIt);
diff --git a/Engine/Plugins/Runtime/nDisplay/Source/DisplayClusterConfigurator/Private/DisplayClusterConfiguratorCompiler.cpp b/Engine/Plugins/Runtime/nDisplay/Source/DisplayClusterConfigurator/Private/DisplayClusterConfiguratorCompiler.cpp
--- a/Engine/Plugins/Runtime/nDisplay/Source/DisplayClusterConfigurator/Private/DisplayClusterConfiguratorCompiler.cpp
+++ b/Engine/Plugins/Runtime/nDisplay/Source/DisplayClusterConfigurator/Private/DisplayClusterConfiguratorCompiler.cpp
@@ -81,20 +81,20 @@ void FDisplayClusterConfiguratorKismetCompilerContext::PreCompile()
 void FDisplayClusterConfiguratorKismetCompilerContext::SaveSubObjectsFromCleanAndSanitizeClass(
 	FSubobjectCollection& SubObjectsToSave, UBlueprintGeneratedClass* ClassToClean)
 {
-	SavedSubObjects.Empty();
+	OldSubObjects.Empty();
 	FKismetCompilerContext::SaveSubObjectsFromCleanAndSanitizeClass(SubObjectsToSave, ClassToClean);
 
-	UDisplayClusterBlueprint* DCBlueprint = CastChecked<UDisplayClusterBlueprint>(Blueprint);
+	const UDisplayClusterBlueprint* DCBlueprint = CastChecked<UDisplayClusterBlueprint>(Blueprint);
 	if (UDisplayClusterConfigurationData* ConfigData = DCBlueprint->GetConfig())
 	{
 		SubObjectsToSave.AddObject(ConfigData);
 
-		// Find all subobjects owned by this object.. same as how `AddObject` above works.
+		// Find all subobjects owned by this object. Same as how `AddObject` above works.
 		{
-			SavedSubObjects.Add(ConfigData);
+			OldSubObjects.AddUnique(ConfigData);
 			ForEachObjectWithOuter(ConfigData, [this](UObject* Child)
 			{
-				SavedSubObjects.Add(Child);
+				OldSubObjects.AddUnique(Child);
 			});
 		}
 	}
@@ -109,71 +109,108 @@ void FDisplayClusterConfiguratorKismetCompilerContext::CopyTermDefaultsToDefault
 	{
 		return;
 	}
-	
+
+	ADisplayClusterRootActor* RootActor = CastChecked<ADisplayClusterRootActor>(DefaultObject);
 	if (Blueprint->bIsNewlyCreated)
 	{
-		ADisplayClusterRootActor* RootActor = CastChecked<ADisplayClusterRootActor>(DefaultObject);
 		RootActor->PreviewNodeId = DisplayClusterConfigurationStrings::gui::preview::PreviewNodeAll;
 	}
-
-	auto GetCorrectOuter = [this](UObject* InObject)
+	
+	// Restore the old config data back to the CDO to preserve transaction history.
+	// Ideally transaction history would persist through sub-object re-instancing, but that requires
+	// the sub-objects are marked as default sub objects which is incompatible for our use.
+	if (OldSubObjects.Num() > 0)
 	{
-		check (InObject);
+		UDisplayClusterConfigurationData* OldConfigData = CastChecked<UDisplayClusterConfigurationData>(OldSubObjects[0]);
+		UDisplayClusterConfigurationData* NewConfigData = RootActor->CurrentConfigData;
+
+		TMap<UObject*, UObject*> NewToOldSubObjects;
+		{
+			NewToOldSubObjects.Reserve(OldSubObjects.Num());
+			NewToOldSubObjects.Add(NewConfigData, OldConfigData);
+			
+			ForEachObjectWithOuter(NewConfigData, [&](UObject* NewSubObject)
+			{
+				if (UObject** MatchingOldSubObject = OldSubObjects.FindByPredicate([NewSubObject](const UObject* OldSubObject)
+				{
+					return OldSubObject->GetName() == NewSubObject->GetName() &&
+						OldSubObject->GetClass() == NewSubObject->GetClass();
+				}))
+				{
+					NewToOldSubObjects.Add(NewSubObject, *MatchingOldSubObject);
+					NewSubObject->SetFlags(RF_Transient);
+					NewSubObject->ClearFlags(RF_Transactional);
+					FLinkerLoad::InvalidateExport(NewSubObject);
+				}
+			});
+		}
 		
-		UObject* InOuter = InObject->GetOuter();
-		check(InOuter);
+		const ERenameFlags RenFlags = REN_DontCreateRedirectors | REN_NonTransactional | REN_DoNotDirty | REN_ForceNoResetLoaders;
 
-		// Find the saved sub-object outer.
-		if (const TArray<UObject*>::ElementType* FoundSavedOuter = SavedSubObjects.FindByPredicate([InOuter](const UObject* SubObject)
+		// Rename our new config data (along with all new sub-objects) to the transient package.
 		{
-			check(SubObject);
-			return InOuter->GetName() == SubObject->GetName();
-		}))
+			NewConfigData->Rename(nullptr, GetTransientPackage(), RenFlags);
+			NewConfigData->SetFlags(RF_Transient);
+			NewConfigData->ClearFlags(RF_Transactional);
+			FLinkerLoad::InvalidateExport(NewConfigData);
+		}
+
+		// Rename our old config data (along with all sub-objects) to our new CDO.
 		{
-			return *FoundSavedOuter;
+			OldConfigData->Rename(nullptr, DefaultObject, RenFlags);
+
+			// Set the property value directly to avoid going through our custom load logic.
+			const FObjectProperty* CurrentConfigProperty =
+				CastFieldChecked<FObjectProperty>(RootActor->GetClass()->FindPropertyByName(RootActor->GetCurrentConfigDataMemberName()));
+
+			void* PropertyAddress = CurrentConfigProperty->ContainerPtrToValuePtr<UObject>(DefaultObject);
+			CurrentConfigProperty->SetObjectPropertyValue(PropertyAddress, OldConfigData);
 		}
 		
-		// Fallback to original outer.
-		return InOuter;
-	};
-
-	// Restore all saved sub-objects to their original locations. The sub-objects are dynamically added and
-	// already correct by this point. They need to be restored so transaction history can be preserved and undos
-	// function after a compile. If we ever add in compile modification of sub-objects under
-	// CopyTermDefaultsToDefaultObject then we'll need to adjust this logic.
-	for (UObject* SavedSubObject : SavedSubObjects)
-	{
-		if (UObject* NewObject = static_cast<UObject*>(FindObjectWithOuter(DefaultObject, SavedSubObject->GetClass(), *SavedSubObject->GetName())))
+		if (GEditor)
 		{
-			const ERenameFlags RenFlags = REN_DontCreateRedirectors | REN_ForceNoResetLoaders | REN_NonTransactional | REN_DoNotDirty;
-			
-			const bool bSubObjectIsNew = NewObject != SavedSubObject;
-			ensure(bSubObjectIsNew);
-
-			UObject* DesiredOuter = GetCorrectOuter(NewObject);
-			check(DesiredOuter != GetTransientPackage());
-			
-			if (bSubObjectIsNew)
-			{
-				// Invalidate the newly created sub-object.
-				NewObject->Rename(nullptr, GetTransientPackage(), RenFlags);
-				NewObject->SetFlags(RF_Transient);
-				FLinkerLoad::InvalidateExport(NewObject);
-			}
+			GEditor->NotifyToolsOfObjectReplacement(NewToOldSubObjects);
+		}
 
-			if (DesiredOuter != SavedSubObject->GetOuter())
+		// Validate
+		for (UObject* OldSubObject : OldSubObjects)
+		{
+			if (OldSubObject->GetTypedOuter(ADisplayClusterRootActor::StaticClass()) != DefaultObject)
 			{
-				// Restore the original sub-object.
-				SavedSubObject->Rename(nullptr, DesiredOuter, RenFlags);
+				MessageLog.Error(TEXT("Could not preserve original sub-object @@. Transaction history may be lost."), OldSubObject);
 			}
+		}
 
-			if (bSubObjectIsNew)
+		// Fix/hack for UE-136629: Containers (specifically hash based) may need to be reset.
+		// The internal index of the CDO container may mismatch the instanced index because we aren't re-instantiating our
+		// sub-objects so we can maintain transaction history. Once in this state it takes an asset reload to correct.
+		// 
+		// If the index is off and a user adds a new element, it may inadvertently overwrite the wrong index on the instance.
+		// This occurs when we delete a cluster node, compile, then add a new cluster node in. The map is modified via handle,
+		// AddItem() is called, then the keypair value is set. AddItem() will add the value to the CDO, but in the wrong index.
+		// When UE propagates this to the instances, it will add it there in the correct index. When the value is set later the wrong
+		// index of the instance is used since it propagates based on the CDO index.
+		//
+		// It's not clear if this issue would be better solved in the engine's internal map/set propagation or if there
+		// is a better approach for switching out sub-objects on the CDO before the compile process finishes.
+
+		for (UObject* OldSubObject : OldSubObjects)
+		{
+			for (const FProperty* Property : TFieldRange<FProperty>(OldSubObject->GetClass()))
 			{
-				// Update all properties to point to the original sub-object.
-				TArray<UObject*> ObjectsToReplace { NewObject };
-				ObjectTools::ForceReplaceReferences(SavedSubObject, ObjectsToReplace);
+				// Is container
+				if (Property->IsA<FArrayProperty>() ||
+					Property->IsA<FSetProperty>() ||
+					Property->IsA<FMapProperty>())
+				{
+					void* PropertyValue = Property->ContainerPtrToValuePtr<void*>(OldSubObject);
+			
+					FString Value;
+					Property->ExportTextItem(Value, PropertyValue, nullptr, nullptr, 0, nullptr);
+					Property->ImportText(*Value, PropertyValue, 0, OldSubObject);
+				}
 			}
-		}		
+		}
 	}
 }
 
diff --git a/Engine/Plugins/Runtime/nDisplay/Source/DisplayClusterConfigurator/Private/DisplayClusterConfiguratorCompiler.h b/Engine/Plugins/Runtime/nDisplay/Source/DisplayClusterConfigurator/Private/DisplayClusterConfiguratorCompiler.h
--- a/Engine/Plugins/Runtime/nDisplay/Source/DisplayClusterConfigurator/Private/DisplayClusterConfiguratorCompiler.h
+++ b/Engine/Plugins/Runtime/nDisplay/Source/DisplayClusterConfigurator/Private/DisplayClusterConfiguratorCompiler.h
@@ -41,5 +41,5 @@ protected:
 private:
 	class UDisplayClusterBlueprintGeneratedClass* DCGeneratedBP;
 	/** Sub-objects which should survive a compile. */
-	TArray<UObject*> SavedSubObjects;
+	TArray<UObject*> OldSubObjects;
 };
\ No newline at end of file
diff --git a/Engine/Plugins/Runtime/nDisplay/Source/DisplayClusterConfigurator/Private/Views/Details/Cluster/DisplayClusterConfiguratorViewportDetailsCustomization.cpp b/Engine/Plugins/Runtime/nDisplay/Source/DisplayClusterConfigurator/Private/Views/Details/Cluster/DisplayClusterConfiguratorViewportDetailsCustomization.cpp
--- a/Engine/Plugins/Runtime/nDisplay/Source/DisplayClusterConfigurator/Private/Views/Details/Cluster/DisplayClusterConfiguratorViewportDetailsCustomization.cpp
+++ b/Engine/Plugins/Runtime/nDisplay/Source/DisplayClusterConfigurator/Private/Views/Details/Cluster/DisplayClusterConfiguratorViewportDetailsCustomization.cpp
@@ -6,6 +6,7 @@
 #include "DisplayClusterConfigurationTypes.h"
 #include "DisplayClusterConfigurationTypes_Viewport.h"
 #include "DisplayClusterProjectionStrings.h"
+#include "DisplayClusterConfiguratorLog.h"
 #include "Components/DisplayClusterCameraComponent.h"
 #include "Views/Details/Widgets/SDisplayClusterConfigurationSearchableComboBox.h"
 
@@ -23,7 +24,11 @@ void FDisplayClusterConfiguratorViewportDetailsCustomization::CustomizeDetails(I
 
 	// Set config data pointer
 	UDisplayClusterConfigurationData* ConfigurationData = GetConfigData();
-	check(ConfigurationData != nullptr);
+	if (ConfigurationData == nullptr)
+	{
+		UE_LOG(DisplayClusterConfiguratorLog, Warning, TEXT("Details panel config data invalid."));
+		return;
+	}
 
 	ConfigurationDataPtr = ConfigurationData;
 
diff --git a/Engine/Plugins/Runtime/nDisplay/Source/DisplayClusterConfigurator/Private/Views/Details/DisplayClusterConfiguratorBaseDetailCustomization.cpp b/Engine/Plugins/Runtime/nDisplay/Source/DisplayClusterConfigurator/Private/Views/Details/DisplayClusterConfiguratorBaseDetailCustomization.cpp
--- a/Engine/Plugins/Runtime/nDisplay/Source/DisplayClusterConfigurator/Private/Views/Details/DisplayClusterConfiguratorBaseDetailCustomization.cpp
+++ b/Engine/Plugins/Runtime/nDisplay/Source/DisplayClusterConfigurator/Private/Views/Details/DisplayClusterConfiguratorBaseDetailCustomization.cpp
@@ -5,6 +5,7 @@
 #include "DisplayClusterRootActor.h"
 #include "DisplayClusterConfiguratorBlueprintEditor.h"
 #include "DisplayClusterConfiguratorUtils.h"
+#include "DisplayClusterConfiguratorLog.h"
 
 #include "PropertyHandle.h"
 #include "DetailLayoutBuilder.h"
@@ -41,7 +42,12 @@ void FDisplayClusterConfiguratorBaseDetailCustomization::CustomizeDetails(IDetai
 		}
 	}
 
-	check(RootActorPtr.IsValid() || ToolkitPtr.IsValid());
+	if (!(RootActorPtr.IsValid() || ToolkitPtr.IsValid()))
+	{
+		// Possible to hit if details panel selected during undo/redo.
+		UE_LOG(DisplayClusterConfiguratorLog, Warning, TEXT("Details panel root actor and toolkit invalid."));
+		return;
+	}
 
 	// Iterate over all of the properties in the object being edited and process any custom metadata that the properties may have
 	{
@@ -129,6 +135,5 @@ UDisplayClusterConfigurationData* FDisplayClusterConfiguratorBaseDetailCustomiza
 		ConfigData = RootActorPtr->GetConfigData();
 	}
 
-	check(ConfigData);
 	return ConfigData;
 }
\ No newline at end of file
diff --git a/Engine/Plugins/Runtime/nDisplay/Source/DisplayClusterConfigurator/Private/Views/Details/Policies/DisplayClusterConfiguratorPolicyDetailCustomization.cpp b/Engine/Plugins/Runtime/nDisplay/Source/DisplayClusterConfigurator/Private/Views/Details/Policies/DisplayClusterConfiguratorPolicyDetailCustomization.cpp
--- a/Engine/Plugins/Runtime/nDisplay/Source/DisplayClusterConfigurator/Private/Views/Details/Policies/DisplayClusterConfiguratorPolicyDetailCustomization.cpp
+++ b/Engine/Plugins/Runtime/nDisplay/Source/DisplayClusterConfigurator/Private/Views/Details/Policies/DisplayClusterConfiguratorPolicyDetailCustomization.cpp
@@ -6,6 +6,7 @@
 #include "DisplayClusterConfigurationTypes.h"
 #include "DisplayClusterConfiguratorUtils.h"
 #include "DisplayClusterConfiguratorPropertyUtils.h"
+#include "DisplayClusterConfiguratorLog.h"
 #include "DisplayClusterConfiguratorPolicyParameterCustomization.h"
 #include "Views/Details/Widgets/SDisplayClusterConfigurationSearchableComboBox.h"
 
@@ -341,7 +342,11 @@ void FDisplayClusterConfiguratorProjectionCustomization::BuildParametersForPolic
 	}
 	
 	UDisplayClusterBlueprint* Blueprint = FDisplayClusterConfiguratorUtils::FindBlueprintFromObject(EditingObject.Get());
-	check(Blueprint);
+	if (Blueprint == nullptr)
+	{
+		UE_LOG(DisplayClusterConfiguratorLog, Warning, TEXT("Details policy selection blueprint invalid."));
+		return;
+	}
 
 	const FString PolicyLower = Policy.ToLower();
 	/*
