ndisplay: fixed cis build error

diff --git a/Engine/Plugins/Runtime/nDisplay/Source/DisplayCluster/Private/Game/EngineClasses/Scene/DisplayClusterICVFXCameraComponent.cpp b/Engine/Plugins/Runtime/nDisplay/Source/DisplayCluster/Private/Game/EngineClasses/Scene/DisplayClusterICVFXCameraComponent.cpp
--- a/Engine/Plugins/Runtime/nDisplay/Source/DisplayCluster/Private/Game/EngineClasses/Scene/DisplayClusterICVFXCameraComponent.cpp
+++ b/Engine/Plugins/Runtime/nDisplay/Source/DisplayCluster/Private/Game/EngineClasses/Scene/DisplayClusterICVFXCameraComponent.cpp
@@ -85,7 +85,10 @@ void UDisplayClusterICVFXCameraComponent::OnRegister()
 
 #if WITH_EDITORONLY_DATA
 	// disable frustum for icvfx camera component
-	DrawFrustum->bFrustumEnabled = false;
+	if (DrawFrustum != nullptr)
+	{
+		DrawFrustum->bFrustumEnabled = false;
+	}
 
 	// Update ExternalCineactor behaviour
 	UpdateICVFXPreviewState();
