nDisplay: Fixed issue that prevented node movement operations from being added to the undo stack in the Ouput Mapping graph editor.

diff --git a/Engine/Plugins/Runtime/nDisplay/Source/DisplayClusterConfigurator/Private/Views/OutputMapping/GraphNodes/SDisplayClusterConfiguratorBaseNode.cpp b/Engine/Plugins/Runtime/nDisplay/Source/DisplayClusterConfigurator/Private/Views/OutputMapping/GraphNodes/SDisplayClusterConfiguratorBaseNode.cpp
--- a/Engine/Plugins/Runtime/nDisplay/Source/DisplayClusterConfigurator/Private/Views/OutputMapping/GraphNodes/SDisplayClusterConfiguratorBaseNode.cpp
+++ b/Engine/Plugins/Runtime/nDisplay/Source/DisplayClusterConfigurator/Private/Views/OutputMapping/GraphNodes/SDisplayClusterConfiguratorBaseNode.cpp
@@ -186,6 +186,26 @@ void SDisplayClusterConfiguratorBaseNode::MoveTo(const FVector2D& NewPosition, F
 		return;
 	}
 
+	// The graph editor applies the "true" movement of the node after the user is done dragging the node around,
+	// first resetting the selected nodes back to their starting position, then moving them to their final position,
+	// at which point the nodes should be marking themselves as modified for the undo buffer. In this case,
+	// all positions have been properly computed and the NewPosition can be applied directly to each node.
+	if (!FSlateApplication::Get().GetPressedMouseButtons().Contains(EKeys::LeftMouseButton))
+	{
+		SGraphNode::MoveTo(NewPosition, NodeFilter, bMarkDirty);
+
+		// If the parent node is being auto-positioned, add it to the undo stack here because we need to store its old position with the this node's old position
+		// so that if the move operation is undone, this node can appropriately reset the backing config object's position without requiring a full auto-positioning pass.
+		if (ParentEdNode && ParentEdNode->IsNodeAutoPositioned())
+		{
+			ParentEdNode->Modify(bMarkDirty);
+		}
+
+		EdNode->UpdateObject();
+
+		return;
+	}
+
 	const bool bIsNodeFiltered = NodeFilter.Contains(SharedThis(this));
 	if (!bIsNodeFiltered && !EdNode->IsUserInteractingWithNode())
 	{
@@ -256,7 +276,7 @@ void SDisplayClusterConfiguratorBaseNode::MoveTo(const FVector2D& NewPosition, F
 		{
 			if (BaseNode != EdNode)
 			{
-				BaseNode->Modify();
+				BaseNode->Modify(bMarkDirty);
 				BaseNode->NodePosX += BestOffset.X + AlignmentOffset.X;
 				BaseNode->NodePosY += BestOffset.Y + AlignmentOffset.Y;
 
@@ -269,13 +289,6 @@ void SDisplayClusterConfiguratorBaseNode::MoveTo(const FVector2D& NewPosition, F
 
 	if (!bIsNodeFiltered)
 	{
-		// If the parent node is being auto-positioned, add it to the undo stack here because we need to store its old position with the this node's old position
-		// so that if the move operation is undone, this node can appropriately reset the backing config object's position without requiring a full auto-positioning pass.
-		if (ParentEdNode && ParentEdNode->IsNodeAutoPositioned())
-		{
-			ParentEdNode->Modify();
-		}
-
 		EdNode->UpdateObject();
 		EdNode->UpdateChildNodes();
 	}
