nDisplay: Disabled Lumen on lightcard and chromakey renders.

diff --git a/Engine/Plugins/Runtime/nDisplay/Source/DisplayCluster/Private/Render/Viewport/DisplayClusterViewportManager.cpp b/Engine/Plugins/Runtime/nDisplay/Source/DisplayCluster/Private/Render/Viewport/DisplayClusterViewportManager.cpp
--- a/Engine/Plugins/Runtime/nDisplay/Source/DisplayCluster/Private/Render/Viewport/DisplayClusterViewportManager.cpp
+++ b/Engine/Plugins/Runtime/nDisplay/Source/DisplayCluster/Private/Render/Viewport/DisplayClusterViewportManager.cpp
@@ -448,15 +448,21 @@ void FDisplayClusterViewportManager::ConfigureViewFamily(const FDisplayClusterRe
 			break;
 		}
 	}
-	// Disable nanite for lightcards or chromakey
+	// Disable nanite and lumen for lightcards/chromakey
 	{
 		switch (InFrameTarget.CaptureMode)
 		{
 		case EDisplayClusterViewportCaptureMode::Chromakey:
 		case EDisplayClusterViewportCaptureMode::Lightcard:
 		case EDisplayClusterViewportCaptureMode::Lightcard_OCIO:
-			ViewFamily.EngineShowFlags.NaniteMeshes = false;
+
+			ViewFamily.EngineShowFlags.SetNaniteMeshes(0);
+			
+			ViewFamily.EngineShowFlags.SetLumenReflections(0);
+			ViewFamily.EngineShowFlags.SetLumenGlobalIllumination(0);
+
 			break;
+
 		default:
 			break;
 		}
