nDisplay: Fixed issue where property categories that were supposed to be hidden on the Display Cluster Root Actor were visible.

diff --git a/Engine/Plugins/Runtime/nDisplay/Source/DisplayCluster/Public/DisplayClusterRootActor.h b/Engine/Plugins/Runtime/nDisplay/Source/DisplayCluster/Public/DisplayClusterRootActor.h
--- a/Engine/Plugins/Runtime/nDisplay/Source/DisplayCluster/Public/DisplayClusterRootActor.h
+++ b/Engine/Plugins/Runtime/nDisplay/Source/DisplayCluster/Public/DisplayClusterRootActor.h
@@ -40,7 +40,7 @@ class UProceduralMeshComponent;
 /**
  * VR root. This contains nDisplay VR hierarchy in the game.
  */
-UCLASS(HideCategories=(Replication, Collision, Input, Actor, HLOD, Cooking, Physics, Activation, AssetUserData, ActorTick, Advanced), meta=(DisplayName = "nDisplay Root Actor"))
+UCLASS(HideCategories=(Replication, Collision, Input, Actor, HLOD, Cooking, Physics, Activation, AssetUserData, ActorTick, Advanced, WorldPartition, DataLayers, Events), meta=(DisplayName = "nDisplay Root Actor"))
 class DISPLAYCLUSTER_API ADisplayClusterRootActor
 	: public AActor
 {
