nDisplay: various updates
- got rid of master/slave terminology
- new JSON 500 config scheme
- removed old CFG config scheme

diff --git a/Engine/Plugins/Runtime/LiveLinkOvernDisplay/Source/LiveLinkOverNDisplay/Private/NDisplayLiveLinkSubjectReplicator.cpp b/Engine/Plugins/Runtime/LiveLinkOvernDisplay/Source/LiveLinkOverNDisplay/Private/NDisplayLiveLinkSubjectReplicator.cpp
--- a/Engine/Plugins/Runtime/LiveLinkOvernDisplay/Source/LiveLinkOverNDisplay/Private/NDisplayLiveLinkSubjectReplicator.cpp
+++ b/Engine/Plugins/Runtime/LiveLinkOvernDisplay/Source/LiveLinkOverNDisplay/Private/NDisplayLiveLinkSubjectReplicator.cpp
@@ -87,12 +87,12 @@ void FNDisplayLiveLinkSubjectReplicator::Initialize()
 
 		if (IDisplayCluster::IsAvailable())
 		{
-			if (IDisplayCluster::Get().GetClusterMgr()->IsMaster())
+			if (IDisplayCluster::Get().GetClusterMgr()->IsPrimary())
 			{
 				//If we're a master, listen to LiveLinkTicked callback to bundle up the SyncObject. Subjects frames will have been processed at that point.
 				LiveLinkClient->OnLiveLinkTicked().AddRaw(this, &FNDisplayLiveLinkSubjectReplicator::OnLiveLinkTicked);
 			}
-			else if (IDisplayCluster::Get().GetClusterMgr()->IsSlave())
+			else if (IDisplayCluster::Get().GetClusterMgr()->IsSecondary())
 			{
 				//If we're a slave, listen for new subject and disables them if we're tracking that subject
 				FCoreDelegates::OnBeginFrame.AddRaw(this, &FNDisplayLiveLinkSubjectReplicator::OnEngineBeginFrame);
@@ -164,7 +164,7 @@ void FNDisplayLiveLinkSubjectReplicator::OnEngineBeginFrame()
 	//Only our replicated VirtualSubjects should be enabled
 	if (LiveLinkClient)
 	{
-		check(IDisplayCluster::Get().GetClusterMgr()->IsSlave());
+		check(IDisplayCluster::Get().GetClusterMgr()->IsSecondary());
 
 		const bool bIncludeDisabled = false;
 		const bool bIncludeVirtual = true;
diff --git a/Engine/Plugins/Runtime/nDisplay/Source/DisplayCluster/Private/Blueprints/DisplayClusterBlueprintAPIImpl.cpp b/Engine/Plugins/Runtime/nDisplay/Source/DisplayCluster/Private/Blueprints/DisplayClusterBlueprintAPIImpl.cpp
--- a/Engine/Plugins/Runtime/nDisplay/Source/DisplayCluster/Private/Blueprints/DisplayClusterBlueprintAPIImpl.cpp
+++ b/Engine/Plugins/Runtime/nDisplay/Source/DisplayCluster/Private/Blueprints/DisplayClusterBlueprintAPIImpl.cpp
@@ -80,18 +80,18 @@ int32 UDisplayClusterBlueprintAPIImpl::GetActiveNodesAmount() const
 	return NodesAmount;
 }
 
-bool UDisplayClusterBlueprintAPIImpl::IsMaster() const
+bool UDisplayClusterBlueprintAPIImpl::IsPrimary() const
 {
-	const bool bIsMaster = IDisplayCluster::Get().GetClusterMgr()->IsMaster();
-	UE_LOG(LogDisplayClusterBlueprint, Verbose, TEXT("IsMaster - %s"), *DisplayClusterTypesConverter::template ToString(bIsMaster));
-	return bIsMaster;
+	const bool bIsPrimary = IDisplayCluster::Get().GetClusterMgr()->IsPrimary();
+	UE_LOG(LogDisplayClusterBlueprint, Verbose, TEXT("IsPrimary - %s"), *DisplayClusterTypesConverter::template ToString(bIsPrimary));
+	return bIsPrimary;
 }
 
-bool UDisplayClusterBlueprintAPIImpl::IsSlave() const
+bool UDisplayClusterBlueprintAPIImpl::IsSecondary() const
 {
-	const bool bIsSlave = IDisplayCluster::Get().GetClusterMgr()->IsSlave();
-	UE_LOG(LogDisplayClusterBlueprint, Verbose, TEXT("IsSlave - %s"), *DisplayClusterTypesConverter::template ToString(bIsSlave));
-	return bIsSlave;
+	const bool bIsSecondary = IDisplayCluster::Get().GetClusterMgr()->IsSecondary();
+	UE_LOG(LogDisplayClusterBlueprint, Verbose, TEXT("IsSecondary - %s"), *DisplayClusterTypesConverter::template ToString(bIsSecondary));
+	return bIsSecondary;
 }
 
 bool UDisplayClusterBlueprintAPIImpl::IsBackup() const
@@ -122,34 +122,34 @@ void UDisplayClusterBlueprintAPIImpl::RemoveClusterEventListener(TScriptInterfac
 	IDisplayCluster::Get().GetClusterMgr()->RemoveClusterEventListener(Listener);
 }
 
-void UDisplayClusterBlueprintAPIImpl::EmitClusterEventJson(const FDisplayClusterClusterEventJson& Event, bool bMasterOnly)
+void UDisplayClusterBlueprintAPIImpl::EmitClusterEventJson(const FDisplayClusterClusterEventJson& Event, bool bPrimaryOnly)
 {
-	UE_LOG(LogDisplayClusterBlueprint, Verbose, TEXT("EmitClusterEventJson - emitting cluster event: bMasterOnly='%s' Category='%s' Type='%s' Name='%s'"),
-		*DisplayClusterTypesConverter::template ToString(bMasterOnly), *Event.Category, *Event.Type, *Event.Name);
+	UE_LOG(LogDisplayClusterBlueprint, Verbose, TEXT("EmitClusterEventJson - emitting cluster event: bPrimaryOnly='%s' Category='%s' Type='%s' Name='%s'"),
+		*DisplayClusterTypesConverter::template ToString(bPrimaryOnly), *Event.Category, *Event.Type, *Event.Name);
 
-	IDisplayCluster::Get().GetClusterMgr()->EmitClusterEventJson(Event, bMasterOnly);
+	IDisplayCluster::Get().GetClusterMgr()->EmitClusterEventJson(Event, bPrimaryOnly);
 }
 
-void UDisplayClusterBlueprintAPIImpl::EmitClusterEventBinary(const FDisplayClusterClusterEventBinary& Event, bool bMasterOnly)
+void UDisplayClusterBlueprintAPIImpl::EmitClusterEventBinary(const FDisplayClusterClusterEventBinary& Event, bool bPrimaryOnly)
 {
-	UE_LOG(LogDisplayClusterBlueprint, Verbose, TEXT("EmitClusterEventBinary - emitting cluster event: bMasterOnly='%s' EventId='%d'"),
-		*DisplayClusterTypesConverter::template ToString(bMasterOnly), Event.EventId);
+	UE_LOG(LogDisplayClusterBlueprint, Verbose, TEXT("EmitClusterEventBinary - emitting cluster event: bPrimaryOnly='%s' EventId='%d'"),
+		*DisplayClusterTypesConverter::template ToString(bPrimaryOnly), Event.EventId);
 
-	IDisplayCluster::Get().GetClusterMgr()->EmitClusterEventBinary(Event, bMasterOnly);
+	IDisplayCluster::Get().GetClusterMgr()->EmitClusterEventBinary(Event, bPrimaryOnly);
 }
 
-void UDisplayClusterBlueprintAPIImpl::SendClusterEventJsonTo(const FString& Address, const int32 Port, const FDisplayClusterClusterEventJson& Event, bool bMasterOnly)
+void UDisplayClusterBlueprintAPIImpl::SendClusterEventJsonTo(const FString& Address, const int32 Port, const FDisplayClusterClusterEventJson& Event, bool bPrimaryOnly)
 {
 	checkSlow(Port > 0 && Port <= UINT16_MAX);
 
 	UE_LOG(LogDisplayClusterBlueprint, Verbose, TEXT("SendClusterEventJsonTo - sending json event to %s:%d, Category='%s' Type='%s' Name='%s'"), *Address, Port, *Event.Category, *Event.Type, *Event.Name);
-	IDisplayCluster::Get().GetClusterMgr()->SendClusterEventTo(Address, static_cast<uint16>(Port), Event, bMasterOnly);
+	IDisplayCluster::Get().GetClusterMgr()->SendClusterEventTo(Address, static_cast<uint16>(Port), Event, bPrimaryOnly);
 }
 
-void UDisplayClusterBlueprintAPIImpl::SendClusterEventBinaryTo(const FString& Address, const int32 Port, const FDisplayClusterClusterEventBinary& Event, bool bMasterOnly)
+void UDisplayClusterBlueprintAPIImpl::SendClusterEventBinaryTo(const FString& Address, const int32 Port, const FDisplayClusterClusterEventBinary& Event, bool bPrimaryOnly)
 {
 	checkSlow(Port > 0 && Port <= UINT16_MAX);
 
 	UE_LOG(LogDisplayClusterBlueprint, Verbose, TEXT("SendClusterEventBinaryTo - sending binary event to %s:%d, EventId='%d'"), *Address, Port, Event.EventId);
-	IDisplayCluster::Get().GetClusterMgr()->SendClusterEventTo(Address, static_cast<uint16>(Port), Event, bMasterOnly);
+	IDisplayCluster::Get().GetClusterMgr()->SendClusterEventTo(Address, static_cast<uint16>(Port), Event, bPrimaryOnly);
 }
diff --git a/Engine/Plugins/Runtime/nDisplay/Source/DisplayCluster/Private/Blueprints/DisplayClusterBlueprintAPIImpl.h b/Engine/Plugins/Runtime/nDisplay/Source/DisplayCluster/Private/Blueprints/DisplayClusterBlueprintAPIImpl.h
--- a/Engine/Plugins/Runtime/nDisplay/Source/DisplayCluster/Private/Blueprints/DisplayClusterBlueprintAPIImpl.h
+++ b/Engine/Plugins/Runtime/nDisplay/Source/DisplayCluster/Private/Blueprints/DisplayClusterBlueprintAPIImpl.h
@@ -47,11 +47,11 @@ public:
 	UFUNCTION(BlueprintCallable, meta = (DisplayName = "Get cluster role"), Category = "NDisplay|Cluster")
 	virtual EDisplayClusterNodeRole GetClusterRole() const override;
 
-	UFUNCTION(BlueprintCallable, meta = (DisplayName = "Is master node"), Category = "NDisplay|Cluster")
-	virtual bool IsMaster() const override;
+	UFUNCTION(BlueprintCallable, meta = (DisplayName = "Is a primary node"), Category = "NDisplay|Cluster")
+	virtual bool IsPrimary() const override;
 	
-	UFUNCTION(BlueprintCallable, meta = (DisplayName = "Is slave node"), Category = "NDisplay|Cluster")
-	virtual bool IsSlave() const override;
+	UFUNCTION(BlueprintCallable, meta = (DisplayName = "Is a secondary node"), Category = "NDisplay|Cluster")
+	virtual bool IsSecondary() const override;
 
 	/** Returns true if current node is a backup node in a cluster. */
 	UFUNCTION(BlueprintCallable, meta = (DisplayName = "Is backup node"), Category = "NDisplay|Cluster")
@@ -65,14 +65,14 @@ public:
 	virtual void RemoveClusterEventListener(TScriptInterface<IDisplayClusterClusterEventListener> Listener) override;
 
 	UFUNCTION(BlueprintCallable, meta = (DisplayName = "Emit JSON cluster event"), Category = "NDisplay|Cluster")
-	virtual void EmitClusterEventJson(const FDisplayClusterClusterEventJson& Event, bool bMasterOnly) override;
+	virtual void EmitClusterEventJson(const FDisplayClusterClusterEventJson& Event, bool bPrimaryOnly) override;
 
 	UFUNCTION(BlueprintCallable, meta = (DisplayName = "Emit binary cluster event"), Category = "NDisplay|Cluster")
-	virtual void EmitClusterEventBinary(const FDisplayClusterClusterEventBinary& Event, bool bMasterOnly) override;
+	virtual void EmitClusterEventBinary(const FDisplayClusterClusterEventBinary& Event, bool bPrimaryOnly) override;
 
 	UFUNCTION(BlueprintCallable, meta = (DisplayName = "Emit JSON cluster event"), Category = "NDisplay|Cluster")
-	virtual void SendClusterEventJsonTo(const FString& Address, const int32 Port, const FDisplayClusterClusterEventJson& Event, bool bMasterOnly) override;
+	virtual void SendClusterEventJsonTo(const FString& Address, const int32 Port, const FDisplayClusterClusterEventJson& Event, bool bPrimaryOnly) override;
 
 	UFUNCTION(BlueprintCallable, meta = (DisplayName = "Emit binary cluster event"), Category = "NDisplay|Cluster")
-	virtual void SendClusterEventBinaryTo(const FString& Address, const int32 Port, const FDisplayClusterClusterEventBinary& Event, bool bMasterOnly) override;
+	virtual void SendClusterEventBinaryTo(const FString& Address, const int32 Port, const FDisplayClusterClusterEventBinary& Event, bool bPrimaryOnly) override;
 };
diff --git a/Engine/Plugins/Runtime/nDisplay/Source/DisplayCluster/Private/Cluster/Controller/DisplayClusterClusterNodeCtrlBase.cpp b/Engine/Plugins/Runtime/nDisplay/Source/DisplayCluster/Private/Cluster/Controller/DisplayClusterClusterNodeCtrlBase.cpp
--- a/Engine/Plugins/Runtime/nDisplay/Source/DisplayCluster/Private/Cluster/Controller/DisplayClusterClusterNodeCtrlBase.cpp
+++ b/Engine/Plugins/Runtime/nDisplay/Source/DisplayCluster/Private/Cluster/Controller/DisplayClusterClusterNodeCtrlBase.cpp
@@ -64,7 +64,7 @@ void FDisplayClusterClusterNodeCtrlBase::Shutdown()
 	StopServers();
 }
 
-void FDisplayClusterClusterNodeCtrlBase::SendClusterEventTo(const FString& Address, const uint16 Port, const FDisplayClusterClusterEventJson& Event, bool bMasterOnly)
+void FDisplayClusterClusterNodeCtrlBase::SendClusterEventTo(const FString& Address, const uint16 Port, const FDisplayClusterClusterEventJson& Event, bool bPrimaryOnly)
 {
 	// We should synchronize access to the client
 	FScopeLock Lock(&ExternEventsClientJsonGuard);
@@ -75,7 +75,7 @@ void FDisplayClusterClusterNodeCtrlBase::SendClusterEventTo(const FString& Addre
 	ExternalEventsClientJson->Disconnect();
 }
 
-void FDisplayClusterClusterNodeCtrlBase::SendClusterEventTo(const FString& Address, const uint16 Port, const FDisplayClusterClusterEventBinary& Event, bool bMasterOnly)
+void FDisplayClusterClusterNodeCtrlBase::SendClusterEventTo(const FString& Address, const uint16 Port, const FDisplayClusterClusterEventBinary& Event, bool bPrimaryOnly)
 {
 	// We should synchronize access to the client
 	FScopeLock Lock(&ExternEventsClientBinaryGuard);
diff --git a/Engine/Plugins/Runtime/nDisplay/Source/DisplayCluster/Private/Cluster/Controller/DisplayClusterClusterNodeCtrlBase.h b/Engine/Plugins/Runtime/nDisplay/Source/DisplayCluster/Private/Cluster/Controller/DisplayClusterClusterNodeCtrlBase.h
--- a/Engine/Plugins/Runtime/nDisplay/Source/DisplayCluster/Private/Cluster/Controller/DisplayClusterClusterNodeCtrlBase.h
+++ b/Engine/Plugins/Runtime/nDisplay/Source/DisplayCluster/Private/Cluster/Controller/DisplayClusterClusterNodeCtrlBase.h
@@ -38,8 +38,8 @@ public:
 		return ControllerName;
 	}
 
-	virtual void SendClusterEventTo(const FString& Address, const uint16 Port, const FDisplayClusterClusterEventJson& Event, bool bMasterOnly) override;
-	virtual void SendClusterEventTo(const FString& Address, const uint16 Port, const FDisplayClusterClusterEventBinary& Event, bool bMasterOnly) override;
+	virtual void SendClusterEventTo(const FString& Address, const uint16 Port, const FDisplayClusterClusterEventJson&   Event, bool bPrimaryOnly) override;
+	virtual void SendClusterEventTo(const FString& Address, const uint16 Port, const FDisplayClusterClusterEventBinary& Event, bool bPrimaryOnly) override;
 
 public:
 	//////////////////////////////////////////////////////////////////////////////////////////////
diff --git a/Engine/Plugins/Runtime/nDisplay/Source/DisplayCluster/Private/Cluster/Controller/DisplayClusterClusterNodeCtrlEditor.cpp b/Engine/Plugins/Runtime/nDisplay/Source/DisplayCluster/Private/Cluster/Controller/DisplayClusterClusterNodeCtrlEditor.cpp
--- a/Engine/Plugins/Runtime/nDisplay/Source/DisplayCluster/Private/Cluster/Controller/DisplayClusterClusterNodeCtrlEditor.cpp
+++ b/Engine/Plugins/Runtime/nDisplay/Source/DisplayCluster/Private/Cluster/Controller/DisplayClusterClusterNodeCtrlEditor.cpp
@@ -51,7 +51,7 @@ bool FDisplayClusterClusterNodeCtrlEditor::StartServers()
 
 	// Always use localhost for PIE because there might be some other host specified in the configuration data
 	const FString HostForPie("127.0.0.1");
-	const FDisplayClusterConfigurationMasterNodePorts& Ports = ConfigData->Cluster->MasterNode.Ports;
+	const FDisplayClusterConfigurationPrimaryNodePorts& Ports = ConfigData->Cluster->PrimaryNode.Ports;
 
 	UE_LOG(LogDisplayClusterCluster, Log, TEXT("Servers: addr %s, port_cs %d, port_ce %d, port_ceb %d"), *HostForPie, Ports.ClusterSync, Ports.ClusterEventsJson, Ports.ClusterEventsBinary);
 
diff --git a/Engine/Plugins/Runtime/nDisplay/Source/DisplayCluster/Private/Cluster/Controller/DisplayClusterClusterNodeCtrlMaster.cpp b/Engine/Plugins/Runtime/nDisplay/Source/DisplayCluster/Private/Cluster/Controller/DisplayClusterClusterNodeCtrlMaster.cpp
deleted file mode 100644
diff --git a/Engine/Plugins/Runtime/nDisplay/Source/DisplayCluster/Private/Cluster/Controller/DisplayClusterClusterNodeCtrlMaster.h b/Engine/Plugins/Runtime/nDisplay/Source/DisplayCluster/Private/Cluster/Controller/DisplayClusterClusterNodeCtrlMaster.h
deleted file mode 100644
diff --git a/Engine/Plugins/Runtime/nDisplay/Source/DisplayCluster/Private/Cluster/Controller/DisplayClusterClusterNodeCtrlNull.h b/Engine/Plugins/Runtime/nDisplay/Source/DisplayCluster/Private/Cluster/Controller/DisplayClusterClusterNodeCtrlNull.h
--- a/Engine/Plugins/Runtime/nDisplay/Source/DisplayCluster/Private/Cluster/Controller/DisplayClusterClusterNodeCtrlNull.h
+++ b/Engine/Plugins/Runtime/nDisplay/Source/DisplayCluster/Private/Cluster/Controller/DisplayClusterClusterNodeCtrlNull.h
@@ -6,7 +6,7 @@
 
 
 /**
- * Null controller. It's used when no 'legal' controller istantiated.
+ * Null controller. It's used when no 'legal' controller instantiated.
  * The main goal is to always have a cluster controller instance.
  */
 class FDisplayClusterClusterNodeCtrlNull final
@@ -47,10 +47,10 @@ public:
 		return ControllerName;
 	}
 
-	virtual void SendClusterEventTo(const FString& Address, const uint16 Port, const FDisplayClusterClusterEventBinary& Event, bool bMasterOnly) override
+	virtual void SendClusterEventTo(const FString& Address, const uint16 Port, const FDisplayClusterClusterEventBinary& Event, bool bPrimaryOnly) override
 	{ }
 
-	virtual void SendClusterEventTo(const FString& Address, const uint16 Port, const FDisplayClusterClusterEventJson& Event, bool bMasterOnly) override
+	virtual void SendClusterEventTo(const FString& Address, const uint16 Port, const FDisplayClusterClusterEventJson& Event, bool bPrimaryOnly) override
 	{ }
 
 public:
diff --git a/Engine/Plugins/Runtime/nDisplay/Source/DisplayCluster/Private/Cluster/Controller/DisplayClusterClusterNodeCtrlPrimary.cpp b/Engine/Plugins/Runtime/nDisplay/Source/DisplayCluster/Private/Cluster/Controller/DisplayClusterClusterNodeCtrlPrimary.cpp
new file mode 100644
--- /dev/null
+++ b/Engine/Plugins/Runtime/nDisplay/Source/DisplayCluster/Private/Cluster/Controller/DisplayClusterClusterNodeCtrlPrimary.cpp
@@ -0,0 +1,225 @@
+// Copyright Epic Games, Inc. All Rights Reserved.
+
+#include "Cluster/Controller/DisplayClusterClusterNodeCtrlPrimary.h"
+
+#include "Cluster/IPDisplayClusterClusterManager.h"
+#include "Cluster/Failover/IDisplayClusterFailoverNodeController.h"
+#include "Config/IPDisplayClusterConfigManager.h"
+
+#include "DisplayClusterConfigurationTypes.h"
+
+#include "Network/Service/ClusterSync/DisplayClusterClusterSyncService.h"
+#include "Network/Service/RenderSync/DisplayClusterRenderSyncService.h"
+#include "Network/Service/ClusterEventsJson/DisplayClusterClusterEventsJsonService.h"
+#include "Network/Service/ClusterEventsBinary/DisplayClusterClusterEventsBinaryService.h"
+#include "Network/Listener/DisplayClusterTcpListener.h"
+#include "Network/DisplayClusterNetworkTypes.h"
+
+#include "Misc/DisplayClusterGlobals.h"
+#include "Misc/DisplayClusterLog.h"
+
+
+FDisplayClusterClusterNodeCtrlPrimary::FDisplayClusterClusterNodeCtrlPrimary(const FString& CtrlName, const FString& NodeName)
+	: FDisplayClusterClusterNodeCtrlSecondary(CtrlName, NodeName)
+{
+	// TCP connection listener for all nDisplay internal services
+	TcpListener = MakeShared<FDisplayClusterTcpListener>(true, FString("nDisplay-TCP-listener"));
+}
+
+FDisplayClusterClusterNodeCtrlPrimary::~FDisplayClusterClusterNodeCtrlPrimary()
+{
+	UE_LOG(LogDisplayClusterNetwork, Log, TEXT("Releasing %s cluster controller..."), *GetControllerName());
+
+	Shutdown();
+}
+
+
+//////////////////////////////////////////////////////////////////////////////////////////////
+// IDisplayClusterProtocolClusterSync
+//////////////////////////////////////////////////////////////////////////////////////////////
+EDisplayClusterCommResult FDisplayClusterClusterNodeCtrlPrimary::GetTimeData(double& OutDeltaTime, double& OutGameTime, TOptional<FQualifiedFrameTime>& OutFrameTime)
+{
+	if (IPDisplayClusterClusterManager* const ClusterMgr = GDisplayCluster->GetPrivateClusterMgr())
+	{
+		ClusterMgr->ExportTimeData(OutDeltaTime, OutGameTime, OutFrameTime);
+		return EDisplayClusterCommResult::Ok;
+	}
+
+	return EDisplayClusterCommResult::InternalError;
+}
+
+EDisplayClusterCommResult FDisplayClusterClusterNodeCtrlPrimary::GetObjectsData(const EDisplayClusterSyncGroup InSyncGroup, TMap<FString, FString>& OutObjectsData)
+{
+	if (IPDisplayClusterClusterManager* const ClusterMgr = GDisplayCluster->GetPrivateClusterMgr())
+	{
+		ClusterMgr->ExportObjectsData(InSyncGroup, OutObjectsData);
+		return EDisplayClusterCommResult::Ok;
+	}
+
+	return EDisplayClusterCommResult::InternalError;
+}
+
+EDisplayClusterCommResult FDisplayClusterClusterNodeCtrlPrimary::GetEventsData(TArray<TSharedPtr<FDisplayClusterClusterEventJson, ESPMode::ThreadSafe>>& OutJsonEvents, TArray<TSharedPtr<FDisplayClusterClusterEventBinary, ESPMode::ThreadSafe>>& OutBinaryEvents)
+{
+	if (IPDisplayClusterClusterManager* const ClusterMgr = GDisplayCluster->GetPrivateClusterMgr())
+	{
+		ClusterMgr->ExportEventsData(OutJsonEvents, OutBinaryEvents);
+		return EDisplayClusterCommResult::Ok;
+	}
+
+	return EDisplayClusterCommResult::InternalError;
+}
+
+EDisplayClusterCommResult FDisplayClusterClusterNodeCtrlPrimary::GetNativeInputData(TMap<FString, FString>& OutNativeInputData)
+{
+	if (IPDisplayClusterClusterManager* const ClusterMgr = GDisplayCluster->GetPrivateClusterMgr())
+	{
+		ClusterMgr->ExportNativeInputData(OutNativeInputData);
+		return EDisplayClusterCommResult::Ok;
+	}
+
+	return EDisplayClusterCommResult::InternalError;
+}
+
+
+//////////////////////////////////////////////////////////////////////////////////////////////
+// IDisplayClusterClusterNodeController
+//////////////////////////////////////////////////////////////////////////////////////////////
+bool FDisplayClusterClusterNodeCtrlPrimary::DropClusterNode(const FString& NodeId)
+{
+	// Kill all sessions of the requested node
+	ClusterSyncServer->KillSession(NodeId);
+	RenderSyncServer->KillSession(NodeId);
+	ClusterEventsJsonServer->KillSession(NodeId);
+	ClusterEventsBinaryServer->KillSession(NodeId);
+
+	return true;
+}
+
+
+//////////////////////////////////////////////////////////////////////////////////////////////
+// FDisplayClusterNodeCtrlBase
+//////////////////////////////////////////////////////////////////////////////////////////////
+bool FDisplayClusterClusterNodeCtrlPrimary::InitializeServers()
+{
+	if (GDisplayCluster->GetOperationMode() == EDisplayClusterOperationMode::Disabled)
+	{
+		return false;
+	}
+
+	if (!FDisplayClusterClusterNodeCtrlSecondary::InitializeServers())
+	{
+		return false;
+	}
+
+	UE_LOG(LogDisplayClusterCluster, Log, TEXT("%s - initializing primary node servers..."), *GetControllerName());
+
+	// Instantiate node servers
+	ClusterSyncServer         = MakeUnique<FDisplayClusterClusterSyncService>       ();
+	RenderSyncServer          = MakeUnique<FDisplayClusterRenderSyncService>        ();
+	ClusterEventsJsonServer   = MakeUnique<FDisplayClusterClusterEventsJsonService> ();
+	ClusterEventsBinaryServer = MakeUnique<FDisplayClusterClusterEventsBinaryService>();
+
+	return ClusterSyncServer && RenderSyncServer && ClusterEventsJsonServer && ClusterEventsBinaryServer;
+}
+
+bool FDisplayClusterClusterNodeCtrlPrimary::StartServers()
+{
+	if (!FDisplayClusterClusterNodeCtrlSecondary::StartServers())
+	{
+		return false;
+	}
+
+	UE_LOG(LogDisplayClusterCluster, Log, TEXT("%s - starting primary node servers..."), *GetControllerName());
+
+	// Get config data
+	const UDisplayClusterConfigurationData* ConfigData = GDisplayCluster->GetPrivateConfigMgr()->GetConfig();
+	if (!ConfigData)
+	{
+		UE_LOG(LogDisplayClusterCluster, Error, TEXT("Couldn't get configuration data"));
+		return false;
+	}
+
+	// Get primary node config data
+	const UDisplayClusterConfigurationClusterNode* const CfgPrimary = GDisplayCluster->GetPrivateConfigMgr()->GetPrimaryNode();
+	if (!CfgPrimary)
+	{
+		UE_LOG(LogDisplayClusterCluster, Error, TEXT("No primary node configuration data found"));
+		return false;
+	}
+
+	const FDisplayClusterConfigurationPrimaryNodePorts& Ports = ConfigData->Cluster->PrimaryNode.Ports;
+	UE_LOG(LogDisplayClusterCluster, Log, TEXT("Servers: addr %s, port_cs %d, port_ce %d, port_ceb %d"), *CfgPrimary->Host, Ports.ClusterSync, Ports.ClusterEventsJson, Ports.ClusterEventsBinary);
+
+	// Connection validation lambda
+	auto IsConnectionAllowedFunc = [](const FDisplayClusterSessionInfo& SessionInfo)
+	{
+		// Here we make sure the node belongs to the cluster
+		TArray<FString> NodeIds;
+		GDisplayCluster->GetPrivateClusterMgr()->GetNodeIds(NodeIds);
+		return NodeIds.ContainsByPredicate([SessionInfo](const FString& Item)
+			{
+				return Item.Equals(SessionInfo.NodeId.Get(FString()), ESearchCase::IgnoreCase);
+			});
+	};
+
+	// Set connection validation for internal sync servers
+	ClusterSyncServer->OnIsConnectionAllowed().BindLambda(IsConnectionAllowedFunc);
+	RenderSyncServer->OnIsConnectionAllowed().BindLambda(IsConnectionAllowedFunc);
+
+	// Listen for node failure notifications
+	ClusterSyncServer->OnNodeFailed().AddRaw(this, &FDisplayClusterClusterNodeCtrlPrimary::ProcessNodeFailed);
+	RenderSyncServer->OnNodeFailed().AddRaw(this, &FDisplayClusterClusterNodeCtrlPrimary::ProcessNodeFailed);
+	ClusterEventsJsonServer->OnNodeFailed().AddRaw(this, &FDisplayClusterClusterNodeCtrlPrimary::ProcessNodeFailed);
+	ClusterEventsBinaryServer->OnNodeFailed().AddRaw(this, &FDisplayClusterClusterNodeCtrlPrimary::ProcessNodeFailed);
+
+	// Start the servers
+	return StartServerWithLogs(ClusterSyncServer.Get(),         TcpListener) // Start with shared listener
+		&& StartServerWithLogs(RenderSyncServer.Get(),          TcpListener) // Start with shared listener
+		&& StartServerWithLogs(ClusterEventsJsonServer.Get(),   CfgPrimary->Host, Ports.ClusterEventsJson)
+		&& StartServerWithLogs(ClusterEventsBinaryServer.Get(), CfgPrimary->Host, Ports.ClusterEventsBinary)
+		&& TcpListener->StartListening(CfgPrimary->Host, Ports.ClusterSync); // Start shared listener as well
+}
+
+void FDisplayClusterClusterNodeCtrlPrimary::StopServers()
+{
+	UE_LOG(LogDisplayClusterCluster, Log, TEXT("Cluster controller %s is shutting down the services..."), *GetControllerName());
+
+	// Stop listening for incoming connections
+	TcpListener->StopListening(true);
+
+	// Stop services
+	ClusterSyncServer->Shutdown();
+	RenderSyncServer->Shutdown();
+	ClusterEventsJsonServer->Shutdown();
+	ClusterEventsBinaryServer->Shutdown();
+
+	FDisplayClusterClusterNodeCtrlSecondary::StopServers();
+}
+
+bool FDisplayClusterClusterNodeCtrlPrimary::InitializeClients()
+{
+	return FDisplayClusterClusterNodeCtrlSecondary::InitializeClients();
+}
+
+bool FDisplayClusterClusterNodeCtrlPrimary::StartClients()
+{
+	return FDisplayClusterClusterNodeCtrlSecondary::StartClients();
+}
+
+void FDisplayClusterClusterNodeCtrlPrimary::StopClients()
+{
+	FDisplayClusterClusterNodeCtrlSecondary::StopClients();
+}
+
+void FDisplayClusterClusterNodeCtrlPrimary::ProcessNodeFailed(const FString& NodeId, IDisplayClusterServer::ENodeFailType NodeFailType)
+{
+	// We need to notify other services about this node (for example, we need to deactivate this node on all barriers of all services)
+	DropClusterNode(NodeId);
+
+	// We need to pass information about node failure so the failover controller will decide what to do
+	if (IDisplayClusterFailoverNodeController* const FailoverCtrl = GDisplayCluster->GetPrivateClusterMgr()->GetFailoverNodeController())
+	{
+		FailoverCtrl->HandleNodeFailed(NodeId, NodeFailType);
+	}
+}
diff --git a/Engine/Plugins/Runtime/nDisplay/Source/DisplayCluster/Private/Cluster/Controller/DisplayClusterClusterNodeCtrlPrimary.h b/Engine/Plugins/Runtime/nDisplay/Source/DisplayCluster/Private/Cluster/Controller/DisplayClusterClusterNodeCtrlPrimary.h
new file mode 100644
--- /dev/null
+++ b/Engine/Plugins/Runtime/nDisplay/Source/DisplayCluster/Private/Cluster/Controller/DisplayClusterClusterNodeCtrlPrimary.h
@@ -0,0 +1,71 @@
+// Copyright Epic Games, Inc. All Rights Reserved.
+
+#pragma once
+
+#include "CoreMinimal.h"
+#include "Cluster/Controller/DisplayClusterClusterNodeCtrlSecondary.h"
+#include "Network/IDisplayClusterServer.h"
+
+class FDisplayClusterClusterSyncService;
+class FDisplayClusterRenderSyncService;
+class FDisplayClusterClusterEventsJsonService;
+class FDisplayClusterClusterEventsBinaryService;
+class FDisplayClusterTcpListener;
+
+
+/**
+ * Primary node controller implementation (cluster mode). Manages servers on primary node side.
+ */
+class FDisplayClusterClusterNodeCtrlPrimary
+	: public FDisplayClusterClusterNodeCtrlSecondary
+{
+public:
+	FDisplayClusterClusterNodeCtrlPrimary(const FString& CtrlName, const FString& NodeName);
+	virtual ~FDisplayClusterClusterNodeCtrlPrimary();
+
+public:
+	//////////////////////////////////////////////////////////////////////////////////////////////
+	// IDisplayClusterProtocolClusterSync
+	//////////////////////////////////////////////////////////////////////////////////////////////
+	virtual EDisplayClusterCommResult GetTimeData(double& OutDeltaTime, double& OutGameTime, TOptional<FQualifiedFrameTime>& OutFrameTime) override;
+	virtual EDisplayClusterCommResult GetObjectsData(const EDisplayClusterSyncGroup InSyncGroup, TMap<FString, FString>& OutObjectsData) override;
+	virtual EDisplayClusterCommResult GetEventsData(TArray<TSharedPtr<FDisplayClusterClusterEventJson, ESPMode::ThreadSafe>>& OutJsonEvents, TArray<TSharedPtr<FDisplayClusterClusterEventBinary, ESPMode::ThreadSafe>>& OutBinaryEvents) override;
+	virtual EDisplayClusterCommResult GetNativeInputData(TMap<FString, FString>& OutNativeInputData) override;
+
+public:
+	//////////////////////////////////////////////////////////////////////////////////////////////
+	// IDisplayClusterClusterNodeController
+	//////////////////////////////////////////////////////////////////////////////////////////////
+	virtual EDisplayClusterNodeRole GetClusterRole() const override
+	{
+		return EDisplayClusterNodeRole::Primary;
+	}
+
+	virtual bool DropClusterNode(const FString& NodeId) override;
+
+protected:
+	//////////////////////////////////////////////////////////////////////////////////////////////
+	// FDisplayClusterNodeCtrlBase
+	//////////////////////////////////////////////////////////////////////////////////////////////
+	virtual bool InitializeServers() override;
+	virtual bool StartServers()      override;
+	virtual void StopServers()       override;
+
+	virtual bool InitializeClients() override;
+	virtual bool StartClients()      override;
+	virtual void StopClients()       override;
+
+private:
+	// Handle node failures
+	void ProcessNodeFailed(const FString& NodeId, IDisplayClusterServer::ENodeFailType NodeFailType);
+
+private:
+	// Node servers
+	TUniquePtr<FDisplayClusterClusterSyncService>         ClusterSyncServer;
+	TUniquePtr<FDisplayClusterRenderSyncService>          RenderSyncServer;
+	TUniquePtr<FDisplayClusterClusterEventsJsonService>   ClusterEventsJsonServer;
+	TUniquePtr<FDisplayClusterClusterEventsBinaryService> ClusterEventsBinaryServer;
+
+	// Shared TCP connection listener. Allows to use single port for multiple internal services.
+	TSharedPtr<FDisplayClusterTcpListener> TcpListener;
+};
diff --git a/Engine/Plugins/Runtime/nDisplay/Source/DisplayCluster/Private/Cluster/Controller/DisplayClusterClusterNodeCtrlSecondary.cpp b/Engine/Plugins/Runtime/nDisplay/Source/DisplayCluster/Private/Cluster/Controller/DisplayClusterClusterNodeCtrlSecondary.cpp
new file mode 100644
--- /dev/null
+++ b/Engine/Plugins/Runtime/nDisplay/Source/DisplayCluster/Private/Cluster/Controller/DisplayClusterClusterNodeCtrlSecondary.cpp
@@ -0,0 +1,210 @@
+// Copyright Epic Games, Inc. All Rights Reserved.
+
+#include "Cluster/Controller/DisplayClusterClusterNodeCtrlSecondary.h"
+
+#include "Cluster/IPDisplayClusterClusterManager.h"
+#include "Cluster/Failover/IDisplayClusterFailoverNodeController.h"
+#include "Config/IPDisplayClusterConfigManager.h"
+#include "Network/Service/ClusterEventsJson/DisplayClusterClusterEventsJsonClient.h"
+#include "Network/Service/ClusterEventsBinary/DisplayClusterClusterEventsBinaryClient.h"
+#include "Network/Service/ClusterSync/DisplayClusterClusterSyncClient.h"
+#include "Network/Service/RenderSync/DisplayClusterRenderSyncClient.h"
+
+#include "DisplayClusterConfigurationTypes.h"
+
+#include "Misc/DisplayClusterGlobals.h"
+#include "Misc/DisplayClusterLog.h"
+
+
+FDisplayClusterClusterNodeCtrlSecondary::FDisplayClusterClusterNodeCtrlSecondary(const FString& CtrlName, const FString& NodeName) :
+	FDisplayClusterClusterNodeCtrlBase(CtrlName, NodeName)
+{
+}
+
+FDisplayClusterClusterNodeCtrlSecondary::~FDisplayClusterClusterNodeCtrlSecondary()
+{
+	UE_LOG(LogDisplayClusterNetwork, Log, TEXT("Releasing %s cluster controller..."), *GetControllerName());
+
+	Shutdown();
+}
+
+
+//////////////////////////////////////////////////////////////////////////////////////////////
+// IDisplayClusterProtocolClusterSync
+//////////////////////////////////////////////////////////////////////////////////////////////
+EDisplayClusterCommResult FDisplayClusterClusterNodeCtrlSecondary::WaitForGameStart()
+{
+	checkSlow(ClusterSyncClient);
+	const EDisplayClusterCommResult CommResult = ClusterSyncClient->WaitForGameStart();
+	GDisplayCluster->GetPrivateClusterMgr()->GetFailoverNodeController()->HandleCommResult(CommResult);
+	return CommResult;
+}
+
+EDisplayClusterCommResult FDisplayClusterClusterNodeCtrlSecondary::WaitForFrameStart()
+{
+	checkSlow(ClusterSyncClient);
+	const EDisplayClusterCommResult CommResult = ClusterSyncClient->WaitForFrameStart();
+	GDisplayCluster->GetPrivateClusterMgr()->GetFailoverNodeController()->HandleCommResult(CommResult);
+	return CommResult;
+}
+
+EDisplayClusterCommResult FDisplayClusterClusterNodeCtrlSecondary::WaitForFrameEnd()
+{
+	checkSlow(ClusterSyncClient);
+	const EDisplayClusterCommResult CommResult = ClusterSyncClient->WaitForFrameEnd();
+	GDisplayCluster->GetPrivateClusterMgr()->GetFailoverNodeController()->HandleCommResult(CommResult);
+	return CommResult;
+}
+
+EDisplayClusterCommResult FDisplayClusterClusterNodeCtrlSecondary::GetTimeData(double& OutDeltaTime, double& OutGameTime, TOptional<FQualifiedFrameTime>& OutFrameTime)
+{
+	checkSlow(ClusterSyncClient);
+	const EDisplayClusterCommResult CommResult = ClusterSyncClient->GetTimeData(OutDeltaTime, OutGameTime, OutFrameTime);
+	GDisplayCluster->GetPrivateClusterMgr()->GetFailoverNodeController()->HandleCommResult(CommResult);
+	return CommResult;
+}
+
+EDisplayClusterCommResult FDisplayClusterClusterNodeCtrlSecondary::GetObjectsData(const EDisplayClusterSyncGroup InSyncGroup, TMap<FString, FString>& OutObjectsData)
+{
+	checkSlow(ClusterSyncClient);
+	const EDisplayClusterCommResult CommResult = ClusterSyncClient->GetObjectsData(InSyncGroup, OutObjectsData);
+	GDisplayCluster->GetPrivateClusterMgr()->GetFailoverNodeController()->HandleCommResult(CommResult);
+	return CommResult;
+}
+
+EDisplayClusterCommResult FDisplayClusterClusterNodeCtrlSecondary::GetEventsData(TArray<TSharedPtr<FDisplayClusterClusterEventJson, ESPMode::ThreadSafe>>& OutJsonEvents, TArray<TSharedPtr<FDisplayClusterClusterEventBinary, ESPMode::ThreadSafe>>& OutBinaryEvents)
+{
+	checkSlow(ClusterSyncClient);
+	const EDisplayClusterCommResult CommResult = ClusterSyncClient->GetEventsData(OutJsonEvents, OutBinaryEvents);
+	GDisplayCluster->GetPrivateClusterMgr()->GetFailoverNodeController()->HandleCommResult(CommResult);
+	return CommResult;
+}
+
+EDisplayClusterCommResult FDisplayClusterClusterNodeCtrlSecondary::GetNativeInputData(TMap<FString, FString>& OutNativeInputData)
+{
+	checkSlow(ClusterSyncClient);
+	const EDisplayClusterCommResult CommResult = ClusterSyncClient->GetNativeInputData(OutNativeInputData);
+	GDisplayCluster->GetPrivateClusterMgr()->GetFailoverNodeController()->HandleCommResult(CommResult);
+	return CommResult;
+}
+
+
+//////////////////////////////////////////////////////////////////////////////////////////////
+// IDisplayClusterProtocolRenderSync
+//////////////////////////////////////////////////////////////////////////////////////////////
+EDisplayClusterCommResult FDisplayClusterClusterNodeCtrlSecondary::WaitForSwapSync()
+{
+	checkSlow(RenderSyncClient.IsValid());
+	const EDisplayClusterCommResult CommResult = RenderSyncClient->WaitForSwapSync();
+	GDisplayCluster->GetPrivateClusterMgr()->GetFailoverNodeController()->HandleCommResult(CommResult);
+	return CommResult;
+}
+
+
+//////////////////////////////////////////////////////////////////////////////////////////////
+// IDisplayClusterProtocolEventsJson
+//////////////////////////////////////////////////////////////////////////////////////////////
+EDisplayClusterCommResult FDisplayClusterClusterNodeCtrlSecondary::EmitClusterEventJson(const FDisplayClusterClusterEventJson& Event)
+{
+	checkSlow(ClusterEventsJsonClient);
+	const EDisplayClusterCommResult CommResult = ClusterEventsJsonClient->EmitClusterEventJson(Event);
+	GDisplayCluster->GetPrivateClusterMgr()->GetFailoverNodeController()->HandleCommResult(CommResult);
+	return CommResult;
+}
+
+
+//////////////////////////////////////////////////////////////////////////////////////////////
+// IDisplayClusterProtocolEventsBinary
+//////////////////////////////////////////////////////////////////////////////////////////////
+EDisplayClusterCommResult FDisplayClusterClusterNodeCtrlSecondary::EmitClusterEventBinary(const FDisplayClusterClusterEventBinary& Event)
+{
+	checkSlow(ClusterEventsBinaryClient);
+	const EDisplayClusterCommResult CommResult = ClusterEventsBinaryClient->EmitClusterEventBinary(Event);
+	GDisplayCluster->GetPrivateClusterMgr()->GetFailoverNodeController()->HandleCommResult(CommResult);
+	return CommResult;
+}
+
+
+//////////////////////////////////////////////////////////////////////////////////////////////
+// FDisplayClusterNodeCtrlBase
+//////////////////////////////////////////////////////////////////////////////////////////////
+bool FDisplayClusterClusterNodeCtrlSecondary::InitializeServers()
+{
+	return FDisplayClusterClusterNodeCtrlBase::InitializeServers();
+}
+
+bool FDisplayClusterClusterNodeCtrlSecondary::StartServers()
+{
+	return FDisplayClusterClusterNodeCtrlBase::StartServers();
+}
+
+void FDisplayClusterClusterNodeCtrlSecondary::StopServers()
+{
+	FDisplayClusterClusterNodeCtrlBase::StopServers();
+}
+
+bool FDisplayClusterClusterNodeCtrlSecondary::InitializeClients()
+{
+	if (!FDisplayClusterClusterNodeCtrlBase::InitializeClients())
+	{
+		return false;
+	}
+
+	UE_LOG(LogDisplayClusterCluster, Log, TEXT("%s - initializing secondary node clients..."), *GetControllerName());
+
+	// Instantiate local clients
+	ClusterSyncClient         = MakeUnique<FDisplayClusterClusterSyncClient>();
+	RenderSyncClient          = MakeUnique<FDisplayClusterRenderSyncClient>();
+	ClusterEventsJsonClient   = MakeUnique<FDisplayClusterClusterEventsJsonClient>();
+	ClusterEventsBinaryClient = MakeUnique<FDisplayClusterClusterEventsBinaryClient>();
+
+	return ClusterSyncClient && RenderSyncClient && ClusterEventsJsonClient && ClusterEventsBinaryClient;
+}
+
+bool FDisplayClusterClusterNodeCtrlSecondary::StartClients()
+{
+	if (!FDisplayClusterClusterNodeCtrlBase::StartClients())
+	{
+		return false;
+	}
+
+	UE_LOG(LogDisplayClusterCluster, Log, TEXT("%s - initializing secondary node clients..."), *GetControllerName());
+
+	// Primary config
+	const UDisplayClusterConfigurationClusterNode* CfgPrimary = GDisplayCluster->GetPrivateConfigMgr()->GetPrimaryNode();
+	if (!CfgPrimary)
+	{
+		UE_LOG(LogDisplayClusterCluster, Error, TEXT("No primary node configuration data found"));
+		return false;
+	}
+
+	// Get configuration data
+	const UDisplayClusterConfigurationData* ConfigData = GDisplayCluster->GetPrivateConfigMgr()->GetConfig();
+	if (!ConfigData)
+	{
+		UE_LOG(LogDisplayClusterCluster, Error, TEXT("Couldn't get configuration data"));
+		return false;
+	}
+
+	// Helper shortcuts to avoid long lines below
+	const FDisplayClusterConfigurationPrimaryNode&     PrimaryNode     = ConfigData->Cluster->PrimaryNode;
+	const FDisplayClusterConfigurationNetworkSettings& NetworkSettings = ConfigData->Cluster->Network;
+
+	// Start the clients
+	return StartClientWithLogs(ClusterSyncClient.Get(),         CfgPrimary->Host, PrimaryNode.Ports.ClusterSync,         NetworkSettings.ConnectRetriesAmount, NetworkSettings.ConnectRetryDelay)
+		&& StartClientWithLogs(RenderSyncClient.Get(),          CfgPrimary->Host, PrimaryNode.Ports.ClusterSync,         NetworkSettings.ConnectRetriesAmount, NetworkSettings.ConnectRetryDelay)
+		&& StartClientWithLogs(ClusterEventsJsonClient.Get(),   CfgPrimary->Host, PrimaryNode.Ports.ClusterEventsJson,   NetworkSettings.ConnectRetriesAmount, NetworkSettings.ConnectRetryDelay)
+		&& StartClientWithLogs(ClusterEventsBinaryClient.Get(), CfgPrimary->Host, PrimaryNode.Ports.ClusterEventsBinary, NetworkSettings.ConnectRetriesAmount, NetworkSettings.ConnectRetryDelay);
+}
+
+void FDisplayClusterClusterNodeCtrlSecondary::StopClients()
+{
+	UE_LOG(LogDisplayClusterCluster, Log, TEXT("Cluster controller %s is shutting down the clients..."), *GetControllerName());
+
+	ClusterEventsJsonClient->Disconnect();
+	ClusterEventsBinaryClient->Disconnect();
+	ClusterSyncClient->Disconnect();
+	RenderSyncClient->Disconnect();
+
+	FDisplayClusterClusterNodeCtrlBase::StopClients();
+}
diff --git a/Engine/Plugins/Runtime/nDisplay/Source/DisplayCluster/Private/Cluster/Controller/DisplayClusterClusterNodeCtrlSecondary.h b/Engine/Plugins/Runtime/nDisplay/Source/DisplayCluster/Private/Cluster/Controller/DisplayClusterClusterNodeCtrlSecondary.h
new file mode 100644
--- /dev/null
+++ b/Engine/Plugins/Runtime/nDisplay/Source/DisplayCluster/Private/Cluster/Controller/DisplayClusterClusterNodeCtrlSecondary.h
@@ -0,0 +1,80 @@
+// Copyright Epic Games, Inc. All Rights Reserved.
+
+#pragma once
+
+#include "Cluster/Controller/DisplayClusterClusterNodeCtrlBase.h"
+
+class FDisplayClusterClusterSyncClient;
+class FDisplayClusterRenderSyncClient;
+class FDisplayClusterClusterEventsJsonClient;
+class FDisplayClusterClusterEventsBinaryClient;
+
+
+/**
+ * Secondary node controller implementation (cluster mode). . Manages clients on client side.
+ */
+class FDisplayClusterClusterNodeCtrlSecondary
+	: public FDisplayClusterClusterNodeCtrlBase
+{
+public:
+	FDisplayClusterClusterNodeCtrlSecondary(const FString& CtrlName, const FString& NodeName);
+	virtual ~FDisplayClusterClusterNodeCtrlSecondary();
+
+public:
+	//////////////////////////////////////////////////////////////////////////////////////////////
+	// IDisplayClusterClusterNodeController
+	//////////////////////////////////////////////////////////////////////////////////////////////
+	virtual EDisplayClusterNodeRole GetClusterRole() const override
+	{
+		return EDisplayClusterNodeRole::Secondary;
+	}
+
+public:
+	//////////////////////////////////////////////////////////////////////////////////////////////
+	// IDisplayClusterProtocolClusterSync
+	//////////////////////////////////////////////////////////////////////////////////////////////
+	virtual EDisplayClusterCommResult WaitForGameStart() override final;
+	virtual EDisplayClusterCommResult WaitForFrameStart() override final;
+	virtual EDisplayClusterCommResult WaitForFrameEnd() override final;
+	virtual EDisplayClusterCommResult GetTimeData(double& OutDeltaTime, double& OutGameTime, TOptional<FQualifiedFrameTime>& OutFrameTime) override;
+	virtual EDisplayClusterCommResult GetObjectsData(const EDisplayClusterSyncGroup InSyncGroup, TMap<FString, FString>& OutObjectsData) override;
+	virtual EDisplayClusterCommResult GetEventsData(TArray<TSharedPtr<FDisplayClusterClusterEventJson, ESPMode::ThreadSafe>>& OutJsonEvents, TArray<TSharedPtr<FDisplayClusterClusterEventBinary, ESPMode::ThreadSafe>>& OutBinaryEvents) override;
+	virtual EDisplayClusterCommResult GetNativeInputData(TMap<FString, FString>& OutNativeInputData) override;
+
+public:
+	//////////////////////////////////////////////////////////////////////////////////////////////
+	// IDisplayClusterProtocolRenderSync
+	//////////////////////////////////////////////////////////////////////////////////////////////
+	virtual EDisplayClusterCommResult WaitForSwapSync() override final;
+
+public:
+	//////////////////////////////////////////////////////////////////////////////////////////////
+	// IDisplayClusterProtocolEventsJson
+	//////////////////////////////////////////////////////////////////////////////////////////////
+	virtual EDisplayClusterCommResult EmitClusterEventJson(const FDisplayClusterClusterEventJson& Event) override;
+
+public:
+	//////////////////////////////////////////////////////////////////////////////////////////////
+	// IDisplayClusterProtocolEventsBinary
+	//////////////////////////////////////////////////////////////////////////////////////////////
+	virtual EDisplayClusterCommResult EmitClusterEventBinary(const FDisplayClusterClusterEventBinary& Event) override;
+
+protected:
+	//////////////////////////////////////////////////////////////////////////////////////////////
+	// FDisplayClusterNodeCtrlBase
+	//////////////////////////////////////////////////////////////////////////////////////////////
+	virtual bool InitializeServers() override;
+	virtual bool StartServers()      override;
+	virtual void StopServers()       override;
+
+	virtual bool InitializeClients() override;
+	virtual bool StartClients()      override;
+	virtual void StopClients()       override;
+
+private:
+	// Cluster node clients
+	TUniquePtr<FDisplayClusterClusterSyncClient>         ClusterSyncClient;
+	TUniquePtr<FDisplayClusterRenderSyncClient>          RenderSyncClient;
+	TUniquePtr<FDisplayClusterClusterEventsJsonClient>   ClusterEventsJsonClient;
+	TUniquePtr<FDisplayClusterClusterEventsBinaryClient> ClusterEventsBinaryClient;
+};
diff --git a/Engine/Plugins/Runtime/nDisplay/Source/DisplayCluster/Private/Cluster/Controller/DisplayClusterClusterNodeCtrlSlave.cpp b/Engine/Plugins/Runtime/nDisplay/Source/DisplayCluster/Private/Cluster/Controller/DisplayClusterClusterNodeCtrlSlave.cpp
deleted file mode 100644
diff --git a/Engine/Plugins/Runtime/nDisplay/Source/DisplayCluster/Private/Cluster/Controller/DisplayClusterClusterNodeCtrlSlave.h b/Engine/Plugins/Runtime/nDisplay/Source/DisplayCluster/Private/Cluster/Controller/DisplayClusterClusterNodeCtrlSlave.h
deleted file mode 100644
diff --git a/Engine/Plugins/Runtime/nDisplay/Source/DisplayCluster/Private/Cluster/Controller/IDisplayClusterClusterNodeController.h b/Engine/Plugins/Runtime/nDisplay/Source/DisplayCluster/Private/Cluster/Controller/IDisplayClusterClusterNodeController.h
--- a/Engine/Plugins/Runtime/nDisplay/Source/DisplayCluster/Private/Cluster/Controller/IDisplayClusterClusterNodeController.h
+++ b/Engine/Plugins/Runtime/nDisplay/Source/DisplayCluster/Private/Cluster/Controller/IDisplayClusterClusterNodeController.h
@@ -47,10 +47,10 @@ public:
 	}
 
 	/** Send binary event to a specific target outside of the cluster */
-	virtual void SendClusterEventTo(const FString& Address, const uint16 Port, const FDisplayClusterClusterEventBinary& Event, bool bMasterOnly)
+	virtual void SendClusterEventTo(const FString& Address, const uint16 Port, const FDisplayClusterClusterEventBinary& Event, bool bPrimaryOnly)
 	{ }
 
 	/** Send JSON event to a specific target outside of the cluster */
-	virtual void SendClusterEventTo(const FString& Address, const uint16 Port, const FDisplayClusterClusterEventJson& Event, bool bMasterOnly)
+	virtual void SendClusterEventTo(const FString& Address, const uint16 Port, const FDisplayClusterClusterEventJson& Event, bool bPrimaryOnly)
 	{ }
 };
diff --git a/Engine/Plugins/Runtime/nDisplay/Source/DisplayCluster/Private/Cluster/DisplayClusterClusterManager.cpp b/Engine/Plugins/Runtime/nDisplay/Source/DisplayCluster/Private/Cluster/DisplayClusterClusterManager.cpp
--- a/Engine/Plugins/Runtime/nDisplay/Source/DisplayCluster/Private/Cluster/DisplayClusterClusterManager.cpp
+++ b/Engine/Plugins/Runtime/nDisplay/Source/DisplayCluster/Private/Cluster/DisplayClusterClusterManager.cpp
@@ -6,10 +6,10 @@
 
 #include "Cluster/IDisplayClusterClusterSyncObject.h"
 #include "Cluster/IDisplayClusterClusterEventListener.h"
-#include "Cluster/Controller/DisplayClusterClusterNodeCtrlMaster.h"
-#include "Cluster/Controller/DisplayClusterClusterNodeCtrlNull.h"
-#include "Cluster/Controller/DisplayClusterClusterNodeCtrlSlave.h"
 #include "Cluster/Controller/DisplayClusterClusterNodeCtrlEditor.h"
+#include "Cluster/Controller/DisplayClusterClusterNodeCtrlNull.h"
+#include "Cluster/Controller/DisplayClusterClusterNodeCtrlPrimary.h"
+#include "Cluster/Controller/DisplayClusterClusterNodeCtrlSecondary.h"
 #include "Cluster/Failover/DisplayClusterFailoverNodeCtrlNull.h"
 #include "Cluster/Failover/DisplayClusterFailoverNodeCtrlPrimary.h"
 #include "Cluster/Failover/DisplayClusterFailoverNodeCtrlSecondary.h"
@@ -255,16 +255,16 @@ void FDisplayClusterClusterManager::PostTick(float DeltaSeconds)
 //////////////////////////////////////////////////////////////////////////////////////////////
 // IDisplayClusterClusterManager
 //////////////////////////////////////////////////////////////////////////////////////////////
-bool FDisplayClusterClusterManager::IsMaster() const
+bool FDisplayClusterClusterManager::IsPrimary() const
 {
 	FScopeLock Lock(&InternalsCS);
-	return ClusterNodeCtrl ? ClusterNodeCtrl->GetClusterRole() == EDisplayClusterNodeRole::Master : false;
+	return ClusterNodeCtrl ? ClusterNodeCtrl->GetClusterRole() == EDisplayClusterNodeRole::Primary : false;
 }
 
-bool FDisplayClusterClusterManager::IsSlave() const
+bool FDisplayClusterClusterManager::IsSecondary() const
 {
 	FScopeLock Lock(&InternalsCS);
-	return ClusterNodeCtrl ? ClusterNodeCtrl->GetClusterRole() == EDisplayClusterNodeRole::Slave : false;
+	return ClusterNodeCtrl ? ClusterNodeCtrl->GetClusterRole() == EDisplayClusterNodeRole::Secondary : false;
 }
 
 bool FDisplayClusterClusterManager::IsBackup() const
@@ -361,7 +361,7 @@ void FDisplayClusterClusterManager::RemoveClusterEventBinaryListener(const FOnCl
 	OnClusterEventBinary.Remove(Listener.GetHandle());
 }
 
-void FDisplayClusterClusterManager::EmitClusterEventJson(const FDisplayClusterClusterEventJson& Event, bool bMasterOnly)
+void FDisplayClusterClusterManager::EmitClusterEventJson(const FDisplayClusterClusterEventJson& Event, bool bPrimaryOnly)
 {
 	UE_LOG(LogDisplayClusterCluster, Verbose, TEXT("JSON event emission request: %s:%s:%s"), *Event.Category, *Event.Type, *Event.Name);
 
@@ -369,8 +369,8 @@ void FDisplayClusterClusterManager::EmitClusterEventJson(const FDisplayClusterCl
 
 	if (CurrentOperationMode == EDisplayClusterOperationMode::Cluster || CurrentOperationMode == EDisplayClusterOperationMode::Editor)
 	{
-		// [Master] Since we receive cluster events asynchronously, we push it to a primary events pool
-		if (IsMaster())
+		// [Primary] Since we receive cluster events asynchronously, we push it to a primary events pool
+		if (IsPrimary())
 		{
 			// Generate event ID
 			const FString EventId = FString::Printf(TEXT("%s-%s-%s"), *Event.Category, *Event.Type, *Event.Name);
@@ -386,11 +386,11 @@ void FDisplayClusterClusterManager::EmitClusterEventJson(const FDisplayClusterCl
 				ClusterEventsJsonNonDiscarded.Add(EventPtr);
 			}
 		}
-		// [Slave] Send event to the master
+		// [Secondary] Send event to the primary node
 		else
 		{
-			// An event will be emitted from a slave node if it's explicitly specified by MasterOnly=false
-			if (!bMasterOnly && ClusterNodeCtrl)
+			// An event will be emitted from a secondary node if it's explicitly specified by bPrimaryOnly=false
+			if (!bPrimaryOnly && ClusterNodeCtrl)
 			{
 				ClusterNodeCtrl->EmitClusterEventJson(Event);
 			}
@@ -398,7 +398,7 @@ void FDisplayClusterClusterManager::EmitClusterEventJson(const FDisplayClusterCl
 	}
 }
 
-void FDisplayClusterClusterManager::EmitClusterEventBinary(const FDisplayClusterClusterEventBinary& Event, bool bMasterOnly)
+void FDisplayClusterClusterManager::EmitClusterEventBinary(const FDisplayClusterClusterEventBinary& Event, bool bPrimaryOnly)
 {
 	UE_LOG(LogDisplayClusterCluster, Verbose, TEXT("BIN event emission request: %d"), Event.EventId);
 
@@ -406,8 +406,8 @@ void FDisplayClusterClusterManager::EmitClusterEventBinary(const FDisplayCluster
 
 	if (CurrentOperationMode == EDisplayClusterOperationMode::Cluster || CurrentOperationMode == EDisplayClusterOperationMode::Editor)
 	{
-		// [Master] Since we receive cluster events asynchronously, we push it to a primary events pool
-		if (IsMaster())
+		// [Primary] Since we receive cluster events asynchronously, we push it to a primary events pool
+		if (IsPrimary())
 		{
 			// Make it shared ptr
 			TSharedPtr<FDisplayClusterClusterEventBinary, ESPMode::ThreadSafe> EventPtr = MakeShared<FDisplayClusterClusterEventBinary, ESPMode::ThreadSafe>(Event);
@@ -421,11 +421,11 @@ void FDisplayClusterClusterManager::EmitClusterEventBinary(const FDisplayCluster
 				ClusterEventsBinaryNonDiscarded.Add(EventPtr);
 			}
 		}
-		// [Slave] Send event to the master
+		// [Secondary] Send event to the primary node
 		else
 		{
-			// An event will be emitted from a slave node if it's explicitly specified by MasterOnly=false
-			if (!bMasterOnly && ClusterNodeCtrl)
+			// An event will be emitted from a secondary node if it's explicitly specified by bPrimaryOnly=false
+			if (!bPrimaryOnly && ClusterNodeCtrl)
 			{
 				ClusterNodeCtrl->EmitClusterEventBinary(Event);
 			}
@@ -433,31 +433,31 @@ void FDisplayClusterClusterManager::EmitClusterEventBinary(const FDisplayCluster
 	}
 }
 
-void FDisplayClusterClusterManager::SendClusterEventTo(const FString& Address, const uint16 Port, const FDisplayClusterClusterEventJson& Event, bool bMasterOnly)
+void FDisplayClusterClusterManager::SendClusterEventTo(const FString& Address, const uint16 Port, const FDisplayClusterClusterEventJson& Event, bool bPrimaryOnly)
 {
 	if (CurrentOperationMode == EDisplayClusterOperationMode::Cluster || CurrentOperationMode == EDisplayClusterOperationMode::Editor)
 	{
 		if (ClusterNodeCtrl)
 		{
-			if (IsMaster() || !bMasterOnly)
+			if (IsPrimary() || !bPrimaryOnly)
 			{
 				UE_LOG(LogDisplayClusterCluster, Verbose, TEXT("JSON event emission request: recipient=%s:%u, event=%s:%s:%s"), *Address, Port, *Event.Category, *Event.Type, *Event.Name);
-				ClusterNodeCtrl->SendClusterEventTo(Address, Port, Event, bMasterOnly);
+				ClusterNodeCtrl->SendClusterEventTo(Address, Port, Event, bPrimaryOnly);
 			}
 		}
 	}
 }
 
-void FDisplayClusterClusterManager::SendClusterEventTo(const FString& Address, const uint16 Port, const FDisplayClusterClusterEventBinary& Event, bool bMasterOnly)
+void FDisplayClusterClusterManager::SendClusterEventTo(const FString& Address, const uint16 Port, const FDisplayClusterClusterEventBinary& Event, bool bPrimaryOnly)
 {
 	if (CurrentOperationMode == EDisplayClusterOperationMode::Cluster || CurrentOperationMode == EDisplayClusterOperationMode::Editor)
 	{
 		if (ClusterNodeCtrl)
 		{
-			if (IsMaster() || !bMasterOnly)
+			if (IsPrimary() || !bPrimaryOnly)
 			{
 				UE_LOG(LogDisplayClusterCluster, Verbose, TEXT("BIN event emission request: recipient=%s:%u, event=%d"), *Address, Port, Event.EventId);
-				ClusterNodeCtrl->SendClusterEventTo(Address, Port, Event, bMasterOnly);
+				ClusterNodeCtrl->SendClusterEventTo(Address, Port, Event, bPrimaryOnly);
 			}
 		}
 	}
@@ -558,7 +558,7 @@ void FDisplayClusterClusterManager::SyncObjects(EDisplayClusterSyncGroup InSyncG
 		UE_LOG(LogDisplayClusterCluster, Verbose, TEXT("Downloading finished. Available %d records (objects)."), ObjectsData.Num());
 
 		// No need to import data on the primary node
-		if (!IsMaster())
+		if (!IsPrimary())
 		{
 			// Perform data load (objects state update)
 			ImportObjectsData(InSyncGroup, ObjectsData);
@@ -772,7 +772,7 @@ void FDisplayClusterClusterManager::ExportNativeInputData(TMap<FString, FString>
 {
 	// If we are on the primary node, that means we got network request for input data.
 	// So, provide with data once it's cached.
-	if (IsMaster())
+	if (IsPrimary())
 	{
 		// Wait for data cache to be ready
 		NativeInputCacheReadySignal->Wait();
@@ -805,16 +805,16 @@ TUniquePtr<IDisplayClusterClusterNodeController> FDisplayClusterClusterManager::
 	// Instantiate appropriate controller depending on operation mode and cluster role
 	if (CurrentOperationMode == EDisplayClusterOperationMode::Cluster)
 	{
-		// Master or slave
-		if (ClusterNodeId.Equals(GDisplayCluster->GetPrivateConfigMgr()->GetMasterNodeId(), ESearchCase::IgnoreCase))
+		// Primary or secondary
+		if (ClusterNodeId.Equals(GDisplayCluster->GetPrivateConfigMgr()->GetPrimaryNodeId(), ESearchCase::IgnoreCase))
 		{
-			UE_LOG(LogDisplayClusterCluster, Log, TEXT("Instantiating master node controller..."));
-			return MakeUnique<FDisplayClusterClusterNodeCtrlMaster>(FString("[CTRL-M]"), ClusterNodeId);
+			UE_LOG(LogDisplayClusterCluster, Log, TEXT("Instantiating primary node controller..."));
+			return MakeUnique<FDisplayClusterClusterNodeCtrlPrimary>(FString("[CTRL-M]"), ClusterNodeId);
 		}
 		else
 		{
-			UE_LOG(LogDisplayClusterCluster, Log, TEXT("Instantiating slave node controller..."));
-			return MakeUnique<FDisplayClusterClusterNodeCtrlSlave>(FString("[CTRL-S]"), ClusterNodeId);
+			UE_LOG(LogDisplayClusterCluster, Log, TEXT("Instantiating secondary node controller..."));
+			return MakeUnique<FDisplayClusterClusterNodeCtrlSecondary>(FString("[CTRL-S]"), ClusterNodeId);
 		}
 	}
 	else if (CurrentOperationMode == EDisplayClusterOperationMode::Editor)
@@ -833,11 +833,11 @@ TUniquePtr<IDisplayClusterFailoverNodeController> FDisplayClusterClusterManager:
 	const EDisplayClusterNodeRole ClusterRole = GetClusterRole();
 	switch (ClusterRole)
 	{
-	case EDisplayClusterNodeRole::Master:
+	case EDisplayClusterNodeRole::Primary:
 		UE_LOG(LogDisplayClusterCluster, Log, TEXT("Instantiating failover primary controller..."));
 		return MakeUnique<FDisplayClusterFailoverNodeCtrlPrimary>();
 
-	case EDisplayClusterNodeRole::Slave:
+	case EDisplayClusterNodeRole::Secondary:
 		UE_LOG(LogDisplayClusterCluster, Log, TEXT("Instantiating failover secondary controller..."));
 		return MakeUnique<FDisplayClusterFailoverNodeCtrlSecondary>();
 	}
diff --git a/Engine/Plugins/Runtime/nDisplay/Source/DisplayCluster/Private/Cluster/DisplayClusterClusterManager.h b/Engine/Plugins/Runtime/nDisplay/Source/DisplayCluster/Private/Cluster/DisplayClusterClusterManager.h
--- a/Engine/Plugins/Runtime/nDisplay/Source/DisplayCluster/Private/Cluster/DisplayClusterClusterManager.h
+++ b/Engine/Plugins/Runtime/nDisplay/Source/DisplayCluster/Private/Cluster/DisplayClusterClusterManager.h
@@ -41,9 +41,9 @@ public:
 	//////////////////////////////////////////////////////////////////////////////////////////////
 	// IDisplayClusterClusterManager
 	//////////////////////////////////////////////////////////////////////////////////////////////
-	virtual bool IsMaster() const override;
-	virtual bool IsSlave()  const override;
-	virtual bool IsBackup() const override;
+	virtual bool IsPrimary()   const override;
+	virtual bool IsSecondary() const override;
+	virtual bool IsBackup()    const override;
 	virtual EDisplayClusterNodeRole GetClusterRole() const override;
 
 	virtual FString GetNodeId() const override
@@ -72,11 +72,11 @@ public:
 	virtual void AddClusterEventBinaryListener(const FOnClusterEventBinaryListener& Listener) override;
 	virtual void RemoveClusterEventBinaryListener(const FOnClusterEventBinaryListener& Listener) override;
 
-	virtual void EmitClusterEventJson(const FDisplayClusterClusterEventJson& Event, bool MasterOnly) override;
-	virtual void EmitClusterEventBinary(const FDisplayClusterClusterEventBinary& Event, bool bMasterOnly) override;
+	virtual void EmitClusterEventJson(const FDisplayClusterClusterEventJson& Event, bool bPrimaryOnly) override;
+	virtual void EmitClusterEventBinary(const FDisplayClusterClusterEventBinary& Event, bool bPrimaryOnly) override;
 
-	virtual void SendClusterEventTo(const FString& Address, const uint16 Port, const FDisplayClusterClusterEventJson&   Event, bool bMasterOnly) override;
-	virtual void SendClusterEventTo(const FString& Address, const uint16 Port, const FDisplayClusterClusterEventBinary& Event, bool bMasterOnly) override;
+	virtual void SendClusterEventTo(const FString& Address, const uint16 Port, const FDisplayClusterClusterEventJson&   Event, bool bPrimaryOnly) override;
+	virtual void SendClusterEventTo(const FString& Address, const uint16 Port, const FDisplayClusterClusterEventBinary& Event, bool bPrimaryOnly) override;
 
 public:
 	//////////////////////////////////////////////////////////////////////////////////////////////
diff --git a/Engine/Plugins/Runtime/nDisplay/Source/DisplayCluster/Private/Config/DisplayClusterConfigManager.cpp b/Engine/Plugins/Runtime/nDisplay/Source/DisplayCluster/Private/Config/DisplayClusterConfigManager.cpp
--- a/Engine/Plugins/Runtime/nDisplay/Source/DisplayCluster/Private/Config/DisplayClusterConfigManager.cpp
+++ b/Engine/Plugins/Runtime/nDisplay/Source/DisplayCluster/Private/Config/DisplayClusterConfigManager.cpp
@@ -47,14 +47,14 @@ void FDisplayClusterConfigManager::EndSession()
 //////////////////////////////////////////////////////////////////////////////////////////////
 // IDisplayClusterConfigManager
 //////////////////////////////////////////////////////////////////////////////////////////////
-FString FDisplayClusterConfigManager::GetMasterNodeId() const
+FString FDisplayClusterConfigManager::GetPrimaryNodeId() const
 {
-	return ConfigData ? ConfigData->Cluster->MasterNode.Id : FString();
+	return ConfigData ? ConfigData->Cluster->PrimaryNode.Id : FString();
 }
 
-const UDisplayClusterConfigurationClusterNode* FDisplayClusterConfigManager::GetMasterNode() const
+const UDisplayClusterConfigurationClusterNode* FDisplayClusterConfigManager::GetPrimaryNode() const
 {
-	return ConfigData ? ConfigData->Cluster->GetNode(GetMasterNodeId()) : nullptr;
+	return ConfigData ? ConfigData->Cluster->GetNode(GetPrimaryNodeId()) : nullptr;
 }
 
 const UDisplayClusterConfigurationClusterNode* FDisplayClusterConfigManager::GetLocalNode() const
diff --git a/Engine/Plugins/Runtime/nDisplay/Source/DisplayCluster/Private/Config/DisplayClusterConfigManager.h b/Engine/Plugins/Runtime/nDisplay/Source/DisplayCluster/Private/Config/DisplayClusterConfigManager.h
--- a/Engine/Plugins/Runtime/nDisplay/Source/DisplayCluster/Private/Config/DisplayClusterConfigManager.h
+++ b/Engine/Plugins/Runtime/nDisplay/Source/DisplayCluster/Private/Config/DisplayClusterConfigManager.h
@@ -47,9 +47,9 @@ public:
 		return ClusterNodeId; 
 	}
 
-	virtual FString GetMasterNodeId() const override;
+	virtual FString GetPrimaryNodeId() const override;
 
-	virtual const UDisplayClusterConfigurationClusterNode* GetMasterNode() const override;
+	virtual const UDisplayClusterConfigurationClusterNode* GetPrimaryNode() const override;
 	virtual const UDisplayClusterConfigurationClusterNode* GetLocalNode() const override;
 	virtual const UDisplayClusterConfigurationViewport*    GetLocalViewport(const FString& ViewportId) const override;
 
diff --git a/Engine/Plugins/Runtime/nDisplay/Source/DisplayCluster/Private/Game/EngineClasses/Basics/DisplayClusterGameEngine.cpp b/Engine/Plugins/Runtime/nDisplay/Source/DisplayCluster/Private/Game/EngineClasses/Basics/DisplayClusterGameEngine.cpp
--- a/Engine/Plugins/Runtime/nDisplay/Source/DisplayCluster/Private/Game/EngineClasses/Basics/DisplayClusterGameEngine.cpp
+++ b/Engine/Plugins/Runtime/nDisplay/Source/DisplayCluster/Private/Game/EngineClasses/Basics/DisplayClusterGameEngine.cpp
@@ -236,21 +236,21 @@ bool UDisplayClusterGameEngine::ValidateConfigFile(const FString& FilePath)
 	const EDisplayClusterConfigurationVersion ConfigVersion = IDisplayClusterConfiguration::Get().GetConfigVersion(FilePath);
 	switch (ConfigVersion)
 	{
-		case EDisplayClusterConfigurationVersion::Version_CFG:
-			// Old .cfg file are not allowed anymore
-			UE_LOG(LogDisplayClusterEngine, Error, TEXT("Old (.cfg) config format is not supported"));
-			break;
-
 		case EDisplayClusterConfigurationVersion::Version_426:
 			// Old 4.26 and 4.27p1 formats are not allowed as well
 			UE_LOG(LogDisplayClusterEngine, Error, TEXT("Detected old (.ndisplay 4.26 or .ndisplay 4.27p1) config format. Please upgrade to the actual version."));
 			return true;
 
 		case EDisplayClusterConfigurationVersion::Version_427:
-			// Ok, it's the actual config format
+			// Ok, it's one of the actual config formats
 			UE_LOG(LogDisplayClusterEngine, Log, TEXT("Detected (.ndisplay 4.27) config format"));
 			return true;
 
+		case EDisplayClusterConfigurationVersion::Version_500:
+			// Ok, it's one of the actual config formats
+			UE_LOG(LogDisplayClusterEngine, Log, TEXT("Detected (.ndisplay 5.00) config format"));
+			return true;
+
 		case EDisplayClusterConfigurationVersion::Unknown:
 		default:
 			// Something unexpected came here
diff --git a/Engine/Plugins/Runtime/nDisplay/Source/DisplayCluster/Private/Game/EngineClasses/Basics/DisplayClusterPlayerInput.cpp b/Engine/Plugins/Runtime/nDisplay/Source/DisplayCluster/Private/Game/EngineClasses/Basics/DisplayClusterPlayerInput.cpp
--- a/Engine/Plugins/Runtime/nDisplay/Source/DisplayCluster/Private/Game/EngineClasses/Basics/DisplayClusterPlayerInput.cpp
+++ b/Engine/Plugins/Runtime/nDisplay/Source/DisplayCluster/Private/Game/EngineClasses/Basics/DisplayClusterPlayerInput.cpp
@@ -26,7 +26,7 @@ void UDisplayClusterPlayerInput::ProcessInputStack(const TArray<UInputComponent*
 	{
 		TMap<FString, FString> KeyStates;
 
-		if (ClusterMgr->IsMaster())
+		if (ClusterMgr->IsPrimary())
 		{
 			// Export key states data for the stack
 			SerializeKeyStateMap(KeyStates);
diff --git a/Engine/Plugins/Runtime/nDisplay/Source/DisplayCluster/Private/Game/EngineClasses/Scene/DisplayClusterRootActor.cpp b/Engine/Plugins/Runtime/nDisplay/Source/DisplayCluster/Private/Game/EngineClasses/Scene/DisplayClusterRootActor.cpp
--- a/Engine/Plugins/Runtime/nDisplay/Source/DisplayCluster/Private/Game/EngineClasses/Scene/DisplayClusterRootActor.cpp
+++ b/Engine/Plugins/Runtime/nDisplay/Source/DisplayCluster/Private/Game/EngineClasses/Scene/DisplayClusterRootActor.cpp
@@ -177,7 +177,7 @@ void ADisplayClusterRootActor::OverrideFromConfig(UDisplayClusterConfigurationDa
 	// Override Cluster
 	if (CurrentConfigData->Cluster)
 	{
-		CurrentConfigData->Cluster->MasterNode = ConfigData->Cluster->MasterNode;
+		CurrentConfigData->Cluster->PrimaryNode = ConfigData->Cluster->PrimaryNode;
 		CurrentConfigData->Cluster->Sync = ConfigData->Cluster->Sync;
 		CurrentConfigData->Cluster->Network = ConfigData->Cluster->Network;
 
@@ -660,7 +660,7 @@ void ADisplayClusterRootActor::BeginPlay()
 		}
 
 		// Optionally activate native input synchronization
-		if (SyncPolicyType.Equals(DisplayClusterConfigurationStrings::config::cluster::input_sync::InputSyncPolicyReplicateMaster, ESearchCase::IgnoreCase))
+		if (SyncPolicyType.Equals(DisplayClusterConfigurationStrings::config::cluster::input_sync::InputSyncPolicyReplicatePrimary, ESearchCase::IgnoreCase))
 		{
 			APlayerController* const PlayerController = GetWorld()->GetFirstPlayerController();
 			if (PlayerController)
diff --git a/Engine/Plugins/Runtime/nDisplay/Source/DisplayCluster/Public/Blueprints/IDisplayClusterBlueprintAPI.h b/Engine/Plugins/Runtime/nDisplay/Source/DisplayCluster/Public/Blueprints/IDisplayClusterBlueprintAPI.h
--- a/Engine/Plugins/Runtime/nDisplay/Source/DisplayCluster/Public/Blueprints/IDisplayClusterBlueprintAPI.h
+++ b/Engine/Plugins/Runtime/nDisplay/Source/DisplayCluster/Public/Blueprints/IDisplayClusterBlueprintAPI.h
@@ -67,16 +67,16 @@ public:
 	UFUNCTION(BlueprintCallable, meta = (DisplayName = "Get nodes amount"), Category = "NDisplay|Cluster")
 	virtual int32 GetActiveNodesAmount() const = 0;
 
-	/** Returns true if current node is a master node in a cluster. */
-	UFUNCTION(BlueprintCallable, meta = (DisplayName = "Is master node"), Category = "NDisplay|Cluster")
-	virtual bool IsMaster() const = 0;
+	/** Returns true if current node is a primary node in a cluster. */
+	UFUNCTION(BlueprintCallable, meta = (DisplayName = "Is a primary node"), Category = "NDisplay|Cluster")
+	virtual bool IsPrimary() const = 0;
 
-	/** Returns true if current node is a slave node in a cluster. */
-	UFUNCTION(BlueprintCallable, meta = (DisplayName = "Is slave node"), Category = "NDisplay|Cluster")
-	virtual bool IsSlave() const = 0;
+	/** Returns true if current node is a secondary node in a cluster. */
+	UFUNCTION(BlueprintCallable, meta = (DisplayName = "Is a secondary node"), Category = "NDisplay|Cluster")
+	virtual bool IsSecondary() const = 0;
 
 	/** Returns true if current node is a backup node in a cluster. */
-	UFUNCTION(BlueprintCallable, meta = (DisplayName = "Is backup node"), Category = "NDisplay|Cluster")
+	UFUNCTION(BlueprintCallable, meta = (DisplayName = "Is a backup node"), Category = "NDisplay|Cluster")
 	virtual bool IsBackup() const = 0;
 
 	/** Returns the role of the current cluster node. */
@@ -95,17 +95,17 @@ public:
 
 	/** Emits JSON cluster event. */
 	UFUNCTION(BlueprintCallable, meta = (DisplayName = "Emit JSON cluster event"), Category = "NDisplay|Cluster")
-	virtual void EmitClusterEventJson(const FDisplayClusterClusterEventJson& Event, bool bMasterOnly) = 0;
+	virtual void EmitClusterEventJson(const FDisplayClusterClusterEventJson& Event, bool bPrimaryOnly) = 0;
 
 	/** Emits binary cluster event. */
 	UFUNCTION(BlueprintCallable, meta = (DisplayName = "Emit binary cluster event"), Category = "NDisplay|Cluster")
-	virtual void EmitClusterEventBinary(const FDisplayClusterClusterEventBinary& Event, bool bMasterOnly) = 0;
+	virtual void EmitClusterEventBinary(const FDisplayClusterClusterEventBinary& Event, bool bPrimaryOnly) = 0;
 
 	/** Sends JSON cluster event to a specific target (outside of the cluster). */
 	UFUNCTION(BlueprintCallable, meta = (DisplayName = "Send JSON event to a specific host"), Category = "NDisplay|Cluster")
-	virtual void SendClusterEventJsonTo(const FString& Address, const int32 Port, const FDisplayClusterClusterEventJson& Event, bool bMasterOnly) = 0;
+	virtual void SendClusterEventJsonTo(const FString& Address, const int32 Port, const FDisplayClusterClusterEventJson& Event, bool bPrimaryOnly) = 0;
 
 	/** Sends binary cluster event to a specific target (outside of the cluster). */
 	UFUNCTION(BlueprintCallable, meta = (DisplayName = "Send binary event to a specific host"), Category = "NDisplay|Cluster")
-	virtual void SendClusterEventBinaryTo(const FString& Address, const int32 Port, const FDisplayClusterClusterEventBinary& Event, bool bMasterOnly) = 0;
+	virtual void SendClusterEventBinaryTo(const FString& Address, const int32 Port, const FDisplayClusterClusterEventBinary& Event, bool bPrimaryOnly) = 0;
 };
diff --git a/Engine/Plugins/Runtime/nDisplay/Source/DisplayCluster/Public/Cluster/IDisplayClusterClusterManager.h b/Engine/Plugins/Runtime/nDisplay/Source/DisplayCluster/Public/Cluster/IDisplayClusterClusterManager.h
--- a/Engine/Plugins/Runtime/nDisplay/Source/DisplayCluster/Public/Cluster/IDisplayClusterClusterManager.h
+++ b/Engine/Plugins/Runtime/nDisplay/Source/DisplayCluster/Public/Cluster/IDisplayClusterClusterManager.h
@@ -30,10 +30,10 @@ public:
 	virtual ~IDisplayClusterClusterManager() = default;
 
 public:
-	/** Returns true if current node has master role. */
-	virtual bool IsMaster() const = 0;
-	/** Returns true if current node has slave role. */
-	virtual bool IsSlave()  const = 0;
+	/** Returns true if current node has primary role. */
+	virtual bool IsPrimary() const = 0;
+	/** Returns true if current node has secondary role. */
+	virtual bool IsSecondary()  const = 0;
 	/** Returns true if current node has backup role. */
 	virtual bool IsBackup() const = 0;
 	/** Returns cluster node role. */
@@ -77,14 +77,14 @@ public:
 	virtual void RemoveClusterEventBinaryListener(const FOnClusterEventBinaryListener& Listener) = 0;
 
 	/** Emits JSON cluster event. */
-	virtual void EmitClusterEventJson(const FDisplayClusterClusterEventJson& Event, bool bMasterOnly) = 0;
+	virtual void EmitClusterEventJson(const FDisplayClusterClusterEventJson& Event, bool bPrimaryOnly) = 0;
 
 	/** Emits binary cluster event. */
-	virtual void EmitClusterEventBinary(const FDisplayClusterClusterEventBinary& Event, bool bMasterOnly) = 0;
+	virtual void EmitClusterEventBinary(const FDisplayClusterClusterEventBinary& Event, bool bPrimaryOnly) = 0;
 
 	/** Sends JSON cluster event to a specific target (outside of the cluster). */
-	virtual void SendClusterEventTo(const FString& Address, const uint16 Port, const FDisplayClusterClusterEventJson& Event, bool bMasterOnly) = 0;
+	virtual void SendClusterEventTo(const FString& Address, const uint16 Port, const FDisplayClusterClusterEventJson& Event, bool bPrimaryOnly) = 0;
 
 	/** Sends binary cluster event to a specific target (outside of the cluster). */
-	virtual void SendClusterEventTo(const FString& Address, const uint16 Port, const FDisplayClusterClusterEventBinary& Event, bool bMasterOnly) = 0;
+	virtual void SendClusterEventTo(const FString& Address, const uint16 Port, const FDisplayClusterClusterEventBinary& Event, bool bPrimaryOnly) = 0;
 };
diff --git a/Engine/Plugins/Runtime/nDisplay/Source/DisplayCluster/Public/Config/IDisplayClusterConfigManager.h b/Engine/Plugins/Runtime/nDisplay/Source/DisplayCluster/Public/Config/IDisplayClusterConfigManager.h
--- a/Engine/Plugins/Runtime/nDisplay/Source/DisplayCluster/Public/Config/IDisplayClusterConfigManager.h
+++ b/Engine/Plugins/Runtime/nDisplay/Source/DisplayCluster/Public/Config/IDisplayClusterConfigManager.h
@@ -26,11 +26,11 @@ public:
 
 	// Returns ID of cluster node that is assigned to this application instance
 	virtual FString GetLocalNodeId() const = 0;
-	// Returns master node ID
-	virtual FString GetMasterNodeId() const = 0;
+	// Returns primary node ID
+	virtual FString GetPrimaryNodeId() const = 0;
 
-	// Returns master node configuration data
-	virtual const UDisplayClusterConfigurationClusterNode* GetMasterNode() const = 0;
+	// Returns primary node configuration data
+	virtual const UDisplayClusterConfigurationClusterNode* GetPrimaryNode() const = 0;
 	// Returns configuration data for cluster node that is assigned to this application instance
 	virtual const UDisplayClusterConfigurationClusterNode* GetLocalNode() const = 0;
 	// Returns configuration data for a specified local viewport 
diff --git a/Engine/Plugins/Runtime/nDisplay/Source/DisplayCluster/Public/DisplayClusterEnums.h b/Engine/Plugins/Runtime/nDisplay/Source/DisplayCluster/Public/DisplayClusterEnums.h
--- a/Engine/Plugins/Runtime/nDisplay/Source/DisplayCluster/Public/DisplayClusterEnums.h
+++ b/Engine/Plugins/Runtime/nDisplay/Source/DisplayCluster/Public/DisplayClusterEnums.h
@@ -26,8 +26,8 @@ UENUM(BlueprintType)
 enum class EDisplayClusterNodeRole : uint8
 {
 	None = 0,
-	Master,
-	Slave,
+	Primary,
+	Secondary,
 	Backup,
 };
 
diff --git a/Engine/Plugins/Runtime/nDisplay/Source/DisplayCluster/Public/Misc/DisplayClusterStrings.h b/Engine/Plugins/Runtime/nDisplay/Source/DisplayCluster/Public/Misc/DisplayClusterStrings.h
--- a/Engine/Plugins/Runtime/nDisplay/Source/DisplayCluster/Public/Misc/DisplayClusterStrings.h
+++ b/Engine/Plugins/Runtime/nDisplay/Source/DisplayCluster/Public/Misc/DisplayClusterStrings.h
@@ -27,9 +27,9 @@ namespace DisplayClusterStrings
 		// Cluster node roles
 		namespace role
 		{
-			static constexpr auto Master = TEXT("master");
-			static constexpr auto Slave  = TEXT("slave");
-			static constexpr auto Backup = TEXT("backup");
+			static constexpr auto Primary   = TEXT("primary");
+			static constexpr auto Secondary = TEXT("secondary");
+			static constexpr auto Backup    = TEXT("backup");
 		}
 
 		// Stereo device types (command line values)
diff --git a/Engine/Plugins/Runtime/nDisplay/Source/DisplayClusterConfiguration/Private/DisplayClusterConfigurationMgr.cpp b/Engine/Plugins/Runtime/nDisplay/Source/DisplayClusterConfiguration/Private/DisplayClusterConfigurationMgr.cpp
--- a/Engine/Plugins/Runtime/nDisplay/Source/DisplayClusterConfiguration/Private/DisplayClusterConfigurationMgr.cpp
+++ b/Engine/Plugins/Runtime/nDisplay/Source/DisplayClusterConfiguration/Private/DisplayClusterConfigurationMgr.cpp
@@ -5,7 +5,7 @@
 #include "Formats/IDisplayClusterConfigurationDataParser.h"
 #include "Formats/JSON426/DisplayClusterConfigurationJsonParser_426.h"
 #include "Formats/JSON427/DisplayClusterConfigurationJsonParser_427.h"
-#include "Formats/Text/DisplayClusterConfigurationTextParser.h"
+#include "Formats/JSON500/DisplayClusterConfigurationJsonParser_500.h"
 
 #include "VersionChecker/DisplayClusterConfigurationVersionChecker.h"
 
@@ -18,7 +18,7 @@
 
 
 // Alias for current config scheme
-namespace ACTUAL_CONFIG_SCHEME = JSON427;
+namespace ACTUAL_CONFIG_SCHEME = JSON500;
 
 
 FDisplayClusterConfigurationMgr& FDisplayClusterConfigurationMgr::Get()
@@ -62,10 +62,6 @@ UDisplayClusterConfigurationData* FDisplayClusterConfigurationMgr::LoadConfig(co
 	TUniquePtr<FDisplayClusterConfigurationVersionChecker> VersionChecker = MakeUnique<FDisplayClusterConfigurationVersionChecker>();
 	switch (VersionChecker->GetConfigVersion(ConfigFile))
 	{
-	case EDisplayClusterConfigurationVersion::Version_CFG:
-		Parser = MakeUnique<FDisplayClusterConfigurationTextParser>();
-		break;
-
 	case EDisplayClusterConfigurationVersion::Version_426:
 		Parser = MakeUnique<JSON426::FDisplayClusterConfigurationJsonParser>();
 		break;
@@ -74,6 +70,10 @@ UDisplayClusterConfigurationData* FDisplayClusterConfigurationMgr::LoadConfig(co
 		Parser = MakeUnique<JSON427::FDisplayClusterConfigurationJsonParser>();
 		break;
 
+	case EDisplayClusterConfigurationVersion::Version_500:
+		Parser = MakeUnique<JSON500::FDisplayClusterConfigurationJsonParser>();
+		break;
+
 	case EDisplayClusterConfigurationVersion::Unknown:
 	default:
 		UE_LOG(LogDisplayClusterConfiguration, Error, TEXT("No parser implemented for file: %s"), *ConfigFile);
diff --git a/Engine/Plugins/Runtime/nDisplay/Source/DisplayClusterConfiguration/Private/DisplayClusterConfigurationTypes.cpp b/Engine/Plugins/Runtime/nDisplay/Source/DisplayClusterConfiguration/Private/DisplayClusterConfigurationTypes.cpp
--- a/Engine/Plugins/Runtime/nDisplay/Source/DisplayClusterConfiguration/Private/DisplayClusterConfigurationTypes.cpp
+++ b/Engine/Plugins/Runtime/nDisplay/Source/DisplayClusterConfiguration/Private/DisplayClusterConfigurationTypes.cpp
@@ -134,7 +134,7 @@ const TSet<FString> UDisplayClusterConfigurationData::RenderSyncPolicies =
 
 const TSet<FString> UDisplayClusterConfigurationData::InputSyncPolicies =
 {
-	TEXT("ReplicateMaster"),
+	TEXT("ReplicatePrimary"),
 	TEXT("None")
 };
 
@@ -271,7 +271,7 @@ void UDisplayClusterConfigurationScene::GetObjectsToExport(TArray<UObject*>& Out
 	SAVE_MAP(Cameras);
 }
 
-FDisplayClusterConfigurationMasterNodePorts::FDisplayClusterConfigurationMasterNodePorts()
+FDisplayClusterConfigurationPrimaryNodePorts::FDisplayClusterConfigurationPrimaryNodePorts()
 	: ClusterSync(41001)
 	, ClusterEventsJson  (41003)
 	, ClusterEventsBinary(41004)
@@ -282,7 +282,7 @@ FDisplayClusterConfigurationClusterSync::FDisplayClusterConfigurationClusterSync
 {
 	using namespace DisplayClusterConfigurationStrings::config;
 	RenderSyncPolicy.Type = cluster::render_sync::Ethernet;
-	InputSyncPolicy.Type = cluster::input_sync::InputSyncPolicyReplicateMaster;
+	InputSyncPolicy.Type  = cluster::input_sync::InputSyncPolicyReplicatePrimary;
 }
 
 FDisplayClusterConfigurationNetworkSettings::FDisplayClusterConfigurationNetworkSettings()
diff --git a/Engine/Plugins/Runtime/nDisplay/Source/DisplayClusterConfiguration/Private/Formats/JSON426/DisplayClusterConfigurationJsonParser_426.cpp b/Engine/Plugins/Runtime/nDisplay/Source/DisplayClusterConfiguration/Private/Formats/JSON426/DisplayClusterConfigurationJsonParser_426.cpp
--- a/Engine/Plugins/Runtime/nDisplay/Source/DisplayClusterConfiguration/Private/Formats/JSON426/DisplayClusterConfigurationJsonParser_426.cpp
+++ b/Engine/Plugins/Runtime/nDisplay/Source/DisplayClusterConfiguration/Private/Formats/JSON426/DisplayClusterConfigurationJsonParser_426.cpp
@@ -125,18 +125,18 @@ namespace JSON426
 
 		// Cluster
 		{
-			// Master node
+			// Primary node
 			{
-				Config->Cluster->MasterNode.Id = CfgJson.Cluster.MasterNode.Id;
+				Config->Cluster->PrimaryNode.Id = CfgJson.Cluster.MasterNode.Id;
 
 				const uint16* ClusterSyncPort = CfgJson.Cluster.MasterNode.Ports.Find(DisplayClusterConfigurationStrings::config::cluster::ports::PortClusterSync);
-				Config->Cluster->MasterNode.Ports.ClusterSync = (ClusterSyncPort ? *ClusterSyncPort : 41001);
+				Config->Cluster->PrimaryNode.Ports.ClusterSync = (ClusterSyncPort ? *ClusterSyncPort : 41001);
 
 				const uint16* ClusterEventsJsonPort = CfgJson.Cluster.MasterNode.Ports.Find(DisplayClusterConfigurationStrings::config::cluster::ports::PortClusterEventsJson);
-				Config->Cluster->MasterNode.Ports.ClusterEventsJson = (ClusterEventsJsonPort ? *ClusterEventsJsonPort : 41003);
+				Config->Cluster->PrimaryNode.Ports.ClusterEventsJson = (ClusterEventsJsonPort ? *ClusterEventsJsonPort : 41003);
 
 				const uint16* ClusterEventsBinaryPort = CfgJson.Cluster.MasterNode.Ports.Find(DisplayClusterConfigurationStrings::config::cluster::ports::PortClusterEventsBinary);
-				Config->Cluster->MasterNode.Ports.ClusterEventsBinary = (ClusterEventsBinaryPort ? *ClusterEventsBinaryPort : 41004);
+				Config->Cluster->PrimaryNode.Ports.ClusterEventsBinary = (ClusterEventsBinaryPort ? *ClusterEventsBinaryPort : 41004);
 			}
 
 			// Cluster sync
diff --git a/Engine/Plugins/Runtime/nDisplay/Source/DisplayClusterConfiguration/Private/Formats/JSON426/DisplayClusterConfigurationJsonTypes_426.h b/Engine/Plugins/Runtime/nDisplay/Source/DisplayClusterConfiguration/Private/Formats/JSON426/DisplayClusterConfigurationJsonTypes_426.h
--- a/Engine/Plugins/Runtime/nDisplay/Source/DisplayClusterConfiguration/Private/Formats/JSON426/DisplayClusterConfigurationJsonTypes_426.h
+++ b/Engine/Plugins/Runtime/nDisplay/Source/DisplayClusterConfiguration/Private/Formats/JSON426/DisplayClusterConfigurationJsonTypes_426.h
@@ -305,7 +305,7 @@ public:
 };
 
 USTRUCT()
-struct FDisplayClusterConfigurationJsonMasterNode_426
+struct FDisplayClusterConfigurationJsonPrimaryNode_426
 {
 	GENERATED_BODY()
 
@@ -411,7 +411,7 @@ struct FDisplayClusterConfigurationJsonCluster_426
 
 public:
 	UPROPERTY()
-	FDisplayClusterConfigurationJsonMasterNode_426 MasterNode;
+	FDisplayClusterConfigurationJsonPrimaryNode_426 MasterNode;
 
 	UPROPERTY()
 	FDisplayClusterConfigurationJsonClusterSync_426 Sync;
diff --git a/Engine/Plugins/Runtime/nDisplay/Source/DisplayClusterConfiguration/Private/Formats/JSON427/DisplayClusterConfigurationJsonParser_427.cpp b/Engine/Plugins/Runtime/nDisplay/Source/DisplayClusterConfiguration/Private/Formats/JSON427/DisplayClusterConfigurationJsonParser_427.cpp
--- a/Engine/Plugins/Runtime/nDisplay/Source/DisplayClusterConfiguration/Private/Formats/JSON427/DisplayClusterConfigurationJsonParser_427.cpp
+++ b/Engine/Plugins/Runtime/nDisplay/Source/DisplayClusterConfiguration/Private/Formats/JSON427/DisplayClusterConfigurationJsonParser_427.cpp
@@ -47,43 +47,14 @@ namespace JSON427
 
 	bool FDisplayClusterConfigurationJsonParser::SaveData(const UDisplayClusterConfigurationData* ConfigData, const FString& FilePath)
 	{
-		FString JsonTextOut;
-
-		// Convert to json string
-		if (!AsString(ConfigData, JsonTextOut))
-		{
-			return false;
-		}
-
-		// Save json string to a file
-		if (!FFileHelper::SaveStringToFile(JsonTextOut, *FilePath))
-		{
-			UE_LOG(LogDisplayClusterConfiguration, Error, TEXT("Couldn't save data to file: %s"), *FilePath);
-			return false;
-		}
-
-		UE_LOG(LogDisplayClusterConfiguration, Log, TEXT("Configuration data has been successfully saved to file: %s"), *FilePath);
-
-		return true;
+		UE_LOG(LogDisplayClusterConfiguration, Error, TEXT("Export to JSON 427 is not supported. Use JSON 500 exporter."));
+		return false;
 	}
 
 	bool FDisplayClusterConfigurationJsonParser::AsString(const UDisplayClusterConfigurationData* ConfigData, FString& OutString)
 	{
-		// Convert nDisplay internal types to json types
-		if (!ConvertDataToExternalTypes(ConfigData))
-		{
-			UE_LOG(LogDisplayClusterConfiguration, Error, TEXT("Couldn't convert data to json data types"));
-			return false;
-		}
-
-		// Serialize json types to json string
-		if (!FJsonObjectConverter::UStructToJsonObjectString<FDisplayClusterConfigurationJsonContainer_427>(JsonData, OutString))
-		{
-			UE_LOG(LogDisplayClusterConfiguration, Error, TEXT("Couldn't serialize data to json"));
-			return false;
-		}
-
-		return true;
+		UE_LOG(LogDisplayClusterConfiguration, Error, TEXT("Export to JSON 427 is not supported. Use JSON 500 exporter."));
+		return false;
 	}
 
 	UDisplayClusterConfigurationData* FDisplayClusterConfigurationJsonParser::ConvertDataToInternalTypes()
@@ -159,18 +130,18 @@ namespace JSON427
 
 		// Cluster
 		{
-			// Master node
+			// Primary node
 			{
-				Config->Cluster->MasterNode.Id = CfgJson.Cluster.MasterNode.Id;
+				Config->Cluster->PrimaryNode.Id = CfgJson.Cluster.MasterNode.Id;
 
 				const uint16* ClusterSyncPort = CfgJson.Cluster.MasterNode.Ports.Find(DisplayClusterConfigurationStrings::config::cluster::ports::PortClusterSync);
-				Config->Cluster->MasterNode.Ports.ClusterSync = (ClusterSyncPort ? *ClusterSyncPort : 41001);
+				Config->Cluster->PrimaryNode.Ports.ClusterSync = (ClusterSyncPort ? *ClusterSyncPort : 41001);
 
 				const uint16* ClusterEventsJsonPort = CfgJson.Cluster.MasterNode.Ports.Find(DisplayClusterConfigurationStrings::config::cluster::ports::PortClusterEventsJson);
-				Config->Cluster->MasterNode.Ports.ClusterEventsJson = (ClusterEventsJsonPort ? *ClusterEventsJsonPort : 41003);
+				Config->Cluster->PrimaryNode.Ports.ClusterEventsJson = (ClusterEventsJsonPort ? *ClusterEventsJsonPort : 41003);
 
 				const uint16* ClusterEventsBinaryPort = CfgJson.Cluster.MasterNode.Ports.Find(DisplayClusterConfigurationStrings::config::cluster::ports::PortClusterEventsBinary);
-				Config->Cluster->MasterNode.Ports.ClusterEventsBinary = (ClusterEventsBinaryPort ? *ClusterEventsBinaryPort : 41004);
+				Config->Cluster->PrimaryNode.Ports.ClusterEventsBinary = (ClusterEventsBinaryPort ? *ClusterEventsBinaryPort : 41004);
 			}
 
 			// Cluster sync
@@ -284,188 +255,7 @@ namespace JSON427
 
 	bool FDisplayClusterConfigurationJsonParser::ConvertDataToExternalTypes(const UDisplayClusterConfigurationData* Config)
 	{
-		if (!(Config && Config->Scene && Config->Cluster))
-		{
-			UE_LOG(LogDisplayClusterConfiguration, Error, TEXT("nullptr detected in the configuration data"));
-			return false;
-		}
-
-		FDisplayClusterConfigurationJsonNdisplay_427& Json = JsonData.nDisplay;
-
-		// Info
-		Json.Description = Config->Info.Description;
-		Json.Version = FString("4.27"); // Use hard-coded version number to be able to detect JSON config version on import
-		Json.AssetPath = Config->Info.AssetPath;
-
-		// Misc
-		Json.Misc.bFollowLocalPlayerCamera = Config->bFollowLocalPlayerCamera;
-		Json.Misc.bExitOnEsc = Config->bExitOnEsc;
-
-		// Scene
-		{
-			// Cameras
-			for (const TPair<FString, UDisplayClusterConfigurationSceneComponentCamera*>& Comp : Config->Scene->Cameras)
-			{
-				FDisplayClusterConfigurationJsonSceneComponentCamera_427 CfgComp;
-
-				// General
-				CfgComp.ParentId = Comp.Value->ParentId;
-				CfgComp.Location = FDisplayClusterConfigurationJsonVector_427::FromVector(Comp.Value->Location);
-				CfgComp.Rotation = FDisplayClusterConfigurationJsonRotator_427::FromRotator(Comp.Value->Rotation);
-				// Camera specific
-				CfgComp.InterpupillaryDistance = Comp.Value->InterpupillaryDistance;
-				CfgComp.SwapEyes = Comp.Value->bSwapEyes;
-				CfgComp.StereoOffset = DisplayClusterConfigurationJsonHelpers::ToString(Comp.Value->StereoOffset);
-
-				Json.Scene.Cameras.Emplace(Comp.Key, CfgComp);
-			}
-
-			// Screens
-			for (const TPair<FString, UDisplayClusterConfigurationSceneComponentScreen*>& Comp : Config->Scene->Screens)
-			{
-				FDisplayClusterConfigurationJsonSceneComponentScreen_427 CfgComp;
-
-				// General
-				CfgComp.ParentId = Comp.Value->ParentId;
-				CfgComp.Location = FDisplayClusterConfigurationJsonVector_427::FromVector(Comp.Value->Location);
-				CfgComp.Rotation = FDisplayClusterConfigurationJsonRotator_427::FromRotator(Comp.Value->Rotation);
-				// Screen specific
-				CfgComp.Size = FDisplayClusterConfigurationJsonSizeFloat_427::FromVector(Comp.Value->Size);
-
-				Json.Scene.Screens.Emplace(Comp.Key, CfgComp);
-			}
-
-			// Xforms
-			for (const TPair<FString, UDisplayClusterConfigurationSceneComponentXform*>& Comp : Config->Scene->Xforms)
-			{
-				FDisplayClusterConfigurationJsonSceneComponentScreen_427 CfgComp;
-
-				// General
-				CfgComp.ParentId = Comp.Value->ParentId;
-				CfgComp.Location = FDisplayClusterConfigurationJsonVector_427::FromVector(Comp.Value->Location);
-				CfgComp.Rotation = FDisplayClusterConfigurationJsonRotator_427::FromRotator(Comp.Value->Rotation);
-
-				Json.Scene.Xforms.Emplace(Comp.Key, CfgComp);
-			}
-		}
-
-		// Cluster
-		{
-			// Master node
-			{
-				Json.Cluster.MasterNode.Id = Config->Cluster->MasterNode.Id;
-				Json.Cluster.MasterNode.Ports.Emplace(DisplayClusterConfigurationStrings::config::cluster::ports::PortClusterSync,         Config->Cluster->MasterNode.Ports.ClusterSync);
-				Json.Cluster.MasterNode.Ports.Emplace(DisplayClusterConfigurationStrings::config::cluster::ports::PortClusterEventsJson,   Config->Cluster->MasterNode.Ports.ClusterEventsJson);
-				Json.Cluster.MasterNode.Ports.Emplace(DisplayClusterConfigurationStrings::config::cluster::ports::PortClusterEventsBinary, Config->Cluster->MasterNode.Ports.ClusterEventsBinary);
-			}
-
-			// Cluster sync
-			{
-				// Native input sync
-				Json.Cluster.Sync.InputSyncPolicy.Type = Config->Cluster->Sync.InputSyncPolicy.Type;
-				Json.Cluster.Sync.InputSyncPolicy.Parameters = Config->Cluster->Sync.InputSyncPolicy.Parameters;
-
-				// Render sync
-				Json.Cluster.Sync.RenderSyncPolicy.Type = Config->Cluster->Sync.RenderSyncPolicy.Type;
-				Json.Cluster.Sync.RenderSyncPolicy.Parameters = Config->Cluster->Sync.RenderSyncPolicy.Parameters;
-			}
-
-			// Network
-			{
-				Json.Cluster.Network.Emplace(DisplayClusterConfigurationStrings::config::cluster::network::NetConnectRetriesAmount,
-					DisplayClusterTypesConverter::template ToString(Config->Cluster->Network.ConnectRetriesAmount));
-
-				Json.Cluster.Network.Emplace(DisplayClusterConfigurationStrings::config::cluster::network::NetConnectRetryDelay,
-					DisplayClusterTypesConverter::template ToString(Config->Cluster->Network.ConnectRetryDelay));
-
-				Json.Cluster.Network.Emplace(DisplayClusterConfigurationStrings::config::cluster::network::NetGameStartBarrierTimeout,
-					DisplayClusterTypesConverter::template ToString(Config->Cluster->Network.GameStartBarrierTimeout));
-
-				Json.Cluster.Network.Emplace(DisplayClusterConfigurationStrings::config::cluster::network::NetFrameStartBarrierTimeout,
-					DisplayClusterTypesConverter::template ToString(Config->Cluster->Network.FrameStartBarrierTimeout));
-
-				Json.Cluster.Network.Emplace(DisplayClusterConfigurationStrings::config::cluster::network::NetFrameEndBarrierTimeout,
-					DisplayClusterTypesConverter::template ToString(Config->Cluster->Network.FrameEndBarrierTimeout));
-
-				Json.Cluster.Network.Emplace(DisplayClusterConfigurationStrings::config::cluster::network::NetRenderSyncBarrierTimeout,
-					DisplayClusterTypesConverter::template ToString(Config->Cluster->Network.RenderSyncBarrierTimeout));
-			}
-
-			// Cluster nodes
-			for (const TPair<FString, UDisplayClusterConfigurationClusterNode*>& CfgNode : Config->Cluster->Nodes)
-			{
-				FDisplayClusterConfigurationJsonClusterNode_427 Node;
-
-				// Base parameters
-				Node.Host = CfgNode.Value->Host;
-				Node.Sound = CfgNode.Value->bIsSoundEnabled;
-				Node.FullScreen = CfgNode.Value->bIsFullscreen;
-				Node.Window = FDisplayClusterConfigurationJsonRectangle_427(CfgNode.Value->WindowRect.X, CfgNode.Value->WindowRect.Y, CfgNode.Value->WindowRect.W, CfgNode.Value->WindowRect.H);
-
-				// Viewports
-				for (const TPair<FString, UDisplayClusterConfigurationViewport*>& CfgViewport : CfgNode.Value->Viewports)
-				{
-					FDisplayClusterConfigurationJsonViewport_427 Viewport;
-
-					// Base parameters
-					Viewport.Camera = CfgViewport.Value->Camera;
-					Viewport.Region = FDisplayClusterConfigurationJsonRectangle_427(CfgViewport.Value->Region.X, CfgViewport.Value->Region.Y, CfgViewport.Value->Region.W, CfgViewport.Value->Region.H);
-					Viewport.GPUIndex = CfgViewport.Value->GPUIndex;
-					Viewport.BufferRatio = CfgViewport.Value->RenderSettings.BufferRatio;
-
-					// Projection policy
-					Viewport.ProjectionPolicy.Type = CfgViewport.Value->ProjectionPolicy.Type;
-					Viewport.ProjectionPolicy.Parameters = CfgViewport.Value->ProjectionPolicy.Parameters;
-
-					// Overscan
-					Viewport.Overscan.Mode = DisplayClusterConfigurationJsonHelpers::ToString(CfgViewport.Value->RenderSettings.Overscan.Mode);
-					Viewport.Overscan.bEnabled = CfgViewport.Value->RenderSettings.Overscan.bEnabled;
-					Viewport.Overscan.Oversize = CfgViewport.Value->RenderSettings.Overscan.bOversize;
-					Viewport.Overscan.Left = CfgViewport.Value->RenderSettings.Overscan.Left;
-					Viewport.Overscan.Right = CfgViewport.Value->RenderSettings.Overscan.Right;
-					Viewport.Overscan.Top = CfgViewport.Value->RenderSettings.Overscan.Top;
-					Viewport.Overscan.Bottom = CfgViewport.Value->RenderSettings.Overscan.Bottom;
-					
-					// Save this viewport
-					Node.Viewports.Emplace(CfgViewport.Key, Viewport);
-				}
-
-				// Postprocess
-				for (const TPair<FString, FDisplayClusterConfigurationPostprocess>& CfgPostprocess : CfgNode.Value->Postprocess)
-				{
-					FDisplayClusterConfigurationJsonPostprocess_427 PostprocessOperation;
-
-					PostprocessOperation.Type = CfgPostprocess.Value.Type;
-					PostprocessOperation.Parameters = CfgPostprocess.Value.Parameters;
-
-					Node.Postprocess.Emplace(CfgPostprocess.Key, PostprocessOperation);
-				}
-
-				// OutputRemap
-				{
-					Node.OutputRemap.bEnable = CfgNode.Value->OutputRemap.bEnable;
-					Node.OutputRemap.DataSource = (CfgNode.Value->OutputRemap.DataSource == EDisplayClusterConfigurationFramePostProcess_OutputRemapSource::ExternalFile) ? "file" : "mesh";
-					Node.OutputRemap.ExternalFile = CfgNode.Value->OutputRemap.ExternalFile;
-
-					if (CfgNode.Value->OutputRemap.StaticMesh)
-					{
-						Node.OutputRemap.StaticMeshAsset = CfgNode.Value->OutputRemap.StaticMesh->GetPathName();
-					}
-				}
-
-				// Store new cluster node
-				Json.Cluster.Nodes.Emplace(CfgNode.Key, Node);
-			}
-		}
-
-		// Custom parameters
-		Json.CustomParameters = Config->CustomParameters;
-
-		// Diagnostics
-		Json.Diagnostics.SimulateLag = Config->Diagnostics.bSimulateLag;
-		Json.Diagnostics.MinLagTime = Config->Diagnostics.MinLagTime;
-		Json.Diagnostics.MaxLagTime = Config->Diagnostics.MaxLagTime;
-
-		return true;
+		// Not supported anymore
+		return false;
 	}
 }
diff --git a/Engine/Plugins/Runtime/nDisplay/Source/DisplayClusterConfiguration/Private/Formats/JSON427/DisplayClusterConfigurationJsonTypes_427.h b/Engine/Plugins/Runtime/nDisplay/Source/DisplayClusterConfiguration/Private/Formats/JSON427/DisplayClusterConfigurationJsonTypes_427.h
--- a/Engine/Plugins/Runtime/nDisplay/Source/DisplayClusterConfiguration/Private/Formats/JSON427/DisplayClusterConfigurationJsonTypes_427.h
+++ b/Engine/Plugins/Runtime/nDisplay/Source/DisplayClusterConfiguration/Private/Formats/JSON427/DisplayClusterConfigurationJsonTypes_427.h
@@ -330,7 +330,7 @@ public:
 
 
 USTRUCT()
-struct FDisplayClusterConfigurationJsonMasterNode_427
+struct FDisplayClusterConfigurationJsonPrimaryNode_427
 {
 	GENERATED_BODY()
 
@@ -489,7 +489,7 @@ struct FDisplayClusterConfigurationJsonCluster_427
 
 public:
 	UPROPERTY()
-	FDisplayClusterConfigurationJsonMasterNode_427 MasterNode;
+	FDisplayClusterConfigurationJsonPrimaryNode_427 MasterNode;
 
 	UPROPERTY()
 	FDisplayClusterConfigurationJsonClusterSync_427 Sync;
diff --git a/Engine/Plugins/Runtime/nDisplay/Source/DisplayClusterConfiguration/Private/Formats/JSON500/DisplayClusterConfigurationJsonHelpers_500.h b/Engine/Plugins/Runtime/nDisplay/Source/DisplayClusterConfiguration/Private/Formats/JSON500/DisplayClusterConfigurationJsonHelpers_500.h
new file mode 100644
--- /dev/null
+++ b/Engine/Plugins/Runtime/nDisplay/Source/DisplayClusterConfiguration/Private/Formats/JSON500/DisplayClusterConfigurationJsonHelpers_500.h
@@ -0,0 +1,129 @@
+// Copyright Epic Games, Inc. All Rights Reserved.
+
+#pragma once
+
+#include "DisplayClusterConfigurationLog.h"
+#include "DisplayClusterConfigurationStrings.h"
+#include "DisplayClusterConfigurationTypes.h"
+
+
+namespace JSON500
+{
+	// since we cannot modify public headers we need to put those string in private headers temporary
+	namespace config
+	{
+		namespace node
+		{
+			namespace viewport
+			{
+				namespace overscan
+				{
+					static constexpr auto OverscanNone = TEXT("none");
+					static constexpr auto OverscanPixels = TEXT("pixels");
+					static constexpr auto OverscanPercent = TEXT("percent");
+				}
+			}
+		}
+	}
+
+	/**
+	 * Auxiliary class with different type conversion functions
+	 */
+	class DisplayClusterConfigurationJsonHelpers
+	{
+	public:
+		//////////////////////////////////////////////////////////////////////////////////////////////
+		// TYPE --> STRING
+		//////////////////////////////////////////////////////////////////////////////////////////////
+		template <typename ConvertFrom>
+		static FString ToString(const ConvertFrom& From);
+
+		//////////////////////////////////////////////////////////////////////////////////////////////
+		// STRING --> TYPE
+		//////////////////////////////////////////////////////////////////////////////////////////////
+		template <typename ConvertTo>
+		static ConvertTo FromString(const FString& From);
+	};
+
+
+	// EDisplayClusterConfigurationEyeStereoOffset
+	template <>
+	inline FString DisplayClusterConfigurationJsonHelpers::ToString<>(const EDisplayClusterConfigurationEyeStereoOffset& From)
+	{
+		switch (From)
+		{
+		case EDisplayClusterConfigurationEyeStereoOffset::None:
+			return DisplayClusterConfigurationStrings::config::scene::camera::CameraStereoOffsetNone;
+
+		case EDisplayClusterConfigurationEyeStereoOffset::Left:
+			return DisplayClusterConfigurationStrings::config::scene::camera::CameraStereoOffsetLeft;
+
+		case EDisplayClusterConfigurationEyeStereoOffset::Right:
+			return DisplayClusterConfigurationStrings::config::scene::camera::CameraStereoOffsetRight;
+
+		default:
+			UE_LOG(LogDisplayClusterConfiguration, Error, TEXT("Unexpected camera stereo offset type"));
+			break;
+		}
+
+		return DisplayClusterConfigurationStrings::config::scene::camera::CameraStereoOffsetNone;
+	}
+
+	template <>
+	inline EDisplayClusterConfigurationEyeStereoOffset DisplayClusterConfigurationJsonHelpers::FromString<>(const FString& From)
+	{
+		EDisplayClusterConfigurationEyeStereoOffset Result = EDisplayClusterConfigurationEyeStereoOffset::None;
+
+		if (From.Equals(DisplayClusterConfigurationStrings::config::scene::camera::CameraStereoOffsetNone, ESearchCase::IgnoreCase))
+		{
+			Result = EDisplayClusterConfigurationEyeStereoOffset::None;
+		}
+		else if (From.Equals(DisplayClusterConfigurationStrings::config::scene::camera::CameraStereoOffsetLeft, ESearchCase::IgnoreCase))
+		{
+			Result = EDisplayClusterConfigurationEyeStereoOffset::Left;
+		}
+		else if (From.Equals(DisplayClusterConfigurationStrings::config::scene::camera::CameraStereoOffsetRight, ESearchCase::IgnoreCase))
+		{
+			Result = EDisplayClusterConfigurationEyeStereoOffset::Right;
+		}
+
+		return Result;
+	}
+
+	// EDisplayClusterConfigurationViewportOverscanMode
+	template <>
+	inline FString DisplayClusterConfigurationJsonHelpers::ToString<>(const EDisplayClusterConfigurationViewportOverscanMode& InOverscanMode)
+	{
+		switch (InOverscanMode)
+		{
+		case EDisplayClusterConfigurationViewportOverscanMode::Pixels:
+			return JSON500::config::node::viewport::overscan::OverscanPixels;
+
+		case EDisplayClusterConfigurationViewportOverscanMode::Percent:
+			return JSON500::config::node::viewport::overscan::OverscanPercent;
+
+		default:
+			UE_LOG(LogDisplayClusterConfiguration, Error, TEXT("Unexpected overscan type"));
+			break;
+		}
+
+		return JSON500::config::node::viewport::overscan::OverscanNone;
+	}
+
+	template <>
+	inline EDisplayClusterConfigurationViewportOverscanMode DisplayClusterConfigurationJsonHelpers::FromString<>(const FString& InOverscanModeString)
+	{
+		EDisplayClusterConfigurationViewportOverscanMode Result = EDisplayClusterConfigurationViewportOverscanMode::Percent;
+
+		if (InOverscanModeString.Equals(JSON500::config::node::viewport::overscan::OverscanPixels, ESearchCase::IgnoreCase))
+		{
+			Result = EDisplayClusterConfigurationViewportOverscanMode::Pixels;
+		}
+		else if (InOverscanModeString.Equals(JSON500::config::node::viewport::overscan::OverscanPercent, ESearchCase::IgnoreCase))
+		{
+			Result = EDisplayClusterConfigurationViewportOverscanMode::Percent;
+		}
+
+		return Result;
+	}
+}
\ No newline at end of file
diff --git a/Engine/Plugins/Runtime/nDisplay/Source/DisplayClusterConfiguration/Private/Formats/JSON500/DisplayClusterConfigurationJsonParser_500.cpp b/Engine/Plugins/Runtime/nDisplay/Source/DisplayClusterConfiguration/Private/Formats/JSON500/DisplayClusterConfigurationJsonParser_500.cpp
new file mode 100644
--- /dev/null
+++ b/Engine/Plugins/Runtime/nDisplay/Source/DisplayClusterConfiguration/Private/Formats/JSON500/DisplayClusterConfigurationJsonParser_500.cpp
@@ -0,0 +1,471 @@
+// Copyright Epic Games, Inc. All Rights Reserved.
+
+#include "Formats/JSON500/DisplayClusterConfigurationJsonParser_500.h"
+#include "Formats/JSON500/DisplayClusterConfigurationJsonHelpers_500.h"
+
+#include "Misc/DisplayClusterHelpers.h"
+
+#include "DisplayClusterConfigurationLog.h"
+#include "DisplayClusterConfigurationTypes.h"
+#include "DisplayClusterConfigurationStrings.h"
+
+#include "JsonObjectConverter.h"
+
+#include "Misc/FileHelper.h"
+#include "Misc/Paths.h"
+
+
+namespace JSON500
+{
+	UDisplayClusterConfigurationData* FDisplayClusterConfigurationJsonParser::LoadData(const FString& FilePath, UObject* Owner)
+	{
+		ConfigDataOwner = Owner;
+
+		FString JsonText;
+
+		// Load json text to the string object
+		if (!FFileHelper::LoadFileToString(JsonText, *FilePath))
+		{
+			UE_LOG(LogDisplayClusterConfiguration, Error, TEXT("Couldn't read file: %s"), *FilePath);
+			return nullptr;
+		}
+
+		UE_LOG(LogDisplayClusterConfiguration, Log, TEXT("nDisplay json configuration: %s"), *JsonText);
+
+		// Parse the string object
+		if (!FJsonObjectConverter::JsonObjectStringToUStruct<FDisplayClusterConfigurationJsonContainer_500>(JsonText, &JsonData, 0, 0))
+		{
+			UE_LOG(LogDisplayClusterConfiguration, Error, TEXT("Couldn't deserialize json file: %s"), *FilePath);
+			return nullptr;
+		}
+
+		ConfigFile = FilePath;
+
+		// Finally, convert the data to nDisplay internal types
+		return ConvertDataToInternalTypes();
+	}
+
+	bool FDisplayClusterConfigurationJsonParser::SaveData(const UDisplayClusterConfigurationData* ConfigData, const FString& FilePath)
+	{
+		FString JsonTextOut;
+
+		// Convert to json string
+		if (!AsString(ConfigData, JsonTextOut))
+		{
+			return false;
+		}
+
+		// Save json string to a file
+		if (!FFileHelper::SaveStringToFile(JsonTextOut, *FilePath))
+		{
+			UE_LOG(LogDisplayClusterConfiguration, Error, TEXT("Couldn't save data to file: %s"), *FilePath);
+			return false;
+		}
+
+		UE_LOG(LogDisplayClusterConfiguration, Log, TEXT("Configuration data has been successfully saved to file: %s"), *FilePath);
+
+		return true;
+	}
+
+	bool FDisplayClusterConfigurationJsonParser::AsString(const UDisplayClusterConfigurationData* ConfigData, FString& OutString)
+	{
+		// Convert nDisplay internal types to json types
+		if (!ConvertDataToExternalTypes(ConfigData))
+		{
+			UE_LOG(LogDisplayClusterConfiguration, Error, TEXT("Couldn't convert data to json data types"));
+			return false;
+		}
+
+		// Serialize json types to json string
+		if (!FJsonObjectConverter::UStructToJsonObjectString<FDisplayClusterConfigurationJsonContainer_500>(JsonData, OutString))
+		{
+			UE_LOG(LogDisplayClusterConfiguration, Error, TEXT("Couldn't serialize data to json"));
+			return false;
+		}
+
+		return true;
+	}
+
+	UDisplayClusterConfigurationData* FDisplayClusterConfigurationJsonParser::ConvertDataToInternalTypes()
+	{
+		UDisplayClusterConfigurationData* Config = UDisplayClusterConfigurationData::CreateNewConfigData(ConfigDataOwner, RF_MarkAsRootSet);
+		check(Config && Config->Scene && Config->Cluster);
+
+		// Fill metadata
+		Config->Meta.ImportDataSource = EDisplayClusterConfigurationDataSource::Json;
+		Config->Meta.ImportFilePath = ConfigFile;
+
+		const FDisplayClusterConfigurationJsonNdisplay_500& CfgJson = JsonData.nDisplay;
+
+		// Info
+		Config->Info.Version = CfgJson.Version;
+		Config->Info.Description = CfgJson.Description;
+		Config->Info.AssetPath = CfgJson.AssetPath;
+
+		// Misc
+		Config->bFollowLocalPlayerCamera = CfgJson.Misc.bFollowLocalPlayerCamera;
+		Config->bExitOnEsc = CfgJson.Misc.bExitOnEsc;
+
+		// Scene
+		{
+			// Cameras
+			for (const TPair<FString, FDisplayClusterConfigurationJsonSceneComponentCamera_500>& CfgComp : CfgJson.Scene.Cameras)
+			{
+				UDisplayClusterConfigurationSceneComponentCamera* Comp = NewObject<UDisplayClusterConfigurationSceneComponentCamera>(Config->Scene, *CfgComp.Key, RF_Transactional);
+				check(Comp);
+
+				// General
+				Comp->ParentId = CfgComp.Value.ParentId;
+				Comp->Location = FDisplayClusterConfigurationJsonVector_500::ToVector(CfgComp.Value.Location);
+				Comp->Rotation = FDisplayClusterConfigurationJsonRotator_500::ToRotator(CfgComp.Value.Rotation);
+				// Camera specific
+				Comp->InterpupillaryDistance = CfgComp.Value.InterpupillaryDistance;
+				Comp->bSwapEyes = CfgComp.Value.SwapEyes;
+				Comp->StereoOffset = DisplayClusterConfigurationJsonHelpers::FromString<EDisplayClusterConfigurationEyeStereoOffset>(CfgComp.Value.StereoOffset);
+
+				Config->Scene->Cameras.Emplace(CfgComp.Key, Comp);
+			}
+
+			// Screens
+			for (const TPair<FString, FDisplayClusterConfigurationJsonSceneComponentScreen_500>& CfgComp : CfgJson.Scene.Screens)
+			{
+				UDisplayClusterConfigurationSceneComponentScreen* Comp = NewObject<UDisplayClusterConfigurationSceneComponentScreen>(Config->Scene, *CfgComp.Key, RF_Transactional);
+				check(Comp);
+
+				// General
+				Comp->ParentId = CfgComp.Value.ParentId;
+				Comp->Location = FDisplayClusterConfigurationJsonVector_500::ToVector(CfgComp.Value.Location);
+				Comp->Rotation = FDisplayClusterConfigurationJsonRotator_500::ToRotator(CfgComp.Value.Rotation);
+				// Screen specific
+				Comp->Size = FDisplayClusterConfigurationJsonSizeFloat_500::ToVector(CfgComp.Value.Size);
+
+				Config->Scene->Screens.Emplace(CfgComp.Key, Comp);
+			}
+
+			// Xforms
+			for (const TPair<FString, FDisplayClusterConfigurationJsonSceneComponentXform_500>& CfgComp : CfgJson.Scene.Xforms)
+			{
+				UDisplayClusterConfigurationSceneComponentXform* Comp = NewObject<UDisplayClusterConfigurationSceneComponentXform>(Config->Scene, *CfgComp.Key, RF_Transactional);
+				check(Comp);
+
+				// General
+				Comp->ParentId = CfgComp.Value.ParentId;
+				Comp->Location = FDisplayClusterConfigurationJsonVector_500::ToVector(CfgComp.Value.Location);
+				Comp->Rotation = FDisplayClusterConfigurationJsonRotator_500::ToRotator(CfgComp.Value.Rotation);
+
+				Config->Scene->Xforms.Emplace(CfgComp.Key, Comp);
+			}
+		}
+
+		// Cluster
+		{
+			// Primary node
+			{
+				Config->Cluster->PrimaryNode.Id = CfgJson.Cluster.PrimaryNode.Id;
+
+				const uint16* ClusterSyncPort = CfgJson.Cluster.PrimaryNode.Ports.Find(DisplayClusterConfigurationStrings::config::cluster::ports::PortClusterSync);
+				Config->Cluster->PrimaryNode.Ports.ClusterSync = (ClusterSyncPort ? *ClusterSyncPort : 41001);
+
+				const uint16* ClusterEventsJsonPort = CfgJson.Cluster.PrimaryNode.Ports.Find(DisplayClusterConfigurationStrings::config::cluster::ports::PortClusterEventsJson);
+				Config->Cluster->PrimaryNode.Ports.ClusterEventsJson = (ClusterEventsJsonPort ? *ClusterEventsJsonPort : 41003);
+
+				const uint16* ClusterEventsBinaryPort = CfgJson.Cluster.PrimaryNode.Ports.Find(DisplayClusterConfigurationStrings::config::cluster::ports::PortClusterEventsBinary);
+				Config->Cluster->PrimaryNode.Ports.ClusterEventsBinary = (ClusterEventsBinaryPort ? *ClusterEventsBinaryPort : 41004);
+			}
+
+			// Cluster sync
+			{
+				// Native input sync
+				Config->Cluster->Sync.InputSyncPolicy.Type = CfgJson.Cluster.Sync.InputSyncPolicy.Type;
+				Config->Cluster->Sync.InputSyncPolicy.Parameters = CfgJson.Cluster.Sync.InputSyncPolicy.Parameters;
+
+				// Render sync
+				Config->Cluster->Sync.RenderSyncPolicy.Type = CfgJson.Cluster.Sync.RenderSyncPolicy.Type;
+				Config->Cluster->Sync.RenderSyncPolicy.Parameters = CfgJson.Cluster.Sync.RenderSyncPolicy.Parameters;
+			}
+
+			// Network
+			{
+				Config->Cluster->Network.ConnectRetriesAmount = DisplayClusterHelpers::map::template ExtractValueFromString(
+					CfgJson.Cluster.Network, FString(DisplayClusterConfigurationStrings::config::cluster::network::NetConnectRetriesAmount), (uint32)15);
+
+				Config->Cluster->Network.ConnectRetryDelay = DisplayClusterHelpers::map::template ExtractValueFromString(
+					CfgJson.Cluster.Network, FString(DisplayClusterConfigurationStrings::config::cluster::network::NetConnectRetryDelay), (uint32)1000);
+
+				Config->Cluster->Network.GameStartBarrierTimeout = DisplayClusterHelpers::map::template ExtractValueFromString(
+					CfgJson.Cluster.Network, FString(DisplayClusterConfigurationStrings::config::cluster::network::NetGameStartBarrierTimeout), (uint32)30000);
+
+				Config->Cluster->Network.FrameStartBarrierTimeout = DisplayClusterHelpers::map::template ExtractValueFromString(
+					CfgJson.Cluster.Network, FString(DisplayClusterConfigurationStrings::config::cluster::network::NetFrameStartBarrierTimeout), (uint32)5000);
+
+				Config->Cluster->Network.FrameEndBarrierTimeout = DisplayClusterHelpers::map::template ExtractValueFromString(
+					CfgJson.Cluster.Network, FString(DisplayClusterConfigurationStrings::config::cluster::network::NetFrameEndBarrierTimeout), (uint32)5000);
+
+				Config->Cluster->Network.RenderSyncBarrierTimeout = DisplayClusterHelpers::map::template ExtractValueFromString(
+					CfgJson.Cluster.Network, FString(DisplayClusterConfigurationStrings::config::cluster::network::NetRenderSyncBarrierTimeout), (uint32)5000);
+			}
+
+			// Cluster nodes
+			for (const TPair<FString, FDisplayClusterConfigurationJsonClusterNode_500>& CfgNode : CfgJson.Cluster.Nodes)
+			{
+				UDisplayClusterConfigurationClusterNode* Node = NewObject<UDisplayClusterConfigurationClusterNode>(Config->Cluster, *CfgNode.Key, RF_Transactional);
+				check(Node);
+
+				// Base parameters
+				Node->Host = CfgNode.Value.Host;
+				Node->bIsSoundEnabled = CfgNode.Value.Sound;
+				Node->bIsFullscreen = CfgNode.Value.FullScreen;
+				Node->WindowRect = FDisplayClusterConfigurationRectangle(CfgNode.Value.Window.X, CfgNode.Value.Window.Y, CfgNode.Value.Window.W, CfgNode.Value.Window.H);
+
+				// Viewports
+				for (const TPair<FString, FDisplayClusterConfigurationJsonViewport_500>& CfgViewport : CfgNode.Value.Viewports)
+				{
+					UDisplayClusterConfigurationViewport* Viewport = NewObject<UDisplayClusterConfigurationViewport>(Node, *CfgViewport.Key, RF_Transactional | RF_ArchetypeObject | RF_Public);
+					check(Viewport);
+
+					// Base parameters
+					Viewport->RenderSettings.BufferRatio = CfgViewport.Value.BufferRatio;
+					Viewport->Camera = CfgViewport.Value.Camera;
+					Viewport->Region = FDisplayClusterConfigurationRectangle(CfgViewport.Value.Region.X, CfgViewport.Value.Region.Y, CfgViewport.Value.Region.W, CfgViewport.Value.Region.H);
+					Viewport->GPUIndex = CfgViewport.Value.GPUIndex;
+
+					// Projection policy
+					Viewport->ProjectionPolicy.Type = CfgViewport.Value.ProjectionPolicy.Type;
+					Viewport->ProjectionPolicy.Parameters = CfgViewport.Value.ProjectionPolicy.Parameters;
+
+					Viewport->RenderSettings.Overscan.Mode = DisplayClusterConfigurationJsonHelpers::FromString<EDisplayClusterConfigurationViewportOverscanMode>(CfgViewport.Value.Overscan.Mode);
+					Viewport->RenderSettings.Overscan.bEnabled = CfgViewport.Value.Overscan.bEnabled;
+					Viewport->RenderSettings.Overscan.bOversize = CfgViewport.Value.Overscan.Oversize;
+					Viewport->RenderSettings.Overscan.Left = CfgViewport.Value.Overscan.Left;
+					Viewport->RenderSettings.Overscan.Right = CfgViewport.Value.Overscan.Right;
+					Viewport->RenderSettings.Overscan.Top = CfgViewport.Value.Overscan.Top;
+					Viewport->RenderSettings.Overscan.Bottom = CfgViewport.Value.Overscan.Bottom;
+
+					// Add this viewport
+					Node->Viewports.Emplace(CfgViewport.Key, Viewport);
+				}
+
+				// Postprocess
+				for (const TPair<FString, FDisplayClusterConfigurationJsonPostprocess_500>& CfgPostprocess : CfgNode.Value.Postprocess)
+				{
+					FDisplayClusterConfigurationPostprocess PostprocessOperation;
+
+					PostprocessOperation.Type = CfgPostprocess.Value.Type;
+					PostprocessOperation.Parameters = CfgPostprocess.Value.Parameters;
+
+					Node->Postprocess.Emplace(CfgPostprocess.Key, PostprocessOperation);
+				}
+
+				// Output remap
+				Node->OutputRemap.bEnable = CfgNode.Value.OutputRemap.bEnable;
+				Node->OutputRemap.DataSource = (CfgNode.Value.OutputRemap.DataSource == "file") ? EDisplayClusterConfigurationFramePostProcess_OutputRemapSource::ExternalFile : EDisplayClusterConfigurationFramePostProcess_OutputRemapSource::StaticMesh;
+				Node->OutputRemap.ExternalFile = CfgNode.Value.OutputRemap.ExternalFile;
+
+				if (!CfgNode.Value.OutputRemap.StaticMeshAsset.IsEmpty())
+				{
+					Node->OutputRemap.StaticMesh = Cast<UStaticMesh>(StaticLoadObject(UStaticMesh::StaticClass(), nullptr, *CfgNode.Value.OutputRemap.StaticMeshAsset));
+				}
+
+				// Store new cluster node
+				Config->Cluster->Nodes.Emplace(CfgNode.Key, Node);
+			}
+		}
+
+		// Custom parameters
+		Config->CustomParameters = CfgJson.CustomParameters;
+
+		// Diagnostics
+		Config->Diagnostics.bSimulateLag = CfgJson.Diagnostics.SimulateLag;
+		Config->Diagnostics.MinLagTime = CfgJson.Diagnostics.MinLagTime;
+		Config->Diagnostics.MaxLagTime = CfgJson.Diagnostics.MaxLagTime;
+
+		return Config;
+	}
+
+	bool FDisplayClusterConfigurationJsonParser::ConvertDataToExternalTypes(const UDisplayClusterConfigurationData* Config)
+	{
+		if (!(Config && Config->Scene && Config->Cluster))
+		{
+			UE_LOG(LogDisplayClusterConfiguration, Error, TEXT("nullptr detected in the configuration data"));
+			return false;
+		}
+
+		FDisplayClusterConfigurationJsonNdisplay_500& Json = JsonData.nDisplay;
+
+		// Info
+		Json.Description = Config->Info.Description;
+		Json.Version   = FString("5.00"); // Use hard-coded version number to be able to detect JSON config version on import
+		Json.AssetPath = Config->Info.AssetPath;
+
+		// Misc
+		Json.Misc.bFollowLocalPlayerCamera = Config->bFollowLocalPlayerCamera;
+		Json.Misc.bExitOnEsc = Config->bExitOnEsc;
+
+		// Scene
+		{
+			// Cameras
+			for (const TPair<FString, UDisplayClusterConfigurationSceneComponentCamera*>& Comp : Config->Scene->Cameras)
+			{
+				FDisplayClusterConfigurationJsonSceneComponentCamera_500 CfgComp;
+
+				// General
+				CfgComp.ParentId = Comp.Value->ParentId;
+				CfgComp.Location = FDisplayClusterConfigurationJsonVector_500::FromVector(Comp.Value->Location);
+				CfgComp.Rotation = FDisplayClusterConfigurationJsonRotator_500::FromRotator(Comp.Value->Rotation);
+				// Camera specific
+				CfgComp.InterpupillaryDistance = Comp.Value->InterpupillaryDistance;
+				CfgComp.SwapEyes = Comp.Value->bSwapEyes;
+				CfgComp.StereoOffset = DisplayClusterConfigurationJsonHelpers::ToString(Comp.Value->StereoOffset);
+
+				Json.Scene.Cameras.Emplace(Comp.Key, CfgComp);
+			}
+
+			// Screens
+			for (const TPair<FString, UDisplayClusterConfigurationSceneComponentScreen*>& Comp : Config->Scene->Screens)
+			{
+				FDisplayClusterConfigurationJsonSceneComponentScreen_500 CfgComp;
+
+				// General
+				CfgComp.ParentId = Comp.Value->ParentId;
+				CfgComp.Location = FDisplayClusterConfigurationJsonVector_500::FromVector(Comp.Value->Location);
+				CfgComp.Rotation = FDisplayClusterConfigurationJsonRotator_500::FromRotator(Comp.Value->Rotation);
+				// Screen specific
+				CfgComp.Size = FDisplayClusterConfigurationJsonSizeFloat_500::FromVector(Comp.Value->Size);
+
+				Json.Scene.Screens.Emplace(Comp.Key, CfgComp);
+			}
+
+			// Xforms
+			for (const TPair<FString, UDisplayClusterConfigurationSceneComponentXform*>& Comp : Config->Scene->Xforms)
+			{
+				FDisplayClusterConfigurationJsonSceneComponentScreen_500 CfgComp;
+
+				// General
+				CfgComp.ParentId = Comp.Value->ParentId;
+				CfgComp.Location = FDisplayClusterConfigurationJsonVector_500::FromVector(Comp.Value->Location);
+				CfgComp.Rotation = FDisplayClusterConfigurationJsonRotator_500::FromRotator(Comp.Value->Rotation);
+
+				Json.Scene.Xforms.Emplace(Comp.Key, CfgComp);
+			}
+		}
+
+		// Cluster
+		{
+			// Primary node
+			{
+				Json.Cluster.PrimaryNode.Id = Config->Cluster->PrimaryNode.Id;
+				Json.Cluster.PrimaryNode.Ports.Emplace(DisplayClusterConfigurationStrings::config::cluster::ports::PortClusterSync,         Config->Cluster->PrimaryNode.Ports.ClusterSync);
+				Json.Cluster.PrimaryNode.Ports.Emplace(DisplayClusterConfigurationStrings::config::cluster::ports::PortClusterEventsJson,   Config->Cluster->PrimaryNode.Ports.ClusterEventsJson);
+				Json.Cluster.PrimaryNode.Ports.Emplace(DisplayClusterConfigurationStrings::config::cluster::ports::PortClusterEventsBinary, Config->Cluster->PrimaryNode.Ports.ClusterEventsBinary);
+			}
+
+			// Cluster sync
+			{
+				// Native input sync
+				Json.Cluster.Sync.InputSyncPolicy.Type = Config->Cluster->Sync.InputSyncPolicy.Type;
+				Json.Cluster.Sync.InputSyncPolicy.Parameters = Config->Cluster->Sync.InputSyncPolicy.Parameters;
+
+				// Render sync
+				Json.Cluster.Sync.RenderSyncPolicy.Type = Config->Cluster->Sync.RenderSyncPolicy.Type;
+				Json.Cluster.Sync.RenderSyncPolicy.Parameters = Config->Cluster->Sync.RenderSyncPolicy.Parameters;
+			}
+
+			// Network
+			{
+				Json.Cluster.Network.Emplace(DisplayClusterConfigurationStrings::config::cluster::network::NetConnectRetriesAmount,
+					DisplayClusterTypesConverter::template ToString(Config->Cluster->Network.ConnectRetriesAmount));
+
+				Json.Cluster.Network.Emplace(DisplayClusterConfigurationStrings::config::cluster::network::NetConnectRetryDelay,
+					DisplayClusterTypesConverter::template ToString(Config->Cluster->Network.ConnectRetryDelay));
+
+				Json.Cluster.Network.Emplace(DisplayClusterConfigurationStrings::config::cluster::network::NetGameStartBarrierTimeout,
+					DisplayClusterTypesConverter::template ToString(Config->Cluster->Network.GameStartBarrierTimeout));
+
+				Json.Cluster.Network.Emplace(DisplayClusterConfigurationStrings::config::cluster::network::NetFrameStartBarrierTimeout,
+					DisplayClusterTypesConverter::template ToString(Config->Cluster->Network.FrameStartBarrierTimeout));
+
+				Json.Cluster.Network.Emplace(DisplayClusterConfigurationStrings::config::cluster::network::NetFrameEndBarrierTimeout,
+					DisplayClusterTypesConverter::template ToString(Config->Cluster->Network.FrameEndBarrierTimeout));
+
+				Json.Cluster.Network.Emplace(DisplayClusterConfigurationStrings::config::cluster::network::NetRenderSyncBarrierTimeout,
+					DisplayClusterTypesConverter::template ToString(Config->Cluster->Network.RenderSyncBarrierTimeout));
+			}
+
+			// Cluster nodes
+			for (const TPair<FString, UDisplayClusterConfigurationClusterNode*>& CfgNode : Config->Cluster->Nodes)
+			{
+				FDisplayClusterConfigurationJsonClusterNode_500 Node;
+
+				// Base parameters
+				Node.Host = CfgNode.Value->Host;
+				Node.Sound = CfgNode.Value->bIsSoundEnabled;
+				Node.FullScreen = CfgNode.Value->bIsFullscreen;
+				Node.Window = FDisplayClusterConfigurationJsonRectangle_500(CfgNode.Value->WindowRect.X, CfgNode.Value->WindowRect.Y, CfgNode.Value->WindowRect.W, CfgNode.Value->WindowRect.H);
+
+				// Viewports
+				for (const TPair<FString, UDisplayClusterConfigurationViewport*>& CfgViewport : CfgNode.Value->Viewports)
+				{
+					FDisplayClusterConfigurationJsonViewport_500 Viewport;
+
+					// Base parameters
+					Viewport.Camera = CfgViewport.Value->Camera;
+					Viewport.Region = FDisplayClusterConfigurationJsonRectangle_500(CfgViewport.Value->Region.X, CfgViewport.Value->Region.Y, CfgViewport.Value->Region.W, CfgViewport.Value->Region.H);
+					Viewport.GPUIndex = CfgViewport.Value->GPUIndex;
+					Viewport.BufferRatio = CfgViewport.Value->RenderSettings.BufferRatio;
+
+					// Projection policy
+					Viewport.ProjectionPolicy.Type = CfgViewport.Value->ProjectionPolicy.Type;
+					Viewport.ProjectionPolicy.Parameters = CfgViewport.Value->ProjectionPolicy.Parameters;
+
+					// Overscan
+					Viewport.Overscan.Mode = DisplayClusterConfigurationJsonHelpers::ToString(CfgViewport.Value->RenderSettings.Overscan.Mode);
+					Viewport.Overscan.bEnabled = CfgViewport.Value->RenderSettings.Overscan.bEnabled;
+					Viewport.Overscan.Oversize = CfgViewport.Value->RenderSettings.Overscan.bOversize;
+					Viewport.Overscan.Left = CfgViewport.Value->RenderSettings.Overscan.Left;
+					Viewport.Overscan.Right = CfgViewport.Value->RenderSettings.Overscan.Right;
+					Viewport.Overscan.Top = CfgViewport.Value->RenderSettings.Overscan.Top;
+					Viewport.Overscan.Bottom = CfgViewport.Value->RenderSettings.Overscan.Bottom;
+					
+					// Save this viewport
+					Node.Viewports.Emplace(CfgViewport.Key, Viewport);
+				}
+
+				// Postprocess
+				for (const TPair<FString, FDisplayClusterConfigurationPostprocess>& CfgPostprocess : CfgNode.Value->Postprocess)
+				{
+					FDisplayClusterConfigurationJsonPostprocess_500 PostprocessOperation;
+
+					PostprocessOperation.Type = CfgPostprocess.Value.Type;
+					PostprocessOperation.Parameters = CfgPostprocess.Value.Parameters;
+
+					Node.Postprocess.Emplace(CfgPostprocess.Key, PostprocessOperation);
+				}
+
+				// OutputRemap
+				{
+					Node.OutputRemap.bEnable = CfgNode.Value->OutputRemap.bEnable;
+					Node.OutputRemap.DataSource = (CfgNode.Value->OutputRemap.DataSource == EDisplayClusterConfigurationFramePostProcess_OutputRemapSource::ExternalFile) ? "file" : "mesh";
+					Node.OutputRemap.ExternalFile = CfgNode.Value->OutputRemap.ExternalFile;
+
+					if (CfgNode.Value->OutputRemap.StaticMesh)
+					{
+						Node.OutputRemap.StaticMeshAsset = CfgNode.Value->OutputRemap.StaticMesh->GetPathName();
+					}
+				}
+
+				// Store new cluster node
+				Json.Cluster.Nodes.Emplace(CfgNode.Key, Node);
+			}
+		}
+
+		// Custom parameters
+		Json.CustomParameters = Config->CustomParameters;
+
+		// Diagnostics
+		Json.Diagnostics.SimulateLag = Config->Diagnostics.bSimulateLag;
+		Json.Diagnostics.MinLagTime = Config->Diagnostics.MinLagTime;
+		Json.Diagnostics.MaxLagTime = Config->Diagnostics.MaxLagTime;
+
+		return true;
+	}
+}
diff --git a/Engine/Plugins/Runtime/nDisplay/Source/DisplayClusterConfiguration/Private/Formats/JSON500/DisplayClusterConfigurationJsonParser_500.h b/Engine/Plugins/Runtime/nDisplay/Source/DisplayClusterConfiguration/Private/Formats/JSON500/DisplayClusterConfigurationJsonParser_500.h
new file mode 100644
--- /dev/null
+++ b/Engine/Plugins/Runtime/nDisplay/Source/DisplayClusterConfiguration/Private/Formats/JSON500/DisplayClusterConfigurationJsonParser_500.h
@@ -0,0 +1,48 @@
+// Copyright Epic Games, Inc. All Rights Reserved.
+
+#pragma once
+
+#include "CoreMinimal.h"
+#include "Formats/IDisplayClusterConfigurationDataParser.h"
+#include "Formats/JSON500/DisplayClusterConfigurationJsonTypes_500.h"
+
+class UDisplayClusterConfigurationData;
+
+
+namespace JSON500
+{
+	/**
+	 * Config parser for JSON based config files
+	 */
+	class FDisplayClusterConfigurationJsonParser
+		: public IDisplayClusterConfigurationDataParser
+	{
+	public:
+		FDisplayClusterConfigurationJsonParser()  = default;
+		~FDisplayClusterConfigurationJsonParser() = default;
+
+	public:
+		// Load data from a specified file
+		virtual UDisplayClusterConfigurationData* LoadData(const FString& FilePath, UObject* Owner = nullptr) override;
+
+		// Save data to a specified file
+		virtual bool SaveData(const UDisplayClusterConfigurationData* ConfigData, const FString& FilePath) override;
+
+		// Convert configuration to string
+		virtual bool AsString(const UDisplayClusterConfigurationData* ConfigData, FString& OutString) override;
+
+	protected:
+		// [import data] Fill generic data container with parsed information
+		UDisplayClusterConfigurationData* ConvertDataToInternalTypes();
+		// [export data] Extract data from generic container to specific containers
+		bool ConvertDataToExternalTypes(const UDisplayClusterConfigurationData* InData);
+
+	private:
+		// Intermediate temporary data (Json types)
+		FDisplayClusterConfigurationJsonContainer_500 JsonData;
+		// Source file
+		FString ConfigFile;
+		// Owner for config data UObject we'll create on success
+		UObject* ConfigDataOwner = nullptr;
+	};
+}
diff --git a/Engine/Plugins/Runtime/nDisplay/Source/DisplayClusterConfiguration/Private/Formats/JSON500/DisplayClusterConfigurationJsonTypes_500.h b/Engine/Plugins/Runtime/nDisplay/Source/DisplayClusterConfiguration/Private/Formats/JSON500/DisplayClusterConfigurationJsonTypes_500.h
new file mode 100644
--- /dev/null
+++ b/Engine/Plugins/Runtime/nDisplay/Source/DisplayClusterConfiguration/Private/Formats/JSON500/DisplayClusterConfigurationJsonTypes_500.h
@@ -0,0 +1,562 @@
+// Copyright Epic Games, Inc. All Rights Reserved.
+
+#pragma once
+
+#include "CoreMinimal.h"
+#include "DisplayClusterConfigurationStrings.h"
+#include "DisplayClusterConfigurationJsonTypes_500.generated.h"
+
+
+USTRUCT()
+struct FDisplayClusterConfigurationJsonRectangle_500
+{
+	GENERATED_BODY()
+
+public:
+	FDisplayClusterConfigurationJsonRectangle_500()
+		: X(0), Y(0), W(0), H(0)
+	{ }
+
+	FDisplayClusterConfigurationJsonRectangle_500(int32 _X, int32 _Y, int32 _W, int32 _H)
+		: X(_X), Y(_Y), W(_W), H(_H)
+	{ }
+
+public:
+	UPROPERTY()
+	int32 X;
+
+	UPROPERTY()
+	int32 Y;
+
+	UPROPERTY()
+	int32 W;
+
+	UPROPERTY()
+	int32 H;
+};
+
+
+USTRUCT()
+struct FDisplayClusterConfigurationJsonVector_500
+{
+	GENERATED_BODY()
+
+public:
+	FDisplayClusterConfigurationJsonVector_500()
+		: FDisplayClusterConfigurationJsonVector_500(0.f, 0.f, 0.f)
+	{ }
+
+	FDisplayClusterConfigurationJsonVector_500(float _X, float _Y, float _Z)
+		: X(_X), Y(_Y), Z(_Z)
+	{ }
+
+	FDisplayClusterConfigurationJsonVector_500(const FVector& Vector)
+		: FDisplayClusterConfigurationJsonVector_500(Vector.X, Vector.Y, Vector.Z)
+	{ }
+
+public:
+	UPROPERTY()
+	float X;
+
+	UPROPERTY()
+	float Y;
+
+	UPROPERTY()
+	float Z;
+
+public:
+	static FVector ToVector(const FDisplayClusterConfigurationJsonVector_500& Data)
+	{
+		return FVector(Data.X, Data.Y, Data.Z);
+	}
+
+	static FDisplayClusterConfigurationJsonVector_500 FromVector(const FVector& Vector)
+	{
+		return FDisplayClusterConfigurationJsonVector_500(Vector.X, Vector.Y, Vector.Z);
+	}
+};
+
+
+USTRUCT()
+struct FDisplayClusterConfigurationJsonRotator_500
+{
+	GENERATED_BODY()
+
+public:
+	FDisplayClusterConfigurationJsonRotator_500()
+		: FDisplayClusterConfigurationJsonRotator_500(0.f, 0.f, 0.f)
+	{ }
+
+	FDisplayClusterConfigurationJsonRotator_500(float P, float Y, float R)
+		: Pitch(P), Yaw(Y), Roll(R)
+	{ }
+
+	FDisplayClusterConfigurationJsonRotator_500(const FRotator& Rotator)
+		: FDisplayClusterConfigurationJsonRotator_500(Rotator.Pitch, Rotator.Yaw, Rotator.Roll)
+	{ }
+
+public:
+	UPROPERTY()
+	float Pitch;
+
+	UPROPERTY()
+	float Yaw;
+
+	UPROPERTY()
+	float Roll;
+
+public:
+	static FRotator ToRotator(const FDisplayClusterConfigurationJsonRotator_500& Data)
+	{
+		return FRotator(Data.Pitch, Data.Yaw, Data.Roll);
+	}
+
+	static FDisplayClusterConfigurationJsonRotator_500 FromRotator(const FRotator& Rotator)
+	{
+		return FDisplayClusterConfigurationJsonRotator_500(Rotator.Pitch, Rotator.Yaw, Rotator.Roll);
+	}
+};
+
+
+USTRUCT()
+struct FDisplayClusterConfigurationJsonSizeInt_500
+{
+	GENERATED_BODY()
+
+public:
+	FDisplayClusterConfigurationJsonSizeInt_500()
+		: FDisplayClusterConfigurationJsonSizeInt_500(0, 0)
+	{ }
+
+	FDisplayClusterConfigurationJsonSizeInt_500(int W, int H)
+		: Width(W), Height(H)
+	{ }
+
+public:
+	UPROPERTY()
+	int Width;
+
+	UPROPERTY()
+	int Height;
+
+public:
+	static FIntPoint ToPoint(const FDisplayClusterConfigurationJsonSizeInt_500& Data)
+	{
+		return FIntPoint(Data.Width, Data.Height);
+	}
+
+	static FDisplayClusterConfigurationJsonSizeInt_500 FromPoint(const FIntPoint& Point)
+	{
+		return FDisplayClusterConfigurationJsonSizeInt_500(Point.X, Point.Y);
+	}
+};
+
+
+USTRUCT()
+struct FDisplayClusterConfigurationJsonSizeFloat_500
+{
+	GENERATED_BODY()
+
+public:
+	FDisplayClusterConfigurationJsonSizeFloat_500()
+		: FDisplayClusterConfigurationJsonSizeFloat_500(0.f, 0.f)
+	{ }
+
+	FDisplayClusterConfigurationJsonSizeFloat_500(float W, float H)
+		: Width(W), Height(H)
+	{ }
+
+	FDisplayClusterConfigurationJsonSizeFloat_500(const FVector2D& Vector)
+		: FDisplayClusterConfigurationJsonSizeFloat_500(Vector.X, Vector.Y)
+	{ }
+
+public:
+	UPROPERTY()
+	float Width;
+
+	UPROPERTY()
+	float Height;
+
+public:
+	static FVector2D ToVector(const FDisplayClusterConfigurationJsonSizeFloat_500& Data)
+	{
+		return FVector2D(Data.Width, Data.Height);
+	}
+
+	static FDisplayClusterConfigurationJsonSizeFloat_500 FromVector(const FVector2D& Vector)
+	{
+		return FDisplayClusterConfigurationJsonSizeFloat_500(Vector.X, Vector.Y);
+	}
+};
+
+
+USTRUCT()
+struct FDisplayClusterConfigurationJsonPolymorphicEntity_500
+{
+	GENERATED_BODY()
+
+public:
+	UPROPERTY()
+	FString Type;
+
+	UPROPERTY()
+	TMap<FString, FString> Parameters;
+};
+
+
+USTRUCT()
+struct FDisplayClusterConfigurationJsonMisc_500
+{
+	GENERATED_BODY()
+
+public:
+	FDisplayClusterConfigurationJsonMisc_500()
+		: bFollowLocalPlayerCamera(false)
+		, bExitOnEsc(true)
+	{ }
+
+public:
+	UPROPERTY()
+	bool bFollowLocalPlayerCamera;
+
+	UPROPERTY()
+	bool bExitOnEsc;
+};
+
+
+USTRUCT()
+struct FDisplayClusterConfigurationJsonSceneComponent_500
+{
+	GENERATED_BODY()
+};
+
+
+USTRUCT()
+struct FDisplayClusterConfigurationJsonSceneComponentXform_500
+	: public FDisplayClusterConfigurationJsonSceneComponent_500
+{
+	GENERATED_BODY()
+
+public:
+	FDisplayClusterConfigurationJsonSceneComponentXform_500()
+		: FDisplayClusterConfigurationJsonSceneComponentXform_500(FString(), FVector::ZeroVector, FRotator::ZeroRotator)
+	{ }
+
+	FDisplayClusterConfigurationJsonSceneComponentXform_500(const FString& InParentId, const FVector& InLocation, const FRotator& InRotation)
+		: ParentId(InParentId)
+		, Location(InLocation)
+		, Rotation(InRotation)
+	{ }
+
+public:
+	UPROPERTY()
+	FString ParentId;
+
+	UPROPERTY()
+	FDisplayClusterConfigurationJsonVector_500 Location;
+
+	UPROPERTY()
+	FDisplayClusterConfigurationJsonRotator_500 Rotation;
+};
+
+
+USTRUCT()
+struct FDisplayClusterConfigurationJsonSceneComponentCamera_500
+	: public FDisplayClusterConfigurationJsonSceneComponentXform_500
+{
+	GENERATED_BODY()
+
+public:
+	FDisplayClusterConfigurationJsonSceneComponentCamera_500()
+		: FDisplayClusterConfigurationJsonSceneComponentCamera_500(FString(), FVector::ZeroVector, FRotator::ZeroRotator, 6.4f, false, DisplayClusterConfigurationStrings::config::scene::camera::CameraStereoOffsetNone)
+	{ }
+
+	FDisplayClusterConfigurationJsonSceneComponentCamera_500(const FString& InParentId, const FVector& InLocation, const FRotator& InRotation, float InInterpupillaryDistance, bool bInSwapEyes, const FString& InStereoOffset)
+		: FDisplayClusterConfigurationJsonSceneComponentXform_500(InParentId, InLocation, InRotation)
+		, InterpupillaryDistance(InInterpupillaryDistance)
+		, SwapEyes(bInSwapEyes)
+		, StereoOffset(InStereoOffset)
+	{ }
+
+public:
+	UPROPERTY()
+	float InterpupillaryDistance;
+
+	UPROPERTY()
+	bool SwapEyes;
+
+	UPROPERTY()
+	FString StereoOffset;
+};
+
+
+USTRUCT()
+struct FDisplayClusterConfigurationJsonSceneComponentScreen_500
+	: public FDisplayClusterConfigurationJsonSceneComponentXform_500
+{
+	GENERATED_BODY()
+
+public:
+	FDisplayClusterConfigurationJsonSceneComponentScreen_500()
+		: FDisplayClusterConfigurationJsonSceneComponentScreen_500(FString(), FVector::ZeroVector, FRotator::ZeroRotator, FVector2D(1.f, 1.f))
+	{ }
+
+	FDisplayClusterConfigurationJsonSceneComponentScreen_500(const FString& InParentId, const FVector& InLocation, const FRotator& InRotation, const FVector2D& InSize)
+		: FDisplayClusterConfigurationJsonSceneComponentXform_500(InParentId, InLocation, InRotation)
+		, Size(InSize)
+	{ }
+
+public:
+	UPROPERTY()
+	FDisplayClusterConfigurationJsonSizeFloat_500 Size;
+};
+
+
+USTRUCT()
+struct FDisplayClusterConfigurationJsonScene_500
+{
+	GENERATED_BODY()
+
+public:
+	UPROPERTY()
+	TMap<FString, FDisplayClusterConfigurationJsonSceneComponentXform_500> Xforms;
+
+	UPROPERTY()
+	TMap<FString, FDisplayClusterConfigurationJsonSceneComponentCamera_500> Cameras;
+
+	UPROPERTY()
+	TMap<FString, FDisplayClusterConfigurationJsonSceneComponentScreen_500> Screens;
+};
+
+
+USTRUCT()
+struct FDisplayClusterConfigurationJsonPrimaryNode_500
+{
+	GENERATED_BODY()
+
+public:
+	UPROPERTY()
+	FString Id;
+
+	UPROPERTY()
+	TMap<FString, uint16> Ports;
+};
+
+USTRUCT()
+struct FDisplayClusterConfigurationJsonClusterSyncPolicy_500
+	: public FDisplayClusterConfigurationJsonPolymorphicEntity_500
+{
+	GENERATED_BODY()
+};
+
+USTRUCT()
+struct FDisplayClusterConfigurationJsonClusterSync_500
+{
+	GENERATED_BODY()
+
+public:
+	UPROPERTY()
+	FDisplayClusterConfigurationJsonClusterSyncPolicy_500 RenderSyncPolicy;
+
+	UPROPERTY()
+	FDisplayClusterConfigurationJsonClusterSyncPolicy_500 InputSyncPolicy;
+};
+
+USTRUCT()
+struct FDisplayClusterConfigurationJsonPostprocess_500
+	: public FDisplayClusterConfigurationJsonPolymorphicEntity_500
+{
+	GENERATED_BODY()
+};
+
+USTRUCT()
+struct FDisplayClusterConfigurationJsonProjectionPolicy_500
+	: public FDisplayClusterConfigurationJsonPolymorphicEntity_500
+{
+	GENERATED_BODY()
+};
+
+USTRUCT()
+struct FDisplayClusterConfigurationJsonOverscan_500
+{
+	GENERATED_BODY()
+
+public:
+	UPROPERTY()
+	bool bEnabled = false;
+
+	UPROPERTY()
+	FString Mode = "";
+
+	UPROPERTY()
+	int32 Left = 0;
+
+	UPROPERTY()
+	int32 Right = 0;
+
+	UPROPERTY()
+	int32 Top = 0;
+
+	UPROPERTY()
+	int32 Bottom = 0;
+
+	UPROPERTY()
+	bool Oversize = 0;
+};
+
+USTRUCT()
+struct FDisplayClusterConfigurationJsonViewport_500
+{
+	GENERATED_BODY()
+
+public:
+	UPROPERTY()
+	FString Camera;
+
+	UPROPERTY()
+	float BufferRatio = 1.0f;
+
+	UPROPERTY()
+	int32 GPUIndex = 0;
+
+	UPROPERTY()
+	bool AllowCrossGPUTransfer = false;
+
+	UPROPERTY()
+	bool IsShared = false;
+
+	UPROPERTY()
+	FDisplayClusterConfigurationJsonOverscan_500 Overscan;
+
+	UPROPERTY()
+	FDisplayClusterConfigurationJsonRectangle_500 Region;
+
+	UPROPERTY()
+	FDisplayClusterConfigurationJsonProjectionPolicy_500 ProjectionPolicy;
+};
+
+USTRUCT()
+struct FDisplayClusterConfigurationFramePostProcess_OutputRemap_500
+{
+	GENERATED_BODY()
+
+public:
+	UPROPERTY()
+	bool bEnable = false;
+
+	UPROPERTY()
+	FString DataSource;
+
+	UPROPERTY()
+	FString StaticMeshAsset;
+
+	UPROPERTY()
+	FString ExternalFile;
+};
+
+USTRUCT()
+struct FDisplayClusterConfigurationJsonClusterNode_500
+{
+	GENERATED_BODY()
+
+public:
+	UPROPERTY()
+	FString Host;
+
+	UPROPERTY()
+	bool Sound = true;
+
+	UPROPERTY()
+	bool FullScreen = false;
+
+	UPROPERTY()
+	FDisplayClusterConfigurationJsonRectangle_500 Window;
+
+	UPROPERTY()
+	TMap<FString, FDisplayClusterConfigurationJsonPostprocess_500> Postprocess;
+
+	UPROPERTY()
+	TMap<FString, FDisplayClusterConfigurationJsonViewport_500> Viewports;
+
+	UPROPERTY()
+	FDisplayClusterConfigurationFramePostProcess_OutputRemap_500 OutputRemap;
+};
+
+USTRUCT()
+struct FDisplayClusterConfigurationJsonCluster_500
+{
+	GENERATED_BODY()
+
+public:
+	UPROPERTY()
+	FDisplayClusterConfigurationJsonPrimaryNode_500 PrimaryNode;
+
+	UPROPERTY()
+	FDisplayClusterConfigurationJsonClusterSync_500 Sync;
+
+	UPROPERTY()
+	TMap<FString, FString> Network;
+
+	UPROPERTY()
+	TMap<FString, FDisplayClusterConfigurationJsonClusterNode_500> Nodes;
+};
+
+USTRUCT()
+struct FDisplayClusterConfigurationJsonDiagnostics_500
+{
+	GENERATED_BODY()
+
+public:
+	UPROPERTY()
+	bool SimulateLag = false;
+
+	UPROPERTY()
+	float MinLagTime = 0.01f;
+
+	UPROPERTY()
+	float MaxLagTime = 0.3f;
+};
+
+
+// "nDisplay" category structure
+USTRUCT()
+struct FDisplayClusterConfigurationJsonNdisplay_500
+{
+	GENERATED_BODY()
+
+public:
+	UPROPERTY()
+	FString Description;
+
+	UPROPERTY()
+	FString Version;
+
+	UPROPERTY()
+	FString AssetPath;
+
+	UPROPERTY()
+	FDisplayClusterConfigurationJsonMisc_500 Misc;
+
+	UPROPERTY()
+	FDisplayClusterConfigurationJsonScene_500 Scene;
+
+	UPROPERTY()
+	FDisplayClusterConfigurationJsonCluster_500 Cluster;
+
+	UPROPERTY()
+	TMap<FString, FString> CustomParameters;
+
+	UPROPERTY()
+	FDisplayClusterConfigurationJsonDiagnostics_500 Diagnostics;
+};
+
+// The main nDisplay configuration structure. It's supposed to extract nDisplay related data from a collecting JSON file.
+USTRUCT()
+struct FDisplayClusterConfigurationJsonContainer_500
+{
+	GENERATED_BODY()
+
+public:
+	UPROPERTY()
+	FDisplayClusterConfigurationJsonNdisplay_500 nDisplay;
+};
diff --git a/Engine/Plugins/Runtime/nDisplay/Source/DisplayClusterConfiguration/Private/Formats/Text/DisplayClusterConfigurationTextParser.cpp b/Engine/Plugins/Runtime/nDisplay/Source/DisplayClusterConfiguration/Private/Formats/Text/DisplayClusterConfigurationTextParser.cpp
deleted file mode 100644
diff --git a/Engine/Plugins/Runtime/nDisplay/Source/DisplayClusterConfiguration/Private/Formats/Text/DisplayClusterConfigurationTextParser.h b/Engine/Plugins/Runtime/nDisplay/Source/DisplayClusterConfiguration/Private/Formats/Text/DisplayClusterConfigurationTextParser.h
deleted file mode 100644
diff --git a/Engine/Plugins/Runtime/nDisplay/Source/DisplayClusterConfiguration/Private/Formats/Text/DisplayClusterConfigurationTextStrings.h b/Engine/Plugins/Runtime/nDisplay/Source/DisplayClusterConfiguration/Private/Formats/Text/DisplayClusterConfigurationTextStrings.h
deleted file mode 100644
diff --git a/Engine/Plugins/Runtime/nDisplay/Source/DisplayClusterConfiguration/Private/Formats/Text/DisplayClusterConfigurationTextTypes.cpp b/Engine/Plugins/Runtime/nDisplay/Source/DisplayClusterConfiguration/Private/Formats/Text/DisplayClusterConfigurationTextTypes.cpp
deleted file mode 100644
diff --git a/Engine/Plugins/Runtime/nDisplay/Source/DisplayClusterConfiguration/Private/Formats/Text/DisplayClusterConfigurationTextTypes.h b/Engine/Plugins/Runtime/nDisplay/Source/DisplayClusterConfiguration/Private/Formats/Text/DisplayClusterConfigurationTextTypes.h
deleted file mode 100644
diff --git a/Engine/Plugins/Runtime/nDisplay/Source/DisplayClusterConfiguration/Private/VersionChecker/DisplayClusterConfigurationVersionChecker.cpp b/Engine/Plugins/Runtime/nDisplay/Source/DisplayClusterConfiguration/Private/VersionChecker/DisplayClusterConfigurationVersionChecker.cpp
--- a/Engine/Plugins/Runtime/nDisplay/Source/DisplayClusterConfiguration/Private/VersionChecker/DisplayClusterConfigurationVersionChecker.cpp
+++ b/Engine/Plugins/Runtime/nDisplay/Source/DisplayClusterConfiguration/Private/VersionChecker/DisplayClusterConfigurationVersionChecker.cpp
@@ -19,14 +19,12 @@ EDisplayClusterConfigurationVersion FDisplayClusterConfigurationVersionChecker::
 
 	switch (FileType)
 	{
-	case EConfigFileType::Text:
-		return EDisplayClusterConfigurationVersion::Version_CFG;
-
 	case EConfigFileType::Json:
 		return GetConfigVersionJson(FilePath);
-	}
 
-	return EDisplayClusterConfigurationVersion::Unknown;
+	default:
+		return EDisplayClusterConfigurationVersion::Unknown;
+	}
 }
 
 EDisplayClusterConfigurationVersion FDisplayClusterConfigurationVersionChecker::GetConfigVersionJson(const FString& FilePath) const
@@ -60,6 +58,10 @@ EDisplayClusterConfigurationVersion FDisplayClusterConfigurationVersionChecker::
 	{
 		return EDisplayClusterConfigurationVersion::Version_427;
 	}
+	else if (JsonData.nDisplay.Version.Equals(TEXT("5.00")))
+	{
+		return EDisplayClusterConfigurationVersion::Version_500;
+	}
 
 	// Otherwise it's something unknown
 	return EDisplayClusterConfigurationVersion::Unknown;
@@ -73,11 +75,6 @@ FDisplayClusterConfigurationVersionChecker::EConfigFileType FDisplayClusterConfi
 		UE_LOG(LogDisplayClusterConfiguration, Log, TEXT("JSON config: %s"), *InConfigPath);
 		return EConfigFileType::Json;
 	}
-	else if (Extension.Equals(FString(DisplayClusterConfigurationStrings::file::FileExtCfg), ESearchCase::IgnoreCase))
-	{
-		UE_LOG(LogDisplayClusterConfiguration, Log, TEXT("TXT config: %s"), *InConfigPath);
-		return EConfigFileType::Text;
-	}
 
 	UE_LOG(LogDisplayClusterConfiguration, Warning, TEXT("Unknown file extension: %s"), *Extension);
 	return EConfigFileType::Unknown;
diff --git a/Engine/Plugins/Runtime/nDisplay/Source/DisplayClusterConfiguration/Private/VersionChecker/DisplayClusterConfigurationVersionChecker.h b/Engine/Plugins/Runtime/nDisplay/Source/DisplayClusterConfiguration/Private/VersionChecker/DisplayClusterConfigurationVersionChecker.h
--- a/Engine/Plugins/Runtime/nDisplay/Source/DisplayClusterConfiguration/Private/VersionChecker/DisplayClusterConfigurationVersionChecker.h
+++ b/Engine/Plugins/Runtime/nDisplay/Source/DisplayClusterConfiguration/Private/VersionChecker/DisplayClusterConfigurationVersionChecker.h
@@ -27,7 +27,6 @@ protected:
 	enum class EConfigFileType
 	{
 		Unknown,
-		Text,
 		Json
 	};
 
diff --git a/Engine/Plugins/Runtime/nDisplay/Source/DisplayClusterConfiguration/Public/DisplayClusterConfigurationStrings.h b/Engine/Plugins/Runtime/nDisplay/Source/DisplayClusterConfiguration/Public/DisplayClusterConfigurationStrings.h
--- a/Engine/Plugins/Runtime/nDisplay/Source/DisplayClusterConfiguration/Public/DisplayClusterConfigurationStrings.h
+++ b/Engine/Plugins/Runtime/nDisplay/Source/DisplayClusterConfiguration/Public/DisplayClusterConfigurationStrings.h
@@ -45,7 +45,6 @@ namespace DisplayClusterConfigurationStrings
 	// Config file extensions
 	namespace file
 	{
-		static constexpr auto FileExtCfg  = TEXT("cfg");
 		static constexpr auto FileExtJson = TEXT("ndisplay");
 	}
 
@@ -72,8 +71,8 @@ namespace DisplayClusterConfigurationStrings
 
 			namespace input_sync
 			{
-				static constexpr auto InputSyncPolicyNone            = TEXT("None");
-				static constexpr auto InputSyncPolicyReplicateMaster = TEXT("ReplicateMaster");
+				static constexpr auto InputSyncPolicyNone             = TEXT("None");
+				static constexpr auto InputSyncPolicyReplicatePrimary = TEXT("ReplicatePrimary");
 			}
 
 			namespace render_sync
diff --git a/Engine/Plugins/Runtime/nDisplay/Source/DisplayClusterConfiguration/Public/DisplayClusterConfigurationTypes.h b/Engine/Plugins/Runtime/nDisplay/Source/DisplayClusterConfiguration/Public/DisplayClusterConfigurationTypes.h
--- a/Engine/Plugins/Runtime/nDisplay/Source/DisplayClusterConfiguration/Public/DisplayClusterConfigurationTypes.h
+++ b/Engine/Plugins/Runtime/nDisplay/Source/DisplayClusterConfiguration/Public/DisplayClusterConfigurationTypes.h
@@ -163,12 +163,12 @@ protected:
 ////////////////////////////////////////////////////////////////
 // Cluster
 USTRUCT(Blueprintable)
-struct DISPLAYCLUSTERCONFIGURATION_API FDisplayClusterConfigurationMasterNodePorts
+struct DISPLAYCLUSTERCONFIGURATION_API FDisplayClusterConfigurationPrimaryNodePorts
 {
 	GENERATED_BODY()
 
 public:
-	FDisplayClusterConfigurationMasterNodePorts();
+	FDisplayClusterConfigurationPrimaryNodePorts();
 
 public:
 	/** Advanced: network port for Cluster Sync Events */
@@ -185,7 +185,7 @@ public:
 };
 
 USTRUCT(Blueprintable)
-struct FDisplayClusterConfigurationMasterNode
+struct FDisplayClusterConfigurationPrimaryNode
 {
 	GENERATED_BODY()
 
@@ -193,8 +193,8 @@ public:
 	UPROPERTY(BlueprintReadOnly, Category = NDisplay)
 	FString Id;
 
-	UPROPERTY(EditAnywhere, BlueprintReadOnly, Category = "Configuration", meta = (DisplayName = "Master Node Ports"))
-	FDisplayClusterConfigurationMasterNodePorts Ports;
+	UPROPERTY(EditAnywhere, BlueprintReadOnly, Category = "Configuration", meta = (DisplayName = "Primary Node Ports"))
+	FDisplayClusterConfigurationPrimaryNodePorts Ports;
 };
 
 USTRUCT(Blueprintable)
@@ -415,7 +415,7 @@ class DISPLAYCLUSTERCONFIGURATION_API UDisplayClusterConfigurationCluster
 
 public:
 	UPROPERTY(EditDefaultsOnly, BlueprintReadOnly, Category = "Configuration", meta = (ShowOnlyInnerProperties))
-	FDisplayClusterConfigurationMasterNode MasterNode;
+	FDisplayClusterConfigurationPrimaryNode PrimaryNode;
 
 	UPROPERTY(EditDefaultsOnly, BlueprintReadOnly, Category = "Configuration", meta = (ShowOnlyInnerProperties))
 	FDisplayClusterConfigurationClusterSync Sync;
diff --git a/Engine/Plugins/Runtime/nDisplay/Source/DisplayClusterConfiguration/Public/DisplayClusterConfigurationVersion.h b/Engine/Plugins/Runtime/nDisplay/Source/DisplayClusterConfiguration/Public/DisplayClusterConfigurationVersion.h
--- a/Engine/Plugins/Runtime/nDisplay/Source/DisplayClusterConfiguration/Public/DisplayClusterConfigurationVersion.h
+++ b/Engine/Plugins/Runtime/nDisplay/Source/DisplayClusterConfiguration/Public/DisplayClusterConfigurationVersion.h
@@ -5,15 +5,15 @@
 enum class EDisplayClusterConfigurationVersion : uint8
 {
 	Unknown,     // Unknown version or not a config file
-	Version_CFG, // Old .cfg format (going to be deprecated)
 	Version_426, // 4.26 JSON based config format
 	Version_427, // 4.27 JSON based config format
+	Version_500, // 5.00 JSON based config format
 };
 
 namespace DisplayClusterConfiguration
 {
 	static constexpr auto GetCurrentConfigurationSchemeMarker()
 	{
-		return TEXT("4.27");
+		return TEXT("5.00");
 	}
 }
diff --git a/Engine/Plugins/Runtime/nDisplay/Source/DisplayClusterConfigurator/Private/ClusterConfiguration/DisplayClusterConfiguratorClusterUtils.cpp b/Engine/Plugins/Runtime/nDisplay/Source/DisplayClusterConfigurator/Private/ClusterConfiguration/DisplayClusterConfiguratorClusterUtils.cpp
--- a/Engine/Plugins/Runtime/nDisplay/Source/DisplayClusterConfigurator/Private/ClusterConfiguration/DisplayClusterConfiguratorClusterUtils.cpp
+++ b/Engine/Plugins/Runtime/nDisplay/Source/DisplayClusterConfigurator/Private/ClusterConfiguration/DisplayClusterConfiguratorClusterUtils.cpp
@@ -106,10 +106,10 @@ UDisplayClusterConfigurationClusterNode* FDisplayClusterConfiguratorClusterUtils
 
 		NewNode = AddClusterNodeToCluster(NodeTemplate, Cluster, ItemName);
 
-		// If the newly added cluster node is the only node in the cluster, it should be the master node, so set it as the master
+		// If the newly added cluster node is the only node in the cluster, it should be the primary node, so set it as the primary
 		if (Cluster->Nodes.Num() == 1)
 		{
-			SetClusterNodeAsMaster(NewNode);
+			SetClusterNodeAsPrimary(NewNode);
 		}
 
 		if (bAddViewport)
@@ -612,12 +612,12 @@ FString FDisplayClusterConfiguratorClusterUtils::GetClusterNodeName(UDisplayClus
 	return "";
 }
 
-bool FDisplayClusterConfiguratorClusterUtils::IsClusterNodeMaster(UDisplayClusterConfigurationClusterNode* ClusterNode)
+bool FDisplayClusterConfiguratorClusterUtils::IsClusterNodePrimary(UDisplayClusterConfigurationClusterNode* ClusterNode)
 {
 	if (const UDisplayClusterConfigurationCluster* Cluster = Cast<UDisplayClusterConfigurationCluster>(ClusterNode->GetOuter()))
 	{
-		const FString MasterNodeId = Cluster->MasterNode.Id;
-		if (Cluster->Nodes.Contains(MasterNodeId) && Cluster->Nodes[MasterNodeId] == ClusterNode)
+		const FString PrimaryNodeId = Cluster->PrimaryNode.Id;
+		if (Cluster->Nodes.Contains(PrimaryNodeId) && Cluster->Nodes[PrimaryNodeId] == ClusterNode)
 		{
 			return true;
 		}
@@ -727,7 +727,7 @@ bool FDisplayClusterConfiguratorClusterUtils::RenameClusterNode(UDisplayClusterC
 				return false;
 			}
 
-			const bool bIsMaster = IsClusterNodeMaster(ClusterNode);
+			const bool bIsPrimary = IsClusterNodePrimary(ClusterNode);
 
 			ClusterNode->Modify();
 			ClusterNodeParent->Modify();
@@ -746,10 +746,10 @@ bool FDisplayClusterConfiguratorClusterUtils::RenameClusterNode(UDisplayClusterC
 			ClusterNode->Rename(nullptr, GetTransientPackage(), REN_DontCreateRedirectors | REN_DoNotDirty | REN_ForceNoResetLoaders);
 			ClusterNode->SetFlags(RF_Transient);
 			
-			// If the cluster node was a master node before the rename, we need to update the master reference in the cluster with the new name
-			if (bIsMaster)
+			// If the cluster node was a primary node before the rename, we need to update the primary reference in the cluster with the new name
+			if (bIsPrimary)
 			{
-				SetClusterNodeAsMaster(NewClusterNode);
+				SetClusterNodeAsPrimary(NewClusterNode);
 			}
 
 			return true;
@@ -759,7 +759,7 @@ bool FDisplayClusterConfiguratorClusterUtils::RenameClusterNode(UDisplayClusterC
 	return false;
 }
 
-bool FDisplayClusterConfiguratorClusterUtils::SetClusterNodeAsMaster(UDisplayClusterConfigurationClusterNode* ClusterNode)
+bool FDisplayClusterConfiguratorClusterUtils::SetClusterNodeAsPrimary(UDisplayClusterConfigurationClusterNode* ClusterNode)
 {
 	if (UDisplayClusterConfigurationCluster* ClusterNodeParent = Cast<UDisplayClusterConfigurationCluster>(ClusterNode->GetOuter()))
 	{
@@ -768,7 +768,7 @@ bool FDisplayClusterConfiguratorClusterUtils::SetClusterNodeAsMaster(UDisplayClu
 		if (NodeIdPtr != nullptr)
 		{
 			ClusterNodeParent->Modify();
-			ClusterNodeParent->MasterNode.Id = *NodeIdPtr;
+			ClusterNodeParent->PrimaryNode.Id = *NodeIdPtr;
 			return true;
 		}
 	}
diff --git a/Engine/Plugins/Runtime/nDisplay/Source/DisplayClusterConfigurator/Private/ClusterConfiguration/DisplayClusterConfiguratorClusterUtils.h b/Engine/Plugins/Runtime/nDisplay/Source/DisplayClusterConfigurator/Private/ClusterConfiguration/DisplayClusterConfiguratorClusterUtils.h
--- a/Engine/Plugins/Runtime/nDisplay/Source/DisplayClusterConfigurator/Private/ClusterConfiguration/DisplayClusterConfiguratorClusterUtils.h
+++ b/Engine/Plugins/Runtime/nDisplay/Source/DisplayClusterConfigurator/Private/ClusterConfiguration/DisplayClusterConfiguratorClusterUtils.h
@@ -133,11 +133,11 @@ public:
 	static FString GetClusterNodeName(UDisplayClusterConfigurationClusterNode* ClusterNode);
 
 	/**
-	 * Determines if a cluster node is the master node in its cluster.
+	 * Determines if a cluster node is the primary node in its cluster.
 	 * @param ClusterNode - The cluster node to check.
-	 * @return true if the node is the master; otherwise, false.
+	 * @return true if the node is the primary; otherwise, false.
 	 */
-	static bool IsClusterNodeMaster(UDisplayClusterConfigurationClusterNode* ClusterNode);
+	static bool IsClusterNodePrimary(UDisplayClusterConfigurationClusterNode* ClusterNode);
 
 	/**
 	 * Gets a unique name for a cluster node owned by the specified cluster.
@@ -173,11 +173,11 @@ public:
 	static bool RenameClusterNode(UDisplayClusterConfigurationClusterNode* ClusterNode, FString NewClusterNodeName);
 
 	/**
-	 * Sets the cluster node as the master node.
-	 * @param ClusterNode - The cluster node to set as the master node
-	 * @return true if the cluster node was successfully set as the master node
+	 * Sets the cluster node as the primary node.
+	 * @param ClusterNode - The cluster node to set as the primary node
+	 * @return true if the cluster node was successfully set as the primary node
 	 */
-	static bool SetClusterNodeAsMaster(UDisplayClusterConfigurationClusterNode* ClusterNode);
+	static bool SetClusterNodeAsPrimary(UDisplayClusterConfigurationClusterNode* ClusterNode);
 
 	/**
 	 * Gets the name that the viewport is stored under in its parent cluster node.
diff --git a/Engine/Plugins/Runtime/nDisplay/Source/DisplayClusterConfigurator/Private/DisplayClusterConfiguratorBlueprintEditor.cpp b/Engine/Plugins/Runtime/nDisplay/Source/DisplayClusterConfigurator/Private/DisplayClusterConfiguratorBlueprintEditor.cpp
--- a/Engine/Plugins/Runtime/nDisplay/Source/DisplayClusterConfigurator/Private/DisplayClusterConfiguratorBlueprintEditor.cpp
+++ b/Engine/Plugins/Runtime/nDisplay/Source/DisplayClusterConfigurator/Private/DisplayClusterConfiguratorBlueprintEditor.cpp
@@ -659,7 +659,7 @@ bool FDisplayClusterConfiguratorBlueprintEditor::ExportConfig()
 bool FDisplayClusterConfiguratorBlueprintEditor::CanExportConfig() const
 {
 	UDisplayClusterConfigurationData* EditingObject = GetEditorData();
-	return FDisplayClusterConfiguratorUtils::IsMasterNodeInConfig(EditingObject);
+	return FDisplayClusterConfiguratorUtils::IsPrimaryNodeInConfig(EditingObject);
 }
 
 bool FDisplayClusterConfiguratorBlueprintEditor::SaveToFile(const FString& InFilePath)
diff --git a/Engine/Plugins/Runtime/nDisplay/Source/DisplayClusterConfigurator/Private/DisplayClusterConfiguratorCommands.cpp b/Engine/Plugins/Runtime/nDisplay/Source/DisplayClusterConfigurator/Private/DisplayClusterConfiguratorCommands.cpp
--- a/Engine/Plugins/Runtime/nDisplay/Source/DisplayClusterConfigurator/Private/DisplayClusterConfiguratorCommands.cpp
+++ b/Engine/Plugins/Runtime/nDisplay/Source/DisplayClusterConfigurator/Private/DisplayClusterConfiguratorCommands.cpp
@@ -9,7 +9,7 @@ void FDisplayClusterConfiguratorCommands::RegisterCommands()
 {
 	UI_COMMAND(Import, "Import", "Import an nDisplay config. This may overwrite components. When upgrading from the deprecated .cfg file format the recommended import workflow is to drag & drop your .cfg file into the content browser.\
  If you need to reimport it is recommended to use the reimport options of the asset context menu", EUserInterfaceActionType::Button, FInputChord());
-	UI_COMMAND(Export, "Export", "Export to an nDisplay config. Requires a master cluster node set. This does not export all components and is meant for launching the cluster from the command line", EUserInterfaceActionType::Button, FInputChord());
+	UI_COMMAND(Export, "Export", "Export to an nDisplay config. Requires a primary cluster node set. This does not export all components and is meant for launching the cluster from the command line", EUserInterfaceActionType::Button, FInputChord());
 	UI_COMMAND(EditConfig, "EditConfig", "Edit config file with text editor", EUserInterfaceActionType::Button, FInputChord());
 	UI_COMMAND(ExportConfigOnSave, "Export on Save", "Export to nDisplay config automatically on save. Requires a config previously exported", EUserInterfaceActionType::ToggleButton, FInputChord());
 
diff --git a/Engine/Plugins/Runtime/nDisplay/Source/DisplayClusterConfigurator/Private/DisplayClusterConfiguratorCompiler.cpp b/Engine/Plugins/Runtime/nDisplay/Source/DisplayClusterConfigurator/Private/DisplayClusterConfiguratorCompiler.cpp
--- a/Engine/Plugins/Runtime/nDisplay/Source/DisplayClusterConfigurator/Private/DisplayClusterConfiguratorCompiler.cpp
+++ b/Engine/Plugins/Runtime/nDisplay/Source/DisplayClusterConfigurator/Private/DisplayClusterConfiguratorCompiler.cpp
@@ -241,9 +241,9 @@ void FDisplayClusterConfiguratorKismetCompilerContext::ValidateConfiguration()
 		return;
 	}
 	
-	if (!FDisplayClusterConfiguratorUtils::IsMasterNodeInConfig(BlueprintData))
+	if (!FDisplayClusterConfiguratorUtils::IsPrimaryNodeInConfig(BlueprintData))
 	{
-		MessageLog.Warning(*LOCTEXT("NoMasterNodeWarning", "Master cluster node not set. Please set a master node.").ToString());
+		MessageLog.Warning(*LOCTEXT("NoPrimaryNodeWarning", "Primary cluster node not set. Please set a primary node.").ToString());
 	}
 
 	bool bAtLeastOneViewportFound = false;
diff --git a/Engine/Plugins/Runtime/nDisplay/Source/DisplayClusterConfigurator/Private/DisplayClusterConfiguratorEditorSubsystem.cpp b/Engine/Plugins/Runtime/nDisplay/Source/DisplayClusterConfigurator/Private/DisplayClusterConfiguratorEditorSubsystem.cpp
--- a/Engine/Plugins/Runtime/nDisplay/Source/DisplayClusterConfigurator/Private/DisplayClusterConfiguratorEditorSubsystem.cpp
+++ b/Engine/Plugins/Runtime/nDisplay/Source/DisplayClusterConfigurator/Private/DisplayClusterConfiguratorEditorSubsystem.cpp
@@ -102,23 +102,8 @@ bool UDisplayClusterConfiguratorEditorSubsystem::SaveConfig(UDisplayClusterConfi
 {
 	if (InConfiguratorEditorData)
 	{
-		FString ConfigFileExtension = FPaths::GetExtension(InConfigPath, true);
-		FString FilePathToSave = InConfigPath;
-
-		FString JsonExtension = FString(".") + FString(DisplayClusterConfigurationStrings::file::FileExtJson);
-		FString CfgExtension = FString(".") + FString(DisplayClusterConfigurationStrings::file::FileExtCfg);
-
-		// Update file extension to .ndisplay if original file is .cfg
-		if (ConfigFileExtension.Equals(CfgExtension))
-		{
-			FilePathToSave = FPaths::ChangeExtension(InConfigPath, JsonExtension);
-		}
-
-		if (InConfiguratorEditorData != nullptr)
-		{
-			InConfiguratorEditorData->PathToConfig = FilePathToSave;
-			return IDisplayClusterConfiguration::Get().SaveConfig(InConfiguratorEditorData, FilePathToSave);
-		}
+		InConfiguratorEditorData->PathToConfig = InConfigPath;
+		return IDisplayClusterConfiguration::Get().SaveConfig(InConfiguratorEditorData, InConfigPath);
 	}
 
 	return false;
diff --git a/Engine/Plugins/Runtime/nDisplay/Source/DisplayClusterConfigurator/Private/DisplayClusterConfiguratorFactory.cpp b/Engine/Plugins/Runtime/nDisplay/Source/DisplayClusterConfigurator/Private/DisplayClusterConfiguratorFactory.cpp
--- a/Engine/Plugins/Runtime/nDisplay/Source/DisplayClusterConfigurator/Private/DisplayClusterConfiguratorFactory.cpp
+++ b/Engine/Plugins/Runtime/nDisplay/Source/DisplayClusterConfigurator/Private/DisplayClusterConfiguratorFactory.cpp
@@ -32,7 +32,6 @@ UDisplayClusterConfiguratorFactory::UDisplayClusterConfiguratorFactory()
 	SupportedClass = UDisplayClusterBlueprint::StaticClass();
 	ParentClass = ADisplayClusterRootActor::StaticClass();
 
-	Formats.Add(FString(DisplayClusterConfigurationStrings::file::FileExtCfg) + TEXT(";Config"));
 	Formats.Add(FString(DisplayClusterConfigurationStrings::file::FileExtJson) + TEXT(";Config"));
 
 	bCreateNew = true;
diff --git a/Engine/Plugins/Runtime/nDisplay/Source/DisplayClusterConfigurator/Private/DisplayClusterConfiguratorUtils.cpp b/Engine/Plugins/Runtime/nDisplay/Source/DisplayClusterConfigurator/Private/DisplayClusterConfiguratorUtils.cpp
--- a/Engine/Plugins/Runtime/nDisplay/Source/DisplayClusterConfigurator/Private/DisplayClusterConfiguratorUtils.cpp
+++ b/Engine/Plugins/Runtime/nDisplay/Source/DisplayClusterConfigurator/Private/DisplayClusterConfiguratorUtils.cpp
@@ -189,9 +189,9 @@ FDisplayClusterConfiguratorBlueprintEditor* FDisplayClusterConfiguratorUtils::Ge
 	return nullptr;
 }
 
-bool FDisplayClusterConfiguratorUtils::IsMasterNodeInConfig(UDisplayClusterConfigurationData* ConfigData)
+bool FDisplayClusterConfiguratorUtils::IsPrimaryNodeInConfig(UDisplayClusterConfigurationData* ConfigData)
 {
-	return ConfigData && ConfigData->Cluster && ConfigData->Cluster->Nodes.Num() > 0 && !ConfigData->Cluster->MasterNode.Id.IsEmpty();
+	return ConfigData && ConfigData->Cluster && ConfigData->Cluster->Nodes.Num() > 0 && !ConfigData->Cluster->PrimaryNode.Id.IsEmpty();
 }
 
 void FDisplayClusterConfiguratorUtils::MarkDisplayClusterBlueprintAsModified(UObject* InObject, const bool bIsStructuralChange)
diff --git a/Engine/Plugins/Runtime/nDisplay/Source/DisplayClusterConfigurator/Private/DisplayClusterConfiguratorUtils.h b/Engine/Plugins/Runtime/nDisplay/Source/DisplayClusterConfigurator/Private/DisplayClusterConfiguratorUtils.h
--- a/Engine/Plugins/Runtime/nDisplay/Source/DisplayClusterConfigurator/Private/DisplayClusterConfiguratorUtils.h
+++ b/Engine/Plugins/Runtime/nDisplay/Source/DisplayClusterConfigurator/Private/DisplayClusterConfiguratorUtils.h
@@ -40,8 +40,8 @@ public:
 	/** Locate an open Blueprint Editor for the object. */
 	static FDisplayClusterConfiguratorBlueprintEditor* GetBlueprintEditorForObject(UObject* InObject);
 
-	/** Checks if a master node is present in a config. */
-	static bool IsMasterNodeInConfig(UDisplayClusterConfigurationData* ConfigData);
+	/** Checks if a primary node is present in a config. */
+	static bool IsPrimaryNodeInConfig(UDisplayClusterConfigurationData* ConfigData);
 	
 	/**
 	 * Find the blueprint associated with the object and signal it should be recompiled.
diff --git a/Engine/Plugins/Runtime/nDisplay/Source/DisplayClusterConfigurator/Private/Views/OutputMapping/EdNodes/DisplayClusterConfiguratorWindowNode.cpp b/Engine/Plugins/Runtime/nDisplay/Source/DisplayClusterConfigurator/Private/Views/OutputMapping/EdNodes/DisplayClusterConfiguratorWindowNode.cpp
--- a/Engine/Plugins/Runtime/nDisplay/Source/DisplayClusterConfigurator/Private/Views/OutputMapping/EdNodes/DisplayClusterConfiguratorWindowNode.cpp
+++ b/Engine/Plugins/Runtime/nDisplay/Source/DisplayClusterConfigurator/Private/Views/OutputMapping/EdNodes/DisplayClusterConfiguratorWindowNode.cpp
@@ -70,7 +70,7 @@ bool UDisplayClusterConfiguratorWindowNode::IsFixedAspectRatio() const
 	return CfgClusterNode->bFixedAspectRatio;
 }
 
-bool UDisplayClusterConfiguratorWindowNode::IsMaster() const
+bool UDisplayClusterConfiguratorWindowNode::IsPrimary() const
 {
 	if (!IsObjectValid())
 	{
@@ -78,7 +78,7 @@ bool UDisplayClusterConfiguratorWindowNode::IsMaster() const
 	}
 
 	UDisplayClusterConfigurationClusterNode* ClusterNode = GetObjectChecked<UDisplayClusterConfigurationClusterNode>();
-	return FDisplayClusterConfiguratorClusterUtils::IsClusterNodeMaster(ClusterNode);
+	return FDisplayClusterConfiguratorClusterUtils::IsClusterNodePrimary(ClusterNode);
 }
 
 FDelegateHandle UDisplayClusterConfiguratorWindowNode::RegisterOnPreviewImageChanged(const FOnPreviewImageChangedDelegate& Delegate)
diff --git a/Engine/Plugins/Runtime/nDisplay/Source/DisplayClusterConfigurator/Private/Views/OutputMapping/EdNodes/DisplayClusterConfiguratorWindowNode.h b/Engine/Plugins/Runtime/nDisplay/Source/DisplayClusterConfigurator/Private/Views/OutputMapping/EdNodes/DisplayClusterConfiguratorWindowNode.h
--- a/Engine/Plugins/Runtime/nDisplay/Source/DisplayClusterConfigurator/Private/Views/OutputMapping/EdNodes/DisplayClusterConfiguratorWindowNode.h
+++ b/Engine/Plugins/Runtime/nDisplay/Source/DisplayClusterConfigurator/Private/Views/OutputMapping/EdNodes/DisplayClusterConfiguratorWindowNode.h
@@ -34,7 +34,7 @@ public:
 	FString GetCfgHost() const;
 	const FString& GetPreviewImagePath() const;
 	bool IsFixedAspectRatio() const;
-	bool IsMaster() const;
+	bool IsPrimary() const;
 
 	FDelegateHandle RegisterOnPreviewImageChanged(const FOnPreviewImageChangedDelegate& Delegate);
 	void UnregisterOnPreviewImageChanged(FDelegateHandle DelegateHandle);
diff --git a/Engine/Plugins/Runtime/nDisplay/Source/DisplayClusterConfigurator/Private/Views/OutputMapping/GraphNodes/SDisplayClusterConfiguratorWindowNode.cpp b/Engine/Plugins/Runtime/nDisplay/Source/DisplayClusterConfigurator/Private/Views/OutputMapping/GraphNodes/SDisplayClusterConfiguratorWindowNode.cpp
--- a/Engine/Plugins/Runtime/nDisplay/Source/DisplayClusterConfigurator/Private/Views/OutputMapping/GraphNodes/SDisplayClusterConfiguratorWindowNode.cpp
+++ b/Engine/Plugins/Runtime/nDisplay/Source/DisplayClusterConfigurator/Private/Views/OutputMapping/GraphNodes/SDisplayClusterConfiguratorWindowNode.cpp
@@ -214,9 +214,9 @@ public:
 	{
 		UDisplayClusterConfiguratorWindowNode* WindowEdNode = ParentNode->GetGraphNodeChecked<UDisplayClusterConfiguratorWindowNode>();
 		const FText NodeName = FText::FromString(WindowEdNode->GetNodeName());
-		if (WindowEdNode->IsMaster())
+		if (WindowEdNode->IsPrimary())
 		{
-			return FText::Format(LOCTEXT("WindowNameWithMaster", "{0} (Master)"), NodeName);
+			return FText::Format(LOCTEXT("WindowNameWithPrimary", "{0} (Primary)"), NodeName);
 		}
 		else
 		{
diff --git a/Engine/Plugins/Runtime/nDisplay/Source/DisplayClusterConfigurator/Private/Views/TreeViews/Cluster/DisplayClusterConfiguratorViewCluster.cpp b/Engine/Plugins/Runtime/nDisplay/Source/DisplayClusterConfigurator/Private/Views/TreeViews/Cluster/DisplayClusterConfiguratorViewCluster.cpp
--- a/Engine/Plugins/Runtime/nDisplay/Source/DisplayClusterConfigurator/Private/Views/TreeViews/Cluster/DisplayClusterConfiguratorViewCluster.cpp
+++ b/Engine/Plugins/Runtime/nDisplay/Source/DisplayClusterConfigurator/Private/Views/TreeViews/Cluster/DisplayClusterConfiguratorViewCluster.cpp
@@ -112,7 +112,7 @@ void FDisplayClusterConfiguratorViewCluster::FillContextMenu(FMenuBuilder& MenuB
 		if (SelectedObject->IsA<UDisplayClusterConfigurationClusterNode>())
 		{
 			MenuBuilder.AddSeparator();
-			MenuBuilder.AddMenuEntry(TreeViewCommands.SetAsMaster);
+			MenuBuilder.AddMenuEntry(TreeViewCommands.SetAsPrimary);
 		}
 	}
 }
@@ -196,9 +196,9 @@ void FDisplayClusterConfiguratorViewCluster::BindPinnableCommands(FUICommandList
 	);
 
 	CommandList.MapAction(
-		TreeViewCommands.SetAsMaster,
-		FExecuteAction::CreateSP(this, &FDisplayClusterConfiguratorViewCluster::SetAsMaster),
-		FCanExecuteAction::CreateSP(this, &FDisplayClusterConfiguratorViewCluster::CanSetAsMaster)
+		TreeViewCommands.SetAsPrimary,
+		FExecuteAction::CreateSP(this, &FDisplayClusterConfiguratorViewCluster::SetAsPrimary),
+		FCanExecuteAction::CreateSP(this, &FDisplayClusterConfiguratorViewCluster::CanSetAsPrimary)
 	);
 }
 
@@ -406,7 +406,7 @@ bool FDisplayClusterConfiguratorViewCluster::CanRenameSelectedNode() const
 	return false;
 }
 
-void FDisplayClusterConfiguratorViewCluster::SetAsMaster()
+void FDisplayClusterConfiguratorViewCluster::SetAsPrimary()
 {
 	TArray<TSharedPtr<IDisplayClusterConfiguratorTreeItem>> SelectedTreeItems = ViewTree->GetSelectedItems();
 
@@ -414,9 +414,9 @@ void FDisplayClusterConfiguratorViewCluster::SetAsMaster()
 	{
 		if (UDisplayClusterConfigurationClusterNode* ClusterNode = Cast<UDisplayClusterConfigurationClusterNode>(SelectedTreeItems[0]->GetObject()))
 		{
-			FScopedTransaction Transaction(LOCTEXT("SetMasterNode", "Set Master Node"));
+			FScopedTransaction Transaction(LOCTEXT("SetPrimaryNode", "Set Primary Node"));
 
-			if (!FDisplayClusterConfiguratorClusterUtils::SetClusterNodeAsMaster(ClusterNode))
+			if (!FDisplayClusterConfiguratorClusterUtils::SetClusterNodeAsPrimary(ClusterNode))
 			{
 				Transaction.Cancel();
 			}
@@ -510,7 +510,7 @@ void FDisplayClusterConfiguratorViewCluster::ShowAll()
 	}
 }
 
-bool FDisplayClusterConfiguratorViewCluster::CanSetAsMaster() const
+bool FDisplayClusterConfiguratorViewCluster::CanSetAsPrimary() const
 {
 	TArray<TSharedPtr<IDisplayClusterConfiguratorTreeItem>> SelectedTreeItems = ViewTree->GetSelectedItems();
 
@@ -519,9 +519,9 @@ bool FDisplayClusterConfiguratorViewCluster::CanSetAsMaster() const
 		if (const UDisplayClusterConfigurationClusterNode* ClusterNode = Cast<UDisplayClusterConfigurationClusterNode>(SelectedTreeItems[0]->GetObject()))
 		{
 			UDisplayClusterConfigurationCluster* Cluster = ToolkitPtr.Pin()->GetEditorData()->Cluster;
-			const FString MasterNodeId = Cluster->MasterNode.Id;
+			const FString PrimaryNodeId = Cluster->PrimaryNode.Id;
 
-			if (!Cluster->Nodes.Contains(MasterNodeId) || Cluster->Nodes[MasterNodeId] != ClusterNode)
+			if (!Cluster->Nodes.Contains(PrimaryNodeId) || Cluster->Nodes[PrimaryNodeId] != ClusterNode)
 			{
 				return true;
 			}
diff --git a/Engine/Plugins/Runtime/nDisplay/Source/DisplayClusterConfigurator/Private/Views/TreeViews/Cluster/DisplayClusterConfiguratorViewCluster.h b/Engine/Plugins/Runtime/nDisplay/Source/DisplayClusterConfigurator/Private/Views/TreeViews/Cluster/DisplayClusterConfiguratorViewCluster.h
--- a/Engine/Plugins/Runtime/nDisplay/Source/DisplayClusterConfigurator/Private/Views/TreeViews/Cluster/DisplayClusterConfiguratorViewCluster.h
+++ b/Engine/Plugins/Runtime/nDisplay/Source/DisplayClusterConfigurator/Private/Views/TreeViews/Cluster/DisplayClusterConfiguratorViewCluster.h
@@ -67,8 +67,8 @@ private:
 
 	void ShowAll();
 
-	void SetAsMaster();
-	bool CanSetAsMaster() const;
+	void SetAsPrimary();
+	bool CanSetAsPrimary() const;
 
 	void AddNewClusterNode(FVector2D PresetSize);
 	bool CanAddNewClusterNode() const;
diff --git a/Engine/Plugins/Runtime/nDisplay/Source/DisplayClusterConfigurator/Private/Views/TreeViews/Cluster/TreeItems/DisplayClusterConfiguratorTreeItemClusterNode.cpp b/Engine/Plugins/Runtime/nDisplay/Source/DisplayClusterConfigurator/Private/Views/TreeViews/Cluster/TreeItems/DisplayClusterConfiguratorTreeItemClusterNode.cpp
--- a/Engine/Plugins/Runtime/nDisplay/Source/DisplayClusterConfigurator/Private/Views/TreeViews/Cluster/TreeItems/DisplayClusterConfiguratorTreeItemClusterNode.cpp
+++ b/Engine/Plugins/Runtime/nDisplay/Source/DisplayClusterConfigurator/Private/Views/TreeViews/Cluster/TreeItems/DisplayClusterConfiguratorTreeItemClusterNode.cpp
@@ -313,8 +313,8 @@ void FDisplayClusterConfiguratorTreeItemClusterNode::FillItemColumn(TSharedPtr<S
 		.VAlign(VAlign_Center)
 		[
 			SNew(STextBlock)
-			.Text(NSLOCTEXT("FDisplayClusterConfiguratorTreeItemClusterNode", "MasterText", "(Master)"))
-			.Visibility(this, &FDisplayClusterConfiguratorTreeItemClusterNode::GetMasterLabelVisibility)
+			.Text(NSLOCTEXT("FDisplayClusterConfiguratorTreeItemClusterNode", "PrimaryText", "(Primary)"))
+			.Visibility(this, &FDisplayClusterConfiguratorTreeItemClusterNode::GetPrimaryLabelVisibility)
 		];
 }
 
@@ -340,12 +340,12 @@ void FDisplayClusterConfiguratorTreeItemClusterNode::OnDisplayNameCommitted(cons
 	}
 }
 
-EVisibility FDisplayClusterConfiguratorTreeItemClusterNode::GetMasterLabelVisibility() const
+EVisibility FDisplayClusterConfiguratorTreeItemClusterNode::GetPrimaryLabelVisibility() const
 {
 	UDisplayClusterConfigurationClusterNode* ClusterNode = GetObjectChecked<UDisplayClusterConfigurationClusterNode>();
-	bool bIsMaster = FDisplayClusterConfiguratorClusterUtils::IsClusterNodeMaster(ClusterNode);
+	const bool bIsPrimary = FDisplayClusterConfiguratorClusterUtils::IsClusterNodePrimary(ClusterNode);
 
-	return bIsMaster ? EVisibility::Visible : EVisibility::Collapsed;
+	return bIsPrimary ? EVisibility::Visible : EVisibility::Collapsed;
 }
 
 FSlateColor FDisplayClusterConfiguratorTreeItemClusterNode::GetHostColor() const
diff --git a/Engine/Plugins/Runtime/nDisplay/Source/DisplayClusterConfigurator/Private/Views/TreeViews/Cluster/TreeItems/DisplayClusterConfiguratorTreeItemClusterNode.h b/Engine/Plugins/Runtime/nDisplay/Source/DisplayClusterConfigurator/Private/Views/TreeViews/Cluster/TreeItems/DisplayClusterConfiguratorTreeItemClusterNode.h
--- a/Engine/Plugins/Runtime/nDisplay/Source/DisplayClusterConfigurator/Private/Views/TreeViews/Cluster/TreeItems/DisplayClusterConfiguratorTreeItemClusterNode.h
+++ b/Engine/Plugins/Runtime/nDisplay/Source/DisplayClusterConfigurator/Private/Views/TreeViews/Cluster/TreeItems/DisplayClusterConfiguratorTreeItemClusterNode.h
@@ -45,7 +45,7 @@ protected:
 	//~ End FDisplayClusterConfiguratorTreeItem Interface
 
 private:
-	EVisibility GetMasterLabelVisibility() const;
+	EVisibility GetPrimaryLabelVisibility() const;
 	FSlateColor GetHostColor() const;
 	FReply OnHostClicked(const FGeometry& MyGeometry, const FPointerEvent& MouseEvent);
 
diff --git a/Engine/Plugins/Runtime/nDisplay/Source/DisplayClusterConfigurator/Private/Views/TreeViews/DisplayClusterConfiguratorTreeViewCommands.cpp b/Engine/Plugins/Runtime/nDisplay/Source/DisplayClusterConfigurator/Private/Views/TreeViews/DisplayClusterConfiguratorTreeViewCommands.cpp
--- a/Engine/Plugins/Runtime/nDisplay/Source/DisplayClusterConfigurator/Private/Views/TreeViews/DisplayClusterConfiguratorTreeViewCommands.cpp
+++ b/Engine/Plugins/Runtime/nDisplay/Source/DisplayClusterConfigurator/Private/Views/TreeViews/DisplayClusterConfiguratorTreeViewCommands.cpp
@@ -17,7 +17,7 @@ void FDisplayClusterConfiguratorTreeViewCommands::RegisterCommands()
 	UI_COMMAND(ShowSelected, "Show Selected", "Shows the selected cluster items", EUserInterfaceActionType::Button, FInputChord(EModifierKey::Shift, EKeys::H));
 	UI_COMMAND(HideSelected, "Hide Selected", "Hides the selected cluster items", EUserInterfaceActionType::Button, FInputChord(EKeys::H));
 
-	UI_COMMAND(SetAsMaster, "Set as Master", "Sets this cluster node as the master node", EUserInterfaceActionType::Button, FInputChord());
+	UI_COMMAND(SetAsPrimary, "Set as Primary", "Sets this cluster node as the primary node", EUserInterfaceActionType::Button, FInputChord());
 }
 
 #undef LOCTEXT_NAMESPACE
diff --git a/Engine/Plugins/Runtime/nDisplay/Source/DisplayClusterConfigurator/Private/Views/TreeViews/DisplayClusterConfiguratorTreeViewCommands.h b/Engine/Plugins/Runtime/nDisplay/Source/DisplayClusterConfigurator/Private/Views/TreeViews/DisplayClusterConfiguratorTreeViewCommands.h
--- a/Engine/Plugins/Runtime/nDisplay/Source/DisplayClusterConfigurator/Private/Views/TreeViews/DisplayClusterConfiguratorTreeViewCommands.h
+++ b/Engine/Plugins/Runtime/nDisplay/Source/DisplayClusterConfigurator/Private/Views/TreeViews/DisplayClusterConfiguratorTreeViewCommands.h
@@ -31,5 +31,5 @@ public:
 	TSharedPtr<FUICommandInfo> ShowSelected;
 	TSharedPtr<FUICommandInfo> HideSelected;
 
-	TSharedPtr<FUICommandInfo> SetAsMaster;
+	TSharedPtr<FUICommandInfo> SetAsPrimary;
 };
diff --git a/Engine/Plugins/Runtime/nDisplay/Source/DisplayClusterEditor/Private/DisplayClusterEditorEngine.cpp b/Engine/Plugins/Runtime/nDisplay/Source/DisplayClusterEditor/Private/DisplayClusterEditorEngine.cpp
--- a/Engine/Plugins/Runtime/nDisplay/Source/DisplayClusterEditor/Private/DisplayClusterEditorEngine.cpp
+++ b/Engine/Plugins/Runtime/nDisplay/Source/DisplayClusterEditor/Private/DisplayClusterEditorEngine.cpp
@@ -117,7 +117,7 @@ void UDisplayClusterEditorEngine::StartPlayInEditorSession(FRequestPlaySessionPa
 			// And start PIE session with that config data
 			if (ConfigData)
 			{
-				if (!DisplayClusterModule->StartSession(ConfigData, ConfigData->Cluster->MasterNode.Id))
+				if (!DisplayClusterModule->StartSession(ConfigData, ConfigData->Cluster->PrimaryNode.Id))
 				{
 					UE_LOG(LogDisplayClusterEditorEngine, Error, TEXT("An error occurred during DisplayCluster session start"));
 				}
diff --git a/Engine/Plugins/Runtime/nDisplay/Source/DisplayClusterMessageInterceptor/Private/DisplayClusterMessageInterceptionModule.cpp b/Engine/Plugins/Runtime/nDisplay/Source/DisplayClusterMessageInterceptor/Private/DisplayClusterMessageInterceptionModule.cpp
--- a/Engine/Plugins/Runtime/nDisplay/Source/DisplayClusterMessageInterceptor/Private/DisplayClusterMessageInterceptionModule.cpp
+++ b/Engine/Plugins/Runtime/nDisplay/Source/DisplayClusterMessageInterceptor/Private/DisplayClusterMessageInterceptionModule.cpp
@@ -162,8 +162,8 @@ private:
 
 	void ResendSyncEvent(IDisplayClusterClusterManager* ClusterManager)
 	{
-		//Master will send out its interceptor settings to the cluster so everyone uses the same things
-		if (ClusterManager->IsMaster())
+		// Primary node will send out its interceptor settings to the cluster so everyone uses the same things
+		if (ClusterManager->IsPrimary())
 		{
 			FString ExportedSettings;
 			const UDisplayClusterMessageInterceptionSettings* CurrentSettings = GetDefault<UDisplayClusterMessageInterceptionSettings>();
@@ -174,8 +174,8 @@ private:
 			SettingsEvent.Name = ClusterManager->GetNodeId();
 			SettingsEvent.bIsSystemEvent = true;
 			SettingsEvent.Parameters.FindOrAdd(DisplayClusterInterceptionModuleUtils::EventParameterSettings) = MoveTemp(ExportedSettings);
-			const bool bMasterOnly = true;
-			ClusterManager->EmitClusterEventJson(SettingsEvent, bMasterOnly);
+			const bool bPrimaryOnly = true;
+			ClusterManager->EmitClusterEventJson(SettingsEvent, bPrimaryOnly);
 		}
 	}
 
@@ -198,8 +198,8 @@ private:
 		SyncMessagesEvent.Name = ClusterManager->GetNodeId();	// which node got the message
 		SyncMessagesEvent.bIsSystemEvent = true;				// nDisplay internal event
 		SyncMessagesEvent.bShouldDiscardOnRepeat = false;		// Don' discard the events with the same cat/type/name
-		const bool bMasterOnly = false;							// All nodes are broadcasting events to synchronize them across cluster
-		ClusterManager->EmitClusterEventJson(SyncMessagesEvent, bMasterOnly);
+		const bool bPrimaryOnly = false;							// All nodes are broadcasting events to synchronize them across cluster
+		ClusterManager->EmitClusterEventJson(SyncMessagesEvent, bPrimaryOnly);
 	}
 
 	void WorkspaceSyncEvent()
@@ -219,8 +219,8 @@ private:
 		SyncMessagesEvent.Name = ClusterManager->GetNodeId();	// which node got the message
 		SyncMessagesEvent.bIsSystemEvent = true;				// nDisplay internal event
 		SyncMessagesEvent.bShouldDiscardOnRepeat = false;		// Don' discard the events with the same cat/type/name
-		const bool bMasterOnly = false;							// All nodes are broadcasting events to synchronize them across cluster
-		ClusterManager->EmitClusterEventJson(SyncMessagesEvent, bMasterOnly);
+		const bool bPrimaryOnly = false;							// All nodes are broadcasting events to synchronize them across cluster
+		ClusterManager->EmitClusterEventJson(SyncMessagesEvent, bPrimaryOnly);
 	}
 
 	bool CanFinalizeWorkspaceSync() const
diff --git a/Engine/Plugins/Runtime/nDisplay/Source/DisplayClusterMessageInterceptor/Private/DisplayClusterMessageInterceptor.cpp b/Engine/Plugins/Runtime/nDisplay/Source/DisplayClusterMessageInterceptor/Private/DisplayClusterMessageInterceptor.cpp
--- a/Engine/Plugins/Runtime/nDisplay/Source/DisplayClusterMessageInterceptor/Private/DisplayClusterMessageInterceptor.cpp
+++ b/Engine/Plugins/Runtime/nDisplay/Source/DisplayClusterMessageInterceptor/Private/DisplayClusterMessageInterceptor.cpp
@@ -109,8 +109,8 @@ void FDisplayClusterMessageInterceptor::SyncMessages()
 		for (const FString& MessageId : MessageIds)
 		{
 			SyncMessagesEvent.Type = MessageId;					// the actually message id we received
-			const bool bMasterOnly = false; //All nodes are broadcasting events to synchronize them across cluster
-			ClusterManager->EmitClusterEventJson(SyncMessagesEvent, bMasterOnly);
+			const bool bPrimaryOnly = false; //All nodes are broadcasting events to synchronize them across cluster
+			ClusterManager->EmitClusterEventJson(SyncMessagesEvent, bPrimaryOnly);
 			UE_LOG(LogDisplayClusterInterception, VeryVerbose, TEXT("Emitting cluster event for message %s on frame %d"), *MessageId, GFrameCounter);
 		}
 	}
diff --git a/Engine/Plugins/Runtime/nDisplay/Source/DisplayClusterMessageInterceptor/Public/DisplayClusterMessageInterceptionSettings.h b/Engine/Plugins/Runtime/nDisplay/Source/DisplayClusterMessageInterceptor/Public/DisplayClusterMessageInterceptionSettings.h
--- a/Engine/Plugins/Runtime/nDisplay/Source/DisplayClusterMessageInterceptor/Public/DisplayClusterMessageInterceptionSettings.h
+++ b/Engine/Plugins/Runtime/nDisplay/Source/DisplayClusterMessageInterceptor/Public/DisplayClusterMessageInterceptionSettings.h
@@ -40,7 +40,7 @@ public:
 
 	/**
 	 * Settings related to interception. 
-	 * @note Settings from master will be synchronized across the cluster
+	 * @note Settings from primary node will be synchronized across the cluster
 	 */
 	UPROPERTY(config, EditAnywhere, Category = "Interception Settings", meta = (ShowOnlyInnerProperties))
 	FMessageInterceptionSettings InterceptionSettings;
diff --git a/Engine/Plugins/Runtime/nDisplay/Source/DisplayClusterPostprocess/Private/PostProcess/Windows/DisplayClusterPostprocessTextureShare.cpp b/Engine/Plugins/Runtime/nDisplay/Source/DisplayClusterPostprocess/Private/PostProcess/Windows/DisplayClusterPostprocessTextureShare.cpp
--- a/Engine/Plugins/Runtime/nDisplay/Source/DisplayClusterPostprocess/Private/PostProcess/Windows/DisplayClusterPostprocessTextureShare.cpp
+++ b/Engine/Plugins/Runtime/nDisplay/Source/DisplayClusterPostprocess/Private/PostProcess/Windows/DisplayClusterPostprocessTextureShare.cpp
@@ -276,7 +276,7 @@ void FDisplayClusterPostprocessTextureShare::PerformPostProcessRenderTargetBefor
 
 void FDisplayClusterPostprocessTextureShare::InitializeResources_RenderThread(FRHICommandListImmediate& RHICmdList, FRHITexture2D* InOutTexture) const
 {
-	// Create master share for all render viewports
+	// Create share for all render viewports
 	for (int CurrentViewportIndex = 0; CurrentViewportIndex < RenderViewports.Num(); CurrentViewportIndex++)
 	{
 		const FDisplayClusterViewport& CurrentViewport = RenderViewports[CurrentViewportIndex];
@@ -286,7 +286,7 @@ void FDisplayClusterPostprocessTextureShare::InitializeResources_RenderThread(FR
 		}
 	}
 
-	// Open slave share for all render viewports
+	// Open share for all render viewports
 	for (const auto& It : DestinationViewportsMap)
 	{
 		ImplOpenResource_RenderThread(RHICmdList, GetDisplayClusterViewportShareName(It.Value));
@@ -300,7 +300,7 @@ void FDisplayClusterPostprocessTextureShare::InitializeResources(const FViewport
 {
 	const FTexture2DRHIRef& Backbuffer = MainViewport.GetRenderTargetTexture();
 
-	// Create master share for all render viewports
+	// Create share for all render viewports
 	for (int CurrentViewportIndex = 0; CurrentViewportIndex < RenderViewports.Num(); CurrentViewportIndex++)
 	{
 		const FDisplayClusterViewport& CurrentViewport = RenderViewports[CurrentViewportIndex];
@@ -310,7 +310,7 @@ void FDisplayClusterPostprocessTextureShare::InitializeResources(const FViewport
 		}
 	}
 
-	// Open slave share for all render viewports
+	// Open share for all render viewports
 	for (const auto& It : DestinationViewportsMap)
 	{
 		ImplOpenResource(GetDisplayClusterViewportShareName(It.Value));
diff --git a/Engine/Plugins/Runtime/nDisplay/Source/DisplayClusterRemoteControlInterceptor/Private/DisplayClusterRemoteControlInterceptor.cpp b/Engine/Plugins/Runtime/nDisplay/Source/DisplayClusterRemoteControlInterceptor/Private/DisplayClusterRemoteControlInterceptor.cpp
--- a/Engine/Plugins/Runtime/nDisplay/Source/DisplayClusterRemoteControlInterceptor/Private/DisplayClusterRemoteControlInterceptor.cpp
+++ b/Engine/Plugins/Runtime/nDisplay/Source/DisplayClusterRemoteControlInterceptor/Private/DisplayClusterRemoteControlInterceptor.cpp
@@ -15,12 +15,12 @@
 #include "Misc/CommandLine.h"
 
 // Data interception source
-static TAutoConsoleVariable<int32> CVarInterceptOnMasterOnly(
-	TEXT("nDisplay.RemoteControlInterceptor.MasterOnly"),
+static TAutoConsoleVariable<int32> CVarInterceptOnPrimaryOnly(
+	TEXT("nDisplay.RemoteControlInterceptor.PrimaryOnly"),
 	1,
 	TEXT("RemoteControl commands interception location\n")
 	TEXT("0 : All nodes\n")
-	TEXT("1 : Master only\n")
+	TEXT("1 : Primary only\n")
 	,
 	ECVF_ReadOnly
 );
@@ -30,7 +30,7 @@ static TAutoConsoleVariable<int32> CVarInterceptOnMasterOnly(
 const int32 EventId_InterceptorQueue		= 0xaabb0704;
 
 FDisplayClusterRemoteControlInterceptor::FDisplayClusterRemoteControlInterceptor()
-	: bInterceptOnMasterOnly(CVarInterceptOnMasterOnly.GetValueOnGameThread() == 1)
+	: bInterceptOnPrimaryOnly(CVarInterceptOnPrimaryOnly.GetValueOnGameThread() == 1)
 	, bForceApply(false)
 {
 	bForceApply = FParse::Param(FCommandLine::Get(), TEXT("ClusterForceApplyResponse"));
@@ -251,8 +251,8 @@ void FDisplayClusterRemoteControlInterceptor::SendReplicationQueue()
 	Event.bShouldDiscardOnRepeat = false;
 	Event.EventData              = MoveTemp(Buffer);
 
-	// Emit cluster event (or not, it depends on the bInterceptOnMasterOnly value and the role of this cluster node)
-	IDisplayCluster::Get().GetClusterMgr()->EmitClusterEventBinary(Event, bInterceptOnMasterOnly);
+	// Emit cluster event (or not, it depends on the bInterceptOnPrimaryOnly value and the role of this cluster node)
+	IDisplayCluster::Get().GetClusterMgr()->EmitClusterEventBinary(Event, bInterceptOnPrimaryOnly);
 
 	InterceptQueueMap.Empty();
 }
diff --git a/Engine/Plugins/Runtime/nDisplay/Source/DisplayClusterRemoteControlInterceptor/Private/DisplayClusterRemoteControlInterceptor.h b/Engine/Plugins/Runtime/nDisplay/Source/DisplayClusterRemoteControlInterceptor/Private/DisplayClusterRemoteControlInterceptor.h
--- a/Engine/Plugins/Runtime/nDisplay/Source/DisplayClusterRemoteControlInterceptor/Private/DisplayClusterRemoteControlInterceptor.h
+++ b/Engine/Plugins/Runtime/nDisplay/Source/DisplayClusterRemoteControlInterceptor/Private/DisplayClusterRemoteControlInterceptor.h
@@ -43,8 +43,8 @@ private:
 	void SendReplicationQueue();
 
 private:
-	// CVar value MasterOnly
-	const bool bInterceptOnMasterOnly;
+	// CVar value PrimaryOnly
+	const bool bInterceptOnPrimaryOnly;
 	// Cluster events listener
 	FOnClusterEventBinaryListener EventsListener;
 	// Force Apply the ERCIResponse
diff --git a/Engine/Plugins/VirtualProduction/DMX/DMXDisplayCluster/Source/DMXDisplayCluster/Private/DMXDisplayClusterReplicator.cpp b/Engine/Plugins/VirtualProduction/DMX/DMXDisplayCluster/Source/DMXDisplayCluster/Private/DMXDisplayClusterReplicator.cpp
--- a/Engine/Plugins/VirtualProduction/DMX/DMXDisplayCluster/Source/DMXDisplayCluster/Private/DMXDisplayClusterReplicator.cpp
+++ b/Engine/Plugins/VirtualProduction/DMX/DMXDisplayCluster/Source/DMXDisplayCluster/Private/DMXDisplayClusterReplicator.cpp
@@ -85,7 +85,7 @@ FDMXDisplayClusterReplicator::FDMXDisplayClusterReplicator()
 
 			if (bAutoMode)
 			{
-				return ClusterManager->IsMaster();
+				return ClusterManager->IsPrimary();
 			}
 
 			return bForceMaster;
