[nDisplay] Cleaning up deprecated stuff

diff --git a/Engine/Plugins/Runtime/nDisplay/Source/DisplayCluster/Public/IDisplayCluster.h b/Engine/Plugins/Runtime/nDisplay/Source/DisplayCluster/Public/IDisplayCluster.h
--- a/Engine/Plugins/Runtime/nDisplay/Source/DisplayCluster/Public/IDisplayCluster.h
+++ b/Engine/Plugins/Runtime/nDisplay/Source/DisplayCluster/Public/IDisplayCluster.h
@@ -79,18 +79,6 @@ public:
 	*/
 	virtual IDisplayClusterClusterManager* GetClusterMgr() const = 0;
 
-	/**
-	* Access to the input manager.
-	*
-	* @return Current cluster manager or nullptr
-	*/
-	class IDisplayClusterInputManager { };
-	UE_DEPRECATED(4.27, "VRPN functionality has been moved to LiveLinkVRPN")
-	virtual IDisplayClusterInputManager* GetInputMgr() const
-	{
-		return nullptr;
-	}
-
 	/**
 	* Access to the config manager.
 	*
diff --git a/Engine/Plugins/Runtime/nDisplay/Source/DisplayCluster/Public/Render/Device/IDisplayClusterRenderDevice.h b/Engine/Plugins/Runtime/nDisplay/Source/DisplayCluster/Public/Render/Device/IDisplayClusterRenderDevice.h
--- a/Engine/Plugins/Runtime/nDisplay/Source/DisplayCluster/Public/Render/Device/IDisplayClusterRenderDevice.h
+++ b/Engine/Plugins/Runtime/nDisplay/Source/DisplayCluster/Public/Render/Device/IDisplayClusterRenderDevice.h
@@ -55,180 +55,6 @@ public:
 	virtual bool BeginNewFrame(FViewport* InViewport, UWorld* InWorld, FDisplayClusterRenderFrame& OutRenderFrame) = 0;
 	virtual void FinalizeNewFrame() = 0;
 
-	/**
-	* Assigns camera to a specified viewport. If InViewportId is empty, all viewports will be assigned to a new camera. Empty camera ID means default active camera.
-	*
-	* @param CameraId   - ID of a camera (see [camera] in the nDisplay config file
-	* @param ViewportId - ID of a viewport to assign the camera (all viewports if empty)
-	*/
-	UE_DEPRECATED(4.27, "This function has been moved to FDisplayClusterViewport. Use GetViewportManager() to access  that interface.")
-	virtual void SetViewportCamera(const FString& CameraId = FString(), const FString& ViewportId = FString()) 
-	{}
-
-	/**
-	* Start postprocess settings
-	*
-	* @param ViewportId - viewport to set PP settings
-	* @param StartPostProcessingSettings - PP settings
-	*
-	*/
-	UE_DEPRECATED(4.27, "This function has been moved to FDisplayClusterViewport. Use GetViewportManager() to access  that interface.")
-	virtual void SetStartPostProcessingSettings(const FString& ViewportId, const FPostProcessSettings& StartPostProcessingSettings)
-	{}
-
-	/**
-	* Override postprocess settings
-	*
-	* @param ViewportId - viewport to set PP settings
-	* @param OverridePostProcessingSettings - PP settings
-	*
-	*/
-	UE_DEPRECATED(4.27, "This function has been moved to FDisplayClusterViewport. Use GetViewportManager() to access  that interface.")
-	virtual void SetOverridePostProcessingSettings(const FString& ViewportId, const FPostProcessSettings& OverridePostProcessingSettings, float BlendWeight = 1.0f)
-	{}
-
-	/**
-	* Final postprocess settings
-	*
-	* @param ViewportId - viewport to set PP settings
-	* @param FinalPostProcessingSettings - PP settings
-	*
-	*/
-	UE_DEPRECATED(4.27, "This function has been moved to FDisplayClusterViewport. Use GetViewportManager() to access  that interface.")
-	virtual void SetFinalPostProcessingSettings(const FString& ViewportId, const FPostProcessSettings& FinalPostProcessingSettings)
-	{}
-
-	/**
-	* Returns location and size of a viewport
-	*
-	* @param ViewportId - ID of a viewport
-	* @param OutRect    - a rectangle that describes location and size of the viewport
-	*/
-	UE_DEPRECATED(4.27, "This function has been moved to FDisplayClusterViewport. Use GetViewportManager() to access  that interface.")
-	virtual bool GetViewportRect(const FString& ViewportId, FIntRect& OutRect)
-	{
-		return false;
-	}
-
-	/**
-	* Returns projection policy object of a specified viewport
-	*
-	* @param ViewportId          - ID of a viewport
-	* @param OutProjectionPolicy - projection policy instance
-	*/
-	UE_DEPRECATED(4.27, "This function has been moved to FDisplayClusterViewport. Use GetViewportManager() to access  that interface.")
-	virtual bool GetViewportProjectionPolicy(const FString& ViewportId, TSharedPtr<IDisplayClusterProjectionPolicy, ESPMode::ThreadSafe>& OutProjectionPolicy)
-	{
-		return false;
-	}
-
-	/**
-	* Returns context of a specified viewport
-	*
-	* @param ViewportId          - ID of a viewport
-	* @param ViewIndex           - View index (left/center/right eye)
-	* @param OutProjectionPolicy - projection policy instance
-	*/
-	UE_DEPRECATED(4.27, "This function has been moved to FDisplayClusterViewport. Use GetViewportManager() to access  that interface.")
-	virtual bool GetViewportContext(const FString& ViewportId, int ViewIndex, FDisplayClusterRenderViewContext& OutViewContext)
-	{
-		return false;
-	}
-
-	/**
-	* Sets buffer ratio
-	*
-	* @param ViewportId    - viewport ID to change buffer ratio
-	* @param InBufferRatio - buffer ratio
-	*
-	* @return - true if succeeded
-	*/
-	UE_DEPRECATED(4.27, "This function has been moved to FDisplayClusterViewport. Use GetViewportManager() to access  that interface.")
-	virtual bool SetBufferRatio(const FString& ViewportId, float InBufferRatio)
-	{
-		return false;
-	};
-	
-	
-	/**
-	* Returns current buffer ratio
-	*
-	* @param ViewportId     - Viewport ID to get buffer ratio
-	* @param OutBufferRatio - buffer ratio
-	*
-	* @return - true if succeeded
-	*/
-	UE_DEPRECATED(4.27, "This function has been moved to FDisplayClusterViewport. Use GetViewportManager() to access  that interface.")
-	virtual bool GetBufferRatio(const FString& ViewportId, float& OutBufferRatio) const
-	{
-		return false;
-	}
-
-	/**
-	* Sets buffer ratio
-	*
-	* @param ViewportIdx   - viewport index to change buffer ratio
-	* @param InBufferRatio - buffer ratio
-	*
-	* @return - true if succeeded
-	*/
-	UE_DEPRECATED(4.27, "This function has been moved to FDisplayClusterViewport. Use GetViewportManager() to access  that interface.")
-	virtual bool SetBufferRatio(int32 ViewportIdx, float InBufferRatio)
-	{
-		return false;
-	}
-
-	/**
-	* Returns current buffer ratio
-	*
-	* @param ViewportIdx    - Viewport index to get buffer ratio
-	* @param OutBufferRatio - buffer ratio
-	*
-	* @return - true if succeeded
-	*/
-	UE_DEPRECATED(4.27, "This function has been moved to FDisplayClusterViewport. Use GetViewportManager() to access  that interface.")
-	virtual bool GetBufferRatio(int32 ViewportIdx, float& OutBufferRatio) const
-	{
-		return false;
-	}
-
-	/**
-	* Returns specified viewport info
-	*
-	* @param ViewportId - Viewport ID to get info
-	*
-	* @return - Viewport info
-	*/
-	UE_DEPRECATED(4.27, "This function has been moved to FDisplayClusterViewport. Use GetViewportManager() to access  that interface.")
-	virtual const class FDisplayClusterRenderViewport* GetRenderViewport(const FString& ViewportId) const
-	{
-		return nullptr;
-	}
-
-	/**
-	* Returns specified viewport info
-	*
-	* @param ViewportIdx - Viewport index to get info
-	*
-	* @return - Viewport info
-	*/
-	UE_DEPRECATED(4.27, "This function has been moved to FDisplayClusterViewport. Use GetViewportManager() to access  that interface.")
-	virtual const class FDisplayClusterRenderViewport* GetRenderViewport(int32 ViewportIdx) const
-	{
-		return nullptr;
-	}
-
-	/**
-	* Returns all available viewports with filter applied
-	*
-	* @param Filter - viewport filter
-	*
-	* @return - Viewport info
-	*/
-	UE_DEPRECATED(4.27, "This function has been moved to FDisplayClusterViewport. Use GetViewportManager() to access  that interface.")
-	virtual const void GetRenderViewports(TArray<class FDisplayClusterRenderViewport*>& OutViewports) const
-	{}
-
 	/**
 	* Returns current presentation handler
 	*
diff --git a/Engine/Plugins/Runtime/nDisplay/Source/DisplayCluster/Public/Render/IDisplayClusterRenderManager.h b/Engine/Plugins/Runtime/nDisplay/Source/DisplayCluster/Public/Render/IDisplayClusterRenderManager.h
--- a/Engine/Plugins/Runtime/nDisplay/Source/DisplayCluster/Public/Render/IDisplayClusterRenderManager.h
+++ b/Engine/Plugins/Runtime/nDisplay/Source/DisplayCluster/Public/Render/IDisplayClusterRenderManager.h
@@ -161,51 +161,6 @@ public:
 	*/
 	virtual void GetRegisteredPostProcess(TArray<FString>& OutPostProcessIDs) const = 0;
 
-	/**
-	* Registers a post process operation
-	*
-	* @param Name      - A unique PP operation name
-	* @param Operation - PP operation implementation
-	* @param Priority  - PP order in chain (the calling order is from the smallest to the largest: -N...0...N)
-	*
-	* @return - True if success
-	*/
-	UE_DEPRECATED(4.27, "This function has been moved to FDisplayClusterViewport. Use GetViewportManager() to access  that interface.")
-	virtual bool RegisterPostprocessOperation(const FString& Name, TSharedPtr<IDisplayClusterPostProcess, ESPMode::ThreadSafe>& Operation, int Priority = 0)
-	{ return false; }
-
-	/**
-	* Registers a post process operation
-	*
-	* @param Name - A unique PP operation name
-	* @param PPInfo - PP info wrapper (see IDisplayClusterRenderManager::FDisplayClusterPPInfo)
-	*
-	* @return - True if success
-	*/
-	UE_DEPRECATED(4.27, "This function has been moved to FDisplayClusterViewport. Use GetViewportManager() to access  that interface.")
-	virtual bool RegisterPostprocessOperation(const FString& Name, FDisplayClusterPPInfo& PPInfo)
-	{ return false; }
-
-	/**
-	* Unregisters a post process operation
-	*
-	* @param Name - PP operation name
-	*
-	* @return - True if success
-	*/
-	UE_DEPRECATED(4.27, "This function has been moved to FDisplayClusterViewport. Use GetViewportManager() to access  that interface.")
-	virtual bool UnregisterPostprocessOperation(const FString& Name)
-	{ return false; }
-
-	/**
-	* Returns all registered post-process operations
-	*
-	* @return - PP operations
-	*/
-	UE_DEPRECATED(4.27, "This function has been moved to FDisplayClusterViewport. Use GetViewportManager() to access  that interface.")
-	virtual TMap<FString, FDisplayClusterPPInfo> GetRegisteredPostprocessOperations() const
-	{ return TMap<FString, FDisplayClusterPPInfo>(); }
-
 	/**
 	* @return - Current viewport manager from root actor
 	*/
diff --git a/Engine/Plugins/Runtime/nDisplay/Source/DisplayCluster/Public/Render/PostProcess/IDisplayClusterPostProcess.h b/Engine/Plugins/Runtime/nDisplay/Source/DisplayCluster/Public/Render/PostProcess/IDisplayClusterPostProcess.h
--- a/Engine/Plugins/Runtime/nDisplay/Source/DisplayCluster/Public/Render/PostProcess/IDisplayClusterPostProcess.h
+++ b/Engine/Plugins/Runtime/nDisplay/Source/DisplayCluster/Public/Render/PostProcess/IDisplayClusterPostProcess.h
@@ -20,15 +20,6 @@ public:
 	virtual ~IDisplayClusterPostProcess() = default;
 
 public:
-	/**
-	* Game thread call. Initialize postprocess from config line
-	*
-	* @param CfgLine - Configuration line for this postprocess
-	*/
-	UE_DEPRECATED(4.26, "This function is deprecated. Use TMap based InitializePostProcess.")
-	virtual void InitializePostProcess(const FString& CfgLine)
-	{ }
-
 	/**
 	* Return postprocess name
 	*/
diff --git a/Engine/Plugins/Runtime/nDisplay/Source/DisplayClusterConfigurator/Private/DisplayClusterConfiguratorBlueprintEditor.cpp b/Engine/Plugins/Runtime/nDisplay/Source/DisplayClusterConfigurator/Private/DisplayClusterConfiguratorBlueprintEditor.cpp
--- a/Engine/Plugins/Runtime/nDisplay/Source/DisplayClusterConfigurator/Private/DisplayClusterConfiguratorBlueprintEditor.cpp
+++ b/Engine/Plugins/Runtime/nDisplay/Source/DisplayClusterConfigurator/Private/DisplayClusterConfiguratorBlueprintEditor.cpp
@@ -577,7 +577,7 @@ void FDisplayClusterConfiguratorBlueprintEditor::UpdateXformGizmos()
 bool FDisplayClusterConfiguratorBlueprintEditor::LoadWithOpenFileDialog()
 {
 	const FString NDisplayFileDescription = LOCTEXT("NDisplayFileDescription", "nDisplay Config").ToString();
-	const FString NDisplayFileExtension = TEXT("*.ndisplay;*.cfg");
+	const FString NDisplayFileExtension = TEXT("*.ndisplay");
 	const FString FileTypes = FString::Printf(TEXT("%s (%s)|%s"), *NDisplayFileDescription, *NDisplayFileExtension, *NDisplayFileExtension);
 
 	// Prompt the user for the filenames
diff --git a/Engine/Plugins/Runtime/nDisplay/Source/DisplayClusterConfigurator/Private/DisplayClusterConfiguratorCommands.cpp b/Engine/Plugins/Runtime/nDisplay/Source/DisplayClusterConfigurator/Private/DisplayClusterConfiguratorCommands.cpp
--- a/Engine/Plugins/Runtime/nDisplay/Source/DisplayClusterConfigurator/Private/DisplayClusterConfiguratorCommands.cpp
+++ b/Engine/Plugins/Runtime/nDisplay/Source/DisplayClusterConfigurator/Private/DisplayClusterConfiguratorCommands.cpp
@@ -7,8 +7,7 @@
 
 void FDisplayClusterConfiguratorCommands::RegisterCommands()
 {
-	UI_COMMAND(Import, "Import", "Import an nDisplay config. This may overwrite components. When upgrading from the deprecated .cfg file format the recommended import workflow is to drag & drop your .cfg file into the content browser.\
- If you need to reimport it is recommended to use the reimport options of the asset context menu", EUserInterfaceActionType::Button, FInputChord());
+	UI_COMMAND(Import, "Import", "Import an nDisplay config. This may overwrite components. If you need to reimport it is recommended to use the reimport options of the asset context menu", EUserInterfaceActionType::Button, FInputChord());
 	UI_COMMAND(Export, "Export", "Export to an nDisplay config. Requires a primary cluster node set. This does not export all components and is meant for launching the cluster from the command line", EUserInterfaceActionType::Button, FInputChord());
 	UI_COMMAND(EditConfig, "EditConfig", "Edit config file with text editor", EUserInterfaceActionType::Button, FInputChord());
 	UI_COMMAND(ExportConfigOnSave, "Export on Save", "Export to nDisplay config automatically on save. Requires a config previously exported", EUserInterfaceActionType::ToggleButton, FInputChord());
diff --git a/Engine/Plugins/Runtime/nDisplay/Source/DisplayClusterProjection/Public/Blueprints/IDisplayClusterProjectionBlueprintAPI.h b/Engine/Plugins/Runtime/nDisplay/Source/DisplayClusterProjection/Public/Blueprints/IDisplayClusterProjectionBlueprintAPI.h
--- a/Engine/Plugins/Runtime/nDisplay/Source/DisplayClusterProjection/Public/Blueprints/IDisplayClusterProjectionBlueprintAPI.h
+++ b/Engine/Plugins/Runtime/nDisplay/Source/DisplayClusterProjection/Public/Blueprints/IDisplayClusterProjectionBlueprintAPI.h
@@ -31,10 +31,4 @@ public:
 	/** Sets active camera component for camera policy */
 	UFUNCTION(BlueprintCallable, meta = (DisplayName = "Set camera"), Category = "NDisplayProjection|Camera")
 	virtual void CameraPolicySetCamera(const FString& ViewportId, UCameraComponent* NewCamera, float FOVMultiplier = 1.f) = 0;
-
-	/** Create link to static mesh geometry as warp source */
-	//UFUNCTION(BlueprintCallable, meta = (DisplayName = "Assign Warp Mesh To Viewport"), Category = "NDisplayProjection|Mesh")
-	UE_DEPRECATED(4.26, "use config reference")
-	virtual void AssignWarpMeshToViewport(const FString& ViewportId, UStaticMeshComponent* MeshComponent, USceneComponent* OriginComponent)
-	{ }
 };
diff --git a/Engine/Plugins/Runtime/nDisplay/Source/DisplayClusterProjection/Public/IDisplayClusterProjection.h b/Engine/Plugins/Runtime/nDisplay/Source/DisplayClusterProjection/Public/IDisplayClusterProjection.h
--- a/Engine/Plugins/Runtime/nDisplay/Source/DisplayClusterProjection/Public/IDisplayClusterProjection.h
+++ b/Engine/Plugins/Runtime/nDisplay/Source/DisplayClusterProjection/Public/IDisplayClusterProjection.h
@@ -57,19 +57,6 @@ public:
 	*/
 	virtual TSharedPtr<IDisplayClusterProjectionPolicyFactory> GetProjectionFactory(const FString& InProjectionType) = 0;
 
-	/**
-	* Create link to static mesh geometry as warp source
-	*
-	* @param ViewportId - viewport name
-	* @param MeshComponent - warp mesh
-	* @param OriginComponent - cave origin 
-	*
-	* @return - true if the mesh linked and ready to warp
-	*/
-	UE_DEPRECATED(4.26, "use config reference")
-	virtual bool AssignWarpMeshToViewport(const FString& ViewportId, UStaticMeshComponent* MeshComponent, USceneComponent* OriginComponent)
-	{ return false; }
-
 	/**
 	* Set camera policy camera
 	*
