[nDisplay]
Refactored the public DCMeshComponent APIs.
Added support for procedural and static mesh components for OutputRemap.
Removed spam in logs.
Added UVs remap for Mesh projection policy
Fixed Horde Issue 125526 (nDisplay related code)

diff --git a/Engine/Plugins/Runtime/nDisplay/Source/DisplayCluster/Private/Game/EngineClasses/Scene/DisplayClusterRootActor.cpp b/Engine/Plugins/Runtime/nDisplay/Source/DisplayCluster/Private/Game/EngineClasses/Scene/DisplayClusterRootActor.cpp
--- a/Engine/Plugins/Runtime/nDisplay/Source/DisplayCluster/Private/Game/EngineClasses/Scene/DisplayClusterRootActor.cpp
+++ b/Engine/Plugins/Runtime/nDisplay/Source/DisplayCluster/Private/Game/EngineClasses/Scene/DisplayClusterRootActor.cpp
@@ -37,9 +37,6 @@
 #include "Misc/DisplayClusterStrings.h"
 
 #include "Render/Viewport/DisplayClusterViewportManager.h"
-#include "Render/Viewport/IDisplayClusterViewport.h"
-#include "Render/Projection/IDisplayClusterProjectionPolicy.h"
-#include "WarpBlend/IDisplayClusterWarpBlend.h"
 
 #include "Game/EngineClasses/Scene/DisplayClusterRootActorInitializer.h"
 
@@ -583,27 +580,8 @@ void ADisplayClusterRootActor::UpdateProceduralMeshComponentData(const UProcedur
 	{
 		FName ProceduralComponentName = (InProceduralMeshComponent==nullptr) ? NAME_None : InProceduralMeshComponent->GetFName();
 
-		const TArrayView<IDisplayClusterViewport*> Viewports = ViewportManager->GetViewports();
-		for (IDisplayClusterViewport* ViewportIt : Viewports)
-		{
-			if (ViewportIt != nullptr)
-			{
-				const TSharedPtr<IDisplayClusterProjectionPolicy, ESPMode::ThreadSafe>& ProjectionPolicy = ViewportIt->GetProjectionPolicy();
-				if (ProjectionPolicy.IsValid())
-				{
-					TSharedPtr<IDisplayClusterWarpBlend, ESPMode::ThreadSafe> WarpBlendInterface;
-					if (ProjectionPolicy->GetWarpBlendInterface(WarpBlendInterface) && WarpBlendInterface.IsValid())
-					{
-						// Update only interfaces with ProceduralMesh as geometry source 
-						if (WarpBlendInterface->GetWarpGeometryType() == EDisplayClusterWarpGeometryType::WarpProceduralMesh)
-						{
-							// Set the ProceduralMeshComponent geometry dirty for all valid WarpBlendInterface
-							WarpBlendInterface->MarkWarpGeometryComponentDirty(ProceduralComponentName);
-						}
-					}
-				}
-			}
-		}
+		// Support for all hidden internal refs
+		ViewportManager->MarkComponentGeometryDirty(ProceduralComponentName);
 	}
 }
 
diff --git a/Engine/Plugins/Runtime/nDisplay/Source/DisplayCluster/Private/Render/Containers/DisplayClusterRender_MeshComponent.cpp b/Engine/Plugins/Runtime/nDisplay/Source/DisplayCluster/Private/Render/Containers/DisplayClusterRender_MeshComponent.cpp
--- a/Engine/Plugins/Runtime/nDisplay/Source/DisplayCluster/Private/Render/Containers/DisplayClusterRender_MeshComponent.cpp
+++ b/Engine/Plugins/Runtime/nDisplay/Source/DisplayCluster/Private/Render/Containers/DisplayClusterRender_MeshComponent.cpp
@@ -44,9 +44,18 @@ static const FStaticMeshLODResources* ImplGetStaticMeshLODResources(const UStati
 {
 	if (InStaticMesh != nullptr)
 	{
-		if (InStaticMesh->bAllowCPUAccess)
+		if (InStaticMesh->bAllowCPUAccess == false)
 		{
-			UE_LOG(LogDisplayClusterRender, Warning, TEXT("If packaging this project, static mesh '%s' requires its AllowCPUAccess flag to be enabled."), *InStaticMesh->GetName());
+			const FString InStaticMeshName = InStaticMesh->GetName();
+
+			// Removed spam in log
+			static TArray<FString> MeshNamesInLog;
+			if (MeshNamesInLog.Find(InStaticMeshName) == INDEX_NONE)
+			{
+				// Show this warning once per mesh
+				UE_LOG(LogDisplayClusterRender, Warning, TEXT("If packaging this project, static mesh '%s' requires its AllowCPUAccess flag to be enabled."), *InStaticMeshName);
+				MeshNamesInLog.Add(InStaticMeshName);
+			}
 #if !WITH_EDITOR
 			// Can't access to cooked data from CPU without this flag
 			return nullptr;
@@ -99,7 +108,7 @@ FDisplayClusterRender_MeshComponent::~FDisplayClusterRender_MeshComponent()
 	MeshComponentProxy = nullptr;
 }
 
-USceneComponent* FDisplayClusterRender_MeshComponent::GetOriginComponent()
+USceneComponent* FDisplayClusterRender_MeshComponent::GetOriginComponent() const
 {
 	check(IsInGameThread());
 
@@ -113,14 +122,14 @@ void FDisplayClusterRender_MeshComponent::SetGeometryFunc(const EDisplayClusterR
 	DataFunc = InDataFunc;
 }
 
-FDisplayClusterRender_MeshComponentProxy* FDisplayClusterRender_MeshComponent::GetMeshComponentProxy_RenderThread() const
+IDisplayClusterRender_MeshComponentProxy* FDisplayClusterRender_MeshComponent::GetMeshComponentProxy_RenderThread() const
 {
 	check(IsInRenderingThread());
 
 	return MeshComponentProxy;
 }
 
-UStaticMeshComponent* FDisplayClusterRender_MeshComponent::GetStaticMeshComponent()
+UStaticMeshComponent* FDisplayClusterRender_MeshComponent::GetStaticMeshComponent() const
 {
 	check(IsInGameThread());
 
@@ -132,6 +141,16 @@ UStaticMeshComponent* FDisplayClusterRender_MeshComponent::GetStaticMeshComponen
 	return nullptr;
 }
 
+const UStaticMesh* FDisplayClusterRender_MeshComponent::GetStaticMesh() const
+{
+	if (StaticMeshRef.IsValid() && !StaticMeshRef.IsStale())
+	{
+		return StaticMeshRef.Get();
+	}
+
+	return nullptr;
+}
+
 const FStaticMeshLODResources* FDisplayClusterRender_MeshComponent::GetStaticMeshComponentLODResources(int32 InLODIndex) const
 {
 	check(IsInGameThread());
@@ -144,7 +163,7 @@ const FStaticMeshLODResources* FDisplayClusterRender_MeshComponent::GetStaticMes
 	return nullptr;
 }
 
-void FDisplayClusterRender_MeshComponent::AssignStaticMeshComponentRefs(UStaticMeshComponent* InStaticMeshComponent, USceneComponent* InOriginComponent, int32 InLODIndex)
+void FDisplayClusterRender_MeshComponent::AssignStaticMeshComponentRefs(UStaticMeshComponent* InStaticMeshComponent, const FDisplayClusterMeshUVs& InUVs, USceneComponent* InOriginComponent, int32 InLODIndex)
 {
 	check(IsInGameThread());
 
@@ -153,6 +172,8 @@ void FDisplayClusterRender_MeshComponent::AssignStaticMeshComponentRefs(UStaticM
 	StaticMeshComponentRef.SetStaticMeshComponentRef(InStaticMeshComponent);
 	ProceduralMeshComponentRef.ResetProceduralMeshComponentRef();
 
+	StaticMeshRef.Reset();
+
 	// Set source geometry type
 	GeometrySource = EDisplayClusterRender_MeshComponentGeometrySource::StaticMeshComponentRef;
 
@@ -160,8 +181,9 @@ void FDisplayClusterRender_MeshComponent::AssignStaticMeshComponentRefs(UStaticM
 	const FStaticMeshLODResources* StaticMeshLODResources = ImplGetStaticMeshComponentLODResources(InStaticMeshComponent, InLODIndex);
 	if (StaticMeshLODResources != nullptr)
 	{
+		const FString SourceGeometryName = InStaticMeshComponent->GetFName().ToString();
 		// Send geometry to proxy
-		ImplUpdateMeshComponentProxyData(MeshComponentProxy, new FDisplayClusterRender_MeshComponentProxyData(DataFunc, *StaticMeshLODResources));
+		ImplUpdateMeshComponentProxyData(MeshComponentProxy, new FDisplayClusterRender_MeshComponentProxyData(SourceGeometryName, DataFunc, *StaticMeshLODResources, InUVs));
 	}
 	else
 	{
@@ -170,7 +192,7 @@ void FDisplayClusterRender_MeshComponent::AssignStaticMeshComponentRefs(UStaticM
 	}
 }
 
-UProceduralMeshComponent* FDisplayClusterRender_MeshComponent::GetProceduralMeshComponent()
+UProceduralMeshComponent* FDisplayClusterRender_MeshComponent::GetProceduralMeshComponent() const
 {
 	check(IsInGameThread());
 
@@ -194,7 +216,7 @@ const FProcMeshSection* FDisplayClusterRender_MeshComponent::GetProceduralMeshCo
 	return nullptr;
 }
 
-void FDisplayClusterRender_MeshComponent::AssignProceduralMeshComponentRefs(UProceduralMeshComponent* InProceduralMeshComponent, USceneComponent* InOriginComponent, const int32 InSectionIndex)
+void FDisplayClusterRender_MeshComponent::AssignProceduralMeshComponentRefs(UProceduralMeshComponent* InProceduralMeshComponent, const FDisplayClusterMeshUVs& InUVs, USceneComponent* InOriginComponent, const int32 InSectionIndex)
 {
 	check(IsInGameThread());
 
@@ -203,6 +225,8 @@ void FDisplayClusterRender_MeshComponent::AssignProceduralMeshComponentRefs(UPro
 	StaticMeshComponentRef.ResetStaticMeshComponentRef();
 	ProceduralMeshComponentRef.SetProceduralMeshComponentRef(InProceduralMeshComponent);
 
+	StaticMeshRef.Reset();
+
 	// Set source geometry type
 	GeometrySource = EDisplayClusterRender_MeshComponentGeometrySource::ProceduralMeshComponentRef;
 
@@ -210,8 +234,9 @@ void FDisplayClusterRender_MeshComponent::AssignProceduralMeshComponentRefs(UPro
 	const FProcMeshSection* ProcMeshSection = ImplGetProceduralMeshComponentSection(InProceduralMeshComponent, InSectionIndex);
 	if (ProcMeshSection != nullptr)
 	{
+		const FString SourceGeometryName = InProceduralMeshComponent->GetFName().ToString();
 		// Send geometry to proxy
-		ImplUpdateMeshComponentProxyData(MeshComponentProxy, new FDisplayClusterRender_MeshComponentProxyData(DataFunc, *ProcMeshSection));
+		ImplUpdateMeshComponentProxyData(MeshComponentProxy, new FDisplayClusterRender_MeshComponentProxyData(SourceGeometryName, DataFunc, *ProcMeshSection, InUVs));
 	}
 	else
 	{
@@ -220,7 +245,7 @@ void FDisplayClusterRender_MeshComponent::AssignProceduralMeshComponentRefs(UPro
 	}
 }
 
-void FDisplayClusterRender_MeshComponent::AssignProceduralMeshSection(const FProcMeshSection& InProcMeshSection)
+void FDisplayClusterRender_MeshComponent::AssignProceduralMeshSection(const FProcMeshSection& InProcMeshSection, const FDisplayClusterMeshUVs& InUVs)
 {
 	check(IsInGameThread());
 
@@ -229,14 +254,17 @@ void FDisplayClusterRender_MeshComponent::AssignProceduralMeshSection(const FPro
 	StaticMeshComponentRef.ResetStaticMeshComponentRef();
 	ProceduralMeshComponentRef.ResetProceduralMeshComponentRef();
 
+	StaticMeshRef.Reset();
+
 	// Set source geometry type
 	GeometrySource = EDisplayClusterRender_MeshComponentGeometrySource::ProceduralMeshSection;
 
+	const FString SourceGeometryName(TEXT("ProcMeshSection"));
 	// Send geometry to proxy
-	ImplUpdateMeshComponentProxyData(MeshComponentProxy, new FDisplayClusterRender_MeshComponentProxyData(DataFunc, InProcMeshSection));
+	ImplUpdateMeshComponentProxyData(MeshComponentProxy, new FDisplayClusterRender_MeshComponentProxyData(SourceGeometryName, DataFunc, InProcMeshSection, InUVs));
 }
 
-void FDisplayClusterRender_MeshComponent::AssignStaticMesh(const UStaticMesh* InStaticMesh, int32 InLODIndex)
+void FDisplayClusterRender_MeshComponent::AssignStaticMesh(const UStaticMesh* InStaticMesh, const FDisplayClusterMeshUVs& InUVs, int32 InLODIndex)
 {
 	check(IsInGameThread());
 
@@ -245,6 +273,8 @@ void FDisplayClusterRender_MeshComponent::AssignStaticMesh(const UStaticMesh* In
 	StaticMeshComponentRef.ResetStaticMeshComponentRef();
 	ProceduralMeshComponentRef.ResetProceduralMeshComponentRef();
 
+	StaticMeshRef = InStaticMesh;
+
 	// Set source geometry type
 	GeometrySource = EDisplayClusterRender_MeshComponentGeometrySource::StaticMeshAsset;
 
@@ -252,8 +282,9 @@ void FDisplayClusterRender_MeshComponent::AssignStaticMesh(const UStaticMesh* In
 	const FStaticMeshLODResources* StaticMeshLODResources = ImplGetStaticMeshLODResources(InStaticMesh, InLODIndex);
 	if (StaticMeshLODResources != nullptr)
 	{
+		const FString SourceGeometryName = InStaticMesh->GetFName().ToString();
 		// Send geometry to proxy
-		ImplUpdateMeshComponentProxyData(MeshComponentProxy, new FDisplayClusterRender_MeshComponentProxyData(DataFunc, *StaticMeshLODResources));
+		ImplUpdateMeshComponentProxyData(MeshComponentProxy, new FDisplayClusterRender_MeshComponentProxyData(SourceGeometryName, DataFunc, *StaticMeshLODResources, InUVs));
 	}
 	else
 	{
@@ -271,13 +302,16 @@ void FDisplayClusterRender_MeshComponent::AssignMeshGeometry(const FDisplayClust
 	StaticMeshComponentRef.ResetStaticMeshComponentRef();
 	ProceduralMeshComponentRef.ResetProceduralMeshComponentRef();
 
+	StaticMeshRef.Reset();
+
 	// Set source geometry type
 	GeometrySource = EDisplayClusterRender_MeshComponentGeometrySource::MeshGeometry;
 
 	if (InMeshGeometry != nullptr)
 	{
+		const FString SourceGeometryName(TEXT("nDCRender_MeshGeometry"));
 		// Send geometry to proxy
-		ImplUpdateMeshComponentProxyData(MeshComponentProxy, new FDisplayClusterRender_MeshComponentProxyData(DataFunc, *InMeshGeometry));
+		ImplUpdateMeshComponentProxyData(MeshComponentProxy, new FDisplayClusterRender_MeshComponentProxyData(SourceGeometryName, DataFunc, *InMeshGeometry));
 	}
 	else
 	{
@@ -286,6 +320,18 @@ void FDisplayClusterRender_MeshComponent::AssignMeshGeometry(const FDisplayClust
 	}
 }
 
+void FDisplayClusterRender_MeshComponent::AssignMeshGeometry_RenderThread(const FDisplayClusterRender_MeshGeometry* InMeshGeometry, const EDisplayClusterRender_MeshComponentProxyDataFunc InDataFunc) const
+{
+	check(IsInRenderingThread());
+
+	if (InMeshGeometry != nullptr)
+	{
+		const FString SourceGeometryName(TEXT("nDCRender_MeshGeometry_RenderThread"));
+		// Send geometry to proxy
+		ImplUpdateMeshComponentProxyData(MeshComponentProxy, new FDisplayClusterRender_MeshComponentProxyData(SourceGeometryName, InDataFunc, *InMeshGeometry));
+	}
+}
+
 void FDisplayClusterRender_MeshComponent::ReleaseMeshComponent()
 {
 	check(IsInGameThread());
@@ -295,6 +341,8 @@ void FDisplayClusterRender_MeshComponent::ReleaseMeshComponent()
 	StaticMeshComponentRef.ResetStaticMeshComponentRef();
 	ProceduralMeshComponentRef.ResetProceduralMeshComponentRef();
 
+	StaticMeshRef.Reset();
+
 	// Set source geometry type
 	GeometrySource = EDisplayClusterRender_MeshComponentGeometrySource::Disabled;
 
diff --git a/Engine/Plugins/Runtime/nDisplay/Source/DisplayCluster/Private/Render/Containers/DisplayClusterRender_MeshComponent.h b/Engine/Plugins/Runtime/nDisplay/Source/DisplayCluster/Private/Render/Containers/DisplayClusterRender_MeshComponent.h
new file mode 100644
--- /dev/null
+++ b/Engine/Plugins/Runtime/nDisplay/Source/DisplayCluster/Private/Render/Containers/DisplayClusterRender_MeshComponent.h
@@ -0,0 +1,64 @@
+// Copyright Epic Games, Inc. All Rights Reserved.
+
+#pragma once
+
+#include "CoreMinimal.h"
+#include "Render/Containers/IDisplayClusterRender_MeshComponent.h"
+
+#include "Render/Containers/DisplayClusterRender_StaticMeshComponentRef.h"
+#include "Render/Containers/DisplayClusterRender_ProceduralMeshComponentRef.h"
+
+class FDisplayClusterRender_MeshComponentProxy;
+
+class FDisplayClusterRender_MeshComponent
+	: public IDisplayClusterRender_MeshComponent
+{
+public:
+	FDisplayClusterRender_MeshComponent();
+	virtual ~FDisplayClusterRender_MeshComponent();
+
+public:
+	virtual void AssignStaticMeshComponentRefs(UStaticMeshComponent* InStaticMeshComponent, const FDisplayClusterMeshUVs& InUVs, USceneComponent* InOriginComponent = nullptr, int32 InLODIndex = 0) override;
+	virtual void AssignProceduralMeshComponentRefs(UProceduralMeshComponent* InProceduralMeshComponent, const FDisplayClusterMeshUVs& InUVs, USceneComponent* InOriginComponent = nullptr, const int32 InSectionIndex = 0) override;
+	virtual void AssignProceduralMeshSection(const FProcMeshSection& InProcMeshSection, const FDisplayClusterMeshUVs& InUVs) override;
+	virtual void AssignStaticMesh(const UStaticMesh* InStaticMesh, const FDisplayClusterMeshUVs& InUVs, int32 InLODIndex = 0) override;
+	virtual void AssignMeshGeometry(const FDisplayClusterRender_MeshGeometry* InMeshGeometry) override;
+	
+	virtual void AssignMeshGeometry_RenderThread(const FDisplayClusterRender_MeshGeometry* InMeshGeometry, const EDisplayClusterRender_MeshComponentProxyDataFunc InDataFunc) const override;
+
+	virtual void ReleaseMeshComponent() override;
+	virtual void ReleaseProxyGeometry() override;
+
+	virtual USceneComponent* GetOriginComponent() const override;
+	virtual UStaticMeshComponent* GetStaticMeshComponent() const override;
+	virtual const UStaticMesh* GetStaticMesh() const override;
+	virtual UProceduralMeshComponent* GetProceduralMeshComponent() const override;
+	virtual const FStaticMeshLODResources* GetStaticMeshComponentLODResources(int32 InLODIndex = 0) const override;
+	virtual const FProcMeshSection* GetProceduralMeshComponentSection(const int32 SectionIndex = 0) const override;
+
+	virtual IDisplayClusterRender_MeshComponentProxy* GetMeshComponentProxy_RenderThread() const override;
+
+	virtual void SetGeometryFunc(const EDisplayClusterRender_MeshComponentProxyDataFunc InDataFunc) override;
+
+	virtual bool IsMeshComponentRefGeometryDirty() const override;
+	virtual void MarkMeshComponentRefGeometryDirty() const override;
+	virtual void ResetMeshComponentRefGeometryDirty() const override;
+	virtual bool EqualsMeshComponentName(const FName& InMeshComponentName) const override;
+
+	virtual EDisplayClusterRender_MeshComponentGeometrySource GetGeometrySource() const override;
+
+private:
+	EDisplayClusterRender_MeshComponentGeometrySource GeometrySource = EDisplayClusterRender_MeshComponentGeometrySource::Disabled;
+	EDisplayClusterRender_MeshComponentProxyDataFunc  DataFunc = EDisplayClusterRender_MeshComponentProxyDataFunc::Disabled;
+
+	// Reference containers:
+	FDisplayClusterSceneComponentRef                 OriginComponentRef;
+	FDisplayClusterRender_StaticMeshComponentRef     StaticMeshComponentRef;
+	FDisplayClusterRender_ProceduralMeshComponentRef ProceduralMeshComponentRef;
+
+	// store ref to the static mesh
+	TWeakObjectPtr<const UStaticMesh> StaticMeshRef;
+
+private:
+	FDisplayClusterRender_MeshComponentProxy* MeshComponentProxy = nullptr;
+};
diff --git a/Engine/Plugins/Runtime/nDisplay/Source/DisplayCluster/Private/Render/Containers/DisplayClusterRender_MeshComponentProxy.cpp b/Engine/Plugins/Runtime/nDisplay/Source/DisplayCluster/Private/Render/Containers/DisplayClusterRender_MeshComponentProxy.cpp
--- a/Engine/Plugins/Runtime/nDisplay/Source/DisplayCluster/Private/Render/Containers/DisplayClusterRender_MeshComponentProxy.cpp
+++ b/Engine/Plugins/Runtime/nDisplay/Source/DisplayCluster/Private/Render/Containers/DisplayClusterRender_MeshComponentProxy.cpp
@@ -16,13 +16,20 @@ FDisplayClusterRender_MeshComponentProxy::FDisplayClusterRender_MeshComponentPro
 
 FDisplayClusterRender_MeshComponentProxy::~FDisplayClusterRender_MeshComponentProxy()
 {
-	Release_RenderThread();
+	check(IsInRenderingThread());
+
+	ImplRelease_RenderThread();
 }
 
 void FDisplayClusterRender_MeshComponentProxy::Release_RenderThread()
 {
 	check(IsInRenderingThread());
 
+	ImplRelease_RenderThread();
+}
+
+void FDisplayClusterRender_MeshComponentProxy::ImplRelease_RenderThread()
+{
 	VertexBufferRHI.SafeRelease();
 	IndexBufferRHI.SafeRelease();
 
@@ -30,11 +37,16 @@ void FDisplayClusterRender_MeshComponentProxy::Release_RenderThread()
 	NumVertices = 0;
 }
 
-bool FDisplayClusterRender_MeshComponentProxy::BeginRender_RenderThread(FRHICommandListImmediate& RHICmdList, FGraphicsPipelineStateInitializer& GraphicsPSOInit) const
+bool FDisplayClusterRender_MeshComponentProxy::IsEnabled_RenderThread() const
 {
 	check(IsInRenderingThread());
 
-	if (IsValid_RenderThread())
+	return NumTriangles > 0 && NumVertices > 0 && VertexBufferRHI.IsValid() && IndexBufferRHI.IsValid();
+}
+
+bool FDisplayClusterRender_MeshComponentProxy::BeginRender_RenderThread(FRHICommandListImmediate& RHICmdList, FGraphicsPipelineStateInitializer& GraphicsPSOInit) const
+{
+	if (IsEnabled_RenderThread())
 	{
 		GraphicsPSOInit.BoundShaderState.VertexDeclarationRHI = GDisplayClusterMeshVertexDeclaration.VertexDeclarationRHI;
 		return true;
@@ -45,9 +57,7 @@ bool FDisplayClusterRender_MeshComponentProxy::BeginRender_RenderThread(FRHIComm
 
 bool  FDisplayClusterRender_MeshComponentProxy::FinishRender_RenderThread(FRHICommandListImmediate& RHICmdList) const
 {
-	check(IsInRenderingThread());
-
-	if (IsValid_RenderThread())
+	if (IsEnabled_RenderThread())
 	{
 		// Support update
 		RHICmdList.SetStreamSource(0, VertexBufferRHI, 0);
@@ -62,7 +72,7 @@ void FDisplayClusterRender_MeshComponentProxy::UpdateRHI_RenderThread(FRHIComman
 {
 	check(IsInRenderingThread());
 
-	Release_RenderThread();
+	ImplRelease_RenderThread();
 
 	if (InMeshData && InMeshData->IsValid())
 	{
diff --git a/Engine/Plugins/Runtime/nDisplay/Source/DisplayCluster/Private/Render/Containers/DisplayClusterRender_MeshComponentProxy.h b/Engine/Plugins/Runtime/nDisplay/Source/DisplayCluster/Private/Render/Containers/DisplayClusterRender_MeshComponentProxy.h
new file mode 100644
--- /dev/null
+++ b/Engine/Plugins/Runtime/nDisplay/Source/DisplayCluster/Private/Render/Containers/DisplayClusterRender_MeshComponentProxy.h
@@ -0,0 +1,32 @@
+// Copyright Epic Games, Inc. All Rights Reserved.
+
+#pragma once
+#include "CoreMinimal.h"
+#include "Render/Containers/IDisplayClusterRender_MeshComponentProxy.h"
+
+class FDisplayClusterRender_MeshComponentProxy
+	: public IDisplayClusterRender_MeshComponentProxy
+{
+public:
+	FDisplayClusterRender_MeshComponentProxy();
+	virtual ~FDisplayClusterRender_MeshComponentProxy();
+
+public:
+	virtual bool IsEnabled_RenderThread() const override;
+	virtual bool BeginRender_RenderThread(FRHICommandListImmediate& RHICmdList, FGraphicsPipelineStateInitializer& GraphicsPSOInit) const override;
+	virtual bool FinishRender_RenderThread(FRHICommandListImmediate& RHICmdList) const override;
+
+	void UpdateRHI_RenderThread(FRHICommandListImmediate& RHICmdList, class FDisplayClusterRender_MeshComponentProxyData* InMeshData);
+	void Release_RenderThread();
+
+private:
+	void ImplRelease_RenderThread();
+
+private:
+	/* RenderThread resources */
+	FBufferRHIRef VertexBufferRHI;
+	FBufferRHIRef IndexBufferRHI;
+
+	uint32 NumTriangles = 0;
+	uint32 NumVertices = 0;
+};
diff --git a/Engine/Plugins/Runtime/nDisplay/Source/DisplayCluster/Private/Render/Containers/DisplayClusterRender_MeshComponentProxyData.cpp b/Engine/Plugins/Runtime/nDisplay/Source/DisplayCluster/Private/Render/Containers/DisplayClusterRender_MeshComponentProxyData.cpp
--- a/Engine/Plugins/Runtime/nDisplay/Source/DisplayCluster/Private/Render/Containers/DisplayClusterRender_MeshComponentProxyData.cpp
+++ b/Engine/Plugins/Runtime/nDisplay/Source/DisplayCluster/Private/Render/Containers/DisplayClusterRender_MeshComponentProxyData.cpp
@@ -12,7 +12,8 @@
 //*************************************************************************
 //* FDisaplyClusterMeshComponentProxyData
 //*************************************************************************
-FDisplayClusterRender_MeshComponentProxyData::FDisplayClusterRender_MeshComponentProxyData(const EDisplayClusterRender_MeshComponentProxyDataFunc InDataFunc, const FStaticMeshLODResources& InStaticMeshLODResource, const int32 InUVChromakeyIndex)
+FDisplayClusterRender_MeshComponentProxyData::FDisplayClusterRender_MeshComponentProxyData(const FString& InSourceGeometryName, const EDisplayClusterRender_MeshComponentProxyDataFunc InDataFunc, const FStaticMeshLODResources& InStaticMeshLODResource, const FDisplayClusterMeshUVs& InUVs)
+	: SourceGeometryName(InSourceGeometryName)
 {
 	check(IsInGameThread());
 
@@ -20,8 +21,15 @@ FDisplayClusterRender_MeshComponentProxyData::FDisplayClusterRender_MeshComponen
 	const FStaticMeshVertexBuffer& VertexBuffer = InStaticMeshLODResource.VertexBuffers.StaticMeshVertexBuffer;
 	const FRawStaticIndexBuffer& IndexBuffer = InStaticMeshLODResource.IndexBuffer;
 
-	uint32 UVBaseIndex = 0;
-	uint32 UVChromakeyIndex = (InUVChromakeyIndex < InStaticMeshLODResource.GetNumTexCoords()) ? InUVChromakeyIndex : 0;
+	// Max tex coords used in StaticMesh
+	const int32 NumTexCoords = InStaticMeshLODResource.GetNumTexCoords();
+
+	// Remap BaseUVs
+	// By default, StaticMesh uses index 0 for BaseUV
+	const int32 BaseUVIndex = InUVs.GetSourceGeometryUVIndex(FDisplayClusterMeshUVs::ERemapTarget::Base, NumTexCoords, 0);
+
+	// By default, StaticMesh uses index 1 for ChromakeyUV
+	const int32 ChromakeyUVIndex = InUVs.GetSourceGeometryUVIndex(FDisplayClusterMeshUVs::ERemapTarget::Chromakey, NumTexCoords, 1);
 
 	IndexBuffer.GetCopy(IndexData);
 
@@ -29,37 +37,79 @@ FDisplayClusterRender_MeshComponentProxyData::FDisplayClusterRender_MeshComponen
 	for (int32 VertexIdx = 0; VertexIdx < VertexData.Num(); VertexIdx++)
 	{
 		VertexData[VertexIdx].Position = PositionBuffer.VertexPosition(VertexIdx);
-		VertexData[VertexIdx].UV = VertexBuffer.GetVertexUV(VertexIdx, UVBaseIndex);
-		VertexData[VertexIdx].UV_Chromakey = VertexBuffer.GetVertexUV(VertexIdx, UVChromakeyIndex);
+		VertexData[VertexIdx].UV = VertexBuffer.GetVertexUV(VertexIdx, BaseUVIndex);
+		VertexData[VertexIdx].UV_Chromakey = VertexBuffer.GetVertexUV(VertexIdx, ChromakeyUVIndex);
 	}
 
 	UpdateData(InDataFunc);
 }
 
-FDisplayClusterRender_MeshComponentProxyData::FDisplayClusterRender_MeshComponentProxyData(const EDisplayClusterRender_MeshComponentProxyDataFunc InDataFunc, const FProcMeshSection& InProcMeshSection, const int32 InUVChromakeyIndex)
+FDisplayClusterRender_MeshComponentProxyData::FDisplayClusterRender_MeshComponentProxyData(const FString& InSourceGeometryName, const EDisplayClusterRender_MeshComponentProxyDataFunc InDataFunc, const FProcMeshSection& InProcMeshSection, const FDisplayClusterMeshUVs& InUVs)
+	: SourceGeometryName(InSourceGeometryName)
 {
+	if (InUVs[FDisplayClusterMeshUVs::ERemapTarget::Base] > 3)
+	{
+		// Show warning for values >3
+		UE_LOG(LogDisplayClusterRender, Warning, TEXT("MeshComponentProxyData('%s'): Base UVIndex value '%d' is invalid - ProceduralMesh support range 0..3"), *SourceGeometryName, InUVs[FDisplayClusterMeshUVs::ERemapTarget::Base]);
+	}
+
+	if (InUVs[FDisplayClusterMeshUVs::ERemapTarget::Chromakey] > 3)
+	{
+		// Show warning for values >3
+		UE_LOG(LogDisplayClusterRender, Warning, TEXT("MeshComponentProxyData('%s'): Chromakey UVIndex value '%d' is invalid - ProceduralMesh support range 0..3"), *SourceGeometryName, InUVs[FDisplayClusterMeshUVs::ERemapTarget::Chromakey]);
+	}
+
+	// Remap BaseUVs
+	// Max tex coords used in ProceduralMesh now 4
+	const int32 NumTexCoords = 4;
+
+	// By default, ProceduralMesh uses index 0 for BaseUV
+	const int32 BaseUVIndex = InUVs.GetSourceGeometryUVIndex(FDisplayClusterMeshUVs::ERemapTarget::Base, NumTexCoords, 0);
+
+	// By default, ProceduralMesh uses index 0 for ChromakeyUV
+	const int32 ChromakeyUVIndex = InUVs.GetSourceGeometryUVIndex(FDisplayClusterMeshUVs::ERemapTarget::Chromakey, NumTexCoords, 0);
+
 	VertexData.AddZeroed(InProcMeshSection.ProcVertexBuffer.Num());
 	for (int32 VertexIdx = 0; VertexIdx < VertexData.Num(); VertexIdx++)
 	{
 		const FProcMeshVertex& InVertex = InProcMeshSection.ProcVertexBuffer[VertexIdx];
 
 		VertexData[VertexIdx].Position = InVertex.Position;
-		VertexData[VertexIdx].UV = InVertex.UV0;
 
-		switch (InUVChromakeyIndex)
+		// Remap source geometry UV for Base
+		switch (BaseUVIndex)
 		{
+		default:
+		case 0:
+			VertexData[VertexIdx].UV = InVertex.UV0;
+			break;
 		case 1:
-			VertexData[VertexIdx].UV_Chromakey = InVertex.UV1;
+			VertexData[VertexIdx].UV = InVertex.UV1;
 			break;
 		case 2:
-			VertexData[VertexIdx].UV_Chromakey = InVertex.UV2;
+			VertexData[VertexIdx].UV = InVertex.UV2;
 			break;
 		case 3:
-			VertexData[VertexIdx].UV_Chromakey = InVertex.UV3;
+			VertexData[VertexIdx].UV = InVertex.UV3;
 			break;
+		}
+
+		// Remap source geometry UV for Chromakey
+		switch (ChromakeyUVIndex)
+		{
 		default:
+		case 0:
 			VertexData[VertexIdx].UV_Chromakey = InVertex.UV0;
 			break;
+		case 1:
+			VertexData[VertexIdx].UV_Chromakey = InVertex.UV1;
+			break;
+		case 2:
+			VertexData[VertexIdx].UV_Chromakey = InVertex.UV2;
+			break;
+		case 3:
+			VertexData[VertexIdx].UV_Chromakey = InVertex.UV3;
+			break;
 		}
 	}
 
@@ -73,7 +123,8 @@ FDisplayClusterRender_MeshComponentProxyData::FDisplayClusterRender_MeshComponen
 	UpdateData(InDataFunc);
 }
 
-FDisplayClusterRender_MeshComponentProxyData::FDisplayClusterRender_MeshComponentProxyData(const EDisplayClusterRender_MeshComponentProxyDataFunc InDataFunc, const FDisplayClusterRender_MeshGeometry& InMeshGeometry)
+FDisplayClusterRender_MeshComponentProxyData::FDisplayClusterRender_MeshComponentProxyData(const FString& InSourceGeometryName, const EDisplayClusterRender_MeshComponentProxyDataFunc InDataFunc, const FDisplayClusterRender_MeshGeometry& InMeshGeometry)
+	: SourceGeometryName(InSourceGeometryName)
 {
 	bool bUseChromakeyUV = InMeshGeometry.ChromakeyUV.Num() > 0;
 
@@ -103,8 +154,8 @@ void FDisplayClusterRender_MeshComponentProxyData::UpdateData(const EDisplayClus
 		{
 		case EDisplayClusterRender_MeshComponentProxyDataFunc::OutputRemapScreenSpace:
 			// Output remap require normalize mesh to screen space coords:
-			NormalizeToScreenSpace();
-			RemoveInvisibleFaces();
+			ImplNormalizeToScreenSpace();
+			ImplRemoveInvisibleFaces();
 			break;
 
 		default:
@@ -128,19 +179,11 @@ void FDisplayClusterRender_MeshComponentProxyData::UpdateData(const EDisplayClus
 
 	if (!IsValid())
 	{
-		UE_LOG(LogDisplayClusterRender, Error, TEXT("MeshComponentProxyData::UpdateData() Invalid mesh - ignored"));
+		UE_LOG(LogDisplayClusterRender, Error, TEXT("MeshComponentProxyData::UpdateData('%s') Invalid mesh - ignored"), *SourceGeometryName);
 	}
 }
 
-enum class EGeometryPlane2DAxis : uint8
-{
-	XY,
-	XZ,
-	YZ,
-	Undefined
-};
-
-void FDisplayClusterRender_MeshComponentProxyData::NormalizeToScreenSpace()
+void FDisplayClusterRender_MeshComponentProxyData::ImplNormalizeToScreenSpace()
 {
 	FBox AABBox(FVector(FLT_MAX, FLT_MAX, FLT_MAX), FVector(-FLT_MAX, -FLT_MAX, -FLT_MAX));
 
@@ -165,18 +208,20 @@ void FDisplayClusterRender_MeshComponentProxyData::NormalizeToScreenSpace()
 	);
 
 	bool bIsValidMeshAxis = true;
+
 	// Support axis rules: Z=UP(screen Y), Y=RIGHT(screen X), X=NotUsed
-	EGeometryPlane2DAxis GeometryPlane2DAxis = EGeometryPlane2DAxis::Undefined;
 	if (FMath::Abs(Size.Y) <= KINDA_SMALL_NUMBER)
 	{
-		UE_LOG(LogDisplayClusterRender, Error, TEXT("MeshComponentProxyData::NormalizeToScreenSpace(): The Y axis is used in screen space as 'x' and the distance cannot be zero."));
+		UE_LOG(LogDisplayClusterRender, Error, TEXT("MeshComponentProxyData::NormalizeToScreenSpace('%s'): The Y axis is used in screen space as 'x' and the distance cannot be zero."), *SourceGeometryName);
 		bIsValidMeshAxis = false;
-		}
+	}
+
 	if (FMath::Abs(Size.Z) <= KINDA_SMALL_NUMBER)
-		{
-		UE_LOG(LogDisplayClusterRender, Error, TEXT("MeshComponentProxyData::NormalizeToScreenSpace(): The Z axis is used in screen space as 'y' and the distance cannot be zero."));
+	{
+		UE_LOG(LogDisplayClusterRender, Error, TEXT("MeshComponentProxyData::NormalizeToScreenSpace('%s'): The Z axis is used in screen space as 'y' and the distance cannot be zero."), *SourceGeometryName);
 		bIsValidMeshAxis = false;
-		}
+	}
+
 	if (bIsValidMeshAxis)
 	{
 		// Checking for strange aspect ratio
@@ -185,7 +230,7 @@ void FDisplayClusterRender_MeshComponentProxyData::NormalizeToScreenSpace()
 		if (AspectRatio > 10)
 		{
 			// just warning, experimental
-			UE_LOG(LogDisplayClusterRender, Warning, TEXT("MeshComponentProxyData::NormalizeToScreenSpace(): Aspect ratio is to big '%d'- <%d,%d>"), AspectRatio, ScreenSize.X, ScreenSize.Y);
+			UE_LOG(LogDisplayClusterRender, Warning, TEXT("MeshComponentProxyData::NormalizeToScreenSpace('%s'): Aspect ratio is to big '%d'- <%d,%d>"), *SourceGeometryName, AspectRatio, ScreenSize.X, ScreenSize.Y);
 		}
 	}
 
@@ -197,7 +242,7 @@ void FDisplayClusterRender_MeshComponentProxyData::NormalizeToScreenSpace()
 		NumTriangles = 0;
 		NumVertices = 0;
 
-		UE_LOG(LogDisplayClusterRender, Error, TEXT("MeshComponentProxyData::NormalizeToScreenSpace() Invalid mesh. Expected mesh axis: Z=UP(screen Y), Y=RIGHT(screen X), X=NotUsed"));
+		UE_LOG(LogDisplayClusterRender, Error, TEXT("MeshComponentProxyData::NormalizeToScreenSpace('%s') Invalid mesh. Expected mesh axis: Z=UP(screen Y), Y=RIGHT(screen X), X=NotUsed"), *SourceGeometryName);
 		return;
 	}
 
@@ -213,8 +258,12 @@ void FDisplayClusterRender_MeshComponentProxyData::NormalizeToScreenSpace()
 	}
 }
 
-void FDisplayClusterRender_MeshComponentProxyData::RemoveInvisibleFaces()
+void FDisplayClusterRender_MeshComponentProxyData::ImplRemoveInvisibleFaces()
 {
+	// The geometry is created by the 3D artist and is sometimes incorrect. 
+	// For example, in the OutputRemap post-process, it is necessary that all UVs be in the range 0..1. 
+	// For visual validation, all points outside the 0..1 range are excluded during geometry loading when called function RemoveInvisibleFaces().
+
 	TArray<uint32> VisibleIndexData;
 	int32 RemovedFacesNum = 0;
 	const int32 FacesNum = IndexData.Num() / 3;
@@ -246,7 +295,7 @@ void FDisplayClusterRender_MeshComponentProxyData::RemoveInvisibleFaces()
 
 	if (RemovedFacesNum)
 	{
-		UE_LOG(LogDisplayClusterRender, Error, TEXT("MeshComponentProxyData::RemoveInvisibleFaces() Removed %d/%d faces"), RemovedFacesNum, FacesNum);
+		UE_LOG(LogDisplayClusterRender, Error, TEXT("MeshComponentProxyData::RemoveInvisibleFaces('%s') Removed %d/%d faces"), *SourceGeometryName, RemovedFacesNum, FacesNum);
 	}
 }
 
diff --git a/Engine/Plugins/Runtime/nDisplay/Source/DisplayCluster/Private/Render/Containers/DisplayClusterRender_MeshComponentProxyData.h b/Engine/Plugins/Runtime/nDisplay/Source/DisplayCluster/Private/Render/Containers/DisplayClusterRender_MeshComponentProxyData.h
new file mode 100644
--- /dev/null
+++ b/Engine/Plugins/Runtime/nDisplay/Source/DisplayCluster/Private/Render/Containers/DisplayClusterRender_MeshComponentProxyData.h
@@ -0,0 +1,93 @@
+// Copyright Epic Games, Inc. All Rights Reserved.
+
+#pragma once
+
+#include "CoreMinimal.h"
+#include "Render/Containers/DisplayClusterRender_MeshComponentTypes.h"
+
+
+class FDisplayClusterRender_MeshGeometry;
+struct FStaticMeshLODResources;
+struct FProcMeshSection;
+
+class FDisplayClusterRender_MeshComponentProxyData
+{
+public:
+	/**
+	* Initialize geometry data from static mesh geometry
+	*
+	* @param InDataFunc              - Geometry function type
+	* @param InStaticMeshLODResource - Source of geometry
+	* @param InUVs                   - Map source geometry UVs to DCMeshComponent UVs
+	*/
+	FDisplayClusterRender_MeshComponentProxyData(const FString& InSourceGeometryName, const EDisplayClusterRender_MeshComponentProxyDataFunc InDataFunc, const FStaticMeshLODResources& InStaticMeshLODResource, const FDisplayClusterMeshUVs& InUVs);
+
+	/**
+	* Initialize geometry data from procedural mesh geometry
+	*
+	* @param InDataFunc        - Geometry function type
+	* @param InProcMeshSection - Source of geometry
+	* @param InUVs             - Map source geometry UVs to DCMeshComponent UVs
+	*/
+	FDisplayClusterRender_MeshComponentProxyData(const FString& InSourceGeometryName, const EDisplayClusterRender_MeshComponentProxyDataFunc InDataFunc, const FProcMeshSection& InProcMeshSection, const FDisplayClusterMeshUVs& InUVs);
+
+	/**
+	* Initialize geometry data from nDisplay geometry container
+	*
+	* @param InDataFunc     - Geometry function type
+	* @param InMeshGeometry - Source of geometry
+	*/
+	FDisplayClusterRender_MeshComponentProxyData(const FString& InSourceGeometryName, const EDisplayClusterRender_MeshComponentProxyDataFunc InDataFunc, const FDisplayClusterRender_MeshGeometry& InMeshGeometry);
+
+	const TArray<uint32>& GetIndexData() const
+	{ return IndexData; }
+
+	const TArray<FDisplayClusterMeshVertex>& GetVertexData() const
+	{ return VertexData; }
+
+	uint32 GetNumTriangles() const
+	{ return NumTriangles; }
+
+	uint32 GetNumVertices() const
+	{ return NumVertices; }
+
+	bool IsValid() const
+	{
+		return NumTriangles > 0 && NumVertices > 0 && IndexData.Num() > 0 && VertexData.Num() > 0;
+	}
+
+private:
+	/**
+	* Update geometry data with specified func
+	*
+	* @param InDataFunc     - Geometry function type
+	*/
+	void UpdateData(const EDisplayClusterRender_MeshComponentProxyDataFunc InDataFunc);
+
+	/**
+	* Normalize geometry points to screen space
+	* [OutputRemapScreenSpace]
+	*/
+	void ImplNormalizeToScreenSpace();
+
+	/**
+	* Visual check: remove faces with UV off-screen
+	* [OutputRemapScreenSpace]
+	*/
+	void ImplRemoveInvisibleFaces();
+
+private:
+	bool IsFaceVisible(int32 Face);
+	bool IsUVVisible(int32 UVIndex);
+
+public:
+	// Purpose of logging
+	const FString SourceGeometryName;
+
+private:
+
+	TArray<uint32> IndexData;
+	TArray<FDisplayClusterMeshVertex> VertexData;
+	uint32 NumTriangles = 0;
+	uint32 NumVertices = 0;
+};
diff --git a/Engine/Plugins/Runtime/nDisplay/Source/DisplayCluster/Private/Render/DisplayClusterRenderManager.cpp b/Engine/Plugins/Runtime/nDisplay/Source/DisplayCluster/Private/Render/DisplayClusterRenderManager.cpp
--- a/Engine/Plugins/Runtime/nDisplay/Source/DisplayCluster/Private/Render/DisplayClusterRenderManager.cpp
+++ b/Engine/Plugins/Runtime/nDisplay/Source/DisplayCluster/Private/Render/DisplayClusterRenderManager.cpp
@@ -29,6 +29,8 @@
 #include "Render/Projection/IDisplayClusterProjectionPolicy.h"
 #include "Render/Synchronization/IDisplayClusterRenderSyncPolicyFactory.h"
 
+#include "Render/Containers/DisplayClusterRender_MeshComponent.h"
+
 #include "Render/Presentation/DisplayClusterPresentationNative.h"
 
 #include "Render/Synchronization/DisplayClusterRenderSyncPolicyFactoryInternal.h"
@@ -476,6 +478,11 @@ void FDisplayClusterRenderManager::GetRegisteredPostProcess(TArray<FString>& Out
 	PostProcessFactories.GetKeys(OutPostProcessIDs);
 }
 
+TSharedPtr<IDisplayClusterRender_MeshComponent, ESPMode::ThreadSafe> FDisplayClusterRenderManager::CreateMeshComponent() const
+{
+	return MakeShared<FDisplayClusterRender_MeshComponent, ESPMode::ThreadSafe>();
+}
+
 IDisplayClusterViewportManager* FDisplayClusterRenderManager::GetViewportManager() const
 {
 	ADisplayClusterRootActor* RootActor = GDisplayCluster->GetGameMgr()->GetRootActor();
diff --git a/Engine/Plugins/Runtime/nDisplay/Source/DisplayCluster/Private/Render/DisplayClusterRenderManager.h b/Engine/Plugins/Runtime/nDisplay/Source/DisplayCluster/Private/Render/DisplayClusterRenderManager.h
--- a/Engine/Plugins/Runtime/nDisplay/Source/DisplayCluster/Private/Render/DisplayClusterRenderManager.h
+++ b/Engine/Plugins/Runtime/nDisplay/Source/DisplayCluster/Private/Render/DisplayClusterRenderManager.h
@@ -59,6 +59,7 @@ public:
 	virtual bool UnregisterPostProcessFactory(const FString& InPostProcessType) override;
 	virtual TSharedPtr<IDisplayClusterPostProcessFactory> GetPostProcessFactory(const FString& InPostProcessType) override;
 	virtual void GetRegisteredPostProcess(TArray<FString>& OutPostProcessIDs) const override;
+	virtual TSharedPtr<IDisplayClusterRender_MeshComponent, ESPMode::ThreadSafe> CreateMeshComponent() const override;
 
 	virtual IDisplayClusterViewportManager* GetViewportManager() const override;
 
diff --git a/Engine/Plugins/Runtime/nDisplay/Source/DisplayCluster/Private/Render/Viewport/Configuration/DisplayClusterViewportConfigurationBase.cpp b/Engine/Plugins/Runtime/nDisplay/Source/DisplayCluster/Private/Render/Viewport/Configuration/DisplayClusterViewportConfigurationBase.cpp
--- a/Engine/Plugins/Runtime/nDisplay/Source/DisplayCluster/Private/Render/Viewport/Configuration/DisplayClusterViewportConfigurationBase.cpp
+++ b/Engine/Plugins/Runtime/nDisplay/Source/DisplayCluster/Private/Render/Viewport/Configuration/DisplayClusterViewportConfigurationBase.cpp
@@ -9,6 +9,13 @@
 
 #include "DisplayClusterViewportConfigurationHelpers.h"
 #include "DisplayClusterConfigurationTypes.h"
+
+#include "DisplayClusterRootActor.h"
+
+#include "Engine/StaticMesh.h"
+#include "ProceduralMeshComponent.h"
+
+
 ///////////////////////////////////////////////////////////////////
 // FDisplayClusterViewportConfigurationBase
 ///////////////////////////////////////////////////////////////////
@@ -123,6 +130,28 @@ void FDisplayClusterViewportConfigurationBase::UpdateClusterNodePostProcess(cons
 						PPManager->GetOutputRemap()->UpdateConfiguration_StaticMesh(OutputRemapCfg.StaticMesh);
 						break;
 
+					case EDisplayClusterConfigurationFramePostProcess_OutputRemapSource::MeshComponent:
+						if (!OutputRemapCfg.MeshComponentName.IsEmpty())
+						{
+							// Get the StaticMeshComponent
+							UStaticMeshComponent* StaticMeshComponent = RootActor.GetComponentByName<UStaticMeshComponent>(OutputRemapCfg.MeshComponentName);
+							if (StaticMeshComponent != nullptr)
+							{
+								PPManager->GetOutputRemap()->UpdateConfiguration_StaticMeshComponent(StaticMeshComponent);
+							}
+							else
+							{
+								// Get the procedural mesh component
+								UProceduralMeshComponent* ProceduralMeshComponent = RootActor.GetComponentByName<UProceduralMeshComponent>(OutputRemapCfg.MeshComponentName);
+								PPManager->GetOutputRemap()->UpdateConfiguration_ProceduralMeshComponent(ProceduralMeshComponent);
+							}
+						}
+						else
+						{
+							PPManager->GetOutputRemap()->UpdateConfiguration_Disabled();
+						}
+						break;
+
 					case EDisplayClusterConfigurationFramePostProcess_OutputRemapSource::ExternalFile:
 						PPManager->GetOutputRemap()->UpdateConfiguration_ExternalFile(OutputRemapCfg.ExternalFile);
 						break;
diff --git a/Engine/Plugins/Runtime/nDisplay/Source/DisplayCluster/Private/Render/Viewport/Containers/DisplayClusterViewportProxyData.h b/Engine/Plugins/Runtime/nDisplay/Source/DisplayCluster/Private/Render/Viewport/Containers/DisplayClusterViewportProxyData.h
--- a/Engine/Plugins/Runtime/nDisplay/Source/DisplayCluster/Private/Render/Viewport/Containers/DisplayClusterViewportProxyData.h
+++ b/Engine/Plugins/Runtime/nDisplay/Source/DisplayCluster/Private/Render/Viewport/Containers/DisplayClusterViewportProxyData.h
@@ -13,7 +13,7 @@ class FDisplayClusterRenderTargetResource;
 class FDisplayClusterViewport;
 class FDisplayClusterViewportProxy;
 class FDisplayClusterViewport_Context;
-class FDisplayClusterRender_MeshComponent;
+class IDisplayClusterRender_MeshComponent;
 
 //
 // Container for data exchange game->render threads
@@ -29,7 +29,7 @@ public:
 private:
 	FDisplayClusterViewportProxy* DstViewportProxy;
 
-	TSharedPtr<FDisplayClusterRender_MeshComponent, ESPMode::ThreadSafe> RemapMesh;
+	TSharedPtr<IDisplayClusterRender_MeshComponent, ESPMode::ThreadSafe> RemapMesh;
 
 	FDisplayClusterViewport_OverscanSettings     OverscanSettings;
 
diff --git a/Engine/Plugins/Runtime/nDisplay/Source/DisplayCluster/Private/Render/Viewport/Containers/DisplayClusterViewportRemap.cpp b/Engine/Plugins/Runtime/nDisplay/Source/DisplayCluster/Private/Render/Viewport/Containers/DisplayClusterViewportRemap.cpp
--- a/Engine/Plugins/Runtime/nDisplay/Source/DisplayCluster/Private/Render/Viewport/Containers/DisplayClusterViewportRemap.cpp
+++ b/Engine/Plugins/Runtime/nDisplay/Source/DisplayCluster/Private/Render/Viewport/Containers/DisplayClusterViewportRemap.cpp
@@ -3,10 +3,12 @@
 #include "Render/Viewport/Containers/DisplayClusterViewportRemap.h"
 #include "DisplayClusterConfigurationTypes_ViewportRemap.h"
 
+#include "Misc/DisplayClusterGlobals.h"
+#include "Render/IDisplayClusterRenderManager.h"
+
 #include "Render/Viewport/DisplayClusterViewport.h"
 #include "Render/Viewport/DisplayClusterViewportManager.h"
 #include "Render/Viewport/RenderFrame/DisplayClusterRenderFrameSettings.h"
-#include "Render/Containers/DisplayClusterRender_MeshComponent.h"
 #include "Render/Containers/DisplayClusterRender_MeshGeometry.h"
 
 bool FDisplayClusterViewportRemap::UpdateConfiguration(const FDisplayClusterViewport& InViewport, const FDisplayClusterConfigurationViewport_Remap& InRemapConfiguration)
@@ -94,7 +96,7 @@ void FDisplayClusterViewportRemap::Empty()
 	RemapMesh.Reset();
 }
 
-TSharedPtr<FDisplayClusterRender_MeshComponent, ESPMode::ThreadSafe> FDisplayClusterViewportRemap::GetRemapMesh() const
+TSharedPtr<IDisplayClusterRender_MeshComponent, ESPMode::ThreadSafe> FDisplayClusterViewportRemap::GetRemapMesh() const
 {
 	check(IsInGameThread());
 
@@ -121,7 +123,7 @@ void FDisplayClusterViewportRemap::ImplUpdateRemapMesh(const FDisplayClusterView
 		if (!RemapMesh.IsValid())
 		{
 			// Create new viewport remap mesh component
-			RemapMesh = MakeShared<FDisplayClusterRender_MeshComponent, ESPMode::ThreadSafe>();
+			RemapMesh = IDisplayCluster::Get().GetRenderMgr()->CreateMeshComponent();
 		}
 
 		RemapMesh->AssignMeshGeometry(&RemapGeometry);
diff --git a/Engine/Plugins/Runtime/nDisplay/Source/DisplayCluster/Private/Render/Viewport/Containers/DisplayClusterViewportRemap.h b/Engine/Plugins/Runtime/nDisplay/Source/DisplayCluster/Private/Render/Viewport/Containers/DisplayClusterViewportRemap.h
--- a/Engine/Plugins/Runtime/nDisplay/Source/DisplayCluster/Private/Render/Viewport/Containers/DisplayClusterViewportRemap.h
+++ b/Engine/Plugins/Runtime/nDisplay/Source/DisplayCluster/Private/Render/Viewport/Containers/DisplayClusterViewportRemap.h
@@ -4,7 +4,7 @@
 
 #include "CoreMinimal.h"
 
-#include "Render/Containers/DisplayClusterRender_MeshComponent.h"
+#include "Render/Containers/IDisplayClusterRender_MeshComponent.h"
 #include "Render/Viewport/Containers/DisplayClusterViewportRemapData.h"
 
 class FDisplayClusterViewport;
@@ -24,7 +24,7 @@ public:
 	bool UpdateConfiguration(const FDisplayClusterViewport& InViewport, const FDisplayClusterConfigurationViewport_Remap& InRemapConfiguration);
 
 	/** Gets a pointer to the mesh geometry to use for rendering the remapped regions */
-	TSharedPtr<FDisplayClusterRender_MeshComponent, ESPMode::ThreadSafe> GetRemapMesh() const;
+	TSharedPtr<IDisplayClusterRender_MeshComponent, ESPMode::ThreadSafe> GetRemapMesh() const;
 
 	/** Gets whether there are any remapped regions to render */
 	bool IsUsed() const { return RemapRegions.Num() > 0; }
@@ -52,7 +52,7 @@ private:
 
 private:
 	/** The mesh used to render the remapped regions */
-	TSharedPtr<FDisplayClusterRender_MeshComponent, ESPMode::ThreadSafe> RemapMesh;
+	TSharedPtr<IDisplayClusterRender_MeshComponent, ESPMode::ThreadSafe> RemapMesh;
 	
 	/** A list of all the remapped regions to render */
 	TArray<FDisplayClusterViewportRemapData> RemapRegions;
diff --git a/Engine/Plugins/Runtime/nDisplay/Source/DisplayCluster/Private/Render/Viewport/DisplayClusterViewportManager.cpp b/Engine/Plugins/Runtime/nDisplay/Source/DisplayCluster/Private/Render/Viewport/DisplayClusterViewportManager.cpp
--- a/Engine/Plugins/Runtime/nDisplay/Source/DisplayCluster/Private/Render/Viewport/DisplayClusterViewportManager.cpp
+++ b/Engine/Plugins/Runtime/nDisplay/Source/DisplayCluster/Private/Render/Viewport/DisplayClusterViewportManager.cpp
@@ -14,11 +14,13 @@
 #include "Render/Viewport/RenderFrame/DisplayClusterRenderFrameManager.h"
 #include "Render/Viewport/RenderTarget/DisplayClusterRenderTargetManager.h"
 #include "Render/Viewport/RenderTarget/DisplayClusterRenderTargetResource.h"
-#include "Render/Viewport/DisplayClusterViewportStrings.h"
 #include "Render/Viewport/Postprocess/DisplayClusterViewportPostProcessManager.h"
-
+#include "Render/Viewport/Postprocess/DisplayClusterViewportPostProcessOutputRemap.h"
 #include "Render/Viewport/Configuration/DisplayClusterViewportConfiguration.h"
 #include "Render/Viewport/Configuration/DisplayClusterViewportConfigurationBase.h"
+#include "Render/Viewport/DisplayClusterViewportStrings.h"
+
+#include "WarpBlend/IDisplayClusterWarpBlend.h"
 
 #include "SceneViewExtension.h"
 
@@ -648,3 +650,35 @@ TSharedPtr<IDisplayClusterProjectionPolicy, ESPMode::ThreadSafe> FDisplayCluster
 	return nullptr;
 }
 
+void FDisplayClusterViewportManager::MarkComponentGeometryDirty(const FName InComponentName)
+{
+	check(IsInGameThread());
+
+	// 1. Update all ProceduralMeshComponent references for projection policies
+	for (IDisplayClusterViewport* ViewportIt : GetViewports())
+	{
+		if (ViewportIt != nullptr)
+		{
+			const TSharedPtr<IDisplayClusterProjectionPolicy, ESPMode::ThreadSafe>& ProjectionPolicy = ViewportIt->GetProjectionPolicy();
+			if (ProjectionPolicy.IsValid())
+			{
+				TSharedPtr<IDisplayClusterWarpBlend, ESPMode::ThreadSafe> WarpBlendInterface;
+				if (ProjectionPolicy->GetWarpBlendInterface(WarpBlendInterface) && WarpBlendInterface.IsValid())
+				{
+					// Update only interfaces with ProceduralMesh as geometry source 
+					if (WarpBlendInterface->GetWarpGeometryType() == EDisplayClusterWarpGeometryType::WarpProceduralMesh)
+					{
+						// Set the ProceduralMeshComponent geometry dirty for all valid WarpBlendInterface
+						WarpBlendInterface->MarkWarpGeometryComponentDirty(InComponentName);
+					}
+				}
+			}
+		}
+	}
+
+	// 2. Update all ProceduralMeshComponent references for OutputRemap
+	if (PostProcessManager.IsValid())
+	{
+		PostProcessManager->GetOutputRemap()->MarkProceduralMeshComponentGeometryDirty(InComponentName);
+	}
+}
diff --git a/Engine/Plugins/Runtime/nDisplay/Source/DisplayCluster/Private/Render/Viewport/DisplayClusterViewportManager.h b/Engine/Plugins/Runtime/nDisplay/Source/DisplayCluster/Private/Render/Viewport/DisplayClusterViewportManager.h
--- a/Engine/Plugins/Runtime/nDisplay/Source/DisplayCluster/Private/Render/Viewport/DisplayClusterViewportManager.h
+++ b/Engine/Plugins/Runtime/nDisplay/Source/DisplayCluster/Private/Render/Viewport/DisplayClusterViewportManager.h
@@ -77,6 +77,8 @@ public:
 		return TArrayView<IDisplayClusterViewport*>((IDisplayClusterViewport**)(Viewports.GetData()), Viewports.Num());
 	}
 
+	virtual void MarkComponentGeometryDirty(const FName InComponentName = NAME_None) override;
+
 	// internal use only
 	bool CreateViewport(const FString& ViewportId, const class UDisplayClusterConfigurationViewport* ConfigurationViewport);
 	IDisplayClusterViewport* CreateViewport(const FString& ViewportId, const TSharedPtr<IDisplayClusterProjectionPolicy, ESPMode::ThreadSafe>& InProjectionPolicy);
diff --git a/Engine/Plugins/Runtime/nDisplay/Source/DisplayCluster/Private/Render/Viewport/DisplayClusterViewportProxy.cpp b/Engine/Plugins/Runtime/nDisplay/Source/DisplayCluster/Private/Render/Viewport/DisplayClusterViewportProxy.cpp
--- a/Engine/Plugins/Runtime/nDisplay/Source/DisplayCluster/Private/Render/Viewport/DisplayClusterViewportProxy.cpp
+++ b/Engine/Plugins/Runtime/nDisplay/Source/DisplayCluster/Private/Render/Viewport/DisplayClusterViewportProxy.cpp
@@ -12,8 +12,8 @@
 
 #include "Render/Projection/IDisplayClusterProjectionPolicy.h"
 
-#include "Render/Containers/DisplayClusterRender_MeshComponent.h"
-#include "Render/Containers/DisplayClusterRender_MeshComponentProxy.h"
+#include "Render/Containers/IDisplayClusterRender_MeshComponent.h"
+#include "Render/Containers/IDisplayClusterRender_MeshComponentProxy.h"
 
 #if WITH_EDITOR
 #include "DisplayClusterRootActor.h"
@@ -202,8 +202,8 @@ void FDisplayClusterViewportProxy::ImplViewportRemap_RenderThread(FRHICommandLis
 
 	if (RemapMesh.IsValid())
 	{
-		const FDisplayClusterRender_MeshComponentProxy* MeshProxy = RemapMesh->GetMeshComponentProxy_RenderThread();
-		if (MeshProxy && MeshProxy->IsValid_RenderThread())
+		const IDisplayClusterRender_MeshComponentProxy* MeshProxy = RemapMesh->GetMeshComponentProxy_RenderThread();
+		if (MeshProxy!=nullptr && MeshProxy->IsEnabled_RenderThread())
 		{
 			if (AdditionalFrameTargetableResources.Num() != OutputFrameTargetableResources.Num())
 			{
@@ -241,8 +241,8 @@ EDisplayClusterViewportResourceType FDisplayClusterViewportProxy::GetOutputResou
 
 	if (RemapMesh.IsValid())
 	{
-		const FDisplayClusterRender_MeshComponentProxy* MeshProxy = RemapMesh->GetMeshComponentProxy_RenderThread();
-		if (MeshProxy && MeshProxy->IsValid_RenderThread())
+		const IDisplayClusterRender_MeshComponentProxy* MeshProxy = RemapMesh->GetMeshComponentProxy_RenderThread();
+		if (MeshProxy!=nullptr && MeshProxy->IsEnabled_RenderThread())
 		{
 			// In this case render to additional frame targetable
 			return EDisplayClusterViewportResourceType::AdditionalFrameTargetableResource;
diff --git a/Engine/Plugins/Runtime/nDisplay/Source/DisplayCluster/Private/Render/Viewport/DisplayClusterViewportProxy.h b/Engine/Plugins/Runtime/nDisplay/Source/DisplayCluster/Private/Render/Viewport/DisplayClusterViewportProxy.h
--- a/Engine/Plugins/Runtime/nDisplay/Source/DisplayCluster/Private/Render/Viewport/DisplayClusterViewportProxy.h
+++ b/Engine/Plugins/Runtime/nDisplay/Source/DisplayCluster/Private/Render/Viewport/DisplayClusterViewportProxy.h
@@ -13,7 +13,7 @@ class FDisplayClusterViewportManagerProxy;
 class FDisplayClusterViewport;
 class IDisplayClusterShaders;
 class IDisplayClusterProjectionPolicy;
-class FDisplayClusterRender_MeshComponent;
+class IDisplayClusterRender_MeshComponent;
 
 class FDisplayClusterViewportProxy
 	: public IDisplayClusterViewportProxy
@@ -120,7 +120,7 @@ protected:
 	// Additional parameters
 	FDisplayClusterViewport_OverscanSettings     OverscanSettings;
 
-	TSharedPtr<FDisplayClusterRender_MeshComponent, ESPMode::ThreadSafe> RemapMesh;
+	TSharedPtr<IDisplayClusterRender_MeshComponent, ESPMode::ThreadSafe> RemapMesh;
 
 	// Projection policy instance that serves this viewport
 	TSharedPtr<IDisplayClusterProjectionPolicy, ESPMode::ThreadSafe> ProjectionPolicy;
diff --git a/Engine/Plugins/Runtime/nDisplay/Source/DisplayCluster/Private/Render/Viewport/Postprocess/DisplayClusterViewportPostProcessOutputRemap.cpp b/Engine/Plugins/Runtime/nDisplay/Source/DisplayCluster/Private/Render/Viewport/Postprocess/DisplayClusterViewportPostProcessOutputRemap.cpp
--- a/Engine/Plugins/Runtime/nDisplay/Source/DisplayCluster/Private/Render/Viewport/Postprocess/DisplayClusterViewportPostProcessOutputRemap.cpp
+++ b/Engine/Plugins/Runtime/nDisplay/Source/DisplayCluster/Private/Render/Viewport/Postprocess/DisplayClusterViewportPostProcessOutputRemap.cpp
@@ -6,12 +6,15 @@
 
 #include "Misc/DisplayClusterHelpers.h"
 #include "Misc/DisplayClusterLog.h"
+#include "Misc/DisplayClusterGlobals.h"
+#include "Misc/Paths.h"
 
 #include "DisplayClusterRootActor.h"
 
-#include "Render/Containers/DisplayClusterRender_MeshComponent.h"
-#include "Render/Containers/DisplayClusterRender_MeshComponentProxy.h"
-#include "Render/Containers/DisplayClusterRender_MeshComponentProxyData.h"
+#include "Render/IDisplayClusterRenderManager.h"
+
+#include "Render/Containers/IDisplayClusterRender_MeshComponent.h"
+#include "Render/Containers/IDisplayClusterRender_MeshComponentProxy.h"
 #include "Render/Containers/DisplayClusterRender_MeshGeometry.h"
 
 #include "Render/Viewport/IDisplayClusterViewportManager.h"
@@ -20,9 +23,9 @@
 
 #include "Engine/RendererSettings.h"
 #include "Engine/StaticMesh.h"
-#include "ClearQuad.h"
+#include "ProceduralMeshComponent.h"
 
-#include "Misc/Paths.h"
+#include "ClearQuad.h"
 
 // Enable/Disable ClearTexture for OutputRemap RTT
 static TAutoConsoleVariable<int32> CVarClearOutputRemapFrameTextureEnabled(
@@ -39,109 +42,216 @@ static TAutoConsoleVariable<int32> CVarClearOutputRemapFrameTextureEnabled(
 // FDisplayClusterViewportPostProcessOutputRemap
 //////////////////////////////////////////////////////////////////////////////////////////////
 FDisplayClusterViewportPostProcessOutputRemap::FDisplayClusterViewportPostProcessOutputRemap()
-	: OutputRemapMesh(*(new FDisplayClusterRender_MeshComponent()))
-	, ShadersAPI(IDisplayClusterShaders::Get())
+	: ShadersAPI(IDisplayClusterShaders::Get())
 {
 }
 
 FDisplayClusterViewportPostProcessOutputRemap::~FDisplayClusterViewportPostProcessOutputRemap()
 {
-	delete (&OutputRemapMesh);
 }
 
-bool FDisplayClusterViewportPostProcessOutputRemap::UpdateConfiguration_ExternalFile(const FString& InExternalFile)
+bool FDisplayClusterViewportPostProcessOutputRemap::ImplInitializeConfiguration()
 {
-	if (ExternalFile == InExternalFile)
+	if (!OutputRemapMesh.IsValid())
 	{
-		return true;
+		OutputRemapMesh = IDisplayCluster::Get().GetRenderMgr()->CreateMeshComponent();
 	}
 
-	UpdateConfiguration_Disabled();
+	// Reset stored values
+	ExternalFile.Empty();
+
+	return OutputRemapMesh.IsValid();
+}
+
+bool FDisplayClusterViewportPostProcessOutputRemap::UpdateConfiguration_ExternalFile(const FString& InExternalFile)
+{
+	check(IsInGameThread());
 
+	// Empty input filename -> disable output remap
 	if (InExternalFile.IsEmpty())
 	{
-		return true;
+		UpdateConfiguration_Disabled();
+
+		return false;
 	}
 
 	// Support related paths
 	FString FullPathFileName = DisplayClusterHelpers::filesystem::GetFullPathForConfigResource(InExternalFile);
-	if (FullPathFileName.IsEmpty() || !FPaths::FileExists(FullPathFileName))
+
+	// Don't update twice
+	if (ExternalFile.Equals(FullPathFileName, ESearchCase::IgnoreCase))
 	{
-		if (!bErrorCantFindFileOnce)
-		{
-			UE_LOG(LogDisplayClusterViewport, Error, TEXT("OutputRemap - Failed to find file '%s'"), *InExternalFile);
-			bErrorCantFindFileOnce = true;
-		}
+		return IsEnabled();
+	}
+
+	if (!FPaths::FileExists(FullPathFileName))
+	{
+		UE_LOG(LogDisplayClusterViewport, Error, TEXT("OutputRemap - Failed to find file '%s'"), *FullPathFileName);
+		UpdateConfiguration_Disabled();
+
+		// Store input filename for next update
+		ExternalFile = FullPathFileName;
 
 		return false;
 	}
 
-	bErrorCantFindFileOnce = false;
-
+	// Try load geometry from file
 	FDisplayClusterRender_MeshGeometry MeshGeometry;
 	if (!MeshGeometry.LoadFromFile(FullPathFileName))
 	{
-		if (!bErrorFailLoadFromFileOnce)
-		{
-			UE_LOG(LogDisplayClusterViewport, Error, TEXT("OutputRemap - Failed to load ext mesh from file '%s'"), *FullPathFileName);
-			bErrorFailLoadFromFileOnce = true;
-		}
+		UE_LOG(LogDisplayClusterViewport, Error, TEXT("OutputRemap - Failed to load ext mesh from file '%s'"), *FullPathFileName);
+		UpdateConfiguration_Disabled();
 
-		return false;
+		// Store input filename for next update
+		ExternalFile = FullPathFileName;
 	}
+	else if (ImplInitializeConfiguration())
+	{
+		// Assign loaded geometry
+		OutputRemapMesh->SetGeometryFunc(EDisplayClusterRender_MeshComponentProxyDataFunc::OutputRemapScreenSpace);
+		OutputRemapMesh->AssignMeshGeometry(&MeshGeometry);
 
-	bErrorFailLoadFromFileOnce = false;
-
-
-	OutputRemapMesh.SetGeometryFunc(EDisplayClusterRender_MeshComponentProxyDataFunc::OutputRemapScreenSpace);
-	OutputRemapMesh.AssignMeshGeometry(&MeshGeometry);
+		// Store input filename for next update
+		ExternalFile = FullPathFileName;
 
-	bIsEnabled = true;
-	ExternalFile = InExternalFile;
+		return true;
+	}
 
-	return true;
+	return false;
 }
 
 bool FDisplayClusterViewportPostProcessOutputRemap::UpdateConfiguration_StaticMesh(UStaticMesh* InStaticMesh)
 {
-	if (StaticMesh == InStaticMesh)
+	check(IsInGameThread());
+
+	// nullptr as input - disable
+	if (InStaticMesh == nullptr)
+	{
+		UpdateConfiguration_Disabled();
+
+		return false;
+	}
+
+	// don't update twice
+	if (OutputRemapMesh.IsValid() && !OutputRemapMesh->IsMeshComponentRefGeometryDirty())
 	{
+		const UStaticMesh* AssignedStaticMesh = OutputRemapMesh->GetStaticMesh();
+		if (AssignedStaticMesh == InStaticMesh)
+		{
+			return IsEnabled();
+		}
+	}
+
+	if (ImplInitializeConfiguration())
+	{
+		OutputRemapMesh->SetGeometryFunc(EDisplayClusterRender_MeshComponentProxyDataFunc::OutputRemapScreenSpace);
+		OutputRemapMesh->AssignStaticMesh(InStaticMesh, FDisplayClusterMeshUVs());
+
 		return true;
 	}
 
-	UpdateConfiguration_Disabled();
+	return false;
+}
+
+bool FDisplayClusterViewportPostProcessOutputRemap::UpdateConfiguration_StaticMeshComponent(UStaticMeshComponent* InStaticMeshComponent)
+{
+	check(IsInGameThread());
 
-	if (InStaticMesh == nullptr)
+	// nullptr as input - disable
+	if (InStaticMeshComponent == nullptr)
+	{
+		UpdateConfiguration_Disabled();
+
+		return false;
+	}
+
+	// don't update twice
+	if (OutputRemapMesh.IsValid() && !OutputRemapMesh->IsMeshComponentRefGeometryDirty())
 	{
+		const UStaticMeshComponent* AssignedStaticMeshComponent = OutputRemapMesh->GetStaticMeshComponent();
+		if (AssignedStaticMeshComponent == InStaticMeshComponent)
+		{
+			return IsEnabled();
+		}
+	}
+
+	// Begin new configuration:
+	if (ImplInitializeConfiguration())
+	{
+		OutputRemapMesh->SetGeometryFunc(EDisplayClusterRender_MeshComponentProxyDataFunc::OutputRemapScreenSpace);
+		OutputRemapMesh->AssignStaticMeshComponentRefs(InStaticMeshComponent, FDisplayClusterMeshUVs());
+
 		return true;
 	}
 
-	OutputRemapMesh.SetGeometryFunc(EDisplayClusterRender_MeshComponentProxyDataFunc::OutputRemapScreenSpace);
-	OutputRemapMesh.AssignStaticMesh(InStaticMesh);
+	return false;
+}
+
+bool FDisplayClusterViewportPostProcessOutputRemap::UpdateConfiguration_ProceduralMeshComponent(UProceduralMeshComponent* InProceduralMeshComponent)
+{
+	check(IsInGameThread());
+
+	// nullptr as input - disable
+	if (InProceduralMeshComponent == nullptr)
+	{
+		UpdateConfiguration_Disabled();
+
+		return false;
+	}
+
+	// don't update twice
+	if (OutputRemapMesh.IsValid() && !OutputRemapMesh->IsMeshComponentRefGeometryDirty())
+	{
+		const UProceduralMeshComponent* AssignedProceduralMeshComponent = OutputRemapMesh->GetProceduralMeshComponent();
+		if (AssignedProceduralMeshComponent == InProceduralMeshComponent)
+		{
+			return IsEnabled();
+		}
+	}
 
-	bIsEnabled = true;
-	StaticMesh = InStaticMesh;
+	// Begin new configuration:
+	if (ImplInitializeConfiguration())
+	{
+		OutputRemapMesh->SetGeometryFunc(EDisplayClusterRender_MeshComponentProxyDataFunc::OutputRemapScreenSpace);
+		OutputRemapMesh->AssignProceduralMeshComponentRefs(InProceduralMeshComponent, FDisplayClusterMeshUVs());
 
-	return true;
+		return true;
+	}
+
+	return false;
 }
 
 void FDisplayClusterViewportPostProcessOutputRemap::UpdateConfiguration_Disabled()
 {
-	StaticMesh = nullptr;
+	check(IsInGameThread());
+
 	ExternalFile.Empty();
-	bIsEnabled = false;
+	OutputRemapMesh.Reset();
+}
+
+bool FDisplayClusterViewportPostProcessOutputRemap::MarkProceduralMeshComponentGeometryDirty(const FName& InComponentName)
+{
+	// Mark the procedural mesh component as dirty
+	if (OutputRemapMesh.IsValid() && OutputRemapMesh->GetGeometrySource() == EDisplayClusterRender_MeshComponentGeometrySource::ProceduralMeshComponentRef)
+	{
+		if (InComponentName == NAME_None || OutputRemapMesh->EqualsMeshComponentName(InComponentName))
+		{
+			OutputRemapMesh->MarkMeshComponentRefGeometryDirty();
+			return true;
+		}
+	}
 
-	OutputRemapMesh.ReleaseMeshComponent();
+	return false;
 }
 
 void FDisplayClusterViewportPostProcessOutputRemap::PerformPostProcessFrame_RenderThread(FRHICommandListImmediate& RHICmdList, const TArray<FRHITexture2D*>* InFrameTargets, const TArray<FRHITexture2D*>* InAdditionalFrameTargets) const
 {
 	check(IsInRenderingThread());
 
-	if (InFrameTargets && InAdditionalFrameTargets)
+	if (InFrameTargets && InAdditionalFrameTargets && OutputRemapMesh.IsValid())
 	{
-		const FDisplayClusterRender_MeshComponentProxy* MeshProxy = OutputRemapMesh.GetMeshComponentProxy_RenderThread();
-		if (MeshProxy && MeshProxy->IsValid_RenderThread())
+		const IDisplayClusterRender_MeshComponentProxy* MeshProxy = OutputRemapMesh->GetMeshComponentProxy_RenderThread();
+		if (MeshProxy!=nullptr && MeshProxy->IsEnabled_RenderThread())
 		{
 			const bool bClearOutputRemapFrameTextureEnabled = CVarClearOutputRemapFrameTextureEnabled.GetValueOnRenderThread() != 0;
 
diff --git a/Engine/Plugins/Runtime/nDisplay/Source/DisplayCluster/Private/Render/Viewport/Postprocess/DisplayClusterViewportPostProcessOutputRemap.h b/Engine/Plugins/Runtime/nDisplay/Source/DisplayCluster/Private/Render/Viewport/Postprocess/DisplayClusterViewportPostProcessOutputRemap.h
--- a/Engine/Plugins/Runtime/nDisplay/Source/DisplayCluster/Private/Render/Viewport/Postprocess/DisplayClusterViewportPostProcessOutputRemap.h
+++ b/Engine/Plugins/Runtime/nDisplay/Source/DisplayCluster/Private/Render/Viewport/Postprocess/DisplayClusterViewportPostProcessOutputRemap.h
@@ -8,8 +8,10 @@
 #include "RHIUtilities.h"
 
 class IDisplayClusterViewportManager;
-class FDisplayClusterRender_MeshComponent;
+class IDisplayClusterRender_MeshComponent;
 class UStaticMesh;
+class UStaticMeshComponent;
+class UProceduralMeshComponent;
 
 /**
  * Display Cluster Viewport PostProcess OutputRemap
@@ -23,25 +25,27 @@ public:
 	// Game thread update calls
 	bool UpdateConfiguration_ExternalFile(const FString& InExternalFile);
 	bool UpdateConfiguration_StaticMesh(UStaticMesh* InStaticMesh);
+	bool UpdateConfiguration_StaticMeshComponent(UStaticMeshComponent* InStaticMeshComponent);
+	bool UpdateConfiguration_ProceduralMeshComponent(UProceduralMeshComponent* InProceduralMeshComponent);
 	void UpdateConfiguration_Disabled();
 
+	bool MarkProceduralMeshComponentGeometryDirty(const FName& InComponentName);
+
 	bool IsEnabled() const
 	{
-		return bIsEnabled;
+		return OutputRemapMesh.IsValid();
 	}
 
 public:
 	void PerformPostProcessFrame_RenderThread(FRHICommandListImmediate& RHICmdList, const TArray<FRHITexture2D*>* InFrameTargets = nullptr, const TArray<FRHITexture2D*>* InAdditionalFrameTargets = nullptr) const;
 
 private:
-	bool bIsEnabled = false;
-	UStaticMesh* StaticMesh = nullptr;
-	FString ExternalFile;
+	bool ImplInitializeConfiguration();
 
-	bool bErrorCantFindFileOnce = false;
-	bool bErrorFailLoadFromFileOnce = false;
+private:
+	FString ExternalFile;
 
 private:
-	FDisplayClusterRender_MeshComponent& OutputRemapMesh;
+	TSharedPtr<IDisplayClusterRender_MeshComponent, ESPMode::ThreadSafe> OutputRemapMesh;
 	class IDisplayClusterShaders& ShadersAPI;
 };
diff --git a/Engine/Plugins/Runtime/nDisplay/Source/DisplayCluster/Public/Render/Containers/DisplayClusterRender_MeshComponent.h b/Engine/Plugins/Runtime/nDisplay/Source/DisplayCluster/Public/Render/Containers/DisplayClusterRender_MeshComponent.h
deleted file mode 100644
diff --git a/Engine/Plugins/Runtime/nDisplay/Source/DisplayCluster/Public/Render/Containers/DisplayClusterRender_MeshComponentProxy.h b/Engine/Plugins/Runtime/nDisplay/Source/DisplayCluster/Public/Render/Containers/DisplayClusterRender_MeshComponentProxy.h
deleted file mode 100644
diff --git a/Engine/Plugins/Runtime/nDisplay/Source/DisplayCluster/Public/Render/Containers/DisplayClusterRender_MeshComponentProxyData.h b/Engine/Plugins/Runtime/nDisplay/Source/DisplayCluster/Public/Render/Containers/DisplayClusterRender_MeshComponentProxyData.h
deleted file mode 100644
diff --git a/Engine/Plugins/Runtime/nDisplay/Source/DisplayCluster/Public/Render/Containers/DisplayClusterRender_MeshComponentTypes.h b/Engine/Plugins/Runtime/nDisplay/Source/DisplayCluster/Public/Render/Containers/DisplayClusterRender_MeshComponentTypes.h
--- a/Engine/Plugins/Runtime/nDisplay/Source/DisplayCluster/Public/Render/Containers/DisplayClusterRender_MeshComponentTypes.h
+++ b/Engine/Plugins/Runtime/nDisplay/Source/DisplayCluster/Public/Render/Containers/DisplayClusterRender_MeshComponentTypes.h
@@ -29,3 +29,43 @@ struct FDisplayClusterMeshVertex
 	FVector2D UV;
 	FVector2D UV_Chromakey;
 };
+
+// Map source geometry UVs to DCMeshComponent UVs
+struct FDisplayClusterMeshUVs
+{
+	enum class ERemapTarget: uint8
+	{
+		Base = 0,
+		Chromakey,
+		COUNT
+	};
+
+	FDisplayClusterMeshUVs() = default;
+
+	FDisplayClusterMeshUVs(const int32 InBaseUVIndex, const int32 InChromakeyUVIndex)
+	{
+		RemapValue[(uint8)ERemapTarget::Base]      = InBaseUVIndex;
+		RemapValue[(uint8)ERemapTarget::Chromakey] = InChromakeyUVIndex;
+	}
+
+	const int32 operator[](const ERemapTarget RemapTarget) const
+	{
+		check((uint8)RemapTarget < (uint8)ERemapTarget::COUNT);
+
+		return RemapValue[(uint8)RemapTarget];
+	}
+
+	const int32 GetSourceGeometryUVIndex(const ERemapTarget RemapTarget, const int32 MaxValue, const int32 DefaultValue) const
+	{
+		const int32 CurrentValue = operator[](RemapTarget);
+
+		// Replace INDEX_NONE to default value
+		const int32 Result = (CurrentValue == INDEX_NONE) ? DefaultValue : CurrentValue;
+
+		// Return UV in range 0..MaxValue
+		return (Result >= 0 && Result < MaxValue) ? Result : 0;
+	}
+
+private:
+	int32 RemapValue[(uint8)ERemapTarget::COUNT] = {INDEX_NONE, INDEX_NONE};
+};
diff --git a/Engine/Plugins/Runtime/nDisplay/Source/DisplayCluster/Public/Render/Containers/IDisplayClusterRender_MeshComponent.h b/Engine/Plugins/Runtime/nDisplay/Source/DisplayCluster/Public/Render/Containers/IDisplayClusterRender_MeshComponent.h
new file mode 100644
--- /dev/null
+++ b/Engine/Plugins/Runtime/nDisplay/Source/DisplayCluster/Public/Render/Containers/IDisplayClusterRender_MeshComponent.h
@@ -0,0 +1,177 @@
+// Copyright Epic Games, Inc. All Rights Reserved.
+
+#pragma once
+
+#include "CoreMinimal.h"
+#include "Render/Containers/DisplayClusterRender_MeshComponentTypes.h"
+#include "Render/Containers/DisplayClusterRender_MeshGeometry.h"
+
+class USceneComponent;
+class UStaticMeshComponent;
+class UProceduralMeshComponent;
+
+class IDisplayClusterRender_MeshComponentProxy;
+
+struct FStaticMeshLODResources;
+struct FProcMeshSection;
+
+class IDisplayClusterRender_MeshComponent
+{
+public:
+	virtual ~IDisplayClusterRender_MeshComponent() = default;
+
+public:
+	/**
+	* Assign static mesh component and origin to this container
+	* Then send geometry from the StaticMesh component to proxy
+	*
+	* @param InStaticMeshComponent - Pointer to the StaticMesh component with geometry for warp
+	* @param InUVs                 - Map source geometry UVs to DCMeshComponent UVs
+	* @param InOriginComponent     - (optional) stage origin point component
+	* @param InLODIndex            - (optional) use geometry from LOD index
+	*/
+	virtual void AssignStaticMeshComponentRefs(UStaticMeshComponent* InStaticMeshComponent, const FDisplayClusterMeshUVs& InUVs, USceneComponent* InOriginComponent = nullptr, int32 InLODIndex = 0) = 0;
+
+	/**
+	* Assign procedural mesh component, section index and origin to this container
+	* Then send geometry from the ProceduralMesh(SectionIndex) to proxy
+	*
+	* @param InProceduralMeshComponent - Pointer to the ProceduralMesh component with geometry for warp
+	* @param InUVs                     - Map source geometry UVs to DCMeshComponent UVs
+	* @param InOriginComponent         - (optional) stage origin point component
+	* @param InSectionIndex            - (optional) The section index in ProceduralMesh
+	*/
+	virtual void AssignProceduralMeshComponentRefs(UProceduralMeshComponent* InProceduralMeshComponent, const FDisplayClusterMeshUVs& InUVs, USceneComponent* InOriginComponent = nullptr, const int32 InSectionIndex = 0) = 0;
+
+	/**
+	* Assign procedural mesh section geometry
+	* Then send geometry from the ProcMeshSection to proxy
+	*
+	* @param InProcMeshSection - Pointer to the Proce mesh section geometry
+	* @param InUVs             - Map source geometry UVs to DCMeshComponent UVs
+	*/
+	virtual void AssignProceduralMeshSection(const FProcMeshSection& InProcMeshSection, const FDisplayClusterMeshUVs& InUVs) = 0;
+
+	/**
+	* Assign static mesh geometry
+	* Then send geometry from the StaticMesh to proxy
+	*
+	* @param InStaticMesh - Pointer to the StaticMesh geometry
+	* @param InUVs        - Map source geometry UVs to DCMeshComponent UVs
+	* @param InLODIndex   - (optional) use geometry from LOD index
+	*/
+	virtual void AssignStaticMesh(const UStaticMesh* InStaticMesh, const FDisplayClusterMeshUVs& InUVs, int32 InLODIndex = 0) = 0;
+
+	/**
+	* Assign mesh geometry
+	* Then send geometry to proxy
+	*
+	* @param InMeshGeometry - Pointer to the source geometry data
+	*/
+	virtual void AssignMeshGeometry(const FDisplayClusterRender_MeshGeometry* InMeshGeometry) = 0;
+
+	/**
+	* Create mesh geometry on render thread (debug purpose)
+	*
+	* @param InMeshGeometry - Pointer to the source geometry data
+	* @param InDataFunc     - Geometry function type
+	*/
+	virtual void AssignMeshGeometry_RenderThread(const FDisplayClusterRender_MeshGeometry* InMeshGeometry, const EDisplayClusterRender_MeshComponentProxyDataFunc InDataFunc) const = 0;
+
+	/**
+	* Release refs, proxy, RHI.
+	*
+	*/
+	virtual void ReleaseMeshComponent() = 0;
+
+	/**
+	* Release proxy
+	*
+	*/
+	virtual void ReleaseProxyGeometry() = 0;
+
+	/**
+	* Get referenced Origin component object
+	*
+	* @return - pointer to the Origin component
+	*/
+	virtual USceneComponent* GetOriginComponent() const = 0;
+
+	/**
+	* Get referenced StaticMesh component object
+	*
+	* @return - pointer to the StaticMesh component
+	*/
+	virtual UStaticMeshComponent* GetStaticMeshComponent() const = 0;
+
+	/**
+	* Get referenced StaticMesh object
+	*
+	* @return - pointer to the StaticMesh object
+	*/
+	virtual const UStaticMesh* GetStaticMesh() const = 0;
+
+	/**
+	* Get referenced ProceduralMesh component object
+	*
+	* @return - pointer to the ProceduralMesh component
+	*/
+	virtual UProceduralMeshComponent* GetProceduralMeshComponent() const = 0;
+
+	/**
+	* Get geometry from assigned StaticMesh component
+	*
+	* @param LodIndex - (optional) lod index
+	*
+	* @return - pointer to the static mesh geometry
+	*/
+	virtual const FStaticMeshLODResources* GetStaticMeshComponentLODResources(int32 InLODIndex = 0) const = 0;
+
+	/**
+	* Get geometry from assigned ProceduralMesh section by index
+	*
+	* @param SectionIndex - Geometry source section index
+	*
+	* @return - pointer to the geometry section data
+	*/
+	virtual const FProcMeshSection* GetProceduralMeshComponentSection(const int32 SectionIndex = 0) const = 0;
+
+	/**
+	* Get nDisplay mesh component proxy object
+	*
+	* @return - pointer to the nDisplay mesh component proxy
+	*/
+	virtual IDisplayClusterRender_MeshComponentProxy* GetMeshComponentProxy_RenderThread() const = 0;
+
+	/**
+	* Set geometry preprocess function. This function change geometry prepared for proxy
+	*
+	* @param InDataFunc - geometry function type
+	*/
+	virtual void SetGeometryFunc(const EDisplayClusterRender_MeshComponentProxyDataFunc InDataFunc) = 0;
+
+	/**
+	* Return true if referenced component geometry changed
+	*/
+	virtual bool IsMeshComponentRefGeometryDirty() const = 0;
+
+	/**
+	* Mark referenced component geometry as changed
+	*/
+	virtual void MarkMeshComponentRefGeometryDirty() const = 0;
+
+	/**
+	* Clear referenced component geometry changed flag
+	*/
+	virtual void ResetMeshComponentRefGeometryDirty() const = 0;
+
+	/**
+	* return true, if referenced component name is equal
+	*/
+	virtual bool EqualsMeshComponentName(const FName& InMeshComponentName) const = 0;
+
+	/**
+	* return assigned geometry source type
+	*/
+	virtual EDisplayClusterRender_MeshComponentGeometrySource GetGeometrySource() const = 0;
+};
diff --git a/Engine/Plugins/Runtime/nDisplay/Source/DisplayCluster/Public/Render/Containers/IDisplayClusterRender_MeshComponentProxy.h b/Engine/Plugins/Runtime/nDisplay/Source/DisplayCluster/Public/Render/Containers/IDisplayClusterRender_MeshComponentProxy.h
new file mode 100644
--- /dev/null
+++ b/Engine/Plugins/Runtime/nDisplay/Source/DisplayCluster/Public/Render/Containers/IDisplayClusterRender_MeshComponentProxy.h
@@ -0,0 +1,34 @@
+// Copyright Epic Games, Inc. All Rights Reserved.
+
+#pragma once
+
+#include "CoreMinimal.h"
+
+#include "RHI.h"
+#include "RHICommandList.h"
+
+class IDisplayClusterRender_MeshComponentProxy
+{
+public:
+	virtual ~IDisplayClusterRender_MeshComponentProxy() = default;
+
+public:
+	/**
+	* Initialize RHI for mesh render
+	* 
+	* return - true, if the mesh can be rendered
+	*/
+	virtual bool BeginRender_RenderThread(FRHICommandListImmediate& RHICmdList, FGraphicsPipelineStateInitializer& GraphicsPSOInit) const = 0;
+
+	/**
+	* Render mesh geometry
+	*
+	* return - true, if the mesh rendered
+	*/
+	virtual bool FinishRender_RenderThread(FRHICommandListImmediate& RHICmdList) const = 0;
+
+	/**
+	* Return true, if mesh valid and can be rendered
+	*/
+	virtual bool IsEnabled_RenderThread() const = 0;
+};
diff --git a/Engine/Plugins/Runtime/nDisplay/Source/DisplayCluster/Public/Render/IDisplayClusterRenderManager.h b/Engine/Plugins/Runtime/nDisplay/Source/DisplayCluster/Public/Render/IDisplayClusterRenderManager.h
--- a/Engine/Plugins/Runtime/nDisplay/Source/DisplayCluster/Public/Render/IDisplayClusterRenderManager.h
+++ b/Engine/Plugins/Runtime/nDisplay/Source/DisplayCluster/Public/Render/IDisplayClusterRenderManager.h
@@ -7,6 +7,7 @@
 class IDisplayClusterRenderDevice;
 class IDisplayClusterPostProcess;
 class IDisplayClusterPostProcessFactory;
+class IDisplayClusterRender_MeshComponent;
 class IDisplayClusterProjectionPolicy;
 class IDisplayClusterProjectionPolicyFactory;
 class IDisplayClusterRenderDeviceFactory;
@@ -166,4 +167,8 @@ public:
 	*/
 	virtual IDisplayClusterViewportManager* GetViewportManager() const = 0;
 
+	/**
+	* @return - new mesh component object
+	*/
+	virtual TSharedPtr<IDisplayClusterRender_MeshComponent, ESPMode::ThreadSafe> CreateMeshComponent() const = 0;
 };
diff --git a/Engine/Plugins/Runtime/nDisplay/Source/DisplayCluster/Public/Render/Viewport/IDisplayClusterViewportManager.h b/Engine/Plugins/Runtime/nDisplay/Source/DisplayCluster/Public/Render/Viewport/IDisplayClusterViewportManager.h
--- a/Engine/Plugins/Runtime/nDisplay/Source/DisplayCluster/Public/Render/Viewport/IDisplayClusterViewportManager.h
+++ b/Engine/Plugins/Runtime/nDisplay/Source/DisplayCluster/Public/Render/Viewport/IDisplayClusterViewportManager.h
@@ -115,5 +115,13 @@ public:
 	* @return - arrays with viewport objects refs
 	*/
 	virtual const TArrayView<IDisplayClusterViewport*> GetViewports() const = 0;
+
+	/**
+	* Mark the geometry of the referenced component(s) as dirty (ProceduralMesh, etc)
+	* [Game thread func]
+	*
+	* @param InComponentName - (optional) Unique component name
+	*/
+	virtual void MarkComponentGeometryDirty(const FName InComponentName = NAME_None) = 0;
 };
 
diff --git a/Engine/Plugins/Runtime/nDisplay/Source/DisplayClusterConfiguration/Public/DisplayClusterConfigurationTypes_OutputRemap.h b/Engine/Plugins/Runtime/nDisplay/Source/DisplayClusterConfiguration/Public/DisplayClusterConfigurationTypes_OutputRemap.h
--- a/Engine/Plugins/Runtime/nDisplay/Source/DisplayClusterConfiguration/Public/DisplayClusterConfigurationTypes_OutputRemap.h
+++ b/Engine/Plugins/Runtime/nDisplay/Source/DisplayClusterConfiguration/Public/DisplayClusterConfigurationTypes_OutputRemap.h
@@ -20,6 +20,9 @@ enum class EDisplayClusterConfigurationFramePostProcess_OutputRemapSource : uint
 
 	/** Use an external .obj file for output remapping when the Data Source is set to File */
 	ExternalFile   UMETA(DisplayName = "External File"),
+
+	/** Use a Mesh component reference for output remapping (ProceduralMeshComponent or StaticMeshComponent) */
+	MeshComponent UMETA(DisplayName = "Mesh Component"),
 };
 
 /* Screen space remapping of the final backbuffer output. Applied at the whole window */
@@ -41,6 +44,10 @@ public:
 	UPROPERTY(EditAnywhere, BlueprintReadOnly, Category = "NDisplay OutputRemap", meta = (EditCondition = "DataSource == EDisplayClusterConfigurationFramePostProcess_OutputRemapSource::StaticMesh && bEnable"))
 	class UStaticMesh* StaticMesh = nullptr;
 
+	/** The MeshComponent reference (ProceduralMeshComponent or StaticMeshComponent) to use for output remapping when the Data Source is set to Mesh Component */
+	UPROPERTY(EditAnywhere, BlueprintReadOnly, Category = "NDisplay OutputRemap", meta = (EditCondition = "DataSource == EDisplayClusterConfigurationFramePostProcess_OutputRemapSource::MeshComponent && bEnable"))
+	FString MeshComponentName;
+
 	/** The external .obj file to use for output remapping when the Data Source is set to File */
 	UPROPERTY(EditAnywhere, BlueprintReadWrite, Category = "NDisplay OutputRemap", meta = (EditCondition = "DataSource == EDisplayClusterConfigurationFramePostProcess_OutputRemapSource::ExternalFile && bEnable"))
 	FString ExternalFile;
diff --git a/Engine/Plugins/Runtime/nDisplay/Source/DisplayClusterConfigurator/Private/Views/Details/Policies/DisplayClusterConfiguratorPolicyDetailCustomization.cpp b/Engine/Plugins/Runtime/nDisplay/Source/DisplayClusterConfigurator/Private/Views/Details/Policies/DisplayClusterConfiguratorPolicyDetailCustomization.cpp
--- a/Engine/Plugins/Runtime/nDisplay/Source/DisplayClusterConfigurator/Private/Views/Details/Policies/DisplayClusterConfiguratorPolicyDetailCustomization.cpp
+++ b/Engine/Plugins/Runtime/nDisplay/Source/DisplayClusterConfigurator/Private/Views/Details/Policies/DisplayClusterConfiguratorPolicyDetailCustomization.cpp
@@ -453,12 +453,30 @@ void FDisplayClusterConfiguratorProjectionCustomization::CreateMeshPolicy(UDispl
 		ConfigurationViewports,
 		TArray<TSubclassOf<UActorComponent>>{ UStaticMeshComponent::StaticClass(), UProceduralMeshComponent::StaticClass()}));
 
-	CustomPolicyParameters.Add(MakeShared<FPolicyParameterInfoNumber<int32>>(
-		"SectionIndex",
+	TSharedPtr<FPolicyParameterInfo> SectionIndexParameterRef = MakeShared<FPolicyParameterInfoNumber<int32>>(
+		"Section Index",
 		DisplayClusterProjectionStrings::cfg::mesh::SectionIndex,
 		Blueprint,
-		ConfigurationViewports));
+		ConfigurationViewports, 0, 0);
+	SectionIndexParameterRef->SetParameterTooltip(LOCTEXT("MeshPolicySectionIndexTooltip", "Section index with geometry in the ProceduralMesh component"));
+
+	TSharedPtr<FPolicyParameterInfo> BaseUVIndexParameterRef = MakeShared<FPolicyParameterInfoNumber<int32>>(
+		"Warp UVs",
+		DisplayClusterProjectionStrings::cfg::mesh::BaseUVIndex,
+		Blueprint,
+		ConfigurationViewports, -1, -1);
+	BaseUVIndexParameterRef->SetParameterTooltip(LOCTEXT("MeshPolicyBaseUVIndexTooltip", "Define custom UV channel in the original geometry as UV for warp. The default is -1."));
+
+	TSharedPtr<FPolicyParameterInfo> ChromakeyUVIndexParameterRef = MakeShared<FPolicyParameterInfoNumber<int32>>(
+		"Chromakey UVs",
+		DisplayClusterProjectionStrings::cfg::mesh::ChromakeyUVIndex,
+		Blueprint,
+		ConfigurationViewports, -1, -1);
+	ChromakeyUVIndexParameterRef->SetParameterTooltip(LOCTEXT("MeshPolicyChromakeyUVIndexTooltip", "Define custom UV channel in the original geometry as UV for ChromakeyMarkers. The default is -1."));
 
+	CustomPolicyParameters.Add(SectionIndexParameterRef);
+	CustomPolicyParameters.Add(BaseUVIndexParameterRef);
+	CustomPolicyParameters.Add(ChromakeyUVIndexParameterRef);
 }
 
 void FDisplayClusterConfiguratorProjectionCustomization::CreateDomePolicy(UDisplayClusterBlueprint* Blueprint)
diff --git a/Engine/Plugins/Runtime/nDisplay/Source/DisplayClusterProjection/Private/Policy/MPCDI/DisplayClusterProjectionMPCDIPolicy.cpp b/Engine/Plugins/Runtime/nDisplay/Source/DisplayClusterProjection/Private/Policy/MPCDI/DisplayClusterProjectionMPCDIPolicy.cpp
--- a/Engine/Plugins/Runtime/nDisplay/Source/DisplayClusterProjection/Private/Policy/MPCDI/DisplayClusterProjectionMPCDIPolicy.cpp
+++ b/Engine/Plugins/Runtime/nDisplay/Source/DisplayClusterProjection/Private/Policy/MPCDI/DisplayClusterProjectionMPCDIPolicy.cpp
@@ -81,7 +81,7 @@ bool FDisplayClusterProjectionMPCDIPolicy::HandleStartScene(IDisplayClusterViewp
 
 	WarpBlendContexts.Empty();
 
-	if (WarpBlendInterface.IsValid() == false && !CreateWarpBlendFromConfig())
+	if (WarpBlendInterface.IsValid() == false && !CreateWarpBlendFromConfig(InViewport))
 	{
 		// Ignore broken MPCDI config for other attempts
 		bInvalidConfiguration = true;
@@ -94,9 +94,6 @@ bool FDisplayClusterProjectionMPCDIPolicy::HandleStartScene(IDisplayClusterViewp
 		return false;
 	}
 
-	// Find origin component if it exists
-	InitializeOriginComponent(InViewport, OriginCompId);
-
 	// Finally, initialize internal views data container
 	WarpBlendContexts.AddDefaulted(2);
 
@@ -370,7 +367,7 @@ void FDisplayClusterProjectionMPCDIPolicy::ReleasePreviewMeshComponent()
 	PreviewMeshComponentRef.ResetSceneComponent();
 }
 
-UMeshComponent* FDisplayClusterProjectionMPCDIPolicy::GetOrCreatePreviewMeshComponent(class IDisplayClusterViewport* InViewport, bool& bOutIsRootActorComponent)
+UMeshComponent* FDisplayClusterProjectionMPCDIPolicy::GetOrCreatePreviewMeshComponent(IDisplayClusterViewport* InViewport, bool& bOutIsRootActorComponent)
 {
 	check(IsInGameThread());
 
@@ -426,7 +423,7 @@ UMeshComponent* FDisplayClusterProjectionMPCDIPolicy::GetOrCreatePreviewMeshComp
 
 #endif
 
-bool FDisplayClusterProjectionMPCDIPolicy::CreateWarpBlendFromConfig()
+bool FDisplayClusterProjectionMPCDIPolicy::CreateWarpBlendFromConfig(IDisplayClusterViewport* InViewport)
 {
 	check(IsInGameThread());
 
@@ -438,7 +435,7 @@ bool FDisplayClusterProjectionMPCDIPolicy::CreateWarpBlendFromConfig()
 		IDisplayClusterShaders& ShadersAPI = IDisplayClusterShaders::Get();
 
 		// Support custom origin node
-		OriginCompId = CfgData.OriginType;
+		InitializeOriginComponent(InViewport, CfgData.OriginType);
 
 		bIsPreviewMeshEnabled = CfgData.bEnablePreview;
 
diff --git a/Engine/Plugins/Runtime/nDisplay/Source/DisplayClusterProjection/Private/Policy/MPCDI/DisplayClusterProjectionMPCDIPolicy.h b/Engine/Plugins/Runtime/nDisplay/Source/DisplayClusterProjection/Private/Policy/MPCDI/DisplayClusterProjectionMPCDIPolicy.h
--- a/Engine/Plugins/Runtime/nDisplay/Source/DisplayClusterProjection/Private/Policy/MPCDI/DisplayClusterProjectionMPCDIPolicy.h
+++ b/Engine/Plugins/Runtime/nDisplay/Source/DisplayClusterProjection/Private/Policy/MPCDI/DisplayClusterProjectionMPCDIPolicy.h
@@ -7,8 +7,11 @@
 #include "WarpBlend/IDisplayClusterWarpBlend.h"
 #include "WarpBlend/DisplayClusterWarpContext.h"
 
+class IDisplayClusterViewport;
+
 /**
  * MPCDI projection policy
+ * Supported load from 'MPCDI' and 'PFM' files
  */
 class FDisplayClusterProjectionMPCDIPolicy
 	: public FDisplayClusterProjectionPolicyBase
@@ -43,11 +46,11 @@ public:
 	//////////////////////////////////////////////////////////////////////////////////////////////
 	// IDisplayClusterProjectionPolicy
 	//////////////////////////////////////////////////////////////////////////////////////////////
-	virtual bool HandleStartScene(class IDisplayClusterViewport* InViewport) override;
-	virtual void HandleEndScene(class IDisplayClusterViewport* InViewport) override;
+	virtual bool HandleStartScene(IDisplayClusterViewport* InViewport) override;
+	virtual void HandleEndScene(IDisplayClusterViewport* InViewport) override;
 
-	virtual bool CalculateView(class IDisplayClusterViewport* InViewport, const uint32 InContextNum, FVector& InOutViewLocation, FRotator& InOutViewRotation, const FVector& ViewOffset, const float WorldToMeters, const float NCP, const float FCP) override;
-	virtual bool GetProjectionMatrix(class IDisplayClusterViewport* InViewport, const uint32 InContextNum, FMatrix& OutPrjMatrix) override;
+	virtual bool CalculateView(IDisplayClusterViewport* InViewport, const uint32 InContextNum, FVector& InOutViewLocation, FRotator& InOutViewRotation, const FVector& ViewOffset, const float WorldToMeters, const float NCP, const float FCP) override;
+	virtual bool GetProjectionMatrix(IDisplayClusterViewport* InViewport, const uint32 InContextNum, FMatrix& OutPrjMatrix) override;
 
 	virtual bool IsWarpBlendSupported() override;
 	virtual void ApplyWarpBlend_RenderThread(FRHICommandListImmediate& RHICmdList, const class IDisplayClusterViewportProxy* InViewportProxy) override;
@@ -61,15 +64,13 @@ public:
 		return true;
 	}
 
-	virtual void UpdateProxyData(class IDisplayClusterViewport* InViewport) override;
+	virtual void UpdateProxyData(IDisplayClusterViewport* InViewport) override;
 
 protected:
-	bool CreateWarpBlendFromConfig();
+	bool CreateWarpBlendFromConfig(IDisplayClusterViewport* InViewport);
 	void ImplRelease();
 
 protected:
-	FString OriginCompId;
-	
 	TSharedPtr<IDisplayClusterWarpBlend, ESPMode::ThreadSafe> WarpBlendInterface;
 	TArray<FDisplayClusterWarpContext> WarpBlendContexts;
 
@@ -89,7 +90,7 @@ protected:
 	{
 		return true;
 	}
-	virtual class UMeshComponent* GetOrCreatePreviewMeshComponent(class IDisplayClusterViewport* InViewport, bool& bOutIsRootActorComponent) override;
+	virtual class UMeshComponent* GetOrCreatePreviewMeshComponent(IDisplayClusterViewport* InViewport, bool& bOutIsRootActorComponent) override;
 
 	void ReleasePreviewMeshComponent();
 
diff --git a/Engine/Plugins/Runtime/nDisplay/Source/DisplayClusterProjection/Private/Policy/Mesh/DisplayClusterProjectionMeshPolicy.cpp b/Engine/Plugins/Runtime/nDisplay/Source/DisplayClusterProjection/Private/Policy/Mesh/DisplayClusterProjectionMeshPolicy.cpp
--- a/Engine/Plugins/Runtime/nDisplay/Source/DisplayClusterProjection/Private/Policy/Mesh/DisplayClusterProjectionMeshPolicy.cpp
+++ b/Engine/Plugins/Runtime/nDisplay/Source/DisplayClusterProjection/Private/Policy/Mesh/DisplayClusterProjectionMeshPolicy.cpp
@@ -22,36 +22,50 @@ FDisplayClusterProjectionMeshPolicy::FDisplayClusterProjectionMeshPolicy(const F
 {
 }
 
-bool FDisplayClusterProjectionMeshPolicy::CreateWarpMeshInterface(class IDisplayClusterViewport* InViewport)
+bool FDisplayClusterProjectionMeshPolicy::CreateWarpMeshInterface(IDisplayClusterViewport* InViewport)
 {
 	check(IsInGameThread());
 
 	if (WarpBlendInterface.IsValid() == false)
 	{
+		// Read configuration from config
 		FWarpMeshConfiguration WarpCfg;
 		if (!GetWarpMeshConfiguration(InViewport, WarpCfg))
 		{
 			return false;
 		}
 
+		// The mesh always uses DCRootActor as its origin because all geometry is in UE space.
+		InitializeOriginComponent(InViewport, TEXT(""));
+
 		IDisplayClusterShaders& ShadersAPI = IDisplayClusterShaders::Get();
 		bool bResult = false;
 
 		if (WarpCfg.StaticMeshComponent != nullptr)
 		{
-			FDisplayClusterWarpBlendConstruct::FAssignWarpMesh CreateParameters;
-			CreateParameters.MeshComponent   = WarpCfg.StaticMeshComponent;
-			CreateParameters.OriginComponent = WarpCfg.OriginComponent;
+			FDisplayClusterWarpBlendConstruct::FAssignWarpStaticMesh CreateParameters;
+
+			CreateParameters.OriginComponent = GetOriginComp();
+
+			CreateParameters.StaticMeshComponent = WarpCfg.StaticMeshComponent;
+			CreateParameters.StaticMeshComponentLODIndex = WarpCfg.StaticMeshComponentLODIndex;
+
+			CreateParameters.BaseUVIndex      = WarpCfg.BaseUVIndex;
+			CreateParameters.ChromakeyUVIndex = WarpCfg.ChromakeyUVIndex;
 
 			bResult = ShadersAPI.GetWarpBlendManager().Create(CreateParameters, WarpBlendInterface);
 		}
 		else
 		{
 			FDisplayClusterWarpBlendConstruct::FAssignWarpProceduralMesh CreateParameters;
-			CreateParameters.OriginComponent = WarpCfg.OriginComponent;
+
+			CreateParameters.OriginComponent = GetOriginComp();
 
 			CreateParameters.ProceduralMeshComponent = WarpCfg.ProceduralMeshComponent;
-			CreateParameters.ProceduralMeshComponentSectionIndex = WarpCfg.SectionIndex;
+			CreateParameters.ProceduralMeshComponentSectionIndex = WarpCfg.ProceduralMeshComponentSectionIndex;
+
+			CreateParameters.BaseUVIndex      = WarpCfg.BaseUVIndex;
+			CreateParameters.ChromakeyUVIndex = WarpCfg.ChromakeyUVIndex;
 
 			bResult = ShadersAPI.GetWarpBlendManager().Create(CreateParameters, WarpBlendInterface);
 		}
@@ -70,7 +84,7 @@ bool FDisplayClusterProjectionMeshPolicy::CreateWarpMeshInterface(class IDisplay
 	return true;
 }
 
-bool FDisplayClusterProjectionMeshPolicy::HandleStartScene(class IDisplayClusterViewport* InViewport)
+bool FDisplayClusterProjectionMeshPolicy::HandleStartScene(IDisplayClusterViewport* InViewport)
 {
 	check(IsInGameThread());
 
@@ -79,9 +93,6 @@ bool FDisplayClusterProjectionMeshPolicy::HandleStartScene(class IDisplayCluster
 
 	WarpBlendContexts.Empty();
 
-	// Find origin component if it exists
-	InitializeOriginComponent(InViewport, OriginCompId);
-
 	if (!CreateWarpMeshInterface(InViewport))
 	{
 		if (!IsEditorOperationMode())
@@ -130,9 +141,56 @@ bool FDisplayClusterProjectionMeshPolicy::GetWarpMeshConfiguration(IDisplayClust
 		return false;
 	}
 
+	int32 CfgBaseUVIndex = INDEX_NONE;
+	if (DisplayClusterHelpers::map::template ExtractValueFromString(GetParameters(), DisplayClusterProjectionStrings::cfg::mesh::BaseUVIndex, CfgBaseUVIndex))
+	{
+		if (CfgBaseUVIndex >= 0)
+		{
+			UE_LOG(LogDisplayClusterProjectionMesh, Verbose, TEXT("Found BaseUVIndex value - '%d'"), CfgBaseUVIndex);
+
+			OutWarpCfg.BaseUVIndex = CfgBaseUVIndex;
+		}
+		else
+		{
+			UE_LOG(LogDisplayClusterProjectionMesh, Error, TEXT("Invalid BaseUVIndex value - '%d'"), CfgBaseUVIndex);
+		}
+	}
+
+	int32 CfgChromakeyUVIndex = INDEX_NONE;
+	if (DisplayClusterHelpers::map::template ExtractValueFromString(GetParameters(), DisplayClusterProjectionStrings::cfg::mesh::ChromakeyUVIndex, CfgChromakeyUVIndex))
+	{
+		if (CfgChromakeyUVIndex >= 0)
+		{
+			UE_LOG(LogDisplayClusterProjectionMesh, Verbose, TEXT("Found ChromakeyUVIndex value - '%d'"), CfgChromakeyUVIndex);
+
+			OutWarpCfg.ChromakeyUVIndex = CfgChromakeyUVIndex;
+		}
+		else
+		{
+			UE_LOG(LogDisplayClusterProjectionMesh, Error, TEXT("Invalid ChromakeyUVIndex value - '%d'"), CfgChromakeyUVIndex);
+		}
+	}
+
 	// Get the StaticMeshComponent
 	OutWarpCfg.StaticMeshComponent = Root->GetComponentByName<UStaticMeshComponent>(ComponentId);
-	if (OutWarpCfg.StaticMeshComponent == nullptr)
+	if (OutWarpCfg.StaticMeshComponent != nullptr)
+	{
+		int CfgLODIndex;
+		if (DisplayClusterHelpers::map::template ExtractValueFromString(GetParameters(), DisplayClusterProjectionStrings::cfg::mesh::LODIndex, CfgLODIndex))
+		{
+			if (CfgLODIndex >= 0)
+			{
+				UE_LOG(LogDisplayClusterProjectionMesh, Verbose, TEXT("Found StaticMeshComponent LODIndex value - '%d'"), CfgLODIndex);
+
+				OutWarpCfg.StaticMeshComponentLODIndex = CfgLODIndex;
+			}
+			else
+			{
+				UE_LOG(LogDisplayClusterProjectionMesh, Error, TEXT("Invalid StaticMeshComponent LODIndex value - '%d'"), CfgLODIndex);
+			}
+		}
+	}
+	else
 	{
 		// Get the ProceduralMeshComponent
 		OutWarpCfg.ProceduralMeshComponent = Root->GetComponentByName<UProceduralMeshComponent>(ComponentId);
@@ -140,30 +198,28 @@ bool FDisplayClusterProjectionMeshPolicy::GetWarpMeshConfiguration(IDisplayClust
 		{
 			if (!IsEditorOperationMode())
 			{
-				UE_LOG(LogDisplayClusterProjectionMesh, Warning, TEXT("Couldn't initialize mesh component '%s'"), *ComponentId);
+				UE_LOG(LogDisplayClusterProjectionMesh, Warning, TEXT("Couldn't find mesh component '%s' in RootActor"), *ComponentId);
 			}
 
 			return false;
 		}
 
-		int CfgSectionIndex;
+		int32 CfgSectionIndex;
 		if (DisplayClusterHelpers::map::template ExtractValueFromString(GetParameters(), DisplayClusterProjectionStrings::cfg::mesh::SectionIndex, CfgSectionIndex))
 		{
 			if (CfgSectionIndex >= 0)
 			{
-				UE_LOG(LogDisplayClusterProjectionMesh, Verbose, TEXT("Found SectionIndex value - '%d'"), CfgSectionIndex);
+				UE_LOG(LogDisplayClusterProjectionMesh, Verbose, TEXT("Found ProceduralMeshComponent SectionIndex value - '%d'"), CfgSectionIndex);
 
-				OutWarpCfg.SectionIndex = CfgSectionIndex;
+				OutWarpCfg.ProceduralMeshComponentSectionIndex = CfgSectionIndex;
 			}
 			else
 			{
-				UE_LOG(LogDisplayClusterProjectionMesh, Error, TEXT("Invalid SectionIndex value - '%d'"), CfgSectionIndex);
+				UE_LOG(LogDisplayClusterProjectionMesh, Error, TEXT("Invalid ProceduralMeshComponent SectionIndex value - '%d'"), CfgSectionIndex);
 			}
 		}
 	}
 
-	OutWarpCfg.OriginComponent = Root->GetRootComponent();
-
 	return true;
 }
 
diff --git a/Engine/Plugins/Runtime/nDisplay/Source/DisplayClusterProjection/Private/Policy/Mesh/DisplayClusterProjectionMeshPolicy.h b/Engine/Plugins/Runtime/nDisplay/Source/DisplayClusterProjection/Private/Policy/Mesh/DisplayClusterProjectionMeshPolicy.h
--- a/Engine/Plugins/Runtime/nDisplay/Source/DisplayClusterProjection/Private/Policy/Mesh/DisplayClusterProjectionMeshPolicy.h
+++ b/Engine/Plugins/Runtime/nDisplay/Source/DisplayClusterProjection/Private/Policy/Mesh/DisplayClusterProjectionMeshPolicy.h
@@ -9,8 +9,9 @@ class UStaticMeshComponent;
 class UProceduralMeshComponent;
 class USceneComponent;
 
-/**
- * Adapter for the Mesh and ProceduralMesh warp
+/*
+ * Mesh projection policy
+ * Supported geometry sources - StaticMeshComponent, ProceduralMeshComponent
  */
 class FDisplayClusterProjectionMeshPolicy
 	: public FDisplayClusterProjectionMPCDIPolicy
@@ -36,10 +37,19 @@ public:
 private:
 	struct FWarpMeshConfiguration
 	{
-		USceneComponent*          OriginComponent = nullptr;
+		// StaticMesh component with source geometry
 		UStaticMeshComponent*     StaticMeshComponent = nullptr;
+		// StaticMesh geometry LOD
+		int32 StaticMeshComponentLODIndex = 0;
+
+		// ProceduralMesh component with source geometry
 		UProceduralMeshComponent* ProceduralMeshComponent = nullptr;
-		int32 SectionIndex = 0;
+		// ProceduralMesh section index
+		int32 ProceduralMeshComponentSectionIndex = 0;
+
+		// Customize source geometry UV channels
+		int32 BaseUVIndex = INDEX_NONE;
+		int32 ChromakeyUVIndex = INDEX_NONE;
 	};
 
 	bool GetWarpMeshConfiguration(IDisplayClusterViewport* InViewport, FWarpMeshConfiguration& OutWarpCfg);
diff --git a/Engine/Plugins/Runtime/nDisplay/Source/DisplayClusterProjection/Public/DisplayClusterProjectionStrings.h b/Engine/Plugins/Runtime/nDisplay/Source/DisplayClusterProjection/Public/DisplayClusterProjectionStrings.h
--- a/Engine/Plugins/Runtime/nDisplay/Source/DisplayClusterProjection/Public/DisplayClusterProjectionStrings.h
+++ b/Engine/Plugins/Runtime/nDisplay/Source/DisplayClusterProjection/Public/DisplayClusterProjectionStrings.h
@@ -76,8 +76,13 @@ namespace DisplayClusterProjectionStrings
 		
 		namespace mesh
 		{
-			static constexpr auto Component    = TEXT("mesh_component");
-			static constexpr auto SectionIndex = TEXT("section_index");
+			static constexpr auto Component        = TEXT("mesh_component");
+
+			static constexpr auto LODIndex         = TEXT("lod_index");
+			static constexpr auto SectionIndex     = TEXT("section_index");
+
+			static constexpr auto BaseUVIndex      = TEXT("base_uv_index");
+			static constexpr auto ChromakeyUVIndex = TEXT("chromakey_uv_index");
 		}
 
 		namespace mpcdi
diff --git a/Engine/Plugins/Runtime/nDisplay/Source/DisplayClusterShaders/Private/DisplayClusterShadersModule.cpp b/Engine/Plugins/Runtime/nDisplay/Source/DisplayClusterShaders/Private/DisplayClusterShadersModule.cpp
--- a/Engine/Plugins/Runtime/nDisplay/Source/DisplayClusterShaders/Private/DisplayClusterShadersModule.cpp
+++ b/Engine/Plugins/Runtime/nDisplay/Source/DisplayClusterShaders/Private/DisplayClusterShadersModule.cpp
@@ -41,7 +41,7 @@ bool FDisplayClusterShadersModule::RenderWarpBlend_ICVFX(FRHICommandListImmediat
 	return FDisplayClusterShadersWarpblend_ICVFX::RenderWarpBlend_ICVFX(RHICmdList, InWarpBlendParameters, InICVFXParameters);
 }
 
-bool FDisplayClusterShadersModule::RenderPostprocess_OutputRemap(FRHICommandListImmediate& RHICmdList, FRHITexture2D* InSourceTexture, FRHITexture2D* InRenderTargetableDestTexture, const FDisplayClusterRender_MeshComponentProxy& MeshProxy) const
+bool FDisplayClusterShadersModule::RenderPostprocess_OutputRemap(FRHICommandListImmediate& RHICmdList, FRHITexture2D* InSourceTexture, FRHITexture2D* InRenderTargetableDestTexture, const IDisplayClusterRender_MeshComponentProxy& MeshProxy) const
 {
 	return FDisplayClusterShadersPostprocess_OutputRemap::RenderPostprocess_OutputRemap(RHICmdList, InSourceTexture, InRenderTargetableDestTexture, MeshProxy);
 }
diff --git a/Engine/Plugins/Runtime/nDisplay/Source/DisplayClusterShaders/Private/DisplayClusterShadersModule.h b/Engine/Plugins/Runtime/nDisplay/Source/DisplayClusterShaders/Private/DisplayClusterShadersModule.h
--- a/Engine/Plugins/Runtime/nDisplay/Source/DisplayClusterShaders/Private/DisplayClusterShadersModule.h
+++ b/Engine/Plugins/Runtime/nDisplay/Source/DisplayClusterShaders/Private/DisplayClusterShadersModule.h
@@ -19,7 +19,7 @@ public:
 public:
 	virtual bool RenderWarpBlend_MPCDI(FRHICommandListImmediate& RHICmdList, const FDisplayClusterShaderParameters_WarpBlend& InWarpBlendParameters) const override;
 	virtual bool RenderWarpBlend_ICVFX(FRHICommandListImmediate& RHICmdList, const FDisplayClusterShaderParameters_WarpBlend& InWarpBlendParameters, const FDisplayClusterShaderParameters_ICVFX& InICVFXParameters) const override;
-	virtual bool RenderPostprocess_OutputRemap(FRHICommandListImmediate& RHICmdList, FRHITexture2D* InSourceTexture, FRHITexture2D* InRenderTargetableDestTexture, const FDisplayClusterRender_MeshComponentProxy& MeshProxy) const override;
+	virtual bool RenderPostprocess_OutputRemap(FRHICommandListImmediate& RHICmdList, FRHITexture2D* InSourceTexture, FRHITexture2D* InRenderTargetableDestTexture, const IDisplayClusterRender_MeshComponentProxy& MeshProxy) const override;
 	virtual bool RenderPostprocess_Blur(FRHICommandListImmediate& RHICmdList, FRHITexture2D* InSourceTexture, FRHITexture2D* InRenderTargetableDestTexture, const FDisplayClusterShaderParameters_PostprocessBlur& InSettings) const override;
 	virtual bool GenerateMips(FRHICommandListImmediate& RHICmdList, FRHITexture2D* InOutMipsTexture, const FDisplayClusterShaderParameters_GenerateMips& InSettings) const override;
 
diff --git a/Engine/Plugins/Runtime/nDisplay/Source/DisplayClusterShaders/Private/Shaders/DisplayClusterShadersPostprocess_OutputRemap.cpp b/Engine/Plugins/Runtime/nDisplay/Source/DisplayClusterShaders/Private/Shaders/DisplayClusterShadersPostprocess_OutputRemap.cpp
--- a/Engine/Plugins/Runtime/nDisplay/Source/DisplayClusterShaders/Private/Shaders/DisplayClusterShadersPostprocess_OutputRemap.cpp
+++ b/Engine/Plugins/Runtime/nDisplay/Source/DisplayClusterShaders/Private/Shaders/DisplayClusterShadersPostprocess_OutputRemap.cpp
@@ -22,11 +22,13 @@
 
 #include "ClearQuad.h"
 
-#include "Render/Containers/DisplayClusterRender_MeshComponent.h"
-#include "Render/Containers/DisplayClusterRender_MeshComponentProxy.h"
-#include "Render/Containers/DisplayClusterRender_MeshComponentProxyData.h"
+#include "Render/Containers/IDisplayClusterRender_MeshComponent.h"
+#include "Render/Containers/IDisplayClusterRender_MeshComponentProxy.h"
 #include "Render/Containers/DisplayClusterRender_MeshGeometry.h"
 
+#include "IDisplayCluster.h"
+#include "Render/IDisplayClusterRenderManager.h"
+
 #define OutputRemapShaderFileName TEXT("/Plugin/nDisplay/Private/OutputRemapShaders.usf")
 
 // Select output remap shader
@@ -87,7 +89,7 @@ IMPLEMENT_SHADER_TYPE(, FOutputRemapVS, OutputRemapShaderFileName, TEXT("OutputR
 
 DECLARE_GPU_STAT_NAMED(nDisplay_PostProcess_OutputRemap, TEXT("nDisplay PostProcess::OutputRemap"));
 
-bool FDisplayClusterShadersPostprocess_OutputRemap::RenderPostprocess_OutputRemap(FRHICommandListImmediate& RHICmdList, FRHITexture2D* InSourceTexture, FRHITexture2D* InRenderTargetableDestTexture, const FDisplayClusterRender_MeshComponentProxy& InMeshProxy)
+bool FDisplayClusterShadersPostprocess_OutputRemap::RenderPostprocess_OutputRemap(FRHICommandListImmediate& RHICmdList, FRHITexture2D* InSourceTexture, FRHITexture2D* InRenderTargetableDestTexture, const IDisplayClusterRender_MeshComponentProxy& InMeshProxy)
 {
 	check(IsInRenderingThread());
 
@@ -96,7 +98,7 @@ bool FDisplayClusterShadersPostprocess_OutputRemap::RenderPostprocess_OutputRema
 		return false;
 	}
 
-	FDisplayClusterRender_MeshComponentProxy const* MeshProxy = &InMeshProxy;
+	IDisplayClusterRender_MeshComponentProxy const* MeshProxy = &InMeshProxy;
 
 	const EVarOutputRemapShaderType ShaderType = (EVarOutputRemapShaderType)CVarOutputRemapShaderType.GetValueOnAnyThread();
 	switch (ShaderType)
@@ -104,22 +106,24 @@ bool FDisplayClusterShadersPostprocess_OutputRemap::RenderPostprocess_OutputRema
 		case EVarOutputRemapShaderType::Passthrough:
 		{
 			// Use simple 1:1 test mesh for shader forwarding
-			static FDisplayClusterRender_MeshComponent TestMesh_Passthrough;
+			static TSharedPtr<IDisplayClusterRender_MeshComponent, ESPMode::ThreadSafe> TestMesh_Passthrough = IDisplayCluster::Get().GetRenderMgr()->CreateMeshComponent();
 
-			FDisplayClusterRender_MeshComponentProxy* TestMeshMeshComponentProxy = TestMesh_Passthrough.GetMeshComponentProxy_RenderThread();
-			if (TestMeshMeshComponentProxy != nullptr)
+			if (TestMesh_Passthrough.IsValid())
 			{
-				if (!TestMeshMeshComponentProxy->IsValid_RenderThread())
+				IDisplayClusterRender_MeshComponentProxy* TestMeshMeshComponentProxy = TestMesh_Passthrough->GetMeshComponentProxy_RenderThread();
+				if (TestMeshMeshComponentProxy != nullptr)
 				{
-					// Initialize once:
-					FDisplayClusterRender_MeshGeometry PassthroughMeshGeometry(EDisplayClusterRender_MeshGeometryCreateType::Passthrough);
-					FDisplayClusterRender_MeshComponentProxyData ProxyData(EDisplayClusterRender_MeshComponentProxyDataFunc::OutputRemapScreenSpace, PassthroughMeshGeometry);
-
-					TestMeshMeshComponentProxy->UpdateRHI_RenderThread(RHICmdList, &ProxyData);
+					if (!TestMeshMeshComponentProxy->IsEnabled_RenderThread())
+					{
+						// Initialize once:
+						FDisplayClusterRender_MeshGeometry PassthroughMeshGeometry(EDisplayClusterRender_MeshGeometryCreateType::Passthrough);
+						TestMesh_Passthrough->AssignMeshGeometry_RenderThread(&PassthroughMeshGeometry, EDisplayClusterRender_MeshComponentProxyDataFunc::OutputRemapScreenSpace);
+					}
+
+					MeshProxy = TestMeshMeshComponentProxy;
 				}
-
-				MeshProxy = TestMeshMeshComponentProxy;
 			}
+
 			break;
 		}
 
diff --git a/Engine/Plugins/Runtime/nDisplay/Source/DisplayClusterShaders/Private/Shaders/DisplayClusterShadersPostprocess_OutputRemap.h b/Engine/Plugins/Runtime/nDisplay/Source/DisplayClusterShaders/Private/Shaders/DisplayClusterShadersPostprocess_OutputRemap.h
--- a/Engine/Plugins/Runtime/nDisplay/Source/DisplayClusterShaders/Private/Shaders/DisplayClusterShadersPostprocess_OutputRemap.h
+++ b/Engine/Plugins/Runtime/nDisplay/Source/DisplayClusterShaders/Private/Shaders/DisplayClusterShadersPostprocess_OutputRemap.h
@@ -6,10 +6,10 @@
 #include "RHI.h"
 #include "RHICommandList.h"
 
-class FDisplayClusterRender_MeshComponentProxy;
+class IDisplayClusterRender_MeshComponentProxy;
 
 class FDisplayClusterShadersPostprocess_OutputRemap
 {
 public:
-	static bool RenderPostprocess_OutputRemap(FRHICommandListImmediate& RHICmdList, FRHITexture2D* InSourceTexture, FRHITexture2D* InRenderTargetableDestTexture, const FDisplayClusterRender_MeshComponentProxy& MeshProxy);
+	static bool RenderPostprocess_OutputRemap(FRHICommandListImmediate& RHICmdList, FRHITexture2D* InSourceTexture, FRHITexture2D* InRenderTargetableDestTexture, const IDisplayClusterRender_MeshComponentProxy& MeshProxy);
 };
diff --git a/Engine/Plugins/Runtime/nDisplay/Source/DisplayClusterShaders/Private/Shaders/DisplayClusterShadersWarpblend_ICVFX.cpp b/Engine/Plugins/Runtime/nDisplay/Source/DisplayClusterShaders/Private/Shaders/DisplayClusterShadersWarpblend_ICVFX.cpp
--- a/Engine/Plugins/Runtime/nDisplay/Source/DisplayClusterShaders/Private/Shaders/DisplayClusterShadersWarpblend_ICVFX.cpp
+++ b/Engine/Plugins/Runtime/nDisplay/Source/DisplayClusterShaders/Private/Shaders/DisplayClusterShadersWarpblend_ICVFX.cpp
@@ -14,8 +14,7 @@
 #include "ShaderParameterStruct.h"
 #include "ShaderPermutation.h"
 
-#include "Render/Containers/DisplayClusterRender_MeshComponent.h"
-#include "Render/Containers/DisplayClusterRender_MeshComponentProxy.h"
+#include "Render/Containers/IDisplayClusterRender_MeshComponentProxy.h"
 
 #include "ShaderParameters/DisplayClusterShaderParameters_WarpBlend.h"
 #include "ShaderParameters/DisplayClusterShaderParameters_ICVFX.h"
@@ -662,7 +661,7 @@ public:
 			case EDisplayClusterWarpGeometryType::WarpMesh:
 			case EDisplayClusterWarpGeometryType::WarpProceduralMesh:
 			{
-				const FDisplayClusterRender_MeshComponentProxy* WarpMeshProxy = WarpBlendParameters.WarpInterface->GetWarpMeshProxy_RenderThread();
+				const IDisplayClusterRender_MeshComponentProxy* WarpMeshProxy = WarpBlendParameters.WarpInterface->GetWarpMeshProxy_RenderThread();
 				if (WarpMeshProxy != nullptr)
 				{
 					return WarpMeshProxy->BeginRender_RenderThread(RHICmdList, GraphicsPSOInit);
@@ -691,7 +690,7 @@ public:
 			case EDisplayClusterWarpGeometryType::WarpMesh:
 			case EDisplayClusterWarpGeometryType::WarpProceduralMesh:
 			{
-				const FDisplayClusterRender_MeshComponentProxy* WarpMeshProxy = WarpBlendParameters.WarpInterface->GetWarpMeshProxy_RenderThread();
+				const IDisplayClusterRender_MeshComponentProxy* WarpMeshProxy = WarpBlendParameters.WarpInterface->GetWarpMeshProxy_RenderThread();
 				if (WarpMeshProxy != nullptr)
 				{
 					return WarpMeshProxy->FinishRender_RenderThread(RHICmdList);
diff --git a/Engine/Plugins/Runtime/nDisplay/Source/DisplayClusterShaders/Private/Shaders/DisplayClusterShadersWarpblend_MPCDI.cpp b/Engine/Plugins/Runtime/nDisplay/Source/DisplayClusterShaders/Private/Shaders/DisplayClusterShadersWarpblend_MPCDI.cpp
--- a/Engine/Plugins/Runtime/nDisplay/Source/DisplayClusterShaders/Private/Shaders/DisplayClusterShadersWarpblend_MPCDI.cpp
+++ b/Engine/Plugins/Runtime/nDisplay/Source/DisplayClusterShaders/Private/Shaders/DisplayClusterShadersWarpblend_MPCDI.cpp
@@ -18,8 +18,7 @@
 #include "ShaderPermutation.h"
 #include "ShaderParameterStruct.h"
 
-#include "Render/Containers/DisplayClusterRender_MeshComponent.h"
-#include "Render/Containers/DisplayClusterRender_MeshComponentProxy.h"
+#include "Render/Containers/IDisplayClusterRender_MeshComponentProxy.h"
 #include "WarpBlend/IDisplayClusterWarpBlend.h"
 
 #include "ShaderParameters/DisplayClusterShaderParameters_WarpBlend.h"
@@ -348,7 +347,7 @@ public:
 			case EDisplayClusterWarpGeometryType::WarpMesh:
 			case EDisplayClusterWarpGeometryType::WarpProceduralMesh:
 			{
-				const FDisplayClusterRender_MeshComponentProxy* WarpMeshProxy = WarpBlendParameters.WarpInterface->GetWarpMeshProxy_RenderThread();
+				const IDisplayClusterRender_MeshComponentProxy* WarpMeshProxy = WarpBlendParameters.WarpInterface->GetWarpMeshProxy_RenderThread();
 				if (WarpMeshProxy != nullptr)
 				{
 					return WarpMeshProxy->BeginRender_RenderThread(RHICmdList, GraphicsPSOInit);
@@ -377,7 +376,7 @@ public:
 			case EDisplayClusterWarpGeometryType::WarpMesh:
 			case EDisplayClusterWarpGeometryType::WarpProceduralMesh:
 			{
-				const FDisplayClusterRender_MeshComponentProxy* WarpMeshProxy = WarpBlendParameters.WarpInterface->GetWarpMeshProxy_RenderThread();
+				const IDisplayClusterRender_MeshComponentProxy* WarpMeshProxy = WarpBlendParameters.WarpInterface->GetWarpMeshProxy_RenderThread();
 				if (WarpMeshProxy != nullptr)
 				{
 					return WarpMeshProxy->FinishRender_RenderThread(RHICmdList);
diff --git a/Engine/Plugins/Runtime/nDisplay/Source/DisplayClusterShaders/Private/WarpBlend/DisplayClusterWarpBlend.h b/Engine/Plugins/Runtime/nDisplay/Source/DisplayClusterShaders/Private/WarpBlend/DisplayClusterWarpBlend.h
--- a/Engine/Plugins/Runtime/nDisplay/Source/DisplayClusterShaders/Private/WarpBlend/DisplayClusterWarpBlend.h
+++ b/Engine/Plugins/Runtime/nDisplay/Source/DisplayClusterShaders/Private/WarpBlend/DisplayClusterWarpBlend.h
@@ -39,7 +39,7 @@ public:
 		return GeometryContext.GeometryProxy.AlphaMapEmbeddedGamma; 
 	}
 
-	virtual const FDisplayClusterRender_MeshComponentProxy* GetWarpMeshProxy_RenderThread() const override
+	virtual const IDisplayClusterRender_MeshComponentProxy* GetWarpMeshProxy_RenderThread() const override
 	{
 		return GeometryContext.GeometryProxy.GetWarpMeshProxy_RenderThread();
 	}
diff --git a/Engine/Plugins/Runtime/nDisplay/Source/DisplayClusterShaders/Private/WarpBlend/DisplayClusterWarpBlendManager.cpp b/Engine/Plugins/Runtime/nDisplay/Source/DisplayClusterShaders/Private/WarpBlend/DisplayClusterWarpBlendManager.cpp
--- a/Engine/Plugins/Runtime/nDisplay/Source/DisplayClusterShaders/Private/WarpBlend/DisplayClusterWarpBlendManager.cpp
+++ b/Engine/Plugins/Runtime/nDisplay/Source/DisplayClusterShaders/Private/WarpBlend/DisplayClusterWarpBlendManager.cpp
@@ -16,7 +16,7 @@ bool FDisplayClusterWarpBlendManager::Create(const FDisplayClusterWarpBlendConst
 	return FDisplayClusterWarpBlendLoader_MPCDI::Load(InConstructParameters, OutWarpBlend);
 }
 
-bool  FDisplayClusterWarpBlendManager::Create(const FDisplayClusterWarpBlendConstruct::FAssignWarpMesh& InConstructParameters, TSharedPtr<IDisplayClusterWarpBlend, ESPMode::ThreadSafe>& OutWarpBlend) const
+bool  FDisplayClusterWarpBlendManager::Create(const FDisplayClusterWarpBlendConstruct::FAssignWarpStaticMesh& InConstructParameters, TSharedPtr<IDisplayClusterWarpBlend, ESPMode::ThreadSafe>& OutWarpBlend) const
 {
 	return FDisplayClusterWarpBlendLoader_MeshComponent::Load(InConstructParameters, OutWarpBlend);
 };
diff --git a/Engine/Plugins/Runtime/nDisplay/Source/DisplayClusterShaders/Private/WarpBlend/DisplayClusterWarpBlendManager.h b/Engine/Plugins/Runtime/nDisplay/Source/DisplayClusterShaders/Private/WarpBlend/DisplayClusterWarpBlendManager.h
--- a/Engine/Plugins/Runtime/nDisplay/Source/DisplayClusterShaders/Private/WarpBlend/DisplayClusterWarpBlendManager.h
+++ b/Engine/Plugins/Runtime/nDisplay/Source/DisplayClusterShaders/Private/WarpBlend/DisplayClusterWarpBlendManager.h
@@ -10,6 +10,6 @@ class FDisplayClusterWarpBlendManager
 public:
 	virtual bool Create(const FDisplayClusterWarpBlendConstruct::FLoadMPCDIFile&            InConstructParameters, TSharedPtr<IDisplayClusterWarpBlend, ESPMode::ThreadSafe>& OutWarpBlend) const override;
 	virtual bool Create(const FDisplayClusterWarpBlendConstruct::FLoadPFMFile&              InConstructParameters, TSharedPtr<IDisplayClusterWarpBlend, ESPMode::ThreadSafe>& OutWarpBlend) const override;
-	virtual bool Create(const FDisplayClusterWarpBlendConstruct::FAssignWarpMesh&           InConstructParameters, TSharedPtr<IDisplayClusterWarpBlend, ESPMode::ThreadSafe>& OutWarpBlend) const override;
+	virtual bool Create(const FDisplayClusterWarpBlendConstruct::FAssignWarpStaticMesh&     InConstructParameters, TSharedPtr<IDisplayClusterWarpBlend, ESPMode::ThreadSafe>& OutWarpBlend) const override;
 	virtual bool Create(const FDisplayClusterWarpBlendConstruct::FAssignWarpProceduralMesh& InConstructParameters, TSharedPtr<IDisplayClusterWarpBlend, ESPMode::ThreadSafe>& OutWarpBlend) const override;
 };
diff --git a/Engine/Plugins/Runtime/nDisplay/Source/DisplayClusterShaders/Private/WarpBlend/DisplayClusterWarpBlend_GeometryProxy.cpp b/Engine/Plugins/Runtime/nDisplay/Source/DisplayClusterShaders/Private/WarpBlend/DisplayClusterWarpBlend_GeometryProxy.cpp
--- a/Engine/Plugins/Runtime/nDisplay/Source/DisplayClusterShaders/Private/WarpBlend/DisplayClusterWarpBlend_GeometryProxy.cpp
+++ b/Engine/Plugins/Runtime/nDisplay/Source/DisplayClusterShaders/Private/WarpBlend/DisplayClusterWarpBlend_GeometryProxy.cpp
@@ -2,8 +2,7 @@
 
 #include "WarpBlend/DisplayClusterWarpBlend_GeometryProxy.h"
 
-#include "Render/DisplayClusterRenderTexture.h"
-#include "Render/Containers/DisplayClusterRender_MeshComponent.h"
+#include "Render/Containers/IDisplayClusterRender_MeshComponent.h"
 
 #include "WarpBlend/Math/DisplayClusterWarpBlendMath_WarpMesh.h"
 #include "WarpBlend/Math/DisplayClusterWarpBlendMath_WarpProceduralMesh.h"
@@ -12,7 +11,6 @@
 /////////////////////////////////////////////////////////////////////////////////
 /// FDisplayClusterWarpBlend_GeometryProxy
 /////////////////////////////////////////////////////////////////////////////////
-
 bool FDisplayClusterWarpBlend_GeometryProxy::UpdateGeometry()
 {
 	switch (GeometryType)
@@ -34,7 +32,7 @@ bool FDisplayClusterWarpBlend_GeometryProxy::UpdateGeometry()
 	return false;
 }
 
-const FDisplayClusterRender_MeshComponentProxy* FDisplayClusterWarpBlend_GeometryProxy::GetWarpMeshProxy_RenderThread() const
+const IDisplayClusterRender_MeshComponentProxy* FDisplayClusterWarpBlend_GeometryProxy::GetWarpMeshProxy_RenderThread() const
 {
 	check(IsInRenderingThread());
 
@@ -115,7 +113,7 @@ bool FDisplayClusterWarpBlend_GeometryProxy::ImplUpdateGeometry_WarpMesh()
 	// If StaticMesh geometry changed, update mpcdi math and RHI resources
 	if (MeshComponent->IsMeshComponentRefGeometryDirty() || bIsMeshComponentLost)
 	{
-		MeshComponent->AssignStaticMeshComponentRefs(StaticMeshComponent, OriginComponent, StaticMeshComponentLODIndex);
+		MeshComponent->AssignStaticMeshComponentRefs(StaticMeshComponent, WarpMeshUVs, OriginComponent, StaticMeshComponentLODIndex);
 		bIsMeshComponentLost = false;
 	}
 	
@@ -164,7 +162,7 @@ bool FDisplayClusterWarpBlend_GeometryProxy::ImplUpdateGeometry_WarpProceduralMe
 	// If ProceduralMesh geometry changed, update mpcdi math and RHI resources
 	if (MeshComponent->IsMeshComponentRefGeometryDirty() || bIsMeshComponentLost)
 	{
-		MeshComponent->AssignProceduralMeshComponentRefs(ProceduralMeshComponent, OriginComponent, ProceduralMeshComponentSectionIndex);
+		MeshComponent->AssignProceduralMeshComponentRefs(ProceduralMeshComponent, WarpMeshUVs, OriginComponent, ProceduralMeshComponentSectionIndex);
 		bIsMeshComponentLost = false;
 	}
 
diff --git a/Engine/Plugins/Runtime/nDisplay/Source/DisplayClusterShaders/Private/WarpBlend/DisplayClusterWarpBlend_GeometryProxy.h b/Engine/Plugins/Runtime/nDisplay/Source/DisplayClusterShaders/Private/WarpBlend/DisplayClusterWarpBlend_GeometryProxy.h
--- a/Engine/Plugins/Runtime/nDisplay/Source/DisplayClusterShaders/Private/WarpBlend/DisplayClusterWarpBlend_GeometryProxy.h
+++ b/Engine/Plugins/Runtime/nDisplay/Source/DisplayClusterShaders/Private/WarpBlend/DisplayClusterWarpBlend_GeometryProxy.h
@@ -3,11 +3,10 @@
 #pragma once
 #include "CoreMinimal.h"
 #include "WarpBlend/DisplayClusterWarpEnums.h"
-#include "Render/Containers/DisplayClusterRender_MeshComponent.h"
+#include "Render/Containers/IDisplayClusterRender_MeshComponent.h"
 #include "Render/IDisplayClusterRenderTexture.h"
 
-class IDisplayClusterRenderTexture;
-class FDisplayClusterRender_MeshComponentProxy;
+class IDisplayClusterRender_MeshComponentProxy;
 
 struct FStaticMeshLODResources;
 struct FProcMeshSection;
@@ -46,7 +45,7 @@ public:
 	bool UpdateGeometry();
 	bool UpdateGeometryLOD(const FIntPoint& InSizeLOD);
 
-	const FDisplayClusterRender_MeshComponentProxy* GetWarpMeshProxy_RenderThread() const;
+	const IDisplayClusterRender_MeshComponentProxy* GetWarpMeshProxy_RenderThread() const;
 
 	const FStaticMeshLODResources* GetStaticMeshComponentLODResources() const;
 	const FProcMeshSection* GetProceduralMeshComponentSection() const;
@@ -70,7 +69,9 @@ public:
 	TUniquePtr<IDisplayClusterRenderTexture> AlphaMapTexture;
 	TUniquePtr<IDisplayClusterRenderTexture> BetaMapTexture;
 
-	TUniquePtr<FDisplayClusterRender_MeshComponent> MeshComponent;
+	TSharedPtr<IDisplayClusterRender_MeshComponent, ESPMode::ThreadSafe> MeshComponent;
+
+	FDisplayClusterMeshUVs WarpMeshUVs;
 
 	// for huge warp mesh geometry, change this value, and use LOD geometry in frustum math
 	int32 StaticMeshComponentLODIndex = 0;
diff --git a/Engine/Plugins/Runtime/nDisplay/Source/DisplayClusterShaders/Private/WarpBlend/Loader/DisplayClusterWarpBlendLoader_MeshComponent.cpp b/Engine/Plugins/Runtime/nDisplay/Source/DisplayClusterShaders/Private/WarpBlend/Loader/DisplayClusterWarpBlendLoader_MeshComponent.cpp
--- a/Engine/Plugins/Runtime/nDisplay/Source/DisplayClusterShaders/Private/WarpBlend/Loader/DisplayClusterWarpBlendLoader_MeshComponent.cpp
+++ b/Engine/Plugins/Runtime/nDisplay/Source/DisplayClusterShaders/Private/WarpBlend/Loader/DisplayClusterWarpBlendLoader_MeshComponent.cpp
@@ -12,16 +12,16 @@
 #include "Misc/ConfigCacheIni.h"
 #include "Misc/Paths.h"
 
-#include "Render/Containers/DisplayClusterRender_MeshComponent.h"
-
 #include "WarpBlend/Loader/DisplayClusterWarpBlendLoader_Texture.h"
 #include "WarpBlend/DisplayClusterWarpBlend.h"
 #include "WarpBlend/DisplayClusterWarpBlend_GeometryContext.h"
 #include "WarpBlend/DisplayClusterWarpBlend_GeometryProxy.h"
 
-bool FDisplayClusterWarpBlendLoader_MeshComponent::Load(const FDisplayClusterWarpBlendConstruct::FAssignWarpMesh& InParameters, TSharedPtr<IDisplayClusterWarpBlend, ESPMode::ThreadSafe>& OutWarpBlend)
+#include "Render/IDisplayClusterRenderManager.h"
+
+bool FDisplayClusterWarpBlendLoader_MeshComponent::Load(const FDisplayClusterWarpBlendConstruct::FAssignWarpStaticMesh& InParameters, TSharedPtr<IDisplayClusterWarpBlend, ESPMode::ThreadSafe>& OutWarpBlend)
 {
-	if (InParameters.MeshComponent != nullptr)
+	if (InParameters.StaticMeshComponent != nullptr)
 	{
 		//ok, Create and initialize warpblend interface:
 		TSharedPtr<FDisplayClusterWarpBlend, ESPMode::ThreadSafe> WarpBlend = MakeShared<FDisplayClusterWarpBlend, ESPMode::ThreadSafe>();
@@ -29,9 +29,14 @@ bool FDisplayClusterWarpBlendLoader_MeshComponent::Load(const FDisplayClusterWar
 		WarpBlend->GeometryContext.ProfileType = EDisplayClusterWarpProfileType::warp_A3D;
 
 		FDisplayClusterWarpBlend_GeometryProxy& Proxy = WarpBlend->GeometryContext.GeometryProxy;
-		Proxy.MeshComponent = MakeUnique<FDisplayClusterRender_MeshComponent>();
-		Proxy.MeshComponent->AssignStaticMeshComponentRefs(InParameters.MeshComponent, InParameters.OriginComponent, InParameters.StaticMeshComponentLODIndex);
+
+		FDisplayClusterMeshUVs MeshUVs(InParameters.BaseUVIndex, InParameters.ChromakeyUVIndex);
+
+		Proxy.MeshComponent = IDisplayCluster::Get().GetRenderMgr()->CreateMeshComponent();
+		Proxy.MeshComponent->AssignStaticMeshComponentRefs(InParameters.StaticMeshComponent, MeshUVs, InParameters.OriginComponent, InParameters.StaticMeshComponentLODIndex);
+
 		Proxy.StaticMeshComponentLODIndex = InParameters.StaticMeshComponentLODIndex;
+		Proxy.WarpMeshUVs = MeshUVs;
 
 		Proxy.GeometryType = EDisplayClusterWarpGeometryType::WarpMesh;
 
diff --git a/Engine/Plugins/Runtime/nDisplay/Source/DisplayClusterShaders/Private/WarpBlend/Loader/DisplayClusterWarpBlendLoader_MeshComponent.h b/Engine/Plugins/Runtime/nDisplay/Source/DisplayClusterShaders/Private/WarpBlend/Loader/DisplayClusterWarpBlendLoader_MeshComponent.h
--- a/Engine/Plugins/Runtime/nDisplay/Source/DisplayClusterShaders/Private/WarpBlend/Loader/DisplayClusterWarpBlendLoader_MeshComponent.h
+++ b/Engine/Plugins/Runtime/nDisplay/Source/DisplayClusterShaders/Private/WarpBlend/Loader/DisplayClusterWarpBlendLoader_MeshComponent.h
@@ -11,5 +11,5 @@ class IDisplayClusterWarpBlend;
 class FDisplayClusterWarpBlendLoader_MeshComponent
 {
 public:
-	static bool Load(const FDisplayClusterWarpBlendConstruct::FAssignWarpMesh& InParameters, TSharedPtr<IDisplayClusterWarpBlend, ESPMode::ThreadSafe>& OutWarpBlend);
+	static bool Load(const FDisplayClusterWarpBlendConstruct::FAssignWarpStaticMesh& InParameters, TSharedPtr<IDisplayClusterWarpBlend, ESPMode::ThreadSafe>& OutWarpBlend);
 };
diff --git a/Engine/Plugins/Runtime/nDisplay/Source/DisplayClusterShaders/Private/WarpBlend/Loader/DisplayClusterWarpBlendLoader_ProceduralMeshComponent.cpp b/Engine/Plugins/Runtime/nDisplay/Source/DisplayClusterShaders/Private/WarpBlend/Loader/DisplayClusterWarpBlendLoader_ProceduralMeshComponent.cpp
--- a/Engine/Plugins/Runtime/nDisplay/Source/DisplayClusterShaders/Private/WarpBlend/Loader/DisplayClusterWarpBlendLoader_ProceduralMeshComponent.cpp
+++ b/Engine/Plugins/Runtime/nDisplay/Source/DisplayClusterShaders/Private/WarpBlend/Loader/DisplayClusterWarpBlendLoader_ProceduralMeshComponent.cpp
@@ -17,7 +17,7 @@
 #include "WarpBlend/DisplayClusterWarpBlend_GeometryContext.h"
 #include "WarpBlend/DisplayClusterWarpBlend_GeometryProxy.h"
 
-#include "Render/Containers/DisplayClusterRender_MeshComponent.h"
+#include "Render/IDisplayClusterRenderManager.h"
 
 bool FDisplayClusterWarpBlendLoader_ProceduralMeshComponent::Load(const FDisplayClusterWarpBlendConstruct::FAssignWarpProceduralMesh& InParameters, TSharedPtr<IDisplayClusterWarpBlend, ESPMode::ThreadSafe>& OutWarpBlend)
 {
@@ -29,12 +29,15 @@ bool FDisplayClusterWarpBlendLoader_ProceduralMeshComponent::Load(const FDisplay
 		WarpBlend->GeometryContext.ProfileType = EDisplayClusterWarpProfileType::warp_A3D;
 
 		FDisplayClusterWarpBlend_GeometryProxy& Proxy = WarpBlend->GeometryContext.GeometryProxy;
-
+		
+		FDisplayClusterMeshUVs MeshUVs(InParameters.BaseUVIndex, InParameters.ChromakeyUVIndex);
 
 		// Assign procedural mesh
-		Proxy.MeshComponent = MakeUnique<FDisplayClusterRender_MeshComponent>();
-		Proxy.MeshComponent->AssignProceduralMeshComponentRefs(InParameters.ProceduralMeshComponent, InParameters.OriginComponent, InParameters.ProceduralMeshComponentSectionIndex);
+		Proxy.MeshComponent = IDisplayCluster::Get().GetRenderMgr()->CreateMeshComponent();
+		Proxy.MeshComponent->AssignProceduralMeshComponentRefs(InParameters.ProceduralMeshComponent, MeshUVs, InParameters.OriginComponent, InParameters.ProceduralMeshComponentSectionIndex);
+		
 		Proxy.ProceduralMeshComponentSectionIndex = InParameters.ProceduralMeshComponentSectionIndex;
+		Proxy.WarpMeshUVs = MeshUVs;
 
 		Proxy.GeometryType = EDisplayClusterWarpGeometryType::WarpProceduralMesh;
 
diff --git a/Engine/Plugins/Runtime/nDisplay/Source/DisplayClusterShaders/Private/WarpBlend/Math/DisplayClusterWarpBlendMath_Frustum.cpp b/Engine/Plugins/Runtime/nDisplay/Source/DisplayClusterShaders/Private/WarpBlend/Math/DisplayClusterWarpBlendMath_Frustum.cpp
--- a/Engine/Plugins/Runtime/nDisplay/Source/DisplayClusterShaders/Private/WarpBlend/Math/DisplayClusterWarpBlendMath_Frustum.cpp
+++ b/Engine/Plugins/Runtime/nDisplay/Source/DisplayClusterShaders/Private/WarpBlend/Math/DisplayClusterWarpBlendMath_Frustum.cpp
@@ -3,8 +3,6 @@
 #include "WarpBlend/Math/DisplayClusterWarpBlendMath_Frustum.h"
 
 #include "Render/IDisplayClusterRenderTexture.h"
-#include "Render/Containers/DisplayClusterRender_MeshComponent.h"
-
 #include "Misc/DisplayClusterHelpers.h"
 
 #include "Render/Viewport/IDisplayClusterViewport.h"
diff --git a/Engine/Plugins/Runtime/nDisplay/Source/DisplayClusterShaders/Public/IDisplayClusterShaders.h b/Engine/Plugins/Runtime/nDisplay/Source/DisplayClusterShaders/Public/IDisplayClusterShaders.h
--- a/Engine/Plugins/Runtime/nDisplay/Source/DisplayClusterShaders/Public/IDisplayClusterShaders.h
+++ b/Engine/Plugins/Runtime/nDisplay/Source/DisplayClusterShaders/Public/IDisplayClusterShaders.h
@@ -10,7 +10,7 @@
 #include "RHIResources.h"
 
 
-class FDisplayClusterRender_MeshComponentProxy;
+class IDisplayClusterRender_MeshComponentProxy;
 class FDisplayClusterShaderParameters_ICVFX;
 class FDisplayClusterShaderParameters_Override;
 struct FDisplayClusterShaderParameters_GenerateMips;
@@ -88,7 +88,7 @@ public:
 	*
 	* @return - true if success
 	*/
-	virtual bool RenderPostprocess_OutputRemap(FRHICommandListImmediate& RHICmdList, FRHITexture2D* InSourceTexture, FRHITexture2D* InRenderTargetableDestTexture, const FDisplayClusterRender_MeshComponentProxy& MeshProxy) const = 0;
+	virtual bool RenderPostprocess_OutputRemap(FRHICommandListImmediate& RHICmdList, FRHITexture2D* InSourceTexture, FRHITexture2D* InRenderTargetableDestTexture, const IDisplayClusterRender_MeshComponentProxy& MeshProxy) const = 0;
 
 	/**
 	* Render postprocess Blur
diff --git a/Engine/Plugins/Runtime/nDisplay/Source/DisplayClusterShaders/Public/WarpBlend/IDisplayClusterWarpBlend.h b/Engine/Plugins/Runtime/nDisplay/Source/DisplayClusterShaders/Public/WarpBlend/IDisplayClusterWarpBlend.h
--- a/Engine/Plugins/Runtime/nDisplay/Source/DisplayClusterShaders/Public/WarpBlend/IDisplayClusterWarpBlend.h
+++ b/Engine/Plugins/Runtime/nDisplay/Source/DisplayClusterShaders/Public/WarpBlend/IDisplayClusterWarpBlend.h
@@ -32,7 +32,7 @@ public:
 	virtual class FRHITexture* GetTexture(EDisplayClusterWarpBlendTextureType TextureType) const = 0;
 	virtual float              GetAlphaMapEmbeddedGamma() const = 0;
 
-	virtual const class FDisplayClusterRender_MeshComponentProxy* GetWarpMeshProxy_RenderThread() const = 0;
+	virtual const class IDisplayClusterRender_MeshComponentProxy* GetWarpMeshProxy_RenderThread() const = 0;
 
 	// Return current warp profile type
 	virtual EDisplayClusterWarpProfileType  GetWarpProfileType() const = 0;
diff --git a/Engine/Plugins/Runtime/nDisplay/Source/DisplayClusterShaders/Public/WarpBlend/IDisplayClusterWarpBlendManager.h b/Engine/Plugins/Runtime/nDisplay/Source/DisplayClusterShaders/Public/WarpBlend/IDisplayClusterWarpBlendManager.h
--- a/Engine/Plugins/Runtime/nDisplay/Source/DisplayClusterShaders/Public/WarpBlend/IDisplayClusterWarpBlendManager.h
+++ b/Engine/Plugins/Runtime/nDisplay/Source/DisplayClusterShaders/Public/WarpBlend/IDisplayClusterWarpBlendManager.h
@@ -11,44 +11,72 @@ class IDisplayClusterWarpBlend;
 
 struct FDisplayClusterWarpBlendConstruct
 {
+	// Initialize WarpBlend from MPCDI file
 	struct FLoadMPCDIFile
 	{
+		// The file with MPCDI calibration data (extension of this file '.mpcdi')
 		FString MPCDIFileName;
 
+		// Name of buffer (mpcdi internal topology)
 		FString BufferId;
+
+		// Name of region (mpcdi internal topology)
 		FString RegionId;
 	};
 
+	// Initialize WarpBlend from PFM file
 	struct FLoadPFMFile
 	{
-		EDisplayClusterWarpProfileType ProfileType;
+		// geometry profile type
+		EDisplayClusterWarpProfileType ProfileType = EDisplayClusterWarpProfileType::warp_3D;
 
+		// The file with region calibration data (extension of this file '.pfm')
 		FString PFMFileName;
+
+		// Geometry scale of this file
 		float   PFMScale = 1.f;
+
+		// Wrap axis to mpcdi space for unreal 
 		bool    bIsUnrealGameSpace = false;
 
+		// AlphaMap file (supported format: PNG)
 		FString AlphaMapFileName;
+		// Control alpha multiplier for AlphaMap texture
 		float   AlphaMapEmbeddedAlpha = 0.f;
 
+		// BetaMap file (supported format: PNG)
 		FString BetaMapFileName;
 	};
 
-	struct FAssignWarpMesh
+	struct FWarpMeshBase
+	{
+		// Origin component ptr
+		class USceneComponent* OriginComponent = nullptr;
+
+		// Map source mesh geometry UVs to WarpMesh UVs
+		int32 BaseUVIndex = INDEX_NONE;
+		int32 ChromakeyUVIndex = INDEX_NONE;
+	};
+
+	struct FAssignWarpStaticMesh : FWarpMeshBase
 	{
-		class UStaticMeshComponent* MeshComponent   = nullptr;
-		class USceneComponent*      OriginComponent = nullptr;
+		// StaticMesh component ptr (warp geometry data source)
+		class UStaticMeshComponent* StaticMeshComponent = nullptr;
+
+		// (optional) get geometry data from specified LOD index
 		int32 StaticMeshComponentLODIndex = 0;
 	};
 
-	struct FAssignWarpProceduralMesh
+	struct FAssignWarpProceduralMesh : FWarpMeshBase
 	{
+		// ProceduralMesh component ptr (warp geometry data source)
 		class UProceduralMeshComponent* ProceduralMeshComponent = nullptr;
-		class USceneComponent* OriginComponent = nullptr;
+
+		// (optional) get geometry data from specified SectionIndex
 		int32 ProceduralMeshComponentSectionIndex = 0;
 	};
 };
 
-
 class IDisplayClusterWarpBlendManager
 {
 public:
@@ -57,6 +85,6 @@ public:
 public:
 	virtual bool Create(const FDisplayClusterWarpBlendConstruct::FLoadMPCDIFile&            InConstructParameters, TSharedPtr<IDisplayClusterWarpBlend, ESPMode::ThreadSafe>& OutWarpBlend) const = 0;
 	virtual bool Create(const FDisplayClusterWarpBlendConstruct::FLoadPFMFile&              InConstructParameters, TSharedPtr<IDisplayClusterWarpBlend, ESPMode::ThreadSafe>& OutWarpBlend) const = 0;
-	virtual bool Create(const FDisplayClusterWarpBlendConstruct::FAssignWarpMesh&           InConstructParameters, TSharedPtr<IDisplayClusterWarpBlend, ESPMode::ThreadSafe>& OutWarpBlend) const = 0;
+	virtual bool Create(const FDisplayClusterWarpBlendConstruct::FAssignWarpStaticMesh&     InConstructParameters, TSharedPtr<IDisplayClusterWarpBlend, ESPMode::ThreadSafe>& OutWarpBlend) const = 0;
 	virtual bool Create(const FDisplayClusterWarpBlendConstruct::FAssignWarpProceduralMesh& InConstructParameters, TSharedPtr<IDisplayClusterWarpBlend, ESPMode::ThreadSafe>& OutWarpBlend) const = 0;
 };
