nDisplay: fixed non-Unity build

diff --git a/Engine/Plugins/Runtime/nDisplay/Source/DisplayCluster/Public/Render/Containers/IDisplayClusterRender_MeshComponent.h b/Engine/Plugins/Runtime/nDisplay/Source/DisplayCluster/Public/Render/Containers/IDisplayClusterRender_MeshComponent.h
--- a/Engine/Plugins/Runtime/nDisplay/Source/DisplayCluster/Public/Render/Containers/IDisplayClusterRender_MeshComponent.h
+++ b/Engine/Plugins/Runtime/nDisplay/Source/DisplayCluster/Public/Render/Containers/IDisplayClusterRender_MeshComponent.h
@@ -7,6 +7,7 @@
 #include "Render/Containers/DisplayClusterRender_MeshGeometry.h"
 
 class USceneComponent;
+class UStaticMesh;
 class UStaticMeshComponent;
 class UProceduralMeshComponent;
 
