[nDisplay]
[Configurator]

Added tooltip support for projection policy properties.

diff --git a/Engine/Plugins/Runtime/nDisplay/Source/DisplayClusterConfigurator/Private/Views/Details/Policies/DisplayClusterConfiguratorPolicyParameterCustomization.cpp b/Engine/Plugins/Runtime/nDisplay/Source/DisplayClusterConfigurator/Private/Views/Details/Policies/DisplayClusterConfiguratorPolicyParameterCustomization.cpp
--- a/Engine/Plugins/Runtime/nDisplay/Source/DisplayClusterConfigurator/Private/Views/Details/Policies/DisplayClusterConfiguratorPolicyParameterCustomization.cpp
+++ b/Engine/Plugins/Runtime/nDisplay/Source/DisplayClusterConfigurator/Private/Views/Details/Policies/DisplayClusterConfiguratorPolicyParameterCustomization.cpp
@@ -321,6 +321,7 @@ void FPolicyParameterInfoText::CreateCustomRowWidget(IDetailChildrenBuilder& InD
 			SNew(STextBlock)
 			.Font(IDetailLayoutBuilder::GetDetailFont())
 			.Text(GetParameterDisplayName())
+			.ToolTipText(this, &FPolicyParameterInfo::GetParameterTooltip)
 		]
 		.ValueContent()
 		[
@@ -349,6 +350,7 @@ void FPolicyParameterInfoBool::CreateCustomRowWidget(IDetailChildrenBuilder& InD
 			SNew(STextBlock)
 			.Font(IDetailLayoutBuilder::GetDetailFont())
 			.Text(GetParameterDisplayName())
+			.ToolTipText(this, &FPolicyParameterInfo::GetParameterTooltip)
 		]
 		.ValueContent()
 		[
@@ -385,6 +387,7 @@ void FPolicyParameterInfoFile::CreateCustomRowWidget(IDetailChildrenBuilder& InD
 			SNew(STextBlock)
 			.Font(IDetailLayoutBuilder::GetDetailFont())
 			.Text(GetParameterDisplayName())
+			.ToolTipText(this, &FPolicyParameterInfo::GetParameterTooltip)
 		]
 		.ValueContent()
 		[
@@ -552,6 +555,7 @@ void FPolicyParameterInfoMatrix::CreateCustomRowWidget(IDetailChildrenBuilder& I
 		SNew(STextBlock)
 		.Font(IDetailLayoutBuilder::GetDetailFont())
 		.Text(GetParameterDisplayName())
+		.ToolTipText(this, &FPolicyParameterInfo::GetParameterTooltip)
 		.Visibility(this, &FPolicyParameterInfoMatrix::IsParameterVisible)
 	]
 	.Visibility(MakeAttributeRaw(this, &FPolicyParameterInfoMatrix::IsParameterVisible));
@@ -765,6 +769,7 @@ void FPolicyParameterInfo4x4Matrix::CreateCustomRowWidget(IDetailChildrenBuilder
 		SNew(STextBlock)
 		.Font(IDetailLayoutBuilder::GetDetailFont())
 		.Text(GetParameterDisplayName())
+		.ToolTipText(this, &FPolicyParameterInfo::GetParameterTooltip)
 		.Visibility(this, &FPolicyParameterInfo4x4Matrix::IsParameterVisible)
 	]
 	.Visibility(MakeAttributeRaw(this, &FPolicyParameterInfo4x4Matrix::IsParameterVisible));
@@ -892,6 +897,7 @@ void FPolicyParameterInfoRotator::CreateCustomRowWidget(IDetailChildrenBuilder&
 		SNew(STextBlock)
 		.Font(IDetailLayoutBuilder::GetDetailFont())
 		.Text(LOCTEXT("RotationLabel", "Rotation"))
+		.ToolTipText(this, &FPolicyParameterInfo::GetParameterTooltip)
 	]
 	.ValueContent()
 	.MinDesiredWidth(375.0f)
@@ -1007,13 +1013,15 @@ FPolicyParameterInfoFrustumAngle::FPolicyParameterInfoFrustumAngle(const FString
 }
 
 void FPolicyParameterInfoFrustumAngle::CreateCustomRowWidget(IDetailChildrenBuilder& InDetailWidgetRow)
-{	InDetailWidgetRow.AddCustomRow(GetParameterDisplayName())
+{
+	InDetailWidgetRow.AddCustomRow(GetParameterDisplayName())
 	.Visibility(MakeAttributeRaw(this, &FPolicyParameterInfoFrustumAngle::IsParameterVisible))
 	.NameContent()
 	[
 		SNew(STextBlock)
 		.Font(IDetailLayoutBuilder::GetDetailFont())
 		.Text(GetParameterDisplayName())
+		.ToolTipText(this, &FPolicyParameterInfo::GetParameterTooltip)
 	]
 	.ValueContent()
 	.MinDesiredWidth(375.0f)
diff --git a/Engine/Plugins/Runtime/nDisplay/Source/DisplayClusterConfigurator/Private/Views/Details/Policies/DisplayClusterConfiguratorPolicyParameterCustomization.h b/Engine/Plugins/Runtime/nDisplay/Source/DisplayClusterConfigurator/Private/Views/Details/Policies/DisplayClusterConfiguratorPolicyParameterCustomization.h
--- a/Engine/Plugins/Runtime/nDisplay/Source/DisplayClusterConfigurator/Private/Views/Details/Policies/DisplayClusterConfiguratorPolicyParameterCustomization.h
+++ b/Engine/Plugins/Runtime/nDisplay/Source/DisplayClusterConfigurator/Private/Views/Details/Policies/DisplayClusterConfiguratorPolicyParameterCustomization.h
@@ -217,6 +217,7 @@ public:
 				SNew(STextBlock)
 				.Font(IDetailLayoutBuilder::GetDetailFont())
 				.Text(GetParameterDisplayName())
+				.ToolTipText(this, &FPolicyParameterInfo::GetParameterTooltip)
 			]
 			.ValueContent()
 			[
