[nDisplay]
[Configurator]

Code commented out because it causes a crash.
You cannot change the structure of the DCRootActor inside the customization delegate.
This causes a crash inside the LayoutData.ClassToPropertyMap iterator.

diff --git a/Engine/Plugins/Runtime/nDisplay/Source/DisplayClusterConfigurator/Private/Views/Details/DisplayClusterRootActorDetailsCustomization.cpp b/Engine/Plugins/Runtime/nDisplay/Source/DisplayClusterConfigurator/Private/Views/Details/DisplayClusterRootActorDetailsCustomization.cpp
--- a/Engine/Plugins/Runtime/nDisplay/Source/DisplayClusterConfigurator/Private/Views/Details/DisplayClusterRootActorDetailsCustomization.cpp
+++ b/Engine/Plugins/Runtime/nDisplay/Source/DisplayClusterConfigurator/Private/Views/Details/DisplayClusterRootActorDetailsCustomization.cpp
@@ -80,6 +80,13 @@ void FDisplayClusterRootActorDetailsCustomization::CustomizeDetails(IDetailLayou
 
 	InLayoutBuilder.SortCategories(DisplayClusterRootActorDetailsCustomizationUtils::SortCategories);
 	
+	/*
+	* The code below is commented out because it causes a crash.
+	* You cannot change the structure of the DCRootActor inside the customization delegate.
+	* This causes a crash inside the LayoutData.ClassToPropertyMap iterator.
+	* Customization called from  \\Engine\Source\Editor\PropertyEditor\Private\DetailLayoutHelpers.cpp:356
+	*/
+#if 0
 	// Manually add the transform properties' data to the layout builder's property in order to generate property handles for them.
 	{
 		TArray<UObject*> RootComponents;
@@ -98,6 +105,7 @@ void FDisplayClusterRootActorDetailsCustomization::CustomizeDetails(IDetailLayou
 		InLayoutBuilder.AddObjectPropertyData(RootComponents, USceneComponent::GetRelativeRotationPropertyName());
 		InLayoutBuilder.AddObjectPropertyData(RootComponents, USceneComponent::GetRelativeScale3DPropertyName());
 	}
+#endif
 
 	// Manually label the ICVFX category to propery format it to have the dash in "In-Camera"
 	InLayoutBuilder.EditCategory(DisplayClusterConfigurationStrings::categories::ICVFXCategory, LOCTEXT("ICVFXCategoryLabel", "In-Camera VFX"));
