nDisplay: Removing a primary cluster node will remove the primary status and set a new primary node. The compiler will now throw an error instead of a warning if no primary node is detected.

diff --git a/Engine/Plugins/Runtime/nDisplay/Source/DisplayClusterConfigurator/Private/ClusterConfiguration/DisplayClusterConfiguratorClusterUtils.cpp b/Engine/Plugins/Runtime/nDisplay/Source/DisplayClusterConfigurator/Private/ClusterConfiguration/DisplayClusterConfiguratorClusterUtils.cpp
--- a/Engine/Plugins/Runtime/nDisplay/Source/DisplayClusterConfigurator/Private/ClusterConfiguration/DisplayClusterConfiguratorClusterUtils.cpp
+++ b/Engine/Plugins/Runtime/nDisplay/Source/DisplayClusterConfigurator/Private/ClusterConfiguration/DisplayClusterConfiguratorClusterUtils.cpp
@@ -703,10 +703,27 @@ bool FDisplayClusterConfiguratorClusterUtils::RemoveClusterNodeFromCluster(UDisp
 			ClusterNode->Modify();
 
 			ClusterNodeParent->Modify();
+
+			const bool bRemovingPrimaryNode = IsClusterNodePrimary(ClusterNode);
+			if (bRemovingPrimaryNode)
+			{
+				ClusterNodeParent->PrimaryNode.Id.Empty();
+			}
+			
 			RemoveKeyFromMap(ClusterNodeParent, GET_MEMBER_NAME_CHECKED(UDisplayClusterConfigurationCluster, Nodes), *KeyPtr);
 
 			ClusterNode->Rename(nullptr, GetTransientPackage(), REN_DontCreateRedirectors | REN_ForceNoResetLoaders | REN_DoNotDirty);
 			ClusterNode->SetFlags(RF_Transient);
+
+			if (bRemovingPrimaryNode)
+			{
+				for (const TTuple<FString, UDisplayClusterConfigurationClusterNode*>& NodesKeyVal : ClusterNodeParent->Nodes)
+				{
+					SetClusterNodeAsPrimary(NodesKeyVal.Value);
+					break;
+				}
+			}
+			
 			FDisplayClusterConfiguratorUtils::MarkDisplayClusterBlueprintAsModified(ClusterNodeParent, true);
 			
 			return true;
diff --git a/Engine/Plugins/Runtime/nDisplay/Source/DisplayClusterConfigurator/Private/DisplayClusterConfiguratorCompiler.cpp b/Engine/Plugins/Runtime/nDisplay/Source/DisplayClusterConfigurator/Private/DisplayClusterConfiguratorCompiler.cpp
--- a/Engine/Plugins/Runtime/nDisplay/Source/DisplayClusterConfigurator/Private/DisplayClusterConfiguratorCompiler.cpp
+++ b/Engine/Plugins/Runtime/nDisplay/Source/DisplayClusterConfigurator/Private/DisplayClusterConfiguratorCompiler.cpp
@@ -243,7 +243,7 @@ void FDisplayClusterConfiguratorKismetCompilerContext::ValidateConfiguration()
 	
 	if (!FDisplayClusterConfiguratorUtils::IsPrimaryNodeInConfig(BlueprintData))
 	{
-		MessageLog.Warning(*LOCTEXT("NoPrimaryNodeWarning", "Primary cluster node not set. Please set a primary node.").ToString());
+		MessageLog.Error(*LOCTEXT("NoPrimaryNodeError", "Primary cluster node not set. Please set a primary node.").ToString());
 	}
 
 	bool bAtLeastOneViewportFound = false;
