nDisplay: fixed ICVFX Color Grading

diff --git a/Engine/Plugins/Runtime/nDisplay/Source/DisplayCluster/Private/Render/Viewport/DisplayClusterViewportStereoscopicPass.h b/Engine/Plugins/Runtime/nDisplay/Source/DisplayCluster/Private/Render/Viewport/DisplayClusterViewportStereoscopicPass.h
--- a/Engine/Plugins/Runtime/nDisplay/Source/DisplayCluster/Private/Render/Viewport/DisplayClusterViewportStereoscopicPass.h
+++ b/Engine/Plugins/Runtime/nDisplay/Source/DisplayCluster/Private/Render/Viewport/DisplayClusterViewportStereoscopicPass.h
@@ -9,7 +9,7 @@ struct FDisplayClusterViewportStereoscopicPass
 {
 	inline static EStereoscopicPass EncodeStereoscopicPass(const uint32 ContextNum, const uint32 ContextAmmount)
 	{
-		EStereoscopicPass EncodedPass = EStereoscopicPass::eSSP_FULL;
+		EStereoscopicPass EncodedPass = EStereoscopicPass::eSSP_PRIMARY;
 
 		// Monoscopic rendering
 		if (ContextAmmount == 1)
