nDisplay: Add struct and property core redirectors for master -> primary

diff --git a/Engine/Plugins/Runtime/nDisplay/Config/BasenDisplay.ini b/Engine/Plugins/Runtime/nDisplay/Config/BasenDisplay.ini
new file mode 100644
--- /dev/null
+++ b/Engine/Plugins/Runtime/nDisplay/Config/BasenDisplay.ini
@@ -0,0 +1,5 @@
+[CoreRedirects]
++StructRedirects=(OldName="DisplayClusterConfigurationMasterNode",NewName="DisplayClusterConfigurationPrimaryNode")
++StructRedirects=(OldName="DisplayClusterConfigurationMasterNodePorts",NewName="DisplayClusterConfigurationPrimaryNodePorts")
+
++PropertyRedirects=(OldName="DisplayClusterConfigurationCluster.MasterNode",NewName="DisplayClusterConfigurationCluster.PrimaryNode")
\ No newline at end of file
