nDisplay: Updated nDisplay plugin icons to new SVG icons, refactored the plugin's editor style set to adhere to UE5 recommendations, removed unused icons and style elements, and tweaked various UI element bahaviors and appearances in the nDisplay Config Editor to more closely adhere UE5 style guidelines.

diff --git a/Engine/Plugins/Runtime/nDisplay/Source/DisplayCluster/Private/Game/EngineClasses/Scene/DisplayClusterPreviewComponent.cpp b/Engine/Plugins/Runtime/nDisplay/Source/DisplayCluster/Private/Game/EngineClasses/Scene/DisplayClusterPreviewComponent.cpp
--- a/Engine/Plugins/Runtime/nDisplay/Source/DisplayCluster/Private/Game/EngineClasses/Scene/DisplayClusterPreviewComponent.cpp
+++ b/Engine/Plugins/Runtime/nDisplay/Source/DisplayCluster/Private/Game/EngineClasses/Scene/DisplayClusterPreviewComponent.cpp
@@ -95,7 +95,7 @@ void UDisplayClusterPreviewComponent::UpdatePreviewMeshMaterial(bool bRestoreOri
 
 	if (PreviewMesh)
 	{
-		bool bViewportPreviewEnabled = (ViewportConfig && ViewportConfig->bIsEnabled && RootActor && RootActor->bPreviewEnable);
+		bool bViewportPreviewEnabled = (ViewportConfig && RootActor && RootActor->bPreviewEnable);
 
 		if (bRestoreOriginalMaterial || !bViewportPreviewEnabled)
 		{
diff --git a/Engine/Plugins/Runtime/nDisplay/Source/DisplayClusterConfiguration/Private/DisplayClusterConfigurationTypes.cpp b/Engine/Plugins/Runtime/nDisplay/Source/DisplayClusterConfiguration/Private/DisplayClusterConfigurationTypes.cpp
--- a/Engine/Plugins/Runtime/nDisplay/Source/DisplayClusterConfiguration/Private/DisplayClusterConfigurationTypes.cpp
+++ b/Engine/Plugins/Runtime/nDisplay/Source/DisplayClusterConfiguration/Private/DisplayClusterConfigurationTypes.cpp
@@ -169,7 +169,7 @@ UDisplayClusterConfigurationViewport::UDisplayClusterConfigurationViewport()
 {
 #if WITH_EDITORONLY_DATA
 	bIsVisible = true;
-	bIsEnabled = true;
+	bIsUnlocked = true;
 #endif
 }
 
@@ -201,14 +201,14 @@ UDisplayClusterConfigurationClusterNode::UDisplayClusterConfigurationClusterNode
 	: bIsSoundEnabled(false)
 #if WITH_EDITORONLY_DATA
 	, bIsVisible(true)
-	, bIsEnabled(true)
+	, bIsUnlocked(true)
 #endif
 {
 }
 
 UDisplayClusterConfigurationHostDisplayData::UDisplayClusterConfigurationHostDisplayData()
 	: bIsVisible(true)
-	, bIsEnabled(true)
+	, bIsUnlocked(true)
 {
 	SetFlags(RF_Public);
 }
diff --git a/Engine/Plugins/Runtime/nDisplay/Source/DisplayClusterConfiguration/Public/DisplayClusterConfigurationTypes.h b/Engine/Plugins/Runtime/nDisplay/Source/DisplayClusterConfiguration/Public/DisplayClusterConfigurationTypes.h
--- a/Engine/Plugins/Runtime/nDisplay/Source/DisplayClusterConfiguration/Public/DisplayClusterConfigurationTypes.h
+++ b/Engine/Plugins/Runtime/nDisplay/Source/DisplayClusterConfiguration/Public/DisplayClusterConfigurationTypes.h
@@ -340,7 +340,7 @@ public:
 	bool bIsVisible;
 
 	UPROPERTY(EditDefaultsOnly, Category = "Configuration", meta = (HideProperty))
-	bool bIsEnabled;
+	bool bIsUnlocked;
 
 	/** Binds a background preview image for easier output mapping */
 	UPROPERTY(EditDefaultsOnly, Category = "Configuration")
@@ -404,7 +404,7 @@ public:
 	bool bIsVisible;
 
 	UPROPERTY()
-	bool bIsEnabled;
+	bool bIsUnlocked;
 };
 
 UCLASS(Blueprintable)
diff --git a/Engine/Plugins/Runtime/nDisplay/Source/DisplayClusterConfiguration/Public/DisplayClusterConfigurationTypes_Viewport.h b/Engine/Plugins/Runtime/nDisplay/Source/DisplayClusterConfiguration/Public/DisplayClusterConfigurationTypes_Viewport.h
--- a/Engine/Plugins/Runtime/nDisplay/Source/DisplayClusterConfiguration/Public/DisplayClusterConfigurationTypes_Viewport.h
+++ b/Engine/Plugins/Runtime/nDisplay/Source/DisplayClusterConfiguration/Public/DisplayClusterConfigurationTypes_Viewport.h
@@ -200,7 +200,7 @@ public:
 
 #if WITH_EDITORONLY_DATA
 	UPROPERTY(EditDefaultsOnly, Category = "Configuration", meta = (HideProperty))
-	bool bIsEnabled = true;
+	bool bIsUnlocked = true;
 
 	UPROPERTY(EditDefaultsOnly, Category = "Configuration", meta = (HideProperty))
 	bool bIsVisible = true;
diff --git a/Engine/Plugins/Runtime/nDisplay/Source/DisplayClusterConfigurator/Private/ClusterConfiguration/DisplayClusterConfiguratorClusterUtils.cpp b/Engine/Plugins/Runtime/nDisplay/Source/DisplayClusterConfigurator/Private/ClusterConfiguration/DisplayClusterConfiguratorClusterUtils.cpp
--- a/Engine/Plugins/Runtime/nDisplay/Source/DisplayClusterConfigurator/Private/ClusterConfiguration/DisplayClusterConfiguratorClusterUtils.cpp
+++ b/Engine/Plugins/Runtime/nDisplay/Source/DisplayClusterConfigurator/Private/ClusterConfiguration/DisplayClusterConfiguratorClusterUtils.cpp
@@ -472,7 +472,7 @@ UDisplayClusterConfigurationHostDisplayData* FDisplayClusterConfiguratorClusterU
 		const FString HostName = GetUniqueNameForHost(DefaultNewHostName, Cluster, true);
 		UDisplayClusterConfigurationHostDisplayData* NewData = NewObject<UDisplayClusterConfigurationHostDisplayData>(Cluster, NAME_None, RF_Transactional);
 		NewData->HostName = FText::FromString(HostName);
-		NewData->Color = FDisplayClusterConfiguratorStyle::GetCornerColor(Cluster->HostDisplayData.Num()).CopyWithNewOpacity(1.0f);
+		NewData->Color = FDisplayClusterConfiguratorStyle::Get().GetDefaultColor(Cluster->HostDisplayData.Num()).CopyWithNewOpacity(1.0f);
 
 		Cluster->HostDisplayData.Add(HostIPAddress, NewData);
 	}
diff --git a/Engine/Plugins/Runtime/nDisplay/Source/DisplayClusterConfigurator/Private/DisplayClusterConfiguratorModule.cpp b/Engine/Plugins/Runtime/nDisplay/Source/DisplayClusterConfigurator/Private/DisplayClusterConfiguratorModule.cpp
--- a/Engine/Plugins/Runtime/nDisplay/Source/DisplayClusterConfigurator/Private/DisplayClusterConfiguratorModule.cpp
+++ b/Engine/Plugins/Runtime/nDisplay/Source/DisplayClusterConfigurator/Private/DisplayClusterConfiguratorModule.cpp
@@ -89,7 +89,7 @@ void FDisplayClusterConfiguratorModule::StartupModule()
 	RegisterSettings();
 	RegisterPlacementModeItems();
 	
-	FDisplayClusterConfiguratorStyle::Initialize();
+	FDisplayClusterConfiguratorStyle::Get();
 
 	FDisplayClusterConfiguratorCommands::Register();
 	FDisplayClusterConfiguratorOutputMappingCommands::Register();
@@ -126,7 +126,6 @@ void FDisplayClusterConfiguratorModule::ShutdownModule()
 	UnregisterCustomLayouts();
 	UnregisterPlacementModeItems();
 
-	FDisplayClusterConfiguratorStyle::Shutdown();
 	MenuExtensibilityManager.Reset();
 	ToolBarExtensibilityManager.Reset();
 
diff --git a/Engine/Plugins/Runtime/nDisplay/Source/DisplayClusterConfigurator/Private/DisplayClusterConfiguratorStyle.cpp b/Engine/Plugins/Runtime/nDisplay/Source/DisplayClusterConfigurator/Private/DisplayClusterConfiguratorStyle.cpp
--- a/Engine/Plugins/Runtime/nDisplay/Source/DisplayClusterConfigurator/Private/DisplayClusterConfiguratorStyle.cpp
+++ b/Engine/Plugins/Runtime/nDisplay/Source/DisplayClusterConfigurator/Private/DisplayClusterConfiguratorStyle.cpp
@@ -2,65 +2,37 @@
 
 #include "DisplayClusterConfiguratorStyle.h"
 
-#include "Brushes/SlateImageBrush.h"
-#include "Brushes/SlateBoxBrush.h"
-#include "Brushes/SlateBorderBrush.h"
 #include "EditorStyleSet.h"
 #include "Framework/Application/SlateApplication.h"
 #include "Interfaces/IPluginManager.h"
-#include "Styling/SlateStyle.h"
 #include "Styling/SlateStyleRegistry.h"
+#include "Styling/SlateStyleMacros.h"
 
-TSharedPtr<FSlateStyleSet> FDisplayClusterConfiguratorStyle::StyleInstance = nullptr;
-TArray<FDisplayClusterConfiguratorStyle::FCornerColor> FDisplayClusterConfiguratorStyle::CornerColors;
-
-void FDisplayClusterConfiguratorStyle::Initialize()
+FDisplayClusterConfiguratorStyle::FDisplayClusterConfiguratorStyle()
+	: FSlateStyleSet(TEXT("DisplayClusterConfiguratorStyle"))
 {
-	if (!StyleInstance.IsValid())
-	{
-		StyleInstance = Create();
-		FSlateStyleRegistry::RegisterSlateStyle(*StyleInstance);
-	}
-}
+	Initialize();
 
-void FDisplayClusterConfiguratorStyle::Shutdown()
-{
-	FSlateStyleRegistry::UnRegisterSlateStyle(*StyleInstance);
-	ensure(StyleInstance.IsUnique());
-	StyleInstance.Reset();
+	FSlateStyleRegistry::RegisterSlateStyle(*this);
 }
 
-FName FDisplayClusterConfiguratorStyle::GetStyleSetName()
+FDisplayClusterConfiguratorStyle::~FDisplayClusterConfiguratorStyle()
 {
-	static FName StyleSetName(TEXT("DisplayClusterConfiguratorStyle"));
-	return StyleSetName;
+	FSlateStyleRegistry::UnRegisterSlateStyle(*this);
 }
 
-const FLinearColor& FDisplayClusterConfiguratorStyle::GetColor(FName PropertyName, const ANSICHAR* Specifier)
-{
-	return StyleInstance->GetColor(PropertyName, Specifier);
-}
-
-const FLinearColor& FDisplayClusterConfiguratorStyle::GetCornerColor(uint32 Index)
-{
-	return CornerColors[Index % CornerColors.Num()].Color;
-}
-
-const FMargin& FDisplayClusterConfiguratorStyle::GetMargin(FName PropertyName, const ANSICHAR* Specifier)
+void FDisplayClusterConfiguratorStyle::ReloadTextures()
 {
-	return StyleInstance->GetMargin(PropertyName, Specifier);
+	FSlateApplication::Get().GetRenderer()->ReloadTextureResources();
 }
 
-const FSlateBrush* FDisplayClusterConfiguratorStyle::GetBrush(FName PropertyName, const ANSICHAR* Specifier)
+const FLinearColor& FDisplayClusterConfiguratorStyle::GetDefaultColor(uint32 Index)
 {
-	return StyleInstance->GetBrush(PropertyName, Specifier);
+	return DefaultColors[Index % DefaultColors.Num()].Color;
 }
 
-#define IMAGE_BRUSH( RelativePath, ... ) FSlateImageBrush( Style->RootToContentDir( RelativePath, TEXT(".png") ), __VA_ARGS__ )
-#define BOX_BRUSH( RelativePath, ... ) FSlateBoxBrush( Style->RootToContentDir( RelativePath, TEXT(".png") ), __VA_ARGS__ )
-#define BORDER_BRUSH( RelativePath, ... ) FSlateBorderBrush( Style->RootToContentDir( RelativePath, TEXT(".png") ), __VA_ARGS__ )
-#define EDITOR_BOX_BRUSH( RelativePath, ... ) FSlateBoxBrush( FPaths::EngineContentDir() / "Editor/Slate" / RelativePath + TEXT(".png"), __VA_ARGS__ )
-#define DEFAULT_FONT(...) FCoreStyle::GetDefaultFontStyle(__VA_ARGS__)
+#define EDITOR_IMAGE_BRUSH_SVG( RelativePath, ... ) FSlateVectorImageBrush(FPaths::EngineContentDir() / "Editor/Slate" / RelativePath + TEXT(".svg"), __VA_ARGS__ )
+#define EDITOR_BOX_BRUSH( RelativePath, ... ) FSlateBoxBrush(FPaths::EngineContentDir() / "Editor/Slate" / RelativePath + TEXT(".png"), __VA_ARGS__ )
 
 const FVector2D Icon128x128(128.0f, 128.0f);
 const FVector2D Icon64x64(64.0f, 64.0f);
@@ -68,189 +40,184 @@ const FVector2D Icon40x40(40.0f, 40.0f);
 const FVector2D Icon20x20(20.0f, 20.0f);
 const FVector2D Icon16x16(16.0f, 16.0f);
 const FVector2D Icon12x12(12.0f, 12.0f);
-TSharedRef< FSlateStyleSet > FDisplayClusterConfiguratorStyle::Create()
-{
-	TSharedRef<FSlateStyleSet> Style = MakeShared<FSlateStyleSet>("DisplayClusterConfiguratorStyle");
 
+void FDisplayClusterConfiguratorStyle::Initialize()
+{
 	TSharedPtr<IPlugin> Plugin = IPluginManager::Get().FindPlugin(TEXT("nDisplay"));
 	check(Plugin.IsValid());
 	if (Plugin.IsValid())
 	{
-		Style->SetContentRoot(FPaths::Combine(Plugin->GetBaseDir(), TEXT("Resources")));
+		SetContentRoot(FPaths::Combine(Plugin->GetBaseDir(), TEXT("Content")));
+		SetCoreContentRoot(FPaths::EngineContentDir() / TEXT("Slate"));
 	}
 
-	Style->Set("ClassIcon.DisplayClusterRootActor", new IMAGE_BRUSH("RootActor/nDisplay_x16", Icon16x16));
-	Style->Set("ClassThumbnail.DisplayClusterRootActor", new IMAGE_BRUSH("RootActor/nDisplay_x64", Icon64x64));
-
-	// Tabs
-	Style->Set("DisplayClusterConfigurator.Tabs.General", new IMAGE_BRUSH("Tabs/General_16x", Icon16x16));
-	Style->Set("DisplayClusterConfigurator.Tabs.Details", new IMAGE_BRUSH("Tabs/Details_16x", Icon16x16));
-	Style->Set("DisplayClusterConfigurator.Tabs.Log", new IMAGE_BRUSH("Tabs/Log_16x", Icon16x16));
-	Style->Set("DisplayClusterConfigurator.Tabs.OutputMapping", new IMAGE_BRUSH("Tabs/OutputMapping_16x", Icon16x16));
-	Style->Set("DisplayClusterConfigurator.Tabs.Scene", new IMAGE_BRUSH("Tabs/Scene_16x", Icon16x16));
-	Style->Set("DisplayClusterConfigurator.Tabs.Cluster", new IMAGE_BRUSH("Tabs/Cluster_16x", Icon16x16));
-	Style->Set("DisplayClusterConfigurator.Tabs.Viewport", new IMAGE_BRUSH("Tabs/Viewport_16x", Icon16x16));
-
-
-	// Toolbar
-	Style->Set("DisplayClusterConfigurator.Toolbar.Import", new IMAGE_BRUSH("Toolbar/Import_40x", Icon40x40));
-	Style->Set("DisplayClusterConfigurator.Toolbar.SaveToFile", new IMAGE_BRUSH("Toolbar/SaveToFile_40x", Icon40x40));
-	Style->Set("DisplayClusterConfigurator.Toolbar.EditConfig", new IMAGE_BRUSH("Toolbar/EditConfig_40x", Icon40x40));
-
-	// TreeItems
-	Style->Set("DisplayClusterConfigurator.TreeItems.Scene", new IMAGE_BRUSH("TreeItems/Scene_16x", Icon16x16));
-	Style->Set("DisplayClusterConfigurator.TreeItems.Cluster", new IMAGE_BRUSH("TreeItems/Cluster_16x", Icon16x16));
-	Style->Set("DisplayClusterConfigurator.TreeItems.Container", new IMAGE_BRUSH("TreeItems/Container_12x", Icon12x12));
-
-	// Scene Components
-	Style->Set("DisplayClusterConfigurator.TreeItems.SceneComponentXform", new IMAGE_BRUSH("TreeItems/SceneComponentXform_16x", Icon16x16));
-	Style->Set("DisplayClusterConfigurator.TreeItems.SceneComponentScreen", new IMAGE_BRUSH("TreeItems/SceneComponentScreen_16x", Icon16x16));
-	Style->Set("DisplayClusterConfigurator.TreeItems.SceneComponentCamera", new IMAGE_BRUSH("TreeItems/SceneComponentCamera_16x", Icon16x16));
-
-	Style->Set("ClassIcon.DisplayClusterXformComponent", new IMAGE_BRUSH(TEXT("TreeItems/SceneComponentXform_16x"), Icon16x16));
-	Style->Set("ClassThumbnail.DisplayClusterXformComponent", new IMAGE_BRUSH(TEXT("TreeItems/SceneComponentXform_16x"), Icon16x16));
-	
-	Style->Set("ClassIcon.DisplayClusterScreenComponent", new IMAGE_BRUSH(TEXT("TreeItems/SceneComponentScreen_16x"), Icon16x16));
-	Style->Set("ClassThumbnail.DisplayClusterScreenComponent", new IMAGE_BRUSH(TEXT("TreeItems/SceneComponentScreen_16x"), Icon16x16));
-	
-	Style->Set("ClassIcon.DisplayClusterCameraComponent", new IMAGE_BRUSH(TEXT("Components/ViewOrigin/nDisplayViewOrigin_x16"), Icon16x16));
-	Style->Set("ClassThumbnail.DisplayClusterCameraComponent", new IMAGE_BRUSH(TEXT("Components/ViewOrigin/nDisplayViewOrigin_x64"), Icon64x64));
-	
-	Style->Set("ClassIcon.DisplayClusterICVFXCameraComponent", new IMAGE_BRUSH(TEXT("Components/ICVFXCamera/nDisplayCamera_x16"), Icon16x16));
+	// Class Icons
+	{
+		Set("ClassIcon.DisplayClusterRootActor", new IMAGE_BRUSH_SVG("Icons/RootActor/nDisplay", Icon16x16));
+		Set("ClassThumbnail.DisplayClusterRootActor", new IMAGE_BRUSH_SVG("Icons/RootActor/nDisplay_64", Icon64x64));
+
+		Set("ClassIcon.DisplayClusterXformComponent", new IMAGE_BRUSH_SVG("Icons/Components/nDisplayXform", Icon16x16));
+		Set("ClassThumbnail.DisplayClusterXformComponent", new IMAGE_BRUSH_SVG("Icons/Components/nDisplayXform", Icon64x64));
+
+		Set("ClassIcon.DisplayClusterScreenComponent", new IMAGE_BRUSH_SVG("Icons/Components/nDisplayScreen", Icon16x16));
+		Set("ClassThumbnail.DisplayClusterScreenComponent", new IMAGE_BRUSH_SVG("Icons/Components/nDisplayScreen", Icon64x64));
+
+		Set("ClassIcon.DisplayClusterCameraComponent", new IMAGE_BRUSH_SVG("Icons/Components/nDisplayViewOrigin", Icon16x16));
+		Set("ClassThumbnail.DisplayClusterCameraComponent", new IMAGE_BRUSH_SVG("Icons/Components/nDisplayViewOrigin_64", Icon64x64));
+
+		Set("ClassIcon.DisplayClusterICVFXCameraComponent", new IMAGE_BRUSH_SVG("Icons/Components/nDisplayCamera", Icon16x16));
+		Set("ClassThumbnail.DisplayClusterICVFXCameraComponent", new IMAGE_BRUSH_SVG("Icons/Components/nDisplayCamera", Icon64x64));
+
+		Set("ClassIcon.DisplayClusterOriginComponent", new IMAGE_BRUSH_SVG("Icons/Components/nDisplayRootComponent", Icon16x16));
+		Set("ClassThumbnail.DisplayClusterOriginComponent", new IMAGE_BRUSH_SVG("Icons/Components/nDisplayRootComponent", Icon64x64));
+	}
+
+	// Config Editor Tabs
+	{
+		Set("DisplayClusterConfigurator.Tabs.OutputMapping", new IMAGE_BRUSH_SVG("Icons/Tabs/OutputMapping", Icon16x16));
+		Set("DisplayClusterConfigurator.Tabs.Cluster", new IMAGE_BRUSH_SVG("Icons/Tabs/Cluster", Icon16x16));
+	}
+
+
+	// Config Editor Toolbar
+	{
+		Set("DisplayClusterConfigurator.Toolbar.Import", new CORE_IMAGE_BRUSH_SVG("Starship/Common/import", Icon20x20));
+		Set("DisplayClusterConfigurator.Toolbar.SaveToFile", new EDITOR_IMAGE_BRUSH_SVG("Starship/Common/Export", Icon20x20));
+	}
 
 	// Cluster
-	Style->Set("DisplayClusterConfigurator.TreeItems.Host", new IMAGE_BRUSH("TreeItems/ClusterNode_16x", Icon16x16));
-	Style->Set("DisplayClusterConfigurator.TreeItems.ClusterNode", new IMAGE_BRUSH("TreeItems/Viewport_16x", Icon16x16));
-	Style->Set("DisplayClusterConfigurator.TreeItems.Viewport", new IMAGE_BRUSH("TreeItems/SceneComponentScreen_16x", Icon16x16));
-	Style->Set("DisplayClusterConfigurator.TreeItems.Postprocess", new IMAGE_BRUSH("TreeItems/Postprocess_12x", Icon12x12));
-
-	// Icons
-	Style->Set("DisplayClusterConfigurator.OutputMapping.ToggleWindowInfo", new IMAGE_BRUSH("OutputMapping/ToggleWindowInfo_16x", Icon16x16));
-	Style->Set("DisplayClusterConfigurator.OutputMapping.ToggleWindowCornerImage", new IMAGE_BRUSH("OutputMapping/CornerImage_16x", Icon16x16));
-	Style->Set("DisplayClusterConfigurator.OutputMapping.ToggleOutsideViewports", new IMAGE_BRUSH("OutputMapping/ToggleOutsideViewports_16x", Icon16x16));
-	Style->Set("DisplayClusterConfigurator.OutputMapping.ZoomToFit", new IMAGE_BRUSH("OutputMapping/ZoomToFit_16x", Icon16x16));
-	Style->Set("DisplayClusterConfigurator.OutputMapping.ViewScale", new IMAGE_BRUSH("OutputMapping/ViewScale_16x", Icon16x16));
-	Style->Set("DisplayClusterConfigurator.OutputMapping.ResizeAreaHandle", new IMAGE_BRUSH("OutputMapping/ResizeAreaHandle_20x", Icon20x20));
-
-	Style->Set("DisplayClusterConfigurator.Node.Color.Regular", FLinearColor(FColor(97, 97, 97)));
-	Style->Set("DisplayClusterConfigurator.Node.Color.Selected", FLinearColor(FColor(249, 165, 1)));
-	Style->Set("DisplayClusterConfigurator.Node.Color.Regular.Opacity_50", FLinearColor(FColor(255, 255, 255, 255 * 0.5f)));
-	Style->Set("DisplayClusterConfigurator.Node.Color.Selected.Opacity_50", FLinearColor(FColor(249, 165, 1, 255 * 0.5f)));
-
-	Style->Set("DisplayClusterConfigurator.Node.Text.Color.Regular", FLinearColor(FColor(255, 255, 255)));
-	Style->Set("DisplayClusterConfigurator.Node.Text.Color.WhiteGray", FLinearColor(FColor(226, 226, 226)));
-
-	Style->Set("DisplayClusterConfigurator.Node.Host.Inner.Background", FLinearColor(FColor(38, 38, 38)));
-	Style->Set("DisplayClusterConfigurator.Node.Window.Outer.Background", FLinearColor(FColor(63, 63, 63)));
-	Style->Set("DisplayClusterConfigurator.Node.Window.Inner.Background", FLinearColor(FColor(53, 53, 53)));
-	Style->Set("DisplayClusterConfigurator.Node.Window.Border.Color", FLinearColor(FColor(64, 164, 255, 255)));
-	Style->Set("DisplayClusterConfigurator.Node.Window.Corner.Color", FLinearColor(FColor(64, 164, 255, 255)));
-	Style->Set("DisplayClusterConfigurator.Node.Window.Corner.Color.Locked", FLinearColor(FColor(164, 164, 164, 255)));
-	Style->Set("DisplayClusterConfigurator.Node.Viewport.Border.Color.Regular", FLinearColor(FColor(164, 164, 164, 255)));
-	Style->Set("DisplayClusterConfigurator.Node.Viewport.Border.OutsideColor.Regular", FLinearColor(FColor(247, 129, 91, 255)));
-	Style->Set("DisplayClusterConfigurator.Node.Viewport.Text.Background", FLinearColor(FColor(73, 73, 73, 255 * 0.65f)));
-	Style->Set("DisplayClusterConfigurator.Node.Viewport.Text.Background.Locked", FLinearColor(FColor(164, 164, 164, 255)));
-
-	Style->Set("DisplayClusterConfigurator.Node.Viewport.BackgroundColor.Regular", FLinearColor(FColor(155, 155, 155, 255 * 0.85f)));
-	Style->Set("DisplayClusterConfigurator.Node.Viewport.BackgroundColor.Selected", FLinearColor(FColor(254, 178, 27, 255)));
-	Style->Set("DisplayClusterConfigurator.Node.Viewport.BackgroundImage.Selected", FLinearColor(FColor(255, 204, 102)));
-	Style->Set("DisplayClusterConfigurator.Node.Viewport.BackgroundImage.Locked", FLinearColor(FColor(128, 128, 128)));
-
-	Style->Set("DisplayClusterConfigurator.Node.Viewport.OutsideBackgroundColor.Regular", FLinearColor(FColor(255, 87, 34, 255 * 0.85f)));
-	Style->Set("DisplayClusterConfigurator.Node.Viewport.OutsideBackgroundColor.Selected", FLinearColor(FColor(255, 87, 34, 255)));
-
-	Style->Set("DisplayClusterConfigurator.Node.Body", new BOX_BRUSH("Nodes/Node_Body", FMargin(16.f / 64.f, 16.f / 64.f, 16.f / 64.f, 16.f / 64.f)));
-	Style->Set("DisplayClusterConfigurator.Node.Brush.Corner", new BOX_BRUSH("Nodes/Corner", 4.0f / 16.0f));
+	{
+		Set("DisplayClusterConfigurator.TreeItems.Cluster", new IMAGE_BRUSH_SVG("Icons/Cluster/Cluster", Icon16x16));
+		Set("DisplayClusterConfigurator.TreeItems.Host", new IMAGE_BRUSH_SVG("Icons/Cluster/Host", Icon16x16));
+		Set("DisplayClusterConfigurator.TreeItems.ClusterNode", new IMAGE_BRUSH_SVG("Icons/Cluster/ClusterNode", Icon16x16));
+		Set("DisplayClusterConfigurator.TreeItems.Viewport", new IMAGE_BRUSH_SVG("Icons/Cluster/Viewport", Icon16x16));
+	}
+
+	// Output Mapping Toolbar
+	{
+		Set("DisplayClusterConfigurator.OutputMapping.WindowDisplay", new IMAGE_BRUSH_SVG("Icons/OutputMapping/ToggleWindowInfo", Icon16x16));
+		Set("DisplayClusterConfigurator.OutputMapping.ToggleOutsideViewports", new IMAGE_BRUSH_SVG("Icons/OutputMapping/ToggleOutsideViewports", Icon16x16));
+		Set("DisplayClusterConfigurator.OutputMapping.ZoomToFit", new IMAGE_BRUSH_SVG("Icons/OutputMapping/ZoomToFit", Icon16x16));
+		Set("DisplayClusterConfigurator.OutputMapping.ViewScale", new IMAGE_BRUSH_SVG("Icons/OutputMapping/ViewScale", Icon16x16));
+		Set("DisplayClusterConfigurator.OutputMapping.Transform", new IMAGE_BRUSH_SVG("Icons/OutputMapping/Transform", Icon16x16));
+		Set("DisplayClusterConfigurator.OutputMapping.Snapping", new IMAGE_BRUSH_SVG("Icons/OutputMapping/Snapping", Icon16x16));
+		Set("DisplayClusterConfigurator.OutputMapping.ResizeAreaHandle", new IMAGE_BRUSH("Icons/OutputMapping/ResizeAreaHandle_20x", Icon20x20));
+	}
+
+	// Output Mapping Node Colors
+	{
+		Set("DisplayClusterConfigurator.Node.Color.Regular", FLinearColor(FColor(97, 97, 97)));
+		Set("DisplayClusterConfigurator.Node.Color.Selected", FLinearColor(FColor(249, 165, 1)));
+		Set("DisplayClusterConfigurator.Node.Color.Locked", FLinearColor(FColor(92, 92, 92, 255)));
+		Set("DisplayClusterConfigurator.Node.Color.Regular.Opacity_50", FLinearColor(FColor(255, 255, 255, 255 * 0.5f)));
+		Set("DisplayClusterConfigurator.Node.Color.Selected.Opacity_50", FLinearColor(FColor(249, 165, 1, 255 * 0.5f)));
+
+		Set("DisplayClusterConfigurator.Node.Text.Color.Regular", FLinearColor(FColor(255, 255, 255)));
+		Set("DisplayClusterConfigurator.Node.Text.Color.Selected", FLinearColor(FColor(0, 0, 0)));
+
+		Set("DisplayClusterConfigurator.Node.Host.Inner.Background", FLinearColor(FColor(38, 38, 38)));
+
+		Set("DisplayClusterConfigurator.Node.Window.Inner.Background", FLinearColor(FColor(53, 53, 53)));
+		Set("DisplayClusterConfigurator.Node.Window.Corner.Color", FLinearColor(FColor(47, 47, 47, 255)));
+
+		Set("DisplayClusterConfigurator.Node.Viewport.Border.Color.Regular", FLinearColor(FColor(164, 164, 164, 255)));
+		Set("DisplayClusterConfigurator.Node.Viewport.Border.Color.Outside", FLinearColor(FColor(247, 129, 91, 255)));
+		Set("DisplayClusterConfigurator.Node.Viewport.Text.Background", FLinearColor(FColor(73, 73, 73, 255 * 0.65f)));
+		Set("DisplayClusterConfigurator.Node.Viewport.Text.Background.Locked", FLinearColor(FColor(92, 92, 92, 255)));
+
+		Set("DisplayClusterConfigurator.Node.Viewport.BackgroundColor.Regular", FLinearColor(FColor(155, 155, 155, 255 * 0.85f)));
+		Set("DisplayClusterConfigurator.Node.Viewport.BackgroundColor.Selected", FLinearColor(FColor(254, 178, 27, 255)));
+		Set("DisplayClusterConfigurator.Node.Viewport.BackgroundImage.Selected", FLinearColor(FColor(255, 204, 102)));
+		Set("DisplayClusterConfigurator.Node.Viewport.BackgroundImage.Locked", FLinearColor(FColor(128, 128, 128)));
+
+		Set("DisplayClusterConfigurator.Node.Viewport.OutsideBackgroundColor.Regular", FLinearColor(FColor(255, 87, 34, 255 * 0.85f)));
+		Set("DisplayClusterConfigurator.Node.Viewport.OutsideBackgroundColor.Selected", FLinearColor(FColor(255, 87, 34, 255)));
+	}
 
 	// Node Text
 	{
-		FTextBlockStyle TilteTextBlockStyle = FEditorStyle::GetWidgetStyle< FTextBlockStyle >("NormalText");
+		FTextBlockStyle TilteTextBlockStyle = FEditorStyle::GetWidgetStyle<FTextBlockStyle>("NormalText");
 		TilteTextBlockStyle.SetFont(FCoreStyle::GetDefaultFontStyle("Regular", 28));
-		Style->Set("DisplayClusterConfigurator.Node.Text.Regular", TilteTextBlockStyle);
+		Set("DisplayClusterConfigurator.Node.Text.Regular", TilteTextBlockStyle);
 	}
 
 	{
-		FTextBlockStyle TilteTextBlockStyle = FEditorStyle::GetWidgetStyle< FTextBlockStyle >("NormalText");
+		FTextBlockStyle TilteTextBlockStyle = FEditorStyle::GetWidgetStyle<FTextBlockStyle>("NormalText");
 		TilteTextBlockStyle.SetFont(FCoreStyle::GetDefaultFontStyle("Bold", 28));
-		Style->Set("DisplayClusterConfigurator.Node.Text.Bold", TilteTextBlockStyle);
+		Set("DisplayClusterConfigurator.Node.Text.Bold", TilteTextBlockStyle);
 	}
 
 	{
 		FTextBlockStyle TilteTextBlockStyle = FEditorStyle::GetWidgetStyle<FTextBlockStyle>("Graph.StateNode.NodeTitle");
 		TilteTextBlockStyle.SetFont(FCoreStyle::GetDefaultFontStyle("Bold", 36));
-		Style->Set("DisplayClusterConfigurator.Host.Text.Title", TilteTextBlockStyle);
+		TilteTextBlockStyle.SetShadowColorAndOpacity(FLinearColor::Transparent);
+		Set("DisplayClusterConfigurator.Host.Text.Title", TilteTextBlockStyle);
 	}
 
 	{
-		FTextBlockStyle TilteTextBlockStyle = FEditorStyle::GetWidgetStyle< FTextBlockStyle >("NormalText");
+		FTextBlockStyle TilteTextBlockStyle = FEditorStyle::GetWidgetStyle<FTextBlockStyle>("NormalText");
 		TilteTextBlockStyle.SetFont(FCoreStyle::GetDefaultFontStyle("Italic", 18));
-		Style->Set("DisplayClusterConfigurator.Node.Text.Small", TilteTextBlockStyle);
+		Set("DisplayClusterConfigurator.Node.Text.Small", TilteTextBlockStyle);
 	}
 
 	// Canvas borders
 	{
-		Style->Set("DisplayClusterConfigurator.Selected.Canvas.Brush", new BORDER_BRUSH(TEXT("Slate/NonMarchingAnts"), FMargin(0.25f), Style->GetColor("DisplayClusterConfigurator.Node.Color.Selected.Opacity_50")));
-		Style->Set("DisplayClusterConfigurator.Regular.Canvas.Brush", new BORDER_BRUSH(TEXT("Slate/NonMarchingAnts"), FMargin(0.25f), Style->GetColor("DisplayClusterConfigurator.Node.Color.Regular.Opacity_50")));
+		Set("DisplayClusterConfigurator.Selected.Canvas.Brush", new BORDER_BRUSH(TEXT("Icons/Slate/NonMarchingAnts"), FMargin(0.25f), GetColor("DisplayClusterConfigurator.Node.Color.Selected.Opacity_50")));
+		Set("DisplayClusterConfigurator.Regular.Canvas.Brush", new BORDER_BRUSH(TEXT("Icons/Slate/NonMarchingAnts"), FMargin(0.25f), GetColor("DisplayClusterConfigurator.Node.Color.Regular.Opacity_50")));
 	}
 
 	// Window borders
 	{
+		Set("DisplayClusterConfigurator.Node.Brush.Corner", new BOX_BRUSH_SVG("Icons/OutputMapping/CornerImage", 4.0f / 16.0f));
+
 		FSlateBrush* SelectedBrush = new FSlateBrush();
 		SelectedBrush->Margin = FMargin(6.f);
 		SelectedBrush->DrawAs = ESlateBrushDrawType::Border;
-		SelectedBrush->TintColor = Style->GetColor("DisplayClusterConfigurator.Node.Color.Selected");
+		SelectedBrush->TintColor = GetColor("DisplayClusterConfigurator.Node.Color.Selected");
 
-		Style->Set("DisplayClusterConfigurator.Node.Window.Border.Brush.Selected", SelectedBrush);
+		Set("DisplayClusterConfigurator.Node.Window.Border.Brush.Selected", SelectedBrush);
 
 		FSlateBrush* RegularBrush = new FSlateBrush();
 		RegularBrush->Margin = FMargin(1.f);
 		RegularBrush->DrawAs = ESlateBrushDrawType::Border;
-		RegularBrush->TintColor = Style->GetColor("DisplayClusterConfigurator.Node.Color.Regular");
-		Style->Set("DisplayClusterConfigurator.Node.Window.Border.Brush.Regular", RegularBrush);
+		RegularBrush->TintColor = GetColor("DisplayClusterConfigurator.Node.Color.Regular");
+		Set("DisplayClusterConfigurator.Node.Window.Border.Brush.Regular", RegularBrush);
 	}
 
 	// Viewport borders
 	{
-		FSlateBrush* Brush = new FSlateBrush();
-		Brush->Margin = FMargin(1.f);
-		Brush->DrawAs = ESlateBrushDrawType::Border;
-		Brush->TintColor = Style->GetColor("DisplayClusterConfigurator.Node.Window.Border.Color");
-
-		Style->Set("DisplayClusterConfigurator.Node.Window.Border.Brush", Brush);
-
 		FSlateBrush* SelectedBrush = new FSlateBrush();
 		SelectedBrush->Margin = FMargin(2.f);
 		SelectedBrush->DrawAs = ESlateBrushDrawType::Border;
-		SelectedBrush->TintColor = Style->GetColor("DisplayClusterConfigurator.Node.Color.Selected");
+		SelectedBrush->TintColor = GetColor("DisplayClusterConfigurator.Node.Color.Selected");
 
-		Style->Set("DisplayClusterConfigurator.Node.Viewport.Border.Brush.Selected", SelectedBrush);
+		Set("DisplayClusterConfigurator.Node.Viewport.Border.Brush.Selected", SelectedBrush);
 
 		FSlateBrush* RegularBrush = new FSlateBrush();
 		RegularBrush->Margin = FMargin(1.f);
 		RegularBrush->DrawAs = ESlateBrushDrawType::Border;
-		RegularBrush->TintColor = Style->GetColor("DisplayClusterConfigurator.Node.Viewport.Border.Color.Regular");
-		Style->Set("DisplayClusterConfigurator.Node.Viewport.Border.Brush.Regular", RegularBrush);
-
-		FSlateBrush* OutsideRegularBrush = new FSlateBrush();
-		OutsideRegularBrush->Margin = FMargin(1.f);
-		OutsideRegularBrush->DrawAs = ESlateBrushDrawType::Border;
-		OutsideRegularBrush->TintColor = Style->GetColor("DisplayClusterConfigurator.Node.Viewport.Border.OutsideColor.Regular");
-		Style->Set("DisplayClusterConfigurator.Node.Viewport.Border.OutsideBrush.Regular", OutsideRegularBrush);
+		RegularBrush->TintColor = GetColor("DisplayClusterConfigurator.Node.Viewport.Border.Color.Regular");
+		Set("DisplayClusterConfigurator.Node.Viewport.Border.Brush.Regular", RegularBrush);
+
+		FSlateBrush* OutsideBrush = new FSlateBrush();
+		OutsideBrush->Margin = FMargin(1.f);
+		OutsideBrush->DrawAs = ESlateBrushDrawType::Border;
+		OutsideBrush->TintColor = GetColor("DisplayClusterConfigurator.Node.Viewport.Border.Color.Outside");
+		Set("DisplayClusterConfigurator.Node.Viewport.Border.Brush.Outside", OutsideBrush);
 	}
 
 	// Corner Colors array
 	{
-		CornerColors.Add(FCornerColor("DisplayClusterConfigurator.Node.Corner.Color.0", FLinearColor(FColor(244, 67, 54, 255 * 0.8f))));
-		CornerColors.Add(FCornerColor("DisplayClusterConfigurator.Node.Corner.Color.1", FLinearColor(FColor(156, 39, 176, 255 * 0.8f))));
-		CornerColors.Add(FCornerColor("DisplayClusterConfigurator.Node.Corner.Color.2", FLinearColor(FColor(0, 188, 212, 255 * 0.8f))));
-		CornerColors.Add(FCornerColor("DisplayClusterConfigurator.Node.Corner.Color.3", FLinearColor(FColor(139, 195, 74, 255 * 0.8f))));
-		CornerColors.Add(FCornerColor("DisplayClusterConfigurator.Node.Corner.Color.4", FLinearColor(FColor(255, 235, 59, 255 * 0.8f))));
-		CornerColors.Add(FCornerColor("DisplayClusterConfigurator.Node.Corner.Color.5", FLinearColor(FColor(96, 125, 139, 255 * 0.8f))));
-
-
-		for (const FCornerColor& CornerColor : CornerColors)
+		DefaultColors.Add(FCornerColor("DisplayClusterConfigurator.Node.Corner.Color.0", FLinearColor(FColor(244, 67, 54, 255 * 0.8f))));
+		DefaultColors.Add(FCornerColor("DisplayClusterConfigurator.Node.Corner.Color.1", FLinearColor(FColor(156, 39, 176, 255 * 0.8f))));
+		DefaultColors.Add(FCornerColor("DisplayClusterConfigurator.Node.Corner.Color.2", FLinearColor(FColor(0, 188, 212, 255 * 0.8f))));
+		DefaultColors.Add(FCornerColor("DisplayClusterConfigurator.Node.Corner.Color.3", FLinearColor(FColor(139, 195, 74, 255 * 0.8f))));
+		DefaultColors.Add(FCornerColor("DisplayClusterConfigurator.Node.Corner.Color.4", FLinearColor(FColor(255, 235, 59, 255 * 0.8f))));
+		DefaultColors.Add(FCornerColor("DisplayClusterConfigurator.Node.Corner.Color.5", FLinearColor(FColor(96, 125, 139, 255 * 0.8f))));
+
+		for (const FCornerColor& Color : DefaultColors)
 		{
-			Style->Set(CornerColor.Name, CornerColor.Color);
+			Set(Color.Name, Color.Color);
 		}
 	}
 
@@ -258,26 +225,11 @@ TSharedRef< FSlateStyleSet > FDisplayClusterConfiguratorStyle::Create()
 	{
 		const FTextBlockStyle DefaultTextStyle = FAppStyle::Get().GetWidgetStyle<FTextBlockStyle>("NormalText");
 
-		Style->Set("DisplayClusterConfigurator.NewAssetDialog.OptionText", FTextBlockStyle(DefaultTextStyle).SetFont(DEFAULT_FONT("Bold", 11)));
-		Style->Set("DisplayClusterConfigurator.NewAssetDialog.SubBorder", new EDITOR_BOX_BRUSH("Common/GroupBorderLight", FMargin(4.0f / 16.0f)));
-		Style->Set("DisplayClusterConfigurator.NewAssetDialog.ActiveOptionBorderColor", FLinearColor(FColor(96, 96, 96)));
+		Set("DisplayClusterConfigurator.NewAssetDialog.OptionText", FTextBlockStyle(DefaultTextStyle).SetFont(DEFAULT_FONT("Bold", 11)));
+		Set("DisplayClusterConfigurator.NewAssetDialog.SubBorder", new EDITOR_BOX_BRUSH("Common/GroupBorderLight", FMargin(4.0f / 16.0f)));
+		Set("DisplayClusterConfigurator.NewAssetDialog.ActiveOptionBorderColor", FLinearColor(FColor(96, 96, 96)));
 	}
-
-	return Style;
 }
 
-#undef IMAGE_BRUSH
-#undef BOX_BRUSH
-#undef BORDER_BRUSH
+#undef EDITOR_IMAGE_BRUSH_SVG
 #undef EDITOR_BOX_BRUSH
-#undef DEFAULT_FONT
-
-void FDisplayClusterConfiguratorStyle::ReloadTextures()
-{
-	FSlateApplication::Get().GetRenderer()->ReloadTextureResources();
-}
-
-const ISlateStyle& FDisplayClusterConfiguratorStyle::Get()
-{
-	return *StyleInstance;
-}
diff --git a/Engine/Plugins/Runtime/nDisplay/Source/DisplayClusterConfigurator/Private/DisplayClusterConfiguratorStyle.h b/Engine/Plugins/Runtime/nDisplay/Source/DisplayClusterConfigurator/Private/DisplayClusterConfiguratorStyle.h
--- a/Engine/Plugins/Runtime/nDisplay/Source/DisplayClusterConfigurator/Private/DisplayClusterConfiguratorStyle.h
+++ b/Engine/Plugins/Runtime/nDisplay/Source/DisplayClusterConfigurator/Private/DisplayClusterConfiguratorStyle.h
@@ -5,36 +5,24 @@
 #include "Styling/ISlateStyle.h"
 #include "Styling/SlateStyle.h"
 
-class FDisplayClusterConfiguratorStyle
+class FDisplayClusterConfiguratorStyle : public FSlateStyleSet
 {
 public:
+	FDisplayClusterConfiguratorStyle();
+	~FDisplayClusterConfiguratorStyle();
 
-	static void Initialize();
-
-	static void Shutdown();
+	static FDisplayClusterConfiguratorStyle& Get()
+	{
+		static FDisplayClusterConfiguratorStyle Inst;
+		return Inst;
+	}
 
 	static void ReloadTextures();
 
-	static const ISlateStyle& Get();
-
-	static FName GetStyleSetName();
-
-	static const FLinearColor& GetColor(FName PropertyName, const ANSICHAR* Specifier = nullptr);
-
-	static const FLinearColor& GetCornerColor(uint32 Index);
-
-	static const FMargin& GetMargin(FName PropertyName, const ANSICHAR* Specifier = nullptr);
-
-	static const FSlateBrush* GetBrush(FName PropertyName, const ANSICHAR* Specifier = nullptr);
-
-	template< typename WidgetStyleType >
-	static const WidgetStyleType& GetWidgetStyle(FName PropertyName, const ANSICHAR* Specifier = nullptr)
-	{
-		return StyleInstance->GetWidgetStyle<WidgetStyleType>(PropertyName, Specifier);
-	}
+	const FLinearColor& GetDefaultColor(uint32 Index);
 
 private:
-	static TSharedRef<FSlateStyleSet> Create();
+	void Initialize();
 
 	struct FCornerColor
 	{
@@ -50,8 +38,5 @@ private:
 		FLinearColor Color;
 	};
 
-	static TArray<FCornerColor> CornerColors;
-
-private:
-	static TSharedPtr<FSlateStyleSet> StyleInstance;
+	TArray<FCornerColor> DefaultColors;
 };
diff --git a/Engine/Plugins/Runtime/nDisplay/Source/DisplayClusterConfigurator/Private/DisplayClusterConfiguratorTabSpawners.cpp b/Engine/Plugins/Runtime/nDisplay/Source/DisplayClusterConfigurator/Private/DisplayClusterConfiguratorTabSpawners.cpp
--- a/Engine/Plugins/Runtime/nDisplay/Source/DisplayClusterConfigurator/Private/DisplayClusterConfiguratorTabSpawners.cpp
+++ b/Engine/Plugins/Runtime/nDisplay/Source/DisplayClusterConfigurator/Private/DisplayClusterConfiguratorTabSpawners.cpp
@@ -30,7 +30,7 @@ FDisplayClusterViewOutputMappingSummoner::FDisplayClusterViewOutputMappingSummon
 	TSharedPtr<FDisplayClusterConfiguratorBlueprintEditor> InHostingApp) : FDisplayClusterViewSummoner(FDisplayClusterConfiguratorEditorConfigurationMode::TabID_OutputMapping, InHostingApp)
 {
 	TabLabel = LOCTEXT("OutputMappingTabTitle", "OutputMapping");
-	TabIcon = FSlateIcon(FDisplayClusterConfiguratorStyle::GetStyleSetName(), "DisplayClusterConfigurator.Tabs.OutputMapping");
+	TabIcon = FSlateIcon(FDisplayClusterConfiguratorStyle::Get().GetStyleSetName(), "DisplayClusterConfigurator.Tabs.OutputMapping");
 	
 	ViewMenuDescription = LOCTEXT("OutputMappingTabDescription", "OutputMapping");
 	ViewMenuTooltip = LOCTEXT("OutputMapping_ToolTip", "Shows output mappings");
@@ -46,7 +46,7 @@ FDisplayClusterViewClusterSummoner::FDisplayClusterViewClusterSummoner(
 	TSharedPtr<FDisplayClusterConfiguratorBlueprintEditor> InHostingApp) : FDisplayClusterViewSummoner(FDisplayClusterConfiguratorEditorConfigurationMode::TabID_Cluster, InHostingApp)
 {
 	TabLabel = LOCTEXT("ClusterTabTitle", "Cluster");
-	TabIcon = FSlateIcon(FDisplayClusterConfiguratorStyle::GetStyleSetName(), "DisplayClusterConfigurator.Tabs.Cluster");
+	TabIcon = FSlateIcon(FDisplayClusterConfiguratorStyle::Get().GetStyleSetName(), "DisplayClusterConfigurator.Tabs.Cluster");
 	
 	ViewMenuDescription = LOCTEXT("ClusterTabDescription", "Cluster view");
 	ViewMenuTooltip = LOCTEXT("Cluster_ToolTip", "Shows the cluster view");
@@ -61,7 +61,7 @@ TSharedRef<SWidget> FDisplayClusterViewClusterSummoner::CreateTabBody(const FWor
 FDisplayClusterSCSSummoner::FDisplayClusterSCSSummoner(TSharedPtr<FDisplayClusterConfiguratorBlueprintEditor> InHostingApp) : FDisplayClusterViewSummoner(FDisplayClusterConfiguratorEditorConfigurationMode::TabID_Scene, InHostingApp)
 {
 	TabLabel = LOCTEXT("ComponentsTabLabel", "Components");
-	TabIcon = FSlateIcon(FEditorStyle::GetStyleSetName(), "Kismet.Tabs.Components");
+	TabIcon = FSlateIcon(FAppStyle::Get().GetStyleSetName(), "Kismet.Tabs.Components");
 
 	bIsSingleton = true;
 
@@ -79,7 +79,7 @@ FDisplayClusterSCSViewportSummoner::FDisplayClusterSCSViewportSummoner(
 	TSharedPtr<FDisplayClusterConfiguratorBlueprintEditor> InHostingApp) : FDisplayClusterViewSummoner(FDisplayClusterConfiguratorEditorConfigurationMode::TabID_Viewport, InHostingApp)
 {
 	TabLabel = LOCTEXT("ViewportTabTitle", "Preview");
-	TabIcon = FSlateIcon(FDisplayClusterConfiguratorStyle::GetStyleSetName(), "DisplayClusterConfigurator.Tabs.Viewport");
+	TabIcon = FSlateIcon(FSlateIcon(FAppStyle::Get().GetStyleSetName(), "LevelEditor.Tabs.Viewports"));
 
 	ViewMenuDescription = LOCTEXT("ViewportTabDescription", "Viewport");
 	ViewMenuTooltip = LOCTEXT("Viewport_ToolTip", "Shows the viewport");
diff --git a/Engine/Plugins/Runtime/nDisplay/Source/DisplayClusterConfigurator/Private/Views/DisplayClusterConfiguratorToolbar.cpp b/Engine/Plugins/Runtime/nDisplay/Source/DisplayClusterConfigurator/Private/Views/DisplayClusterConfiguratorToolbar.cpp
--- a/Engine/Plugins/Runtime/nDisplay/Source/DisplayClusterConfigurator/Private/Views/DisplayClusterConfiguratorToolbar.cpp
+++ b/Engine/Plugins/Runtime/nDisplay/Source/DisplayClusterConfigurator/Private/Views/DisplayClusterConfiguratorToolbar.cpp
@@ -25,7 +25,7 @@ public:
 	{
 		SBorder::Construct(
 			SBorder::FArguments()
-			.BorderImage(FEditorStyle::GetBrush("BlueprintEditor.PipelineSeparator"))
+			.BorderImage(FAppStyle::Get().GetBrush("BlueprintEditor.PipelineSeparator"))
 			.Padding(0.0f)
 		);
 	}
@@ -61,13 +61,13 @@ void FDisplayClusterConfiguratorToolbar::FillModesToolbar(FToolBarBuilder& Toolb
 		NAME_None,
 		TAttribute<FText>(),
 		TAttribute<FText>(),
-		FSlateIcon(FDisplayClusterConfiguratorStyle::GetStyleSetName(), "DisplayClusterConfigurator.Toolbar.Import"));
+		FSlateIcon(FDisplayClusterConfiguratorStyle::Get().GetStyleSetName(), "DisplayClusterConfigurator.Toolbar.Import"));
 
 	ToolbarBuilder.AddToolBarButton(Commands.Export,
 		NAME_None,
 		TAttribute<FText>(),
 		TAttribute<FText>(),
-		FSlateIcon(FDisplayClusterConfiguratorStyle::GetStyleSetName(), "DisplayClusterConfigurator.Toolbar.SaveToFile"));
+		FSlateIcon(FDisplayClusterConfiguratorStyle::Get().GetStyleSetName(), "DisplayClusterConfigurator.Toolbar.SaveToFile"));
 	
 	ToolbarBuilder.AddComboButton(FUIAction(),
 		FOnGetContent::CreateRaw(this, &FDisplayClusterConfiguratorToolbar::GenerateExportMenu),
diff --git a/Engine/Plugins/Runtime/nDisplay/Source/DisplayClusterConfigurator/Private/Views/OutputMapping/DisplayClusterConfiguratorOutputMappingCommands.cpp b/Engine/Plugins/Runtime/nDisplay/Source/DisplayClusterConfigurator/Private/Views/OutputMapping/DisplayClusterConfiguratorOutputMappingCommands.cpp
--- a/Engine/Plugins/Runtime/nDisplay/Source/DisplayClusterConfigurator/Private/Views/OutputMapping/DisplayClusterConfiguratorOutputMappingCommands.cpp
+++ b/Engine/Plugins/Runtime/nDisplay/Source/DisplayClusterConfigurator/Private/Views/OutputMapping/DisplayClusterConfiguratorOutputMappingCommands.cpp
@@ -7,13 +7,12 @@
 
 void FDisplayClusterConfiguratorOutputMappingCommands::RegisterCommands()
 {	
-	UI_COMMAND(ToggleWindowInfo, "Show Window Info", "Enables or disables showing the window information", EUserInterfaceActionType::ToggleButton, FInputChord(EKeys::W));
-	UI_COMMAND(ToggleWindowCornerImage, "Show Window Corner Image", "Enables or disables showing the window corner image", EUserInterfaceActionType::ToggleButton, FInputChord(EKeys::E));
-	UI_COMMAND(ToggleOutsideViewports, "Show Viewports outside the Window", "Enables or disables showing the viewport which is compleatly outside window", EUserInterfaceActionType::ToggleButton, FInputChord(EKeys::R));
+	UI_COMMAND(ShowWindowInfo, "Info Bar", "Shows the information bar on cluster nodes", EUserInterfaceActionType::RadioButton, FInputChord());
+	UI_COMMAND(ShowWindowCorner, "Corner", "Shows only the corner nub on cluster nodes", EUserInterfaceActionType::RadioButton, FInputChord());
+	UI_COMMAND(ShowWindowNone, "None", "Hides both the information bar and corner nub on cluster nodes", EUserInterfaceActionType::RadioButton, FInputChord());
+	UI_COMMAND(ToggleOutsideViewports, "Show Viewports outside the Window", "Toggles displaying viewports that are entirely outside a cluster node", EUserInterfaceActionType::ToggleButton, FInputChord(EKeys::R));
 	UI_COMMAND(ToggleClusterItemOverlap, "Allow Cluster Item Overlap", "Enables or disables allowing cluster items to overlap when being manipulated", EUserInterfaceActionType::ToggleButton, FInputChord());
 	UI_COMMAND(ToggleLockClusterNodesInHosts, "Keep Cluster Nodes inside Hosts", "Prevents cluster nodes from being moved outside of hosts when being manipulated", EUserInterfaceActionType::ToggleButton, FInputChord());
-	UI_COMMAND(ToggleLockViewports, "Lock Viewports in place", "Locks viewports in place, preventing them from being selected or dragged", EUserInterfaceActionType::ToggleButton, FInputChord());
-	UI_COMMAND(ToggleLockClusterNodes, "Lock Cluster Nodes in place", "Locks cluster nodes in place, preventing them from being selected or dragged", EUserInterfaceActionType::ToggleButton, FInputChord());
 	UI_COMMAND(ToggleTintViewports, "Tint Selected Viewports", "Toggles tinting selected viewports orange to better indicate that they are selected", EUserInterfaceActionType::ToggleButton, FInputChord());
 	UI_COMMAND(ZoomToFit, "Zoom To Fit", "Zoom To Fit In Graph", EUserInterfaceActionType::Button, FInputChord(EKeys::Z));
 
diff --git a/Engine/Plugins/Runtime/nDisplay/Source/DisplayClusterConfigurator/Private/Views/OutputMapping/DisplayClusterConfiguratorOutputMappingCommands.h b/Engine/Plugins/Runtime/nDisplay/Source/DisplayClusterConfigurator/Private/Views/OutputMapping/DisplayClusterConfiguratorOutputMappingCommands.h
--- a/Engine/Plugins/Runtime/nDisplay/Source/DisplayClusterConfigurator/Private/Views/OutputMapping/DisplayClusterConfiguratorOutputMappingCommands.h
+++ b/Engine/Plugins/Runtime/nDisplay/Source/DisplayClusterConfigurator/Private/Views/OutputMapping/DisplayClusterConfiguratorOutputMappingCommands.h
@@ -13,19 +13,18 @@ class FDisplayClusterConfiguratorOutputMappingCommands
 public:
 	FDisplayClusterConfiguratorOutputMappingCommands()
 		: TCommands<FDisplayClusterConfiguratorOutputMappingCommands>(TEXT("DisplayClusterConfigurator.OutputMapping"), 
-			NSLOCTEXT("Contexts", "DisplayClusterConfiguratorOutputMapping", "Display Cluster Configurator Output Mapping"), NAME_None, FDisplayClusterConfiguratorStyle::GetStyleSetName())
+			NSLOCTEXT("Contexts", "DisplayClusterConfiguratorOutputMapping", "Display Cluster Configurator Output Mapping"), NAME_None, FDisplayClusterConfiguratorStyle::Get().GetStyleSetName())
 	{ }
 
 	virtual void RegisterCommands() override;
 
 public:
-	TSharedPtr<FUICommandInfo> ToggleWindowInfo;
-	TSharedPtr<FUICommandInfo> ToggleWindowCornerImage;
+	TSharedPtr<FUICommandInfo> ShowWindowInfo;
+	TSharedPtr<FUICommandInfo> ShowWindowCorner;
+	TSharedPtr<FUICommandInfo> ShowWindowNone;
 	TSharedPtr<FUICommandInfo> ToggleOutsideViewports;
 	TSharedPtr<FUICommandInfo> ToggleClusterItemOverlap;
 	TSharedPtr<FUICommandInfo> ToggleLockClusterNodesInHosts;
-	TSharedPtr<FUICommandInfo> ToggleLockViewports;
-	TSharedPtr<FUICommandInfo> ToggleLockClusterNodes;
 	TSharedPtr<FUICommandInfo> ToggleTintViewports;
 	TSharedPtr<FUICommandInfo> ZoomToFit;
 	TSharedPtr<FUICommandInfo> BrowseDocumentation;
diff --git a/Engine/Plugins/Runtime/nDisplay/Source/DisplayClusterConfigurator/Private/Views/OutputMapping/DisplayClusterConfiguratorViewOutputMapping.cpp b/Engine/Plugins/Runtime/nDisplay/Source/DisplayClusterConfigurator/Private/Views/OutputMapping/DisplayClusterConfiguratorViewOutputMapping.cpp
--- a/Engine/Plugins/Runtime/nDisplay/Source/DisplayClusterConfigurator/Private/Views/OutputMapping/DisplayClusterConfiguratorViewOutputMapping.cpp
+++ b/Engine/Plugins/Runtime/nDisplay/Source/DisplayClusterConfigurator/Private/Views/OutputMapping/DisplayClusterConfiguratorViewOutputMapping.cpp
@@ -112,13 +112,27 @@ void FDisplayClusterConfiguratorViewOutputMapping::BindCommands()
 
 	const FDisplayClusterConfiguratorOutputMappingCommands& Commands = FDisplayClusterConfiguratorOutputMappingCommands::Get();
 
-	MAP_TOGGLE_COMMAND(Commands.ToggleWindowInfo, OutputMappingSettings.bShowWindowInfo);
-	MAP_TOGGLE_COMMAND(Commands.ToggleWindowCornerImage, OutputMappingSettings.bShowWindowCornerImage);
+	CommandList->MapAction(
+		Commands.ShowWindowInfo, 
+		FExecuteAction::CreateLambda([this]() { OutputMappingSettings.bShowWindowInfo = true; OutputMappingSettings.bShowWindowCornerImage = true; }), 
+		FCanExecuteAction(),
+		FIsActionChecked::CreateLambda([this]() { return OutputMappingSettings.bShowWindowInfo && OutputMappingSettings.bShowWindowCornerImage; }));
+
+	CommandList->MapAction(
+		Commands.ShowWindowCorner, 
+		FExecuteAction::CreateLambda([this]() { OutputMappingSettings.bShowWindowInfo = false; OutputMappingSettings.bShowWindowCornerImage = true; }), 
+		FCanExecuteAction(),
+		FIsActionChecked::CreateLambda([this]() { return !OutputMappingSettings.bShowWindowInfo && OutputMappingSettings.bShowWindowCornerImage; }));
+	
+	CommandList->MapAction(
+		Commands.ShowWindowNone, 
+		FExecuteAction::CreateLambda([this]() { OutputMappingSettings.bShowWindowInfo = false; OutputMappingSettings.bShowWindowCornerImage = false; }), 
+		FCanExecuteAction(),
+		FIsActionChecked::CreateLambda([this]() { return !OutputMappingSettings.bShowWindowInfo && !OutputMappingSettings.bShowWindowCornerImage; }));
+
 	MAP_TOGGLE_COMMAND(Commands.ToggleOutsideViewports, OutputMappingSettings.bShowOutsideViewports);
 	MAP_TOGGLE_COMMAND(Commands.ToggleClusterItemOverlap, OutputMappingSettings.bAllowClusterItemOverlap);
 	MAP_TOGGLE_COMMAND(Commands.ToggleLockClusterNodesInHosts, OutputMappingSettings.bKeepClusterNodesInHosts);
-	MAP_TOGGLE_COMMAND(Commands.ToggleLockViewports, OutputMappingSettings.bLockViewports);
-	MAP_TOGGLE_COMMAND(Commands.ToggleLockClusterNodes, OutputMappingSettings.bLockClusterNodes);
 	MAP_TOGGLE_COMMAND(Commands.ToggleTintViewports, OutputMappingSettings.bTintSelectedViewports);
 
 	MAP_TOGGLE_COMMAND(Commands.ToggleAdjacentEdgeSnapping, NodeAlignmentSettings.bSnapAdjacentEdges);
diff --git a/Engine/Plugins/Runtime/nDisplay/Source/DisplayClusterConfigurator/Private/Views/OutputMapping/EdNodes/DisplayClusterConfiguratorBaseNode.h b/Engine/Plugins/Runtime/nDisplay/Source/DisplayClusterConfigurator/Private/Views/OutputMapping/EdNodes/DisplayClusterConfiguratorBaseNode.h
--- a/Engine/Plugins/Runtime/nDisplay/Source/DisplayClusterConfigurator/Private/Views/OutputMapping/EdNodes/DisplayClusterConfiguratorBaseNode.h
+++ b/Engine/Plugins/Runtime/nDisplay/Source/DisplayClusterConfigurator/Private/Views/OutputMapping/EdNodes/DisplayClusterConfiguratorBaseNode.h
@@ -98,6 +98,7 @@ public:
 	virtual int32 GetAuxiliaryLayer(const TSet<UObject*>& SelectionSet) const;
 	virtual bool IsNodeVisible() const { return true; }
 	virtual bool IsNodeEnabled() const { return true; }
+	virtual bool IsNodeUnlocked() const { return true; }
 	virtual bool IsNodeAutoPositioned() const { return false; }
 	virtual bool IsNodeAutosized() const { return false; }
 	virtual bool CanNodeOverlapSiblings() const { return true; }
diff --git a/Engine/Plugins/Runtime/nDisplay/Source/DisplayClusterConfigurator/Private/Views/OutputMapping/EdNodes/DisplayClusterConfiguratorHostNode.cpp b/Engine/Plugins/Runtime/nDisplay/Source/DisplayClusterConfigurator/Private/Views/OutputMapping/EdNodes/DisplayClusterConfiguratorHostNode.cpp
--- a/Engine/Plugins/Runtime/nDisplay/Source/DisplayClusterConfigurator/Private/Views/OutputMapping/EdNodes/DisplayClusterConfiguratorHostNode.cpp
+++ b/Engine/Plugins/Runtime/nDisplay/Source/DisplayClusterConfigurator/Private/Views/OutputMapping/EdNodes/DisplayClusterConfiguratorHostNode.cpp
@@ -188,7 +188,7 @@ bool UDisplayClusterConfiguratorHostNode::IsNodeVisible() const
 	return bIsChildVisible;
 }
 
-bool UDisplayClusterConfiguratorHostNode::IsNodeEnabled() const
+bool UDisplayClusterConfiguratorHostNode::IsNodeUnlocked() const
 {
 	if (!IsObjectValid())
 	{
@@ -196,24 +196,7 @@ bool UDisplayClusterConfiguratorHostNode::IsNodeEnabled() const
 	}
 
 	UDisplayClusterConfigurationHostDisplayData* HostDisplayData = GetObjectChecked<UDisplayClusterConfigurationHostDisplayData>();
-
-	if (HostDisplayData->bIsEnabled)
-	{
-		return true;
-	}
-
-	// If all children of the host node are disabled, make the host node disabled as well.
-	bool bIsChildEnabled = false;
-	for (UDisplayClusterConfiguratorBaseNode* Child : Children)
-	{
-		if (Child->IsNodeEnabled())
-		{
-			bIsChildEnabled = true;
-			break;
-		}
-	}
-
-	return bIsChildEnabled;
+	return HostDisplayData->bIsUnlocked;
 }
 
 void UDisplayClusterConfiguratorHostNode::DeleteObject()
diff --git a/Engine/Plugins/Runtime/nDisplay/Source/DisplayClusterConfigurator/Private/Views/OutputMapping/EdNodes/DisplayClusterConfiguratorHostNode.h b/Engine/Plugins/Runtime/nDisplay/Source/DisplayClusterConfigurator/Private/Views/OutputMapping/EdNodes/DisplayClusterConfiguratorHostNode.h
--- a/Engine/Plugins/Runtime/nDisplay/Source/DisplayClusterConfigurator/Private/Views/OutputMapping/EdNodes/DisplayClusterConfiguratorHostNode.h
+++ b/Engine/Plugins/Runtime/nDisplay/Source/DisplayClusterConfigurator/Private/Views/OutputMapping/EdNodes/DisplayClusterConfiguratorHostNode.h
@@ -39,7 +39,7 @@ public:
 	virtual FBox2D GetNodeBounds(bool bAsParent = false) const override;
 	virtual FNodeAlignmentAnchors GetNodeAlignmentAnchors(bool bAsParent = false) const override;
 	virtual bool IsNodeVisible() const override;
-	virtual bool IsNodeEnabled() const override;
+	virtual bool IsNodeUnlocked() const override;
 	virtual bool IsNodeAutoPositioned() const { return !CanUserMoveNode(); }
 	virtual bool IsNodeAutosized() const override { return !CanUserResizeNode(); }
 	virtual bool CanNodeOverlapSiblings() const override { return false; }
diff --git a/Engine/Plugins/Runtime/nDisplay/Source/DisplayClusterConfigurator/Private/Views/OutputMapping/EdNodes/DisplayClusterConfiguratorViewportNode.cpp b/Engine/Plugins/Runtime/nDisplay/Source/DisplayClusterConfigurator/Private/Views/OutputMapping/EdNodes/DisplayClusterConfiguratorViewportNode.cpp
--- a/Engine/Plugins/Runtime/nDisplay/Source/DisplayClusterConfigurator/Private/Views/OutputMapping/EdNodes/DisplayClusterConfiguratorViewportNode.cpp
+++ b/Engine/Plugins/Runtime/nDisplay/Source/DisplayClusterConfigurator/Private/Views/OutputMapping/EdNodes/DisplayClusterConfiguratorViewportNode.cpp
@@ -78,7 +78,7 @@ bool UDisplayClusterConfiguratorViewportNode::IsNodeVisible() const
 	return Viewport->bIsVisible;
 }
 
-bool UDisplayClusterConfiguratorViewportNode::IsNodeEnabled() const
+bool UDisplayClusterConfiguratorViewportNode::IsNodeUnlocked() const
 {
 	if (!IsObjectValid())
 	{
@@ -86,7 +86,7 @@ bool UDisplayClusterConfiguratorViewportNode::IsNodeEnabled() const
 	}
 
 	UDisplayClusterConfigurationViewport* Viewport = GetObjectChecked<UDisplayClusterConfigurationViewport>();
-	return Viewport->bIsEnabled;
+	return Viewport->bIsUnlocked;
 }
 
 void UDisplayClusterConfiguratorViewportNode::DeleteObject()
diff --git a/Engine/Plugins/Runtime/nDisplay/Source/DisplayClusterConfigurator/Private/Views/OutputMapping/EdNodes/DisplayClusterConfiguratorViewportNode.h b/Engine/Plugins/Runtime/nDisplay/Source/DisplayClusterConfigurator/Private/Views/OutputMapping/EdNodes/DisplayClusterConfiguratorViewportNode.h
--- a/Engine/Plugins/Runtime/nDisplay/Source/DisplayClusterConfigurator/Private/Views/OutputMapping/EdNodes/DisplayClusterConfiguratorViewportNode.h
+++ b/Engine/Plugins/Runtime/nDisplay/Source/DisplayClusterConfigurator/Private/Views/OutputMapping/EdNodes/DisplayClusterConfiguratorViewportNode.h
@@ -33,7 +33,7 @@ public:
 	
 	//~ Begin UDisplayClusterConfiguratorBaseNode Interface
 	virtual bool IsNodeVisible() const override;
-	virtual bool IsNodeEnabled() const override;
+	virtual bool IsNodeUnlocked() const override;
 	virtual bool CanNodeOverlapSiblings() const override { return false; }
 	virtual bool CanNodeHaveNegativePosition() const { return false; }
 
diff --git a/Engine/Plugins/Runtime/nDisplay/Source/DisplayClusterConfigurator/Private/Views/OutputMapping/EdNodes/DisplayClusterConfiguratorWindowNode.cpp b/Engine/Plugins/Runtime/nDisplay/Source/DisplayClusterConfigurator/Private/Views/OutputMapping/EdNodes/DisplayClusterConfiguratorWindowNode.cpp
--- a/Engine/Plugins/Runtime/nDisplay/Source/DisplayClusterConfigurator/Private/Views/OutputMapping/EdNodes/DisplayClusterConfiguratorWindowNode.cpp
+++ b/Engine/Plugins/Runtime/nDisplay/Source/DisplayClusterConfigurator/Private/Views/OutputMapping/EdNodes/DisplayClusterConfiguratorWindowNode.cpp
@@ -119,7 +119,7 @@ bool UDisplayClusterConfiguratorWindowNode::IsNodeVisible() const
 	return bIsChildVisible;
 }
 
-bool UDisplayClusterConfiguratorWindowNode::IsNodeEnabled() const
+bool UDisplayClusterConfiguratorWindowNode::IsNodeUnlocked() const
 {
 	if (!IsObjectValid())
 	{
@@ -127,24 +127,7 @@ bool UDisplayClusterConfiguratorWindowNode::IsNodeEnabled() const
 	}
 
 	UDisplayClusterConfigurationClusterNode* ClusterNode = GetObjectChecked<UDisplayClusterConfigurationClusterNode>();
-
-	if (ClusterNode->bIsEnabled)
-	{
-		return true;
-	}
-
-	// If this node is marked as disabled but it has a child that is enabled, This node needs to remain enabled.
-	bool bIsChildEnabled = false;
-	for (UDisplayClusterConfiguratorBaseNode* Child : Children)
-	{
-		if (Child->IsNodeEnabled())
-		{
-			bIsChildEnabled = true;
-			break;
-		}
-	}
-
-	return bIsChildEnabled;
+	return ClusterNode->bIsUnlocked;
 }
 
 bool UDisplayClusterConfiguratorWindowNode::CanNodeExceedParentBounds() const
diff --git a/Engine/Plugins/Runtime/nDisplay/Source/DisplayClusterConfigurator/Private/Views/OutputMapping/EdNodes/DisplayClusterConfiguratorWindowNode.h b/Engine/Plugins/Runtime/nDisplay/Source/DisplayClusterConfigurator/Private/Views/OutputMapping/EdNodes/DisplayClusterConfiguratorWindowNode.h
--- a/Engine/Plugins/Runtime/nDisplay/Source/DisplayClusterConfigurator/Private/Views/OutputMapping/EdNodes/DisplayClusterConfiguratorWindowNode.h
+++ b/Engine/Plugins/Runtime/nDisplay/Source/DisplayClusterConfigurator/Private/Views/OutputMapping/EdNodes/DisplayClusterConfiguratorWindowNode.h
@@ -41,7 +41,7 @@ public:
 
 	//~ Begin UDisplayClusterConfiguratorBaseNode Interface
 	virtual bool IsNodeVisible() const override;
-	virtual bool IsNodeEnabled() const override;
+	virtual bool IsNodeUnlocked() const override;
 	virtual bool CanNodeOverlapSiblings() const override { return false; }
 	virtual bool CanNodeExceedParentBounds() const override;
 
diff --git a/Engine/Plugins/Runtime/nDisplay/Source/DisplayClusterConfigurator/Private/Views/OutputMapping/GraphNodes/SDisplayClusterConfiguratorBaseNode.cpp b/Engine/Plugins/Runtime/nDisplay/Source/DisplayClusterConfigurator/Private/Views/OutputMapping/GraphNodes/SDisplayClusterConfiguratorBaseNode.cpp
--- a/Engine/Plugins/Runtime/nDisplay/Source/DisplayClusterConfigurator/Private/Views/OutputMapping/GraphNodes/SDisplayClusterConfiguratorBaseNode.cpp
+++ b/Engine/Plugins/Runtime/nDisplay/Source/DisplayClusterConfigurator/Private/Views/OutputMapping/GraphNodes/SDisplayClusterConfiguratorBaseNode.cpp
@@ -24,7 +24,7 @@ void SAlignmentRuler::Construct(const FArguments& InArgs)
 	[
 		SNew(SBorder)
 		.Padding(FMargin(0.5f))
-		.BorderImage(FEditorStyle::GetBrush("WhiteBrush"))
+		.BorderImage(FAppStyle::Get().GetBrush("WhiteBrush"))
 		.BorderBackgroundColor(InArgs._ColorAndOpacity)
 		[
 			SAssignNew(BoxWidget, SBox)
@@ -157,6 +157,12 @@ FVector2D SDisplayClusterConfiguratorBaseNode::ComputeDesiredSize(float) const
 
 void SDisplayClusterConfiguratorBaseNode::MoveTo(const FVector2D& NewPosition, FNodeSet& NodeFilter, bool bMarkDirty)
 {
+	if (!IsNodeUnlocked())
+	{
+		NodeFilter.Add(SharedThis(this));
+		return;
+	}
+
 	TSharedPtr<FDisplayClusterConfiguratorBlueprintEditor> Toolkit = ToolkitPtr.Pin();
 	check(Toolkit.IsValid());
 
@@ -315,12 +321,12 @@ void SDisplayClusterConfiguratorBaseNode::EndUserInteraction() const
 
 const FSlateBrush* SDisplayClusterConfiguratorBaseNode::GetShadowBrush(bool bSelected) const
 {
-	return FEditorStyle::GetNoBrush();
+	return FStyleDefaults::GetNoBrush();
 }
 
 bool SDisplayClusterConfiguratorBaseNode::CanBeSelected(const FVector2D& MousePositionInNode) const
 {
-	return IsNodeEnabled();
+	return IsNodeEnabled() && IsNodeUnlocked();
 }
 
 bool SDisplayClusterConfiguratorBaseNode::ShouldAllowCulling() const
@@ -459,6 +465,12 @@ bool SDisplayClusterConfiguratorBaseNode::IsNodeEnabled() const
 	return EdNode->IsNodeEnabled();
 }
 
+bool SDisplayClusterConfiguratorBaseNode::IsNodeUnlocked() const
+{
+	UDisplayClusterConfiguratorBaseNode* EdNode = GetGraphNodeChecked<UDisplayClusterConfiguratorBaseNode>();
+	return EdNode->IsNodeUnlocked();
+}
+
 FVector2D SDisplayClusterConfiguratorBaseNode::GetSize() const
 {
 	return FVector2D(GraphNode->NodeWidth, GraphNode->NodeHeight);
diff --git a/Engine/Plugins/Runtime/nDisplay/Source/DisplayClusterConfigurator/Private/Views/OutputMapping/GraphNodes/SDisplayClusterConfiguratorBaseNode.h b/Engine/Plugins/Runtime/nDisplay/Source/DisplayClusterConfigurator/Private/Views/OutputMapping/GraphNodes/SDisplayClusterConfiguratorBaseNode.h
--- a/Engine/Plugins/Runtime/nDisplay/Source/DisplayClusterConfigurator/Private/Views/OutputMapping/GraphNodes/SDisplayClusterConfiguratorBaseNode.h
+++ b/Engine/Plugins/Runtime/nDisplay/Source/DisplayClusterConfigurator/Private/Views/OutputMapping/GraphNodes/SDisplayClusterConfiguratorBaseNode.h
@@ -104,6 +104,11 @@ public:
 	 */
 	virtual bool IsNodeEnabled() const;
 
+	/**
+	 * @return true if the node should be locked, meaning it can't be selected, moved, or resized
+	 */
+	virtual bool IsNodeUnlocked() const;
+
 	/**
 	 * @return The intended size of the node, taken from the backing EdGraphNode 
 	 */
diff --git a/Engine/Plugins/Runtime/nDisplay/Source/DisplayClusterConfigurator/Private/Views/OutputMapping/GraphNodes/SDisplayClusterConfiguratorCanvasNode.cpp b/Engine/Plugins/Runtime/nDisplay/Source/DisplayClusterConfigurator/Private/Views/OutputMapping/GraphNodes/SDisplayClusterConfiguratorCanvasNode.cpp
--- a/Engine/Plugins/Runtime/nDisplay/Source/DisplayClusterConfigurator/Private/Views/OutputMapping/GraphNodes/SDisplayClusterConfiguratorCanvasNode.cpp
+++ b/Engine/Plugins/Runtime/nDisplay/Source/DisplayClusterConfigurator/Private/Views/OutputMapping/GraphNodes/SDisplayClusterConfiguratorCanvasNode.cpp
@@ -30,7 +30,7 @@ void SDisplayClusterConfiguratorCanvasNode::UpdateGraphNode()
 	TAttribute<const FSlateBrush*> SelectedBrush = TAttribute<const FSlateBrush*>::Create(TAttribute<const FSlateBrush*>::FGetter::CreateSP(this, &SDisplayClusterConfiguratorCanvasNode::GetSelectedBrush));
 
 	CanvasSizeTextWidget = SNew(SBorder)
-	.BorderImage(FEditorStyle::GetBrush("NoBorder"))
+	.BorderImage(FAppStyle::Get().GetBrush("NoBorder"))
 	[
 		SNew(SVerticalBox)
 		+ SVerticalBox::Slot()
@@ -46,7 +46,7 @@ void SDisplayClusterConfiguratorCanvasNode::UpdateGraphNode()
 			[
 				SNew(STextBlock)
 				.Text(this, &SDisplayClusterConfiguratorCanvasNode::GetCanvasSizeText)
-				.TextStyle(&FDisplayClusterConfiguratorStyle::GetWidgetStyle<FTextBlockStyle>("DisplayClusterConfigurator.Node.Text.Regular"))
+				.TextStyle(&FDisplayClusterConfiguratorStyle::Get().GetWidgetStyle<FTextBlockStyle>("DisplayClusterConfigurator.Node.Text.Regular"))
 				.Justification(ETextJustify::Center)
 			]
 		]
@@ -57,7 +57,7 @@ void SDisplayClusterConfiguratorCanvasNode::UpdateGraphNode()
 	.VAlign(VAlign_Fill)
 	[
 		SNew(SBorder)
-		.BorderImage(FEditorStyle::GetBrush("NoBorder"))
+		.BorderImage(FAppStyle::Get().GetBrush("NoBorder"))
 		[
 			SNew(SVerticalBox)
 			+SVerticalBox::Slot()
@@ -117,11 +117,11 @@ const FSlateBrush* SDisplayClusterConfiguratorCanvasNode::GetSelectedBrush() con
 	if (GetOwnerPanel()->SelectionManager.SelectedNodes.Contains(GetNodeObj()))
 	{
 		// Selected Case
-		return FDisplayClusterConfiguratorStyle::GetBrush("DisplayClusterConfigurator.Selected.Canvas.Brush");
+		return FDisplayClusterConfiguratorStyle::Get().GetBrush("DisplayClusterConfigurator.Selected.Canvas.Brush");
 	}
 
 	// Regular case
-	return FDisplayClusterConfiguratorStyle::GetBrush("DisplayClusterConfigurator.Regular.Canvas.Brush");
+	return FDisplayClusterConfiguratorStyle::Get().GetBrush("DisplayClusterConfigurator.Regular.Canvas.Brush");
 }
 
 FMargin SDisplayClusterConfiguratorCanvasNode::GetBackgroundPosition() const
diff --git a/Engine/Plugins/Runtime/nDisplay/Source/DisplayClusterConfigurator/Private/Views/OutputMapping/GraphNodes/SDisplayClusterConfiguratorHostNode.cpp b/Engine/Plugins/Runtime/nDisplay/Source/DisplayClusterConfigurator/Private/Views/OutputMapping/GraphNodes/SDisplayClusterConfiguratorHostNode.cpp
--- a/Engine/Plugins/Runtime/nDisplay/Source/DisplayClusterConfigurator/Private/Views/OutputMapping/GraphNodes/SDisplayClusterConfiguratorHostNode.cpp
+++ b/Engine/Plugins/Runtime/nDisplay/Source/DisplayClusterConfigurator/Private/Views/OutputMapping/GraphNodes/SDisplayClusterConfiguratorHostNode.cpp
@@ -117,12 +117,12 @@ public:
 			.AutoSize(true)
 			[
 				SNew(SBorder)
-				.BorderImage(FEditorStyle::GetBrush("NoBorder"))
+				.BorderImage(FAppStyle::Get().GetBrush("NoBorder"))
 				.Padding(5)
 				[
 					SNew(STextBlock)
 					.Text(this, &SNodeOrigin::GetOriginText)
-					.TextStyle(&FDisplayClusterConfiguratorStyle::GetWidgetStyle<FTextBlockStyle>("DisplayClusterConfigurator.Node.Text.Small"))
+					.TextStyle(&FDisplayClusterConfiguratorStyle::Get().GetWidgetStyle<FTextBlockStyle>("DisplayClusterConfigurator.Node.Text.Small"))
 				]
 			]
 		];
@@ -268,7 +268,7 @@ void SDisplayClusterConfiguratorHostNode::UpdateGraphNode()
 					[
 						SNew(SImage)
 						.ColorAndOpacity(this, &SDisplayClusterConfiguratorHostNode::GetBorderColor)
-						.Image(FDisplayClusterConfiguratorStyle::GetBrush("DisplayClusterConfigurator.Node.Body"))
+						.Image(FAppStyle::Get().GetBrush("WhiteBrush"))
 					]
 
 					+ SOverlay::Slot()
@@ -282,71 +282,99 @@ void SDisplayClusterConfiguratorHostNode::UpdateGraphNode()
 						.StretchDirection(EStretchDirection::Both)
 						.VAlign(VAlign_Center)
 						[
-							SNew(SHorizontalBox)
-
-							+ SHorizontalBox::Slot()
-							.AutoWidth()
-							.HAlign(HAlign_Center)
-							.VAlign(VAlign_Center)
+							SNew(SBorder)
+							.BorderImage(FAppStyle::Get().GetBrush("NoBorder"))
+							.ForegroundColor(this, &SDisplayClusterConfiguratorHostNode::GetTextColor)
 							[
-								SNew(SBox)
-								.WidthOverride(48)
-								.HeightOverride(48)
+								SNew(SHorizontalBox)
+
+								+ SHorizontalBox::Slot()
+								.AutoWidth()
+								.HAlign(HAlign_Center)
+								.VAlign(VAlign_Center)
 								[
-									SNew(SImage)
-									.Image(FDisplayClusterConfiguratorStyle::GetBrush(TEXT("DisplayClusterConfigurator.TreeItems.Host")))
+									SNew(SBox)
+									.WidthOverride(48)
+									.HeightOverride(48)
+									[
+										SNew(SImage)
+										.Image(FDisplayClusterConfiguratorStyle::Get().GetBrush(TEXT("DisplayClusterConfigurator.TreeItems.Host")))
+										.ColorAndOpacity(FSlateColor::UseForeground())
+									]
 								]
-							]
 
-							+ SHorizontalBox::Slot()
-							.AutoWidth()
-							[
-								SNew(SSpacer)
-								.Size(FVector2D(25, 1))
-							]
+								+ SHorizontalBox::Slot()
+								.AutoWidth()
+								[
+									SNew(SSpacer)
+									.Size(FVector2D(25, 1))
+								]
 
-							+ SHorizontalBox::Slot()
-							.AutoWidth()
-							.VAlign(VAlign_Center)
-							[
-								SNew(STextBlock)
-								.Text(this, &SDisplayClusterConfiguratorHostNode::GetHostNameText)
-								.TextStyle(&FDisplayClusterConfiguratorStyle::GetWidgetStyle<FTextBlockStyle>("DisplayClusterConfigurator.Host.Text.Title"))
-								.ColorAndOpacity(FDisplayClusterConfiguratorStyle::GetColor("DisplayClusterConfigurator.Node.Text.Color.Regular"))
-							]
+								+ SHorizontalBox::Slot()
+								.AutoWidth()
+								.VAlign(VAlign_Center)
+								[
+									SNew(STextBlock)
+									.Text(this, &SDisplayClusterConfiguratorHostNode::GetHostNameText)
+									.TextStyle(&FDisplayClusterConfiguratorStyle::Get().GetWidgetStyle<FTextBlockStyle>("DisplayClusterConfigurator.Host.Text.Title"))
+									.ColorAndOpacity(FSlateColor::UseForeground())
+								]
 
-							+ SHorizontalBox::Slot()
-							.FillWidth(1)
-							.Padding(20, 0)
-							[
-								SNew(SSpacer)
-							]
+								+ SHorizontalBox::Slot()
+								.AutoWidth()
+								[
+									SNew(SSpacer)
+									.Size(FVector2D(25, 1))
+								]
 
-							+ SHorizontalBox::Slot()
-							.AutoWidth()
-							.VAlign(VAlign_Center)
-							[
-								SNew(STextBlock)
-								.Text(this, &SDisplayClusterConfiguratorHostNode::GetHostResolutionText)
-								.TextStyle(&FDisplayClusterConfiguratorStyle::GetWidgetStyle<FTextBlockStyle>("DisplayClusterConfigurator.Host.Text.Title"))
-								.ColorAndOpacity(FDisplayClusterConfiguratorStyle::GetColor("DisplayClusterConfigurator.Node.Text.Color.Regular"))
-							]
+								+ SHorizontalBox::Slot()
+								.AutoWidth()
+								.VAlign(VAlign_Center)
+								[
+									SNew(SBox)
+									.WidthOverride(48)
+									.HeightOverride(48)
+									.Visibility(this, &SDisplayClusterConfiguratorHostNode::GetLockIconVisibility)
+									[
+										SNew(SImage)
+										.Image(FAppStyle::Get().GetBrush(TEXT("PropertyWindow.Locked")))
+										.ColorAndOpacity(FSlateColor::UseForeground())
+									]
+								]
 
-							+ SHorizontalBox::Slot()
-							.FillWidth(1)
-							.Padding(20, 0)
-							[
-								SNew(SSpacer)
-							]
+								+ SHorizontalBox::Slot()
+								.FillWidth(1)
+								.Padding(20, 0)
+								[
+									SNew(SSpacer)
+								]
 
-							+ SHorizontalBox::Slot()
-							.AutoWidth()
-							.VAlign(VAlign_Center)
-							[
-								SNew(STextBlock)
-								.Text(this, &SDisplayClusterConfiguratorHostNode::GetHostIPText)
-								.TextStyle(&FDisplayClusterConfiguratorStyle::GetWidgetStyle<FTextBlockStyle>("DisplayClusterConfigurator.Host.Text.Title"))
-								.ColorAndOpacity(FDisplayClusterConfiguratorStyle::GetColor("DisplayClusterConfigurator.Node.Text.Color.Regular"))
+								+ SHorizontalBox::Slot()
+								.AutoWidth()
+								.VAlign(VAlign_Center)
+								[
+									SNew(STextBlock)
+									.Text(this, &SDisplayClusterConfiguratorHostNode::GetHostResolutionText)
+									.TextStyle(&FDisplayClusterConfiguratorStyle::Get().GetWidgetStyle<FTextBlockStyle>("DisplayClusterConfigurator.Host.Text.Title"))
+									.ColorAndOpacity(FSlateColor::UseForeground())
+								]
+
+								+ SHorizontalBox::Slot()
+								.FillWidth(1)
+								.Padding(20, 0)
+								[
+									SNew(SSpacer)
+								]
+
+								+ SHorizontalBox::Slot()
+								.AutoWidth()
+								.VAlign(VAlign_Center)
+								[
+									SNew(STextBlock)
+									.Text(this, &SDisplayClusterConfiguratorHostNode::GetHostIPText)
+									.TextStyle(&FDisplayClusterConfiguratorStyle::Get().GetWidgetStyle<FTextBlockStyle>("DisplayClusterConfigurator.Host.Text.Title"))
+									.ColorAndOpacity(FSlateColor::UseForeground())
+								]
 							]
 						]
 					]
@@ -357,8 +385,8 @@ void SDisplayClusterConfiguratorHostNode::UpdateGraphNode()
 					.Padding(BorderThickness)
 					[
 						SNew(SImage)
-						.ColorAndOpacity(FDisplayClusterConfiguratorStyle::GetColor("DisplayClusterConfigurator.Node.Host.Inner.Background"))
-						.Image(FDisplayClusterConfiguratorStyle::GetBrush("DisplayClusterConfigurator.Node.Body"))
+						.ColorAndOpacity(FDisplayClusterConfiguratorStyle::Get().GetColor("DisplayClusterConfigurator.Node.Host.Inner.Background"))
+						.Image(FAppStyle::Get().GetBrush("WhiteBrush"))
 					]
 				]
 			]
@@ -469,7 +497,7 @@ FSlateColor SDisplayClusterConfiguratorHostNode::GetBorderColor() const
 	if (GetOwnerPanel()->SelectionManager.SelectedNodes.Contains(GraphNode))
 	{
 		// Selected Case
-		return FDisplayClusterConfiguratorStyle::GetColor("DisplayClusterConfigurator.Node.Color.Selected");
+		return FDisplayClusterConfiguratorStyle::Get().GetColor("DisplayClusterConfigurator.Node.Color.Selected");
 	}
 
 	// Regular case
@@ -477,4 +505,21 @@ FSlateColor SDisplayClusterConfiguratorHostNode::GetBorderColor() const
 	return PCEdNode->GetHostColor();
 }
 
+
+FSlateColor SDisplayClusterConfiguratorHostNode::GetTextColor() const
+{
+	if (GetOwnerPanel()->SelectionManager.SelectedNodes.Contains(GraphNode))
+	{
+		// Selected Case
+		return FDisplayClusterConfiguratorStyle::Get().GetColor("DisplayClusterConfigurator.Node.Text.Color.Selected");
+	}
+
+	return FDisplayClusterConfiguratorStyle::Get().GetColor("DisplayClusterConfigurator.Node.Text.Color.Regular");
+}
+
+EVisibility SDisplayClusterConfiguratorHostNode::GetLockIconVisibility() const
+{
+	return !IsNodeUnlocked() ? EVisibility::Visible : EVisibility::Hidden;
+}
+
 #undef LOCTEXT_NAMESPACE
\ No newline at end of file
diff --git a/Engine/Plugins/Runtime/nDisplay/Source/DisplayClusterConfigurator/Private/Views/OutputMapping/GraphNodes/SDisplayClusterConfiguratorHostNode.h b/Engine/Plugins/Runtime/nDisplay/Source/DisplayClusterConfigurator/Private/Views/OutputMapping/GraphNodes/SDisplayClusterConfiguratorHostNode.h
--- a/Engine/Plugins/Runtime/nDisplay/Source/DisplayClusterConfigurator/Private/Views/OutputMapping/GraphNodes/SDisplayClusterConfiguratorHostNode.h
+++ b/Engine/Plugins/Runtime/nDisplay/Source/DisplayClusterConfigurator/Private/Views/OutputMapping/GraphNodes/SDisplayClusterConfiguratorHostNode.h
@@ -45,6 +45,9 @@ private:
 	FText GetHostResolutionText() const;
 
 	FSlateColor GetBorderColor() const;
+	FSlateColor GetTextColor() const;
+
+	EVisibility GetLockIconVisibility() const;
 
 private:
 	FMargin BorderThickness;
diff --git a/Engine/Plugins/Runtime/nDisplay/Source/DisplayClusterConfigurator/Private/Views/OutputMapping/GraphNodes/SDisplayClusterConfiguratorViewportNode.cpp b/Engine/Plugins/Runtime/nDisplay/Source/DisplayClusterConfigurator/Private/Views/OutputMapping/GraphNodes/SDisplayClusterConfiguratorViewportNode.cpp
--- a/Engine/Plugins/Runtime/nDisplay/Source/DisplayClusterConfigurator/Private/Views/OutputMapping/GraphNodes/SDisplayClusterConfiguratorViewportNode.cpp
+++ b/Engine/Plugins/Runtime/nDisplay/Source/DisplayClusterConfigurator/Private/Views/OutputMapping/GraphNodes/SDisplayClusterConfiguratorViewportNode.cpp
@@ -66,7 +66,6 @@ void SDisplayClusterConfiguratorViewportNode::UpdateGraphNode()
 					.HAlign(HAlign_Fill)
 					[
 						SNew(SBorder)
-						.BorderImage(FDisplayClusterConfiguratorStyle::GetBrush("DisplayClusterConfigurator.Node.Window.Border.Brush"))
 						.Padding(FMargin(0.f))
 						[
 							BackgroundImage.ToSharedRef()
@@ -86,8 +85,9 @@ void SDisplayClusterConfiguratorViewportNode::UpdateGraphNode()
 							.VAlign(VAlign_Center)
 							[
 								SNew(SBorder)
-								.BorderImage(FEditorStyle::GetBrush("WhiteBrush"))
+								.BorderImage(FAppStyle::Get().GetBrush("WhiteBrush"))
 								.BorderBackgroundColor(this, &SDisplayClusterConfiguratorViewportNode::GetTextBoxColor)
+								.ForegroundColor(this, &SDisplayClusterConfiguratorViewportNode::GetTextColor)
 								.Padding(8.0f)
 								[
 									SNew( SVerticalBox )
@@ -99,8 +99,8 @@ void SDisplayClusterConfiguratorViewportNode::UpdateGraphNode()
 										SNew(STextBlock)
 										.Text(FText::FromString(ViewportEdNode->GetNodeName()))
 										.Justification(ETextJustify::Center)
-										.TextStyle(&FDisplayClusterConfiguratorStyle::GetWidgetStyle<FTextBlockStyle>("DisplayClusterConfigurator.Node.Text.Bold"))
-										.ColorAndOpacity(FDisplayClusterConfiguratorStyle::GetColor("DisplayClusterConfigurator.Node.Text.Color.Regular"))
+										.TextStyle(&FDisplayClusterConfiguratorStyle::Get().GetWidgetStyle<FTextBlockStyle>("DisplayClusterConfigurator.Node.Text.Bold"))
+										.ColorAndOpacity(FSlateColor::UseForeground())
 									]
 
 									+ SVerticalBox::Slot()
@@ -110,8 +110,8 @@ void SDisplayClusterConfiguratorViewportNode::UpdateGraphNode()
 										SNew(STextBlock)
 										.Text(this, &SDisplayClusterConfiguratorViewportNode::GetPositionAndSizeText)
 										.Justification(ETextJustify::Center)
-										.TextStyle(&FDisplayClusterConfiguratorStyle::GetWidgetStyle<FTextBlockStyle>("DisplayClusterConfigurator.Node.Text.Regular"))
-										.ColorAndOpacity(FDisplayClusterConfiguratorStyle::GetColor("DisplayClusterConfigurator.Node.Text.Color.WhiteGray"))
+										.TextStyle(&FDisplayClusterConfiguratorStyle::Get().GetWidgetStyle<FTextBlockStyle>("DisplayClusterConfigurator.Node.Text.Regular"))
+										.ColorAndOpacity(FSlateColor::UseSubduedForeground())
 									]
 
 									+ SVerticalBox::Slot()
@@ -121,8 +121,8 @@ void SDisplayClusterConfiguratorViewportNode::UpdateGraphNode()
 										SNew(STextBlock)
 										.Text(this, &SDisplayClusterConfiguratorViewportNode::GetTransformText)
 										.Justification(ETextJustify::Center)
-										.TextStyle(&FDisplayClusterConfiguratorStyle::GetWidgetStyle<FTextBlockStyle>("DisplayClusterConfigurator.Node.Text.Regular"))
-										.ColorAndOpacity(FDisplayClusterConfiguratorStyle::GetColor("DisplayClusterConfigurator.Node.Text.Color.WhiteGray"))
+										.TextStyle(&FDisplayClusterConfiguratorStyle::Get().GetWidgetStyle<FTextBlockStyle>("DisplayClusterConfigurator.Node.Text.Regular"))
+										.ColorAndOpacity(FSlateColor::UseSubduedForeground())
 										.Visibility(this, &SDisplayClusterConfiguratorViewportNode::GetTransformTextVisibility)
 									]
 
@@ -138,7 +138,8 @@ void SDisplayClusterConfiguratorViewportNode::UpdateGraphNode()
 										.Visibility(this, &SDisplayClusterConfiguratorViewportNode::GetLockIconVisibility)
 										[
 											SNew(SImage)
-											.Image(FEditorStyle::GetBrush(TEXT("GenericLock")))
+											.Image(FAppStyle::Get().GetBrush(TEXT("PropertyWindow.Locked")))
+											.ColorAndOpacity(FSlateColor::UseSubduedForeground())
 										]
 									]
 								]
@@ -166,16 +167,6 @@ void SDisplayClusterConfiguratorViewportNode::Tick(const FGeometry& AllottedGeom
 	UpdatePreviewTexture();
 }
 
-void SDisplayClusterConfiguratorViewportNode::MoveTo(const FVector2D& NewPosition, FNodeSet& NodeFilter, bool bMarkDirty)
-{
-	if (IsViewportLocked())
-	{
-		NodeFilter.Add(SharedThis(this));
-	}
-
-	SDisplayClusterConfiguratorBaseNode::MoveTo(NewPosition, NodeFilter, bMarkDirty);
-}
-
 bool SDisplayClusterConfiguratorViewportNode::IsNodeVisible() const
 {
 	UDisplayClusterConfiguratorViewportNode* ViewportEdNode = GetGraphNodeChecked<UDisplayClusterConfiguratorViewportNode>();
@@ -214,7 +205,7 @@ FSlateColor SDisplayClusterConfiguratorViewportNode::GetBackgroundColor() const
 
 	const bool bIsSelected = GetOwnerPanel()->SelectionManager.SelectedNodes.Contains(GraphNode);
 	const bool bHasImageBackground = BackgroundActiveBrush.GetResourceObject() != nullptr;
-	const bool bIsLocked = IsViewportLocked();
+	const bool bIsLocked = !IsNodeUnlocked();
 	const bool bTintBackground = OutputMapping->GetOutputMappingSettings().bTintSelectedViewports;
 
 	UDisplayClusterConfiguratorViewportNode* ViewportEdNode = GetGraphNodeChecked<UDisplayClusterConfiguratorViewportNode>();
@@ -223,12 +214,12 @@ FSlateColor SDisplayClusterConfiguratorViewportNode::GetBackgroundColor() const
 		if (bIsSelected)
 		{
 			// Selected Case
-			return FDisplayClusterConfiguratorStyle::GetColor("DisplayClusterConfigurator.Node.Viewport.OutsideBackgroundColor.Selected");
+			return FDisplayClusterConfiguratorStyle::Get().GetColor("DisplayClusterConfigurator.Node.Viewport.OutsideBackgroundColor.Selected");
 		}
 		else
 		{
 			// Regular case
-			return FDisplayClusterConfiguratorStyle::GetColor("DisplayClusterConfigurator.Node.Viewport.OutsideBackgroundColor.Regular");
+			return FDisplayClusterConfiguratorStyle::Get().GetColor("DisplayClusterConfigurator.Node.Viewport.OutsideBackgroundColor.Regular");
 		}
 	}
 	else
@@ -237,11 +228,11 @@ FSlateColor SDisplayClusterConfiguratorViewportNode::GetBackgroundColor() const
 		{
 			if (bIsSelected && bTintBackground)
 			{
-				return FDisplayClusterConfiguratorStyle::GetColor("DisplayClusterConfigurator.Node.Viewport.BackgroundImage.Selected");
+				return FDisplayClusterConfiguratorStyle::Get().GetColor("DisplayClusterConfigurator.Node.Viewport.BackgroundImage.Selected");
 			}
 			else if (bIsLocked)
 			{
-				return FDisplayClusterConfiguratorStyle::GetColor("DisplayClusterConfigurator.Node.Viewport.BackgroundImage.Locked");
+				return FDisplayClusterConfiguratorStyle::Get().GetColor("DisplayClusterConfigurator.Node.Viewport.BackgroundImage.Locked");
 			}
 			else
 			{
@@ -253,12 +244,12 @@ FSlateColor SDisplayClusterConfiguratorViewportNode::GetBackgroundColor() const
 			if (bIsSelected && bTintBackground)
 			{
 				// Selected Case
-				return FDisplayClusterConfiguratorStyle::GetColor("DisplayClusterConfigurator.Node.Viewport.BackgroundColor.Selected");
+				return FDisplayClusterConfiguratorStyle::Get().GetColor("DisplayClusterConfigurator.Node.Viewport.BackgroundColor.Selected");
 			}
 			else
 			{
 				// Regular case
-				return FDisplayClusterConfiguratorStyle::GetColor("DisplayClusterConfigurator.Node.Viewport.BackgroundColor.Regular");
+				return FDisplayClusterConfiguratorStyle::Get().GetColor("DisplayClusterConfigurator.Node.Viewport.BackgroundColor.Regular");
 			}
 		}
 	}
@@ -273,7 +264,7 @@ const FSlateBrush* SDisplayClusterConfiguratorViewportNode::GetBackgroundBrush()
 	}
 	else
 	{
-		return FDisplayClusterConfiguratorStyle::GetBrush("DisplayClusterConfigurator.Node.Body");
+		return FAppStyle::Get().GetBrush("WhiteBrush");
 	}
 }
 
@@ -313,44 +304,47 @@ TOptional<FSlateRenderTransform> SDisplayClusterConfiguratorViewportNode::GetBac
 
 const FSlateBrush* SDisplayClusterConfiguratorViewportNode::GetNodeShadowBrush() const
 {
-	return FEditorStyle::GetBrush(TEXT("Graph.Node.Shadow"));
+	return FAppStyle::Get().GetBrush(TEXT("Graph.Node.Shadow"));
 }
 
 const FSlateBrush* SDisplayClusterConfiguratorViewportNode::GetBorderBrush() const
 {
 	if (GetOwnerPanel()->SelectionManager.SelectedNodes.Contains(GraphNode))
 	{
-		return FDisplayClusterConfiguratorStyle::GetBrush("DisplayClusterConfigurator.Node.Viewport.Border.Brush.Selected");
+		return FDisplayClusterConfiguratorStyle::Get().GetBrush("DisplayClusterConfigurator.Node.Viewport.Border.Brush.Selected");
 	}
 	else
 	{
 		UDisplayClusterConfiguratorViewportNode* ViewportEdNode = GetGraphNodeChecked<UDisplayClusterConfiguratorViewportNode>();
 		if (ViewportEdNode->IsOutsideParentBoundary())
 		{
-			return FDisplayClusterConfiguratorStyle::GetBrush("DisplayClusterConfigurator.Node.Viewport.Border.OutsideBrush.Regular");
+			return FDisplayClusterConfiguratorStyle::Get().GetBrush("DisplayClusterConfigurator.Node.Viewport.Border.Brush.Outside");
 		}
 		else
 		{
-			return FDisplayClusterConfiguratorStyle::GetBrush("DisplayClusterConfigurator.Node.Viewport.Border.Brush.Regular");
+			return FDisplayClusterConfiguratorStyle::Get().GetBrush("DisplayClusterConfigurator.Node.Viewport.Border.Brush.Regular");
 		}
 	}
 }
 
 FSlateColor SDisplayClusterConfiguratorViewportNode::GetTextBoxColor() const
 {
-	const bool bIsSelected = GetOwnerPanel()->SelectionManager.SelectedNodes.Contains(GraphNode);
-	const bool bIsLocked = IsViewportLocked();
-	
-	if (bIsSelected)
+	if (GetOwnerPanel()->SelectionManager.SelectedNodes.Contains(GraphNode))
 	{
-		return FDisplayClusterConfiguratorStyle::GetColor("DisplayClusterConfigurator.Node.Color.Selected");
+		return FDisplayClusterConfiguratorStyle::Get().GetColor("DisplayClusterConfigurator.Node.Color.Selected");
 	}
-	else if (bIsLocked)
+
+	return FDisplayClusterConfiguratorStyle::Get().GetColor("DisplayClusterConfigurator.Node.Viewport.Text.Background");
+}
+
+FSlateColor SDisplayClusterConfiguratorViewportNode::GetTextColor() const
+{
+	if (GetOwnerPanel()->SelectionManager.SelectedNodes.Contains(GraphNode))
 	{
-		return FDisplayClusterConfiguratorStyle::GetColor("DisplayClusterConfigurator.Node.Viewport.Text.Background.Locked");
+		return FDisplayClusterConfiguratorStyle::Get().GetColor("DisplayClusterConfigurator.Node.Text.Color.Selected");
 	}
 
-	return FDisplayClusterConfiguratorStyle::GetColor("DisplayClusterConfigurator.Node.Viewport.Text.Background");
+	return FDisplayClusterConfiguratorStyle::Get().GetColor("DisplayClusterConfigurator.Node.Text.Color.Regular");
 }
 
 FText SDisplayClusterConfiguratorViewportNode::GetPositionAndSizeText() const
@@ -431,7 +425,7 @@ bool SDisplayClusterConfiguratorViewportNode::IsViewportLocked() const
 
 EVisibility SDisplayClusterConfiguratorViewportNode::GetLockIconVisibility() const
 {
-	return IsViewportLocked() ? EVisibility::Visible : EVisibility::Collapsed;
+	return !IsNodeUnlocked() ? EVisibility::Visible : EVisibility::Collapsed;
 }
 
 void SDisplayClusterConfiguratorViewportNode::UpdatePreviewTexture()
diff --git a/Engine/Plugins/Runtime/nDisplay/Source/DisplayClusterConfigurator/Private/Views/OutputMapping/GraphNodes/SDisplayClusterConfiguratorViewportNode.h b/Engine/Plugins/Runtime/nDisplay/Source/DisplayClusterConfigurator/Private/Views/OutputMapping/GraphNodes/SDisplayClusterConfiguratorViewportNode.h
--- a/Engine/Plugins/Runtime/nDisplay/Source/DisplayClusterConfigurator/Private/Views/OutputMapping/GraphNodes/SDisplayClusterConfiguratorViewportNode.h
+++ b/Engine/Plugins/Runtime/nDisplay/Source/DisplayClusterConfigurator/Private/Views/OutputMapping/GraphNodes/SDisplayClusterConfiguratorViewportNode.h
@@ -30,13 +30,12 @@ public:
 	//~ Begin SGraphNode interface
 	virtual void UpdateGraphNode() override;
 	virtual void Tick(const FGeometry& AllottedGeometry, const double InCurrentTime, const float InDeltaTime) override;
-	virtual void MoveTo(const FVector2D& NewPosition, FNodeSet& NodeFilter, bool bMarkDirty = true) override;
 	//~ End of SGraphNode interface
 
 	//~ Begin SDisplayClusterConfiguratorBaseNode interface
 	virtual bool IsNodeVisible() const override;
 	virtual bool CanNodeBeSnapAligned() const override { return true; }
-	virtual bool CanNodeBeResized() const { return !IsViewportLocked(); }
+	virtual bool CanNodeBeResized() const { return IsNodeUnlocked(); }
 	virtual float GetNodeMinimumSize() const override;
 	virtual float GetNodeMaximumSize() const override;
 	virtual bool IsAspectRatioFixed() const override;
@@ -49,6 +48,7 @@ private:
 	const FSlateBrush* GetNodeShadowBrush() const;
 	const FSlateBrush* GetBorderBrush() const;
 	FSlateColor GetTextBoxColor() const;
+	FSlateColor GetTextColor() const;
 	FText GetPositionAndSizeText() const;
 	FText GetTransformText() const;
 	EVisibility GetTransformTextVisibility() const;
diff --git a/Engine/Plugins/Runtime/nDisplay/Source/DisplayClusterConfigurator/Private/Views/OutputMapping/GraphNodes/SDisplayClusterConfiguratorWindowNode.cpp b/Engine/Plugins/Runtime/nDisplay/Source/DisplayClusterConfigurator/Private/Views/OutputMapping/GraphNodes/SDisplayClusterConfiguratorWindowNode.cpp
--- a/Engine/Plugins/Runtime/nDisplay/Source/DisplayClusterConfigurator/Private/Views/OutputMapping/GraphNodes/SDisplayClusterConfiguratorWindowNode.cpp
+++ b/Engine/Plugins/Runtime/nDisplay/Source/DisplayClusterConfigurator/Private/Views/OutputMapping/GraphNodes/SDisplayClusterConfiguratorWindowNode.cpp
@@ -51,7 +51,7 @@ public:
 			[
 				SNew(SImage)
 				.ColorAndOpacity(InArgs._ColorAndOpacity)
-				.Image(FDisplayClusterConfiguratorStyle::GetBrush("DisplayClusterConfigurator.Node.Brush.Corner"))
+				.Image(FDisplayClusterConfiguratorStyle::Get().GetBrush("DisplayClusterConfigurator.Node.Brush.Corner"))
 			]
 		];
 	}
@@ -86,6 +86,7 @@ public:
 		, _ZIndexOffset(0)
 	{ }
 		SLATE_ATTRIBUTE(FSlateColor, ColorAndOpacity)
+		SLATE_ATTRIBUTE(FSlateColor, ForegroundColor)
 		SLATE_ARGUMENT(int32, ZIndexOffset)
 	SLATE_END_ARGS()
 
@@ -112,9 +113,10 @@ public:
 					SNew(SBorder)
 					.HAlign(HAlign_Left)
 					.VAlign(VAlign_Center)
-					.BorderImage(FEditorStyle::GetBrush("WhiteBrush"))
+					.BorderImage(FAppStyle::Get().GetBrush("WhiteBrush"))
 					.BorderBackgroundColor(InArgs._ColorAndOpacity)
 					.Padding(FMargin(20, 10, 30, 10))
+					.ForegroundColor(InArgs._ForegroundColor)
 					[
 						SNew(SHorizontalBox)
 
@@ -128,7 +130,8 @@ public:
 							.HeightOverride(36)
 							[
 								SNew(SImage)
-								.Image(FDisplayClusterConfiguratorStyle::GetBrush(TEXT("DisplayClusterConfigurator.TreeItems.ClusterNode")))
+								.Image(FDisplayClusterConfiguratorStyle::Get().GetBrush(TEXT("DisplayClusterConfigurator.TreeItems.ClusterNode")))
+								.ColorAndOpacity(FSlateColor::UseForeground())
 							]
 						]
 		
@@ -146,7 +149,8 @@ public:
 						[
 							SNew(STextBlock)
 							.Text(this, &SNodeInfo::GetNodeName)
-							.TextStyle(&FDisplayClusterConfiguratorStyle::GetWidgetStyle<FTextBlockStyle>("DisplayClusterConfigurator.Node.Text.Bold"))
+							.TextStyle(&FDisplayClusterConfiguratorStyle::Get().GetWidgetStyle<FTextBlockStyle>("DisplayClusterConfigurator.Node.Text.Bold"))
+							.ColorAndOpacity(FSlateColor::UseForeground())
 						]
 		
 						+ SHorizontalBox::Slot()
@@ -163,7 +167,8 @@ public:
 						[
 							SNew(STextBlock)
 							.Text(this, &SNodeInfo::GetPositionAndSizeText)
-							.TextStyle(&FDisplayClusterConfiguratorStyle::GetWidgetStyle<FTextBlockStyle>("DisplayClusterConfigurator.Node.Text.Regular"))
+							.TextStyle(&FDisplayClusterConfiguratorStyle::Get().GetWidgetStyle<FTextBlockStyle>("DisplayClusterConfigurator.Node.Text.Regular"))
+							.ColorAndOpacity(FSlateColor::UseSubduedForeground())
 						]
 		
 						+ SHorizontalBox::Slot()
@@ -184,7 +189,8 @@ public:
 							.Visibility(this, &SNodeInfo::GetLockIconVisibility)
 							[
 								SNew(SImage)
-								.Image(FEditorStyle::GetBrush(TEXT("GenericLock")))
+								.Image(FAppStyle::Get().GetBrush(TEXT("PropertyWindow.Locked")))
+								.ColorAndOpacity(FSlateColor::UseSubduedForeground())
 							]
 						]
 					]
@@ -239,7 +245,7 @@ public:
 
 	EVisibility GetLockIconVisibility() const
 	{
-		if (ParentNode->IsClusterNodeLocked())
+		if (!ParentNode->IsNodeUnlocked())
 		{
 			return EVisibility::Visible;
 		}
@@ -304,6 +310,7 @@ void SDisplayClusterConfiguratorWindowNode::UpdateGraphNode()
 				SNew(SNodeInfo, SharedThis(this))
 				.Visibility(this, &SDisplayClusterConfiguratorWindowNode::GetNodeInfoVisibility)
 				.ColorAndOpacity(this, &SDisplayClusterConfiguratorWindowNode::GetCornerColor)
+				.ForegroundColor(this, &SDisplayClusterConfiguratorWindowNode::GetTextColor)
 			]
 		]
 	];
@@ -326,7 +333,7 @@ void SDisplayClusterConfiguratorWindowNode::UpdateGraphNode()
 				.VAlign(VAlign_Fill)
 				.HAlign(HAlign_Fill)
 				[
-					CreateBackground(FDisplayClusterConfiguratorStyle::GetColor("DisplayClusterConfigurator.Node.Window.Inner.Background"))
+					CreateBackground(FDisplayClusterConfiguratorStyle::Get().GetColor("DisplayClusterConfigurator.Node.Window.Inner.Background"))
 				]
 
 				+ SOverlay::Slot()
@@ -345,26 +352,6 @@ void SDisplayClusterConfiguratorWindowNode::UpdateGraphNode()
 	];
 }
 
-void SDisplayClusterConfiguratorWindowNode::MoveTo(const FVector2D& NewPosition, FNodeSet& NodeFilter, bool bMarkDirty)
-{
-	if (IsClusterNodeLocked())
-	{
-		NodeFilter.Add(SharedThis(this));
-	}
-
-	SDisplayClusterConfiguratorBaseNode::MoveTo(NewPosition, NodeFilter, bMarkDirty);
-}
-
-bool SDisplayClusterConfiguratorWindowNode::CanBeSelected(const FVector2D& MousePositionInNode) const
-{
-	if (IsClusterNodeLocked())
-	{
-		return false;
-	}
-
-	return SDisplayClusterConfiguratorBaseNode::CanBeSelected(MousePositionInNode);
-}
-
 FVector2D SDisplayClusterConfiguratorWindowNode::ComputeDesiredSize(float) const
 {
 	return GetSize() * WindowScaleFactor;
@@ -417,7 +404,7 @@ TSharedRef<SWidget> SDisplayClusterConfiguratorWindowNode::CreateBackground(cons
 		[
 			SNew(SImage)
 			.ColorAndOpacity(InColorAndOpacity)
-			.Image(FDisplayClusterConfiguratorStyle::GetBrush("DisplayClusterConfigurator.Node.Body"))
+			.Image(FAppStyle::Get().GetBrush("WhiteBrush"))
 		]
 	
 		+SOverlay::Slot()
@@ -444,11 +431,11 @@ const FSlateBrush* SDisplayClusterConfiguratorWindowNode::GetBorderBrush() const
 	if (GetOwnerPanel()->SelectionManager.SelectedNodes.Contains(GraphNode))
 	{
 		// Selected Case
-		return FDisplayClusterConfiguratorStyle::GetBrush("DisplayClusterConfigurator.Node.Window.Border.Brush.Selected");
+		return FDisplayClusterConfiguratorStyle::Get().GetBrush("DisplayClusterConfigurator.Node.Window.Border.Brush.Selected");
 	}
 
 	// Regular case
-	return FDisplayClusterConfiguratorStyle::GetBrush("DisplayClusterConfigurator.Node.Window.Border.Brush.Regular");
+	return FDisplayClusterConfiguratorStyle::Get().GetBrush("DisplayClusterConfigurator.Node.Window.Border.Brush.Regular");
 }
 
 int32 SDisplayClusterConfiguratorWindowNode::GetBorderLayerOffset() const
@@ -472,7 +459,7 @@ int32 SDisplayClusterConfiguratorWindowNode::GetBorderLayerOffset() const
 
 const FSlateBrush* SDisplayClusterConfiguratorWindowNode::GetNodeShadowBrush() const
 {
-	return FEditorStyle::GetBrush(TEXT("Graph.Node.Shadow"));
+	return FAppStyle::Get().GetBrush(TEXT("Graph.Node.Shadow"));
 }
 
 FMargin SDisplayClusterConfiguratorWindowNode::GetBackgroundPosition() const
@@ -485,15 +472,20 @@ FSlateColor SDisplayClusterConfiguratorWindowNode::GetCornerColor() const
 {
 	if (GetOwnerPanel()->SelectionManager.SelectedNodes.Contains(GraphNode))
 	{
-		return FDisplayClusterConfiguratorStyle::GetColor("DisplayClusterConfigurator.Node.Color.Selected");
+		return FDisplayClusterConfiguratorStyle::Get().GetColor("DisplayClusterConfigurator.Node.Color.Selected");
 	}
 
-	if (IsClusterNodeLocked())
+	return FDisplayClusterConfiguratorStyle::Get().GetColor("DisplayClusterConfigurator.Node.Window.Corner.Color");
+}
+
+FSlateColor SDisplayClusterConfiguratorWindowNode::GetTextColor() const
+{
+	if (GetOwnerPanel()->SelectionManager.SelectedNodes.Contains(GraphNode))
 	{
-		return FDisplayClusterConfiguratorStyle::GetColor("DisplayClusterConfigurator.Node.Window.Corner.Color.Locked");
+		return FDisplayClusterConfiguratorStyle::Get().GetColor("DisplayClusterConfigurator.Node.Text.Color.Selected");
 	}
 
-	return FDisplayClusterConfiguratorStyle::GetColor("DisplayClusterConfigurator.Node.Window.Corner.Color");
+	return FDisplayClusterConfiguratorStyle::Get().GetColor("DisplayClusterConfigurator.Node.Text.Color.Regular");
 }
 
 FVector2D SDisplayClusterConfiguratorWindowNode::GetPreviewImageSize() const
diff --git a/Engine/Plugins/Runtime/nDisplay/Source/DisplayClusterConfigurator/Private/Views/OutputMapping/GraphNodes/SDisplayClusterConfiguratorWindowNode.h b/Engine/Plugins/Runtime/nDisplay/Source/DisplayClusterConfigurator/Private/Views/OutputMapping/GraphNodes/SDisplayClusterConfiguratorWindowNode.h
--- a/Engine/Plugins/Runtime/nDisplay/Source/DisplayClusterConfigurator/Private/Views/OutputMapping/GraphNodes/SDisplayClusterConfiguratorWindowNode.h
+++ b/Engine/Plugins/Runtime/nDisplay/Source/DisplayClusterConfigurator/Private/Views/OutputMapping/GraphNodes/SDisplayClusterConfiguratorWindowNode.h
@@ -33,15 +33,13 @@ public:
 
 	//~ Begin SGraphNode interface
 	virtual void UpdateGraphNode() override;
-	virtual void MoveTo(const FVector2D& NewPosition, FNodeSet& NodeFilter, bool bMarkDirty = true) override;
-	virtual bool CanBeSelected(const FVector2D& MousePositionInNode) const override;
 	virtual FVector2D ComputeDesiredSize(float) const override;
 	virtual FVector2D GetPosition() const override;
 	//~ End SGraphNode interface
 
 	//~ Begin SDisplayClusterConfiguratorBaseNode interface
 	virtual bool CanNodeBeSnapAligned() const override { return true; }
-	virtual bool CanNodeBeResized() const { return !IsClusterNodeLocked(); }
+	virtual bool CanNodeBeResized() const { return IsNodeUnlocked(); }
 	virtual bool IsAspectRatioFixed() const override;
 
 protected:
@@ -56,6 +54,7 @@ private:
 	const FSlateBrush* GetNodeShadowBrush() const;
 	FMargin GetBackgroundPosition() const;
 	FSlateColor GetCornerColor() const;
+	FSlateColor GetTextColor() const;
 	FVector2D GetPreviewImageSize() const;
 	EVisibility GetPreviewImageVisibility() const;
 	int32 GetNodeTitleLayerOffset() const;
diff --git a/Engine/Plugins/Runtime/nDisplay/Source/DisplayClusterConfigurator/Private/Views/OutputMapping/Widgets/SDisplayClusterConfiguratorOutputMappingToolbar.cpp b/Engine/Plugins/Runtime/nDisplay/Source/DisplayClusterConfigurator/Private/Views/OutputMapping/Widgets/SDisplayClusterConfiguratorOutputMappingToolbar.cpp
--- a/Engine/Plugins/Runtime/nDisplay/Source/DisplayClusterConfigurator/Private/Views/OutputMapping/Widgets/SDisplayClusterConfiguratorOutputMappingToolbar.cpp
+++ b/Engine/Plugins/Runtime/nDisplay/Source/DisplayClusterConfigurator/Private/Views/OutputMapping/Widgets/SDisplayClusterConfiguratorOutputMappingToolbar.cpp
@@ -36,38 +36,36 @@
 {
 	FToolBarBuilder ToolbarBuilder(CommandList, FMultiBoxCustomization::None, InExtenders);
 
-	// Use a custom style
 	FName ToolBarStyle = "EditorViewportToolBar";
-	ToolbarBuilder.SetStyle(&FEditorStyle::Get(), ToolBarStyle);
+	ToolbarBuilder.SetStyle(&FAppStyle::Get(), ToolBarStyle);
 	ToolbarBuilder.SetLabelVisibility(EVisibility::Collapsed);
 
 	ToolbarBuilder.BeginSection("Advanced");
 	{
-		// SEditorViewportToolbarMenu buttons are by default 22 pixels high, while the other toolbar buttons are 26 pixels tall
-		// Use an SBox to manually match this button's height with the rest of the toolbar
 		ToolbarBuilder.AddWidget(
-			SNew(SBox)
-			.HeightOverride(26)
-			[
-				SNew(SEditorViewportToolbarMenu)
-				.ParentToolBar(SharedThis(this))
-				.Cursor(EMouseCursor::Default)
-				.Image("EditorViewportToolBar.OptionsDropdown")
-				.OnGetMenuContent(this, &SDisplayClusterConfiguratorOutputMappingToolbar::MakeAdvancedMenu)
-			]
+			SNew(SEditorViewportToolbarMenu)
+			.ParentToolBar(SharedThis(this))
+			.Cursor(EMouseCursor::Default)
+			.Image("EditorViewportToolBar.OptionsDropdown")
+			.OnGetMenuContent(this, &SDisplayClusterConfiguratorOutputMappingToolbar::MakeAdvancedMenu)
 		);
 	}
 	ToolbarBuilder.EndSection();
 
 	ToolbarBuilder.BeginSection("ClusterItems");
 	{
-		ToolbarBuilder.BeginBlockGroup();
-		{
-			ToolbarBuilder.AddToolBarButton(FDisplayClusterConfiguratorOutputMappingCommands::Get().ToggleWindowInfo);
-			ToolbarBuilder.AddToolBarButton(FDisplayClusterConfiguratorOutputMappingCommands::Get().ToggleWindowCornerImage);
-			ToolbarBuilder.AddToolBarButton(FDisplayClusterConfiguratorOutputMappingCommands::Get().ToggleOutsideViewports);
-		}
-		ToolbarBuilder.EndBlockGroup();
+		ToolbarBuilder.AddWidget(
+			SNew(SViewportToolBarIconMenu)
+			.Cursor(EMouseCursor::Default)
+			.Style(ToolBarStyle)
+			.OnGetMenuContent(this, &SDisplayClusterConfiguratorOutputMappingToolbar::MakeWindowDisplayMenu)
+			.Label(this, &SDisplayClusterConfiguratorOutputMappingToolbar::GetWindowDisplayText)
+			.ToolTipText(LOCTEXT("WindowDisplayMenu_ToolTip", "Cluster Node Display"))
+			.Icon(FSlateIcon(FDisplayClusterConfiguratorStyle::Get().GetStyleSetName(), "DisplayClusterConfigurator.OutputMapping.WindowDisplay"))
+			.ParentToolBar(SharedThis(this))
+		);
+
+		ToolbarBuilder.AddToolBarButton(FDisplayClusterConfiguratorOutputMappingCommands::Get().ToggleOutsideViewports);
 	}
 	ToolbarBuilder.EndSection();
 
@@ -81,7 +79,7 @@
 			.Style(ToolBarStyle)
 			.OnGetMenuContent(this, &SDisplayClusterConfiguratorOutputMappingToolbar::MakeViewScaleMenu)
 			.ToolTipText(LOCTEXT("ViewOptionsMenu_ToolTip", "View Scale"))
-			.Icon(FSlateIcon(FDisplayClusterConfiguratorStyle::GetStyleSetName(), "DisplayClusterConfigurator.OutputMapping.ViewScale"))
+			.Icon(FSlateIcon(FDisplayClusterConfiguratorStyle::Get().GetStyleSetName(), "DisplayClusterConfigurator.OutputMapping.ViewScale"))
 			.Label(this, &SDisplayClusterConfiguratorOutputMappingToolbar::GetViewScaleText)
 			.ParentToolBar(SharedThis(this))
 		);
@@ -96,33 +94,22 @@
 			.Style(ToolBarStyle)
 			.OnGetMenuContent(this, &SDisplayClusterConfiguratorOutputMappingToolbar::MakeTransformMenu)
 			.ToolTipText(LOCTEXT("TransformMenu_ToolTip", "Transform Operations"))
-			.Icon(FSlateIcon(FEditorStyle::GetStyleSetName(), "EditorViewport.TranslateRotate2DMode.Small"))
+			.Icon(FSlateIcon(FDisplayClusterConfiguratorStyle::Get().GetStyleSetName(), "DisplayClusterConfigurator.OutputMapping.Transform"))
 			.Label(LOCTEXT("TransformSettings_Label", "Transform"))
 			.ParentToolBar(SharedThis(this))
 		);
 	}
 	ToolbarBuilder.EndSection();
 
-	ToolbarBuilder.BeginSection("Placement");
+	ToolbarBuilder.BeginSection("Snapping");
 	{
 		ToolbarBuilder.AddWidget(
 			SNew(SViewportToolBarIconMenu)
 			.Cursor(EMouseCursor::Default)
 			.Style(ToolBarStyle)
-			.OnGetMenuContent(this, &SDisplayClusterConfiguratorOutputMappingToolbar::MakePositioningMenu)
-			.ToolTipText(LOCTEXT("PositioningMenu_ToolTip", "Positioning Settings"))
-			.Icon(FSlateIcon(FEditorStyle::GetStyleSetName(), "EditorViewport.TranslateRotate2DMode.Small"))
-			.Label(LOCTEXT("PositioningSettings_Label", "Positioning"))
-			.ParentToolBar(SharedThis(this))
-		);
-
-		ToolbarBuilder.AddWidget(
-			SNew(SViewportToolBarIconMenu)
-			.Cursor(EMouseCursor::Default)
-			.Style(ToolBarStyle)
 			.OnGetMenuContent(this, &SDisplayClusterConfiguratorOutputMappingToolbar::MakeSnappingMenu)
 			.ToolTipText(LOCTEXT("SnappingMenu_ToolTip", "Node Snapping"))
-			.Icon(FSlateIcon(FEditorStyle::GetStyleSetName(), "EditorViewport.Layer2DSnap"))
+			.Icon(FSlateIcon(FDisplayClusterConfiguratorStyle::Get().GetStyleSetName(), "DisplayClusterConfigurator.OutputMapping.Snapping"))
 			.Label(LOCTEXT("AlignmentSettings_Label", "Snapping"))
 			.ParentToolBar(SharedThis(this))
 		);
@@ -132,6 +119,43 @@
 	return ToolbarBuilder.MakeWidget();
 }
 
+TSharedRef<SWidget> SDisplayClusterConfiguratorOutputMappingToolbar::MakeWindowDisplayMenu()
+{
+	const bool bShouldCloseWindowAfterMenuSelection = false;
+	FMenuBuilder MenuBuilder(bShouldCloseWindowAfterMenuSelection, CommandList);
+
+	MenuBuilder.BeginSection(TEXT("WindowDisplaySection"));
+	{
+		MenuBuilder.AddMenuEntry(FDisplayClusterConfiguratorOutputMappingCommands::Get().ShowWindowInfo);
+		MenuBuilder.AddMenuEntry(FDisplayClusterConfiguratorOutputMappingCommands::Get().ShowWindowCorner);
+		MenuBuilder.AddMenuEntry(FDisplayClusterConfiguratorOutputMappingCommands::Get().ShowWindowNone);
+	}
+	MenuBuilder.EndSection();
+
+	return MenuBuilder.MakeWidget();
+}
+
+FText SDisplayClusterConfiguratorOutputMappingToolbar::GetWindowDisplayText() const
+{
+	TSharedPtr<FDisplayClusterConfiguratorViewOutputMapping> ViewOutputMapping = ViewOutputMappingPtr.Pin();
+	check(ViewOutputMapping != nullptr);
+
+	FOutputMappingSettings& OutputMappingSettings = ViewOutputMapping->GetOutputMappingSettings();
+
+	if (OutputMappingSettings.bShowWindowInfo)
+	{
+		return LOCTEXT("WindowDisplay_Info", "Info Bar");
+	}
+	else if (OutputMappingSettings.bShowWindowCornerImage)
+	{
+		return LOCTEXT("WindowDisplay_Corner", "Corner");
+	}
+	else
+	{
+		return LOCTEXT("WindowDisplay_None", "None");
+	}
+}
+
 TSharedRef<SWidget> SDisplayClusterConfiguratorOutputMappingToolbar::MakeTransformMenu()
 {
 	const bool bShouldCloseWindowAfterMenuSelection = false;
@@ -150,25 +174,3 @@
 	}
 	MenuBuilder.EndSection();
 
-	return MenuBuilder.MakeWidget();
-}
-
-TSharedRef<SWidget> SDisplayClusterConfiguratorOutputMappingToolbar::MakePositioningMenu()
-{
-	const bool bShouldCloseWindowAfterMenuSelection = false;
-	FMenuBuilder MenuBuilder(bShouldCloseWindowAfterMenuSelection, CommandList);
-
-	MenuBuilder.BeginSection(TEXT("OverlapBoundsSection"), LOCTEXT("OverlapBoundsSectionLabel", "Overlap & Bounds"));
-	{
-		MenuBuilder.AddMenuEntry(FDisplayClusterConfiguratorOutputMappingCommands::Get().ToggleClusterItemOverlap);
-		MenuBuilder.AddMenuEntry(FDisplayClusterConfiguratorOutputMappingCommands::Get().ToggleLockClusterNodesInHosts);
-	}
-	MenuBuilder.EndSection();
-
-	MenuBuilder.BeginSection(TEXT("LockingSection"), LOCTEXT("LockingSectionLabel", "Locking"));
-	{
-		MenuBuilder.AddMenuEntry(FDisplayClusterConfiguratorOutputMappingCommands::Get().ToggleLockClusterNodes);
-		MenuBuilder.AddMenuEntry(FDisplayClusterConfiguratorOutputMappingCommands::Get().ToggleLockViewports);
-	}
-	MenuBuilder.EndSection();
-
@@ -310,6 +312,13 @@
 	}
 	MenuBuilder.EndSection();
 
+	MenuBuilder.BeginSection(TEXT("OverlapBoundsSection"), LOCTEXT("OverlapBoundsSectionLabel", "Overlap & Bounds"));
+	{
+		MenuBuilder.AddMenuEntry(FDisplayClusterConfiguratorOutputMappingCommands::Get().ToggleClusterItemOverlap);
+		MenuBuilder.AddMenuEntry(FDisplayClusterConfiguratorOutputMappingCommands::Get().ToggleLockClusterNodesInHosts);
+	}
+	MenuBuilder.EndSection();
+
 	MenuBuilder.BeginSection(TEXT("Hosts"), LOCTEXT("HostsSectionLabel", "Hosts"));
 	{
 		MenuBuilder.AddSubMenu(
diff --git a/Engine/Plugins/Runtime/nDisplay/Source/DisplayClusterConfigurator/Private/Views/OutputMapping/Widgets/SDisplayClusterConfiguratorOutputMappingToolbar.h b/Engine/Plugins/Runtime/nDisplay/Source/DisplayClusterConfigurator/Private/Views/OutputMapping/Widgets/SDisplayClusterConfiguratorOutputMappingToolbar.h
--- a/Engine/Plugins/Runtime/nDisplay/Source/DisplayClusterConfigurator/Private/Views/OutputMapping/Widgets/SDisplayClusterConfiguratorOutputMappingToolbar.h
+++ b/Engine/Plugins/Runtime/nDisplay/Source/DisplayClusterConfigurator/Private/Views/OutputMapping/Widgets/SDisplayClusterConfiguratorOutputMappingToolbar.h
@@ -32,9 +32,10 @@ public:
 	TSharedRef< SWidget > MakeToolBar(const TSharedPtr< FExtender > InExtenders);
 
 private:
-	TSharedRef<SWidget> MakeTransformMenu();
+	TSharedRef<SWidget> MakeWindowDisplayMenu();
+	FText GetWindowDisplayText() const;
 
-	TSharedRef<SWidget> MakePositioningMenu();
+	TSharedRef<SWidget> MakeTransformMenu();
 
 	TSharedRef<SWidget> MakeSnappingMenu();
 
diff --git a/Engine/Plugins/Runtime/nDisplay/Source/DisplayClusterConfigurator/Private/Views/OutputMapping/Widgets/SDisplayClusterConfiguratorResizer.cpp b/Engine/Plugins/Runtime/nDisplay/Source/DisplayClusterConfigurator/Private/Views/OutputMapping/Widgets/SDisplayClusterConfiguratorResizer.cpp
--- a/Engine/Plugins/Runtime/nDisplay/Source/DisplayClusterConfigurator/Private/Views/OutputMapping/Widgets/SDisplayClusterConfiguratorResizer.cpp
+++ b/Engine/Plugins/Runtime/nDisplay/Source/DisplayClusterConfigurator/Private/Views/OutputMapping/Widgets/SDisplayClusterConfiguratorResizer.cpp
@@ -31,7 +31,7 @@ void SDisplayClusterConfiguratorResizer::Construct(const FArguments& InArgs, con
 	ChildSlot
 		[
 			SNew(SImage)
-			.Image(FDisplayClusterConfiguratorStyle::GetBrush("DisplayClusterConfigurator.OutputMapping.ResizeAreaHandle"))
+			.Image(FDisplayClusterConfiguratorStyle::Get().GetBrush("DisplayClusterConfigurator.OutputMapping.ResizeAreaHandle"))
 		];
 }
 
diff --git a/Engine/Plugins/Runtime/nDisplay/Source/DisplayClusterConfigurator/Private/Views/SCSEditor/SDisplayClusterConfiguratorComponentCombo.cpp b/Engine/Plugins/Runtime/nDisplay/Source/DisplayClusterConfigurator/Private/Views/SCSEditor/SDisplayClusterConfiguratorComponentCombo.cpp
--- a/Engine/Plugins/Runtime/nDisplay/Source/DisplayClusterConfigurator/Private/Views/SCSEditor/SDisplayClusterConfiguratorComponentCombo.cpp
+++ b/Engine/Plugins/Runtime/nDisplay/Source/DisplayClusterConfigurator/Private/Views/SCSEditor/SDisplayClusterConfiguratorComponentCombo.cpp
@@ -56,14 +56,14 @@ void SDisplayClusterConfiguratorComponentClassCombo::Construct(const FArguments&
 	[
 		SAssignNew(AddNewButton, SPositiveActionButton)
 		.Icon(FAppStyle::Get().GetBrush("Icons.Plus"))
-		.Text(LOCTEXT("AddComponent", "Add Component"))
+		.Text(LOCTEXT("AddComponent", "Add"))
 		.OnComboBoxOpened(this, &SDisplayClusterConfiguratorComponentClassCombo::ClearSelection)
 		.MenuContent()
 		[
 			SNew(SListViewSelectorDropdownMenu<FComponentClassComboEntryPtr>, SearchBox, ComponentClassListView)
 			[
 				SNew(SBorder)
-				.BorderImage(FEditorStyle::GetBrush("Menu.Background"))
+				.BorderImage(FAppStyle::Get().GetBrush("Menu.Background"))
 				.Padding(2)
 				[
 					SNew(SBox)
@@ -291,7 +291,7 @@ TSharedRef<ITableRow> SDisplayClusterConfiguratorComponentClassCombo::GenerateAd
 	{
 		return 
 			SNew( STableRow< TSharedPtr<FString> >, OwnerTable )
-				.Style(&FEditorStyle::Get().GetWidgetStyle<FTableRowStyle>("TableView.NoHoverTableRow"))
+				.Style(&FAppStyle::Get().GetWidgetStyle<FTableRowStyle>("TableView.NoHoverTableRow"))
 				.ShowSelection(false)
 			[
 				SNew(SBox)
@@ -299,7 +299,7 @@ TSharedRef<ITableRow> SDisplayClusterConfiguratorComponentClassCombo::GenerateAd
 				[
 					SNew(STextBlock)
 					.Text(FText::FromString(Entry->GetHeadingText()))
-					.TextStyle(FEditorStyle::Get(), TEXT("Menu.Heading"))
+					.TextStyle(FAppStyle::Get(), TEXT("Menu.Heading"))
 				]
 			];
 	}
@@ -307,11 +307,11 @@ TSharedRef<ITableRow> SDisplayClusterConfiguratorComponentClassCombo::GenerateAd
 	{
 		return 
 			SNew( STableRow< TSharedPtr<FString> >, OwnerTable )
-				.Style(&FEditorStyle::Get().GetWidgetStyle<FTableRowStyle>("TableView.NoHoverTableRow"))
+				.Style(&FAppStyle::Get().GetWidgetStyle<FTableRowStyle>("TableView.NoHoverTableRow"))
 				.ShowSelection(false)
 			[
 				SNew(SSeparator)
-				.SeparatorImage(FEditorStyle::Get().GetBrush("Menu.Separator"))
+				.SeparatorImage(FAppStyle::Get().GetBrush("Menu.Separator"))
 				.Thickness(1.0f)
 			];
 	}
diff --git a/Engine/Plugins/Runtime/nDisplay/Source/DisplayClusterConfigurator/Private/Views/TreeViews/Cluster/DisplayClusterConfiguratorViewCluster.cpp b/Engine/Plugins/Runtime/nDisplay/Source/DisplayClusterConfigurator/Private/Views/TreeViews/Cluster/DisplayClusterConfiguratorViewCluster.cpp
--- a/Engine/Plugins/Runtime/nDisplay/Source/DisplayClusterConfigurator/Private/Views/TreeViews/Cluster/DisplayClusterConfiguratorViewCluster.cpp
+++ b/Engine/Plugins/Runtime/nDisplay/Source/DisplayClusterConfigurator/Private/Views/TreeViews/Cluster/DisplayClusterConfiguratorViewCluster.cpp
@@ -19,6 +19,7 @@
 #include "UnrealExporter.h"
 #include "Exporters/Exporter.h"
 #include "ScopedTransaction.h"
+#include "Widgets/Images/SImage.h"
 
 #define LOCTEXT_NAMESPACE "FDisplayClusterConfiguratorViewCluster"
 
@@ -47,18 +48,44 @@ TSharedRef<SWidget> FDisplayClusterConfiguratorViewCluster::CreateWidget()
 
 void FDisplayClusterConfiguratorViewCluster::ConstructColumns(TArray<SHeaderRow::FColumn::FArguments>& OutColumnArgs) const
 {
+	// For the cluster view, use the group column to indicate which host the cluster nodes belong to with a colored band
+	OutColumnArgs.Add(SHeaderRow::Column(FDisplayClusterConfiguratorViewCluster::Columns::Host)
+		.DefaultLabel(FText::GetEmpty())
+		.DefaultTooltip(LOCTEXT("HostColumn_ToolTip", "Host"))
+		.FixedWidth(12)
+	);
+
 	OutColumnArgs.Add(SHeaderRow::Column(FDisplayClusterConfiguratorViewCluster::Columns::Visible)
 		.DefaultLabel(LOCTEXT("VisibleColumn_Label", "Visible"))
-		.FixedWidth(18));
+		.FixedWidth(24)
+		.HAlignHeader(HAlign_Left)
+		.VAlignHeader(VAlign_Center)
+		.HAlignCell(HAlign_Center)
+		.VAlignCell(VAlign_Center)
+		.DefaultTooltip(LOCTEXT("VisibleColumn_ToolTip", "Visible"))
+		.HeaderContentPadding(FMargin(4.0f, 0.0f, 0.0f, 0.0f))
+		[
+			SNew(SImage)
+			.ColorAndOpacity(FSlateColor::UseForeground())
+			.Image(FAppStyle::Get().GetBrush("Level.VisibleIcon16x"))
+		]
+	);
 
 	OutColumnArgs.Add(SHeaderRow::Column(FDisplayClusterConfiguratorViewCluster::Columns::Enabled)
-		.DefaultLabel(LOCTEXT("EnabledColumn_Label", "Enabled"))
-		.FixedWidth(18));
-
-	// For the cluster view, use the group column to indicate which host the cluster nodes belong to with a colored band
-	OutColumnArgs.Add(SHeaderRow::Column(FDisplayClusterConfiguratorViewCluster::Columns::Host)
-		.DefaultLabel(LOCTEXT("HiostColumn_Label", "Host"))
-		.FixedWidth(8));
+		.DefaultLabel(LOCTEXT("EnabledColumn_Label", "Locked"))
+		.FixedWidth(24)
+		.HAlignHeader(HAlign_Left)
+		.VAlignHeader(VAlign_Center)
+		.HAlignCell(HAlign_Center)
+		.VAlignCell(VAlign_Center)
+		.DefaultTooltip(LOCTEXT("EnabledColumn_ToolTip", "Locked"))
+		.HeaderContentPadding(FMargin(4.0f, 0.0f, 0.0f, 0.0f))
+		[
+			SNew(SImage)
+			.ColorAndOpacity(FSlateColor::UseForeground())
+			.Image(FAppStyle::Get().GetBrush("PropertyWindow.Locked"))
+		]
+	);
 
 	FDisplayClusterConfiguratorViewTree::ConstructColumns(OutColumnArgs);
 }
diff --git a/Engine/Plugins/Runtime/nDisplay/Source/DisplayClusterConfigurator/Private/Views/TreeViews/Cluster/TreeItems/DisplayClusterConfiguratorTreeItemCluster.cpp b/Engine/Plugins/Runtime/nDisplay/Source/DisplayClusterConfigurator/Private/Views/TreeViews/Cluster/TreeItems/DisplayClusterConfiguratorTreeItemCluster.cpp
--- a/Engine/Plugins/Runtime/nDisplay/Source/DisplayClusterConfigurator/Private/Views/TreeViews/Cluster/TreeItems/DisplayClusterConfiguratorTreeItemCluster.cpp
+++ b/Engine/Plugins/Runtime/nDisplay/Source/DisplayClusterConfigurator/Private/Views/TreeViews/Cluster/TreeItems/DisplayClusterConfiguratorTreeItemCluster.cpp
@@ -3,8 +3,84 @@
 #include "Views/TreeViews/Cluster/TreeItems/DisplayClusterConfiguratorTreeItemCluster.h"
 
 #include "DisplayClusterConfiguratorBlueprintEditor.h"
+#include "DisplayClusterConfiguratorStyle.h"
 #include "Interfaces/Views/TreeViews/IDisplayClusterConfiguratorViewTree.h"
 #include "Interfaces/Views/OutputMapping/IDisplayClusterConfiguratorViewOutputMapping.h"
+#include "Views/TreeViews/Cluster/DisplayClusterConfiguratorViewCluster.h"
+
+#include "Widgets/Images/SImage.h"
+#include "Widgets/Input/SButton.h"
+
+#define LOCTEXT_NAMESPACE "FDisplayClusterConfiguratorTreeItemCluster"
+
+
+class SToggleStateButton : public SCompoundWidget
+{
+public:
+	SLATE_BEGIN_ARGS(SToggleStateButton) 
+		: _ToggleState(false)
+	{}
+		SLATE_ATTRIBUTE(const FSlateBrush*, ToggleOnImage)
+		SLATE_ATTRIBUTE(const FSlateBrush*, ToggleOffImage)
+		SLATE_ATTRIBUTE(bool, ToggleState)
+		SLATE_EVENT(FSimpleDelegate, OnToggled)
+	SLATE_END_ARGS()
+
+	void Construct(const FArguments& InArgs)
+	{
+		ToggleOnImageAttribute = InArgs._ToggleOnImage;
+		ToggleOffImageAttribute = InArgs._ToggleOffImage;
+		ToggleStateAttribute = InArgs._ToggleState;
+
+		OnToggled = InArgs._OnToggled;
+
+		ChildSlot
+		[
+			SNew(SButton)
+			.ContentPadding(0)
+			.ButtonStyle(FAppStyle::Get(), "NoBorder")
+			.ToolTipText(InArgs._ToolTipText)
+			.HAlign(HAlign_Center)
+			.VAlign(VAlign_Center)
+			.OnClicked(this, &SToggleStateButton::OnButtonClicked)
+			.Content()
+			[
+				SNew(SImage)
+				.Image(this, &SToggleStateButton::GetImageBrush)
+				.ColorAndOpacity(this, &SToggleStateButton::GetImageForegroundColor)
+			]
+		];
+	}
+
+public:
+	bool GetToggleState() const { return ToggleStateAttribute.Get(); }
+	void SetToggleState(TAttribute<bool> NewToggleState) { ToggleStateAttribute = NewToggleState; }
+
+private:
+	FReply OnButtonClicked()
+	{
+		OnToggled.ExecuteIfBound();
+
+		return FReply::Handled();
+	}
+
+	const FSlateBrush* GetImageBrush() const	
+	{
+		return ToggleStateAttribute.Get(false) ? ToggleOnImageAttribute.Get() : ToggleOffImageAttribute.Get();
+	}
+
+	FSlateColor GetImageForegroundColor() const	
+	{
+		return IsHovered() ? FSlateColor::UseForeground() : FSlateColor::UseSubduedForeground();
+	}
+
+private:
+	TAttribute<const FSlateBrush*> ToggleOnImageAttribute;
+	TAttribute<const FSlateBrush*> ToggleOffImageAttribute;
+	TAttribute<bool> ToggleStateAttribute;
+
+	FSimpleDelegate OnToggled;
+};
 
 
 FDisplayClusterConfiguratorTreeItemCluster::FDisplayClusterConfiguratorTreeItemCluster(const FName& InName,
@@ -18,6 +94,45 @@ FDisplayClusterConfiguratorTreeItemCluster::FDisplayClusterConfiguratorTreeItemC
 	, IconStyle(InIconStyle)
 {}
 
+TSharedRef<SWidget> FDisplayClusterConfiguratorTreeItemCluster::GenerateWidgetForColumn(const FName& ColumnName, TSharedPtr<ITableRow> TableRow, const TAttribute<FText>& FilterText, FIsSelected InIsSelected)
+{
+	if (ColumnName == FDisplayClusterConfiguratorViewCluster::Columns::Host && IsClusterItemGrouped())
+	{
+		return SNew(SBox)
+			.HAlign(HAlign_Center)
+			[
+				SNew(SImage)
+				.ColorAndOpacity(this, &FDisplayClusterConfiguratorTreeItemCluster::GetClusterItemGroupColor)
+				.OnMouseButtonDown(this, &FDisplayClusterConfiguratorTreeItemCluster::OnClusterItemGroupClicked)
+				.Image(FAppStyle::Get().GetBrush("WhiteBrush"))
+				.DesiredSizeOverride(FVector2D(6, 20))
+				.Cursor(EMouseCursor::Hand)
+			];
+	}
+	else if (ColumnName == FDisplayClusterConfiguratorViewCluster::Columns::Visible && CanClusterItemBeHidden())
+	{
+		return SNew(SToggleStateButton)
+			.ToggleOnImage(FAppStyle::Get().GetBrush(TEXT("Level.VisibleIcon16x")))
+			.ToggleOffImage(FAppStyle::Get().GetBrush(TEXT("Level.NotVisibleIcon16x")))
+			.ToggleState(this, &FDisplayClusterConfiguratorTreeItemCluster::IsClusterItemVisible)
+			.Visibility(this, &FDisplayClusterConfiguratorTreeItemCluster::GetVisibleButtonVisibility)
+			.ToolTipText(LOCTEXT("VisibilityButton_Tooltip", "Hides or shows this cluster item and its children in the Output Mapping editor"))
+			.OnToggled(this, &FDisplayClusterConfiguratorTreeItemCluster::ToggleClusterItemVisibility);
+	}
+	else if (ColumnName == FDisplayClusterConfiguratorViewCluster::Columns::Enabled && CanClusterItemBeLocked())
+	{
+		return SNew(SToggleStateButton)
+			.ToggleOnImage(FAppStyle::Get().GetBrush(TEXT("PropertyWindow.Unlocked")))
+			.ToggleOffImage(FAppStyle::Get().GetBrush(TEXT("PropertyWindow.Locked")))
+			.ToggleState(this, &FDisplayClusterConfiguratorTreeItemCluster::IsClusterItemUnlocked)
+			.Visibility(this, &FDisplayClusterConfiguratorTreeItemCluster::GetLockButtonVisibility)
+			.ToolTipText(LOCTEXT("LockButton_Tooltip", "Locks or unlocks this cluster item and its children in the Output Mapping editor"))
+			.OnToggled(this, &FDisplayClusterConfiguratorTreeItemCluster::ToggleClusterItemLock);
+	}
+
+	return FDisplayClusterConfiguratorTreeItem::GenerateWidgetForColumn(ColumnName, TableRow, FilterText, InIsSelected);
+}
+
 void FDisplayClusterConfiguratorTreeItemCluster::OnItemDoubleClicked()
 {
 	if (ToolkitPtr.IsValid())
@@ -46,3 +161,35 @@ bool FDisplayClusterConfiguratorTreeItemCluster::IsHovered() const
 
 	return false;
 }
+
+EVisibility FDisplayClusterConfiguratorTreeItemCluster::GetVisibleButtonVisibility() const
+{
+	bool bIsHovered = IsHovered();
+	bool bIsSelected = ViewTreePtr.Pin()->GetSelectedItems().Contains(SharedThis(this));
+
+	if (bIsHovered || bIsSelected || !IsClusterItemVisible())
+	{
+		return EVisibility::Visible;
+	}
+	else
+	{
+		return EVisibility::Hidden;
+	}
+}
+
+EVisibility FDisplayClusterConfiguratorTreeItemCluster::GetLockButtonVisibility() const
+{
+	bool bIsHovered = IsHovered();
+	bool bIsSelected = ViewTreePtr.Pin()->GetSelectedItems().Contains(SharedThis(this));
+
+	if (bIsHovered || bIsSelected || !IsClusterItemUnlocked())
+	{
+		return EVisibility::Visible;
+	}
+	else
+	{
+		return EVisibility::Hidden;
+	}
+}
+
+#undef LOCTEXT_NAMESPACE
\ No newline at end of file
diff --git a/Engine/Plugins/Runtime/nDisplay/Source/DisplayClusterConfigurator/Private/Views/TreeViews/Cluster/TreeItems/DisplayClusterConfiguratorTreeItemCluster.h b/Engine/Plugins/Runtime/nDisplay/Source/DisplayClusterConfigurator/Private/Views/TreeViews/Cluster/TreeItems/DisplayClusterConfiguratorTreeItemCluster.h
--- a/Engine/Plugins/Runtime/nDisplay/Source/DisplayClusterConfigurator/Private/Views/TreeViews/Cluster/TreeItems/DisplayClusterConfiguratorTreeItemCluster.h
+++ b/Engine/Plugins/Runtime/nDisplay/Source/DisplayClusterConfigurator/Private/Views/TreeViews/Cluster/TreeItems/DisplayClusterConfiguratorTreeItemCluster.h
@@ -3,6 +3,8 @@
 #pragma once
 
 #include "Views/TreeViews/DisplayClusterConfiguratorTreeItem.h"
+#include "Widgets/DeclarativeSyntaxSupport.h"
+#include "Widgets/SCompoundWidget.h"
 
 class IDisplayClusterConfiguratorViewTree;
 
@@ -20,6 +22,7 @@ public:
 		bool InbRoot = false);
 
 	//~ Begin IDisplayClusterConfiguratorTreeItem Interface
+	virtual TSharedRef<SWidget> GenerateWidgetForColumn(const FName& ColumnName, TSharedPtr<ITableRow> TableRow, const TAttribute<FText>& FilterText, FIsSelected InIsSelected) override;
 	virtual void OnItemDoubleClicked() override;
 	virtual void OnMouseEnter() override;
 	virtual void OnMouseLeave() override;
@@ -29,6 +32,23 @@ public:
 	virtual FString GetIconStyle() const override { return IconStyle; }
 	//~ End IDisplayClusterConfiguratorTreeItem Interface
 
+protected:
+	virtual bool CanClusterItemBeHidden() const { return false; }
+	virtual bool CanClusterItemBeLocked() const { return false; }
+	virtual bool IsClusterItemGrouped() const { return false; }
+
+	virtual bool IsClusterItemVisible() const { return true; }
+	virtual bool IsClusterItemUnlocked() const { return true; }
+
+	virtual void ToggleClusterItemVisibility() { }
+	virtual void ToggleClusterItemLock() { }
+
+	virtual FSlateColor GetClusterItemGroupColor() const { return FLinearColor::Transparent; }
+	virtual FReply OnClusterItemGroupClicked(const FGeometry& MyGeometry, const FPointerEvent& MouseEvent) { return FReply::Unhandled(); }
+
+	EVisibility GetVisibleButtonVisibility() const;
+	EVisibility GetLockButtonVisibility() const;
+
 protected:
 	FName Name;
 	FString IconStyle;
diff --git a/Engine/Plugins/Runtime/nDisplay/Source/DisplayClusterConfigurator/Private/Views/TreeViews/Cluster/TreeItems/DisplayClusterConfiguratorTreeItemClusterNode.cpp b/Engine/Plugins/Runtime/nDisplay/Source/DisplayClusterConfigurator/Private/Views/TreeViews/Cluster/TreeItems/DisplayClusterConfiguratorTreeItemClusterNode.cpp
--- a/Engine/Plugins/Runtime/nDisplay/Source/DisplayClusterConfigurator/Private/Views/TreeViews/Cluster/TreeItems/DisplayClusterConfiguratorTreeItemClusterNode.cpp
+++ b/Engine/Plugins/Runtime/nDisplay/Source/DisplayClusterConfigurator/Private/Views/TreeViews/Cluster/TreeItems/DisplayClusterConfiguratorTreeItemClusterNode.cpp
@@ -66,7 +66,7 @@ void FDisplayClusterConfiguratorTreeItemClusterNode::SetVisible(bool bIsVisible)
 	}
 }
 
-void FDisplayClusterConfiguratorTreeItemClusterNode::SetEnabled(bool bIsEnabled)
+void FDisplayClusterConfiguratorTreeItemClusterNode::SetUnlocked(bool bIsUnlocked)
 {
 	UDisplayClusterConfigurationClusterNode* ClusterNode = GetObjectChecked<UDisplayClusterConfigurationClusterNode>();
 
@@ -74,9 +74,9 @@ void FDisplayClusterConfiguratorTreeItemClusterNode::SetEnabled(bool bIsEnabled)
 	SaveToTransactionBuffer(ClusterNode, false);
 
 	const TSharedPtr<ISinglePropertyView> PropertyView = DisplayClusterConfiguratorPropertyUtils::GetPropertyView(
-		ClusterNode, GET_MEMBER_NAME_CHECKED(UDisplayClusterConfigurationClusterNode, bIsEnabled));
+		ClusterNode, GET_MEMBER_NAME_CHECKED(UDisplayClusterConfigurationClusterNode, bIsUnlocked));
 
-	PropertyView->GetPropertyHandle()->SetValue(bIsEnabled);
+	PropertyView->GetPropertyHandle()->SetValue(bIsUnlocked);
 
 	// Set the enabled state of this cluster node's children to match it.
 	for (TSharedPtr<IDisplayClusterConfiguratorTreeItem> Child : Children)
@@ -84,55 +84,11 @@ void FDisplayClusterConfiguratorTreeItemClusterNode::SetEnabled(bool bIsEnabled)
 		TSharedPtr<FDisplayClusterConfiguratorTreeItemViewport> ViewportTreeItem = StaticCastSharedPtr<FDisplayClusterConfiguratorTreeItemViewport>(Child);
 		if (ViewportTreeItem.IsValid())
 		{
-			ViewportTreeItem->SetEnabled(ClusterNode->bIsEnabled);
+			ViewportTreeItem->SetUnlocked(ClusterNode->bIsUnlocked);
 		}
 	}
 }
 
-TSharedRef<SWidget> FDisplayClusterConfiguratorTreeItemClusterNode::GenerateWidgetForColumn(const FName& ColumnName, TSharedPtr<ITableRow> TableRow, const TAttribute<FText>& FilterText, FIsSelected InIsSelected)
-{
-	if (ColumnName == FDisplayClusterConfiguratorViewCluster::Columns::Host)
-	{
-		return SNew(SImage)
-			.ColorAndOpacity(this, &FDisplayClusterConfiguratorTreeItemClusterNode::GetHostColor)
-			.OnMouseButtonDown(this, &FDisplayClusterConfiguratorTreeItemClusterNode::OnHostClicked)
-			.Image(FDisplayClusterConfiguratorStyle::GetBrush("DisplayClusterConfigurator.Node.Body"))
-			.Cursor(EMouseCursor::Hand);
-	}
-	else if (ColumnName == FDisplayClusterConfiguratorViewCluster::Columns::Visible)
-	{
-		return SAssignNew(VisibilityButton, SButton)
-			.ContentPadding(0)
-			.ButtonStyle(FEditorStyle::Get(), "NoBorder")
-			.ToolTipText(LOCTEXT("VisibilityButton_Tooltip", "Hides or shows this cluster node and its child viewports in the Output Mapping editor"))
-			.HAlign(HAlign_Center)
-			.VAlign(VAlign_Center)
-			.OnClicked(this, &FDisplayClusterConfiguratorTreeItemClusterNode::OnVisibilityButtonClicked)
-			.Content()
-			[
-				SNew(SImage)
-				.Image(this, &FDisplayClusterConfiguratorTreeItemClusterNode::GetVisibilityButtonBrush)
-			];
-	}
-	else if (ColumnName == FDisplayClusterConfiguratorViewCluster::Columns::Enabled)
-	{
-		return SAssignNew(EnabledButton, SButton)
-			.ContentPadding(0)
-			.ButtonStyle(FEditorStyle::Get(), "NoBorder")
-			.ToolTipText(LOCTEXT("EnabledButton_Tooltip", "Enables or disables this cluster node and its child viewports in the Output Mapping editor"))
-			.HAlign(HAlign_Center)
-			.VAlign(VAlign_Center)
-			.OnClicked(this, &FDisplayClusterConfiguratorTreeItemClusterNode::OnEnabledButtonClicked)
-			.Content()
-			[
-				SNew(SImage)
-				.Image(this, &FDisplayClusterConfiguratorTreeItemClusterNode::GetEnabledButtonBrush)
-			];
-	}
-
-	return FDisplayClusterConfiguratorTreeItemCluster::GenerateWidgetForColumn(ColumnName, TableRow, FilterText, InIsSelected);
-}
-
 void FDisplayClusterConfiguratorTreeItemClusterNode::DeleteItem() const
 {
 	UDisplayClusterConfigurationClusterNode* ClusterNode = GetObjectChecked<UDisplayClusterConfigurationClusterNode>();
@@ -348,35 +304,6 @@ EVisibility FDisplayClusterConfiguratorTreeItemClusterNode::GetPrimaryLabelVisib
 	return bIsPrimary ? EVisibility::Visible : EVisibility::Collapsed;
 }
 
-FSlateColor FDisplayClusterConfiguratorTreeItemClusterNode::GetHostColor() const
-{
-	if (!HostObject.IsValid())
-	{
-		return FLinearColor::Transparent;
-	}
-
-	if (ToolkitPtr.Pin()->IsObjectSelected(HostObject.Get()))
-	{
-		return FDisplayClusterConfiguratorStyle::GetColor("DisplayClusterConfigurator.Node.Color.Selected");
-	}
-	else
-	{
-		return HostObject.Get()->Color;
-	}
-}
-
-FReply FDisplayClusterConfiguratorTreeItemClusterNode::OnHostClicked(const FGeometry& MyGeometry, const FPointerEvent& MouseEvent)
-{
-	if (!HostObject.IsValid())
-	{
-		return FReply::Unhandled();
-	}
-
-	SelectHostDisplayData();
-
-	return FReply::Handled();
-}
-
 bool FDisplayClusterConfiguratorTreeItemClusterNode::CanDropViewports(TSharedPtr<FDisplayClusterConfiguratorViewportDragDropOp> ViewportDragDropOp, FText& OutErrorMessage) const
 {
 	if (ViewportDragDropOp.IsValid())
@@ -406,23 +333,15 @@ bool FDisplayClusterConfiguratorTreeItemClusterNode::CanDropViewports(TSharedPtr
 	return false;
 }
 
-const FSlateBrush* FDisplayClusterConfiguratorTreeItemClusterNode::GetVisibilityButtonBrush() const
+bool FDisplayClusterConfiguratorTreeItemClusterNode::IsClusterItemVisible() const
 {
 	UDisplayClusterConfigurationClusterNode* ClusterNode = GetObjectChecked<UDisplayClusterConfigurationClusterNode>();
-
-	if (ClusterNode->bIsVisible)
-	{
-		return VisibilityButton->IsHovered() ? FEditorStyle::GetBrush(TEXT("Level.VisibleHighlightIcon16x")) : FEditorStyle::GetBrush(TEXT("Level.VisibleIcon16x"));
-	}
-	else
-	{
-		return VisibilityButton->IsHovered() ? FEditorStyle::GetBrush(TEXT("Level.NotVisibleHighlightIcon16x")) : FEditorStyle::GetBrush(TEXT("Level.NotVisibleIcon16x"));
-	}
+	return ClusterNode->bIsVisible;
 }
 
-FReply FDisplayClusterConfiguratorTreeItemClusterNode::OnVisibilityButtonClicked()
+void FDisplayClusterConfiguratorTreeItemClusterNode::ToggleClusterItemVisibility()
 {
-	const FScopedTransaction Transaction(LOCTEXT("ToggleClusterNodeVisibility", "Set Cluster Node Visibility"));
+	const FScopedTransaction Transaction(LOCTEXT("ToggleClusterNodeVisibility", "Toggle Cluster Node Visibility"));
 
 	UDisplayClusterConfigurationClusterNode* ClusterNode = GetObjectChecked<UDisplayClusterConfigurationClusterNode>();
 
@@ -444,27 +363,17 @@ FReply FDisplayClusterConfiguratorTreeItemClusterNode::OnVisibilityButtonClicked
 			ViewportTreeItem->SetVisible(ClusterNode->bIsVisible);
 		}
 	}
-
-	return FReply::Handled();
 }
 
-const FSlateBrush* FDisplayClusterConfiguratorTreeItemClusterNode::GetEnabledButtonBrush() const
+bool FDisplayClusterConfiguratorTreeItemClusterNode::IsClusterItemUnlocked() const
 {
 	UDisplayClusterConfigurationClusterNode* ClusterNode = GetObjectChecked<UDisplayClusterConfigurationClusterNode>();
-
-	if (ClusterNode->bIsEnabled)
-	{
-		return EnabledButton->IsHovered() ? FEditorStyle::GetBrush(TEXT("Level.UnlockedHighlightIcon16x")) : FEditorStyle::GetBrush(TEXT("Level.UnlockedIcon16x"));
-	}
-	else
-	{
-		return EnabledButton->IsHovered() ? FEditorStyle::GetBrush(TEXT("Level.LockedHighlightIcon16x")) : FEditorStyle::GetBrush(TEXT("Level.LockedIcon16x"));
-	}
+	return ClusterNode->bIsUnlocked;
 }
 
-FReply FDisplayClusterConfiguratorTreeItemClusterNode::OnEnabledButtonClicked()
+void FDisplayClusterConfiguratorTreeItemClusterNode::ToggleClusterItemLock()
 {
-	const FScopedTransaction Transaction(LOCTEXT("ToggleClusterNodeEnabled", "Set Cluster Node Enabled"));
+	const FScopedTransaction Transaction(LOCTEXT("ToggleClusterNodeLock", "Toggle Cluster Node Lock"));
 
 	UDisplayClusterConfigurationClusterNode* ClusterNode = GetObjectChecked<UDisplayClusterConfigurationClusterNode>();
 
@@ -472,9 +381,9 @@ FReply FDisplayClusterConfiguratorTreeItemClusterNode::OnEnabledButtonClicked()
 	SaveToTransactionBuffer(ClusterNode, false);
 
 	const TSharedPtr<ISinglePropertyView> PropertyView = DisplayClusterConfiguratorPropertyUtils::GetPropertyView(
-		ClusterNode, GET_MEMBER_NAME_CHECKED(UDisplayClusterConfigurationClusterNode, bIsEnabled));
+		ClusterNode, GET_MEMBER_NAME_CHECKED(UDisplayClusterConfigurationClusterNode, bIsUnlocked));
 
-	PropertyView->GetPropertyHandle()->SetValue(!ClusterNode->bIsEnabled);
+	PropertyView->GetPropertyHandle()->SetValue(!ClusterNode->bIsUnlocked);
 
 	// Set the enabled state of this cluster node's children to match it.
 	for (TSharedPtr<IDisplayClusterConfiguratorTreeItem> Child : Children)
@@ -483,9 +392,36 @@ FReply FDisplayClusterConfiguratorTreeItemClusterNode::OnEnabledButtonClicked()
 
 		if (ViewportTreeItem.IsValid())
 		{
-			ViewportTreeItem->SetEnabled(ClusterNode->bIsEnabled);
+			ViewportTreeItem->SetUnlocked(ClusterNode->bIsUnlocked);
 		}
 	}
+}
+
+FSlateColor FDisplayClusterConfiguratorTreeItemClusterNode::GetClusterItemGroupColor() const
+{
+	if (!HostObject.IsValid())
+	{
+		return FLinearColor::Transparent;
+	}
+
+	if (ToolkitPtr.Pin()->IsObjectSelected(HostObject.Get()))
+	{
+		return FDisplayClusterConfiguratorStyle::Get().GetColor("DisplayClusterConfigurator.Node.Color.Selected");
+	}
+	else
+	{
+		return HostObject.Get()->Color;
+	}
+}
+
+FReply FDisplayClusterConfiguratorTreeItemClusterNode::OnClusterItemGroupClicked(const FGeometry& MyGeometry, const FPointerEvent& MouseEvent)
+{
+	if (!HostObject.IsValid())
+	{
+		return FReply::Unhandled();
+	}
+
+	SelectHostDisplayData();
 
 	return FReply::Handled();
 }
diff --git a/Engine/Plugins/Runtime/nDisplay/Source/DisplayClusterConfigurator/Private/Views/TreeViews/Cluster/TreeItems/DisplayClusterConfiguratorTreeItemClusterNode.h b/Engine/Plugins/Runtime/nDisplay/Source/DisplayClusterConfigurator/Private/Views/TreeViews/Cluster/TreeItems/DisplayClusterConfiguratorTreeItemClusterNode.h
--- a/Engine/Plugins/Runtime/nDisplay/Source/DisplayClusterConfigurator/Private/Views/TreeViews/Cluster/TreeItems/DisplayClusterConfiguratorTreeItemClusterNode.h
+++ b/Engine/Plugins/Runtime/nDisplay/Source/DisplayClusterConfigurator/Private/Views/TreeViews/Cluster/TreeItems/DisplayClusterConfiguratorTreeItemClusterNode.h
@@ -25,10 +25,9 @@ public:
 	void SelectHostDisplayData();
 
 	void SetVisible(bool bIsVisible);
-	void SetEnabled(bool bIsEnabled);
+	void SetUnlocked(bool bIsUnlocked);
 
 	//~ Begin FDisplayClusterConfiguratorTreeItem Interface
-	virtual TSharedRef<SWidget> GenerateWidgetForColumn(const FName& ColumnName, TSharedPtr<ITableRow> TableRow, const TAttribute<FText>& FilterText, FIsSelected InIsSelected) override;
 	virtual void DeleteItem() const override;
 	virtual bool CanHideItem() const override { return true; }
 	virtual void SetItemHidden(bool bIsHidden) { SetVisible(!bIsHidden); }
@@ -42,23 +41,26 @@ public:
 protected:
 	virtual void FillItemColumn(TSharedPtr<SHorizontalBox> Box, const TAttribute<FText>& FilterText, FIsSelected InIsSelected) override;
 	virtual void OnDisplayNameCommitted(const FText& NewText, ETextCommit::Type CommitInfo) override;
+
+	virtual bool CanClusterItemBeHidden() const override { return true; }
+	virtual bool CanClusterItemBeLocked() const override { return true; }
+	virtual bool IsClusterItemGrouped() const override { return true; }
+
+	virtual bool IsClusterItemVisible() const override;
+	virtual bool IsClusterItemUnlocked() const override;
+
+	virtual void ToggleClusterItemVisibility() override;
+	virtual void ToggleClusterItemLock() override;
+
+	virtual FSlateColor GetClusterItemGroupColor() const override;
+	virtual FReply OnClusterItemGroupClicked(const FGeometry& MyGeometry, const FPointerEvent& MouseEvent) override;
+
 	//~ End FDisplayClusterConfiguratorTreeItem Interface
 
 private:
 	EVisibility GetPrimaryLabelVisibility() const;
-	FSlateColor GetHostColor() const;
-	FReply OnHostClicked(const FGeometry& MyGeometry, const FPointerEvent& MouseEvent);
-
 	bool CanDropViewports(TSharedPtr<FDisplayClusterConfiguratorViewportDragDropOp> ViewportDragDropOp, FText& OutErrorMessage) const;
 
-	const FSlateBrush* GetVisibilityButtonBrush() const;
-	FReply OnVisibilityButtonClicked();
-
-	const FSlateBrush* GetEnabledButtonBrush() const;
-	FReply OnEnabledButtonClicked();
-
 private:
 	TWeakObjectPtr<UDisplayClusterConfigurationHostDisplayData> HostObject;
-	TSharedPtr<SButton> VisibilityButton;
-	TSharedPtr<SButton> EnabledButton;
 };
\ No newline at end of file
diff --git a/Engine/Plugins/Runtime/nDisplay/Source/DisplayClusterConfigurator/Private/Views/TreeViews/Cluster/TreeItems/DisplayClusterConfiguratorTreeItemHost.cpp b/Engine/Plugins/Runtime/nDisplay/Source/DisplayClusterConfigurator/Private/Views/TreeViews/Cluster/TreeItems/DisplayClusterConfiguratorTreeItemHost.cpp
--- a/Engine/Plugins/Runtime/nDisplay/Source/DisplayClusterConfigurator/Private/Views/TreeViews/Cluster/TreeItems/DisplayClusterConfiguratorTreeItemHost.cpp
+++ b/Engine/Plugins/Runtime/nDisplay/Source/DisplayClusterConfigurator/Private/Views/TreeViews/Cluster/TreeItems/DisplayClusterConfiguratorTreeItemHost.cpp
@@ -30,50 +30,6 @@ void FDisplayClusterConfiguratorTreeItemHost::Initialize()
 	HostDisplayData->OnPostEditChangeChainProperty.Add(UDisplayClusterConfigurationHostDisplayData::FOnPostEditChangeChainProperty::FDelegate::CreateSP(this, &FDisplayClusterConfiguratorTreeItemHost::OnPostEditChangeChainProperty));
 }
 
-TSharedRef<SWidget> FDisplayClusterConfiguratorTreeItemHost::GenerateWidgetForColumn(const FName& ColumnName, TSharedPtr<ITableRow> TableRow, const TAttribute<FText>& FilterText, FIsSelected InIsSelected)
-{
-	if (ColumnName == FDisplayClusterConfiguratorViewCluster::Columns::Host)
-	{
-		return SNew(SImage)
-			.ColorAndOpacity(this, &FDisplayClusterConfiguratorTreeItemHost::GetHostColor)
-			.OnMouseButtonDown(this, &FDisplayClusterConfiguratorTreeItemHost::OnHostClicked)
-			.Image(FDisplayClusterConfiguratorStyle::GetBrush("DisplayClusterConfigurator.Node.Body"))
-			.Cursor(EMouseCursor::Hand);
-	}
-	else if (ColumnName == FDisplayClusterConfiguratorViewCluster::Columns::Visible)
-	{
-		return SAssignNew(VisibilityButton, SButton)
-			.ContentPadding(0)
-			.ButtonStyle(FEditorStyle::Get(), "NoBorder")
-			.ToolTipText(LOCTEXT("VisibilityButton_Tooltip", "Hides or shows this host and its children in the Output Mapping editor"))
-			.HAlign(HAlign_Center)
-			.VAlign(VAlign_Center)
-			.OnClicked(this, &FDisplayClusterConfiguratorTreeItemHost::OnVisibilityButtonClicked)
-			.Content()
-			[
-				SNew(SImage)
-				.Image(this, &FDisplayClusterConfiguratorTreeItemHost::GetVisibilityButtonBrush)
-			];
-	}
-	else if (ColumnName == FDisplayClusterConfiguratorViewCluster::Columns::Enabled)
-	{
-		return SAssignNew(EnabledButton, SButton)
-			.ContentPadding(0)
-			.ButtonStyle(FEditorStyle::Get(), "NoBorder")
-			.ToolTipText(LOCTEXT("EnabledButton_Tooltip", "Enables or disables this host and its children in the Output Mapping editor"))
-			.HAlign(HAlign_Center)
-			.VAlign(VAlign_Center)
-			.OnClicked(this, &FDisplayClusterConfiguratorTreeItemHost::OnEnabledButtonClicked)
-			.Content()
-			[
-				SNew(SImage)
-				.Image(this, &FDisplayClusterConfiguratorTreeItemHost::GetEnabledButtonBrush)
-			];
-	}
-
-	return FDisplayClusterConfiguratorTreeItemCluster::GenerateWidgetForColumn(ColumnName, TableRow, FilterText, InIsSelected);
-}
-
 void FDisplayClusterConfiguratorTreeItemHost::DeleteItem() const
 {
 	UDisplayClusterConfigurationHostDisplayData* Host = GetObjectChecked<UDisplayClusterConfigurationHostDisplayData>();
@@ -234,119 +190,99 @@ void FDisplayClusterConfiguratorTreeItemHost::OnDisplayNameCommitted(const FText
 	Name = *NewText.ToString();
 }
 
-FText FDisplayClusterConfiguratorTreeItemHost::GetHostAddress() const
+bool FDisplayClusterConfiguratorTreeItemHost::IsClusterItemVisible() const
 {
 	UDisplayClusterConfigurationHostDisplayData* Host = GetObjectChecked<UDisplayClusterConfigurationHostDisplayData>();
-	return FText::Format(LOCTEXT("HostAddressFormattedLabel", "({0})"), FText::FromString(FDisplayClusterConfiguratorClusterUtils::GetAddressForHost(Host)));
+	return Host->bIsVisible;
 }
 
-FSlateColor FDisplayClusterConfiguratorTreeItemHost::GetHostColor() const
+void FDisplayClusterConfiguratorTreeItemHost::ToggleClusterItemVisibility()
 {
-	UDisplayClusterConfigurationHostDisplayData* Host = GetObjectChecked<UDisplayClusterConfigurationHostDisplayData>();
-	if (ToolkitPtr.Pin()->IsObjectSelected(Host))
-	{
-		return FDisplayClusterConfiguratorStyle::GetColor("DisplayClusterConfigurator.Node.Color.Selected");
-	}
-	else
-	{
-		return Host->Color;
-	}
-}
+	const FScopedTransaction Transaction(LOCTEXT("ToggleHostVisibility", "Set Host Visibility"));
 
-FReply FDisplayClusterConfiguratorTreeItemHost::OnHostClicked(const FGeometry& MyGeometry, const FPointerEvent& MouseEvent)
-{
 	UDisplayClusterConfigurationHostDisplayData* Host = GetObjectChecked<UDisplayClusterConfigurationHostDisplayData>();
-	TArray<UObject*> Objects { Host };
-	ToolkitPtr.Pin()->SelectObjects(Objects);
 
-	return FReply::Handled();
-}
+	// Use SaveToTransactionBuffer to avoid marking the package as dirty
+	SaveToTransactionBuffer(Host, false);
+	Host->bIsVisible = !Host->bIsVisible;
 
-bool FDisplayClusterConfiguratorTreeItemHost::CanDropClusterNodes(TSharedPtr<FDisplayClusterConfiguratorClusterNodeDragDropOp> ClusterNodeDragDropOp, FText& OutErrorMessage) const
-{
-	if (ClusterNodeDragDropOp.IsValid())
+	// Set the visible state of this host's children to match it.
+	for (TSharedPtr<IDisplayClusterConfiguratorTreeItem> Child : Children)
 	{
-		UDisplayClusterConfigurationHostDisplayData* Host = GetObjectChecked<UDisplayClusterConfigurationHostDisplayData>();
-		bool bCanAcceptClusterNodes = true;
-
-		return bCanAcceptClusterNodes;
+		TSharedPtr<FDisplayClusterConfiguratorTreeItemClusterNode> ClusterNodeTreeItem = StaticCastSharedPtr<FDisplayClusterConfiguratorTreeItemClusterNode>(Child);
+		if (ClusterNodeTreeItem.IsValid())
+		{
+			ClusterNodeTreeItem->SetVisible(Host->bIsVisible);
+		}
 	}
-
-	return false;
 }
 
-const FSlateBrush* FDisplayClusterConfiguratorTreeItemHost::GetVisibilityButtonBrush() const
+bool FDisplayClusterConfiguratorTreeItemHost::IsClusterItemUnlocked() const
 {
 	UDisplayClusterConfigurationHostDisplayData* Host = GetObjectChecked<UDisplayClusterConfigurationHostDisplayData>();
-
-	if (Host->bIsVisible)
-	{
-		return VisibilityButton->IsHovered() ? FEditorStyle::GetBrush(TEXT("Level.VisibleHighlightIcon16x")) : FEditorStyle::GetBrush(TEXT("Level.VisibleIcon16x"));
-	}
-	else
-	{
-		return VisibilityButton->IsHovered() ? FEditorStyle::GetBrush(TEXT("Level.NotVisibleHighlightIcon16x")) : FEditorStyle::GetBrush(TEXT("Level.NotVisibleIcon16x"));
-	}
+	return Host->bIsUnlocked;
 }
 
-FReply FDisplayClusterConfiguratorTreeItemHost::OnVisibilityButtonClicked()
+void FDisplayClusterConfiguratorTreeItemHost::ToggleClusterItemLock()
 {
-	const FScopedTransaction Transaction(LOCTEXT("ToggleHostVisibility", "Set Host Visibility"));
+	const FScopedTransaction Transaction(LOCTEXT("ToggleHostEnabled", "Set Host Enabled"));
 
 	UDisplayClusterConfigurationHostDisplayData* Host = GetObjectChecked<UDisplayClusterConfigurationHostDisplayData>();
 
 	// Use SaveToTransactionBuffer to avoid marking the package as dirty
 	SaveToTransactionBuffer(Host, false);
-	Host->bIsVisible = !Host->bIsVisible;
+	Host->bIsUnlocked = !Host->bIsUnlocked;
 
-	// Set the visible state of this host's children to match it.
+	// Set the enabled state of this host's children to match it.
 	for (TSharedPtr<IDisplayClusterConfiguratorTreeItem> Child : Children)
 	{
 		TSharedPtr<FDisplayClusterConfiguratorTreeItemClusterNode> ClusterNodeTreeItem = StaticCastSharedPtr<FDisplayClusterConfiguratorTreeItemClusterNode>(Child);
 		if (ClusterNodeTreeItem.IsValid())
 		{
-			ClusterNodeTreeItem->SetVisible(Host->bIsVisible);
+			ClusterNodeTreeItem->SetUnlocked(Host->bIsUnlocked);
 		}
 	}
-
-	return FReply::Handled();
 }
 
-const FSlateBrush* FDisplayClusterConfiguratorTreeItemHost::GetEnabledButtonBrush() const
+FSlateColor FDisplayClusterConfiguratorTreeItemHost::GetClusterItemGroupColor() const
 {
 	UDisplayClusterConfigurationHostDisplayData* Host = GetObjectChecked<UDisplayClusterConfigurationHostDisplayData>();
-
-	if (Host->bIsEnabled)
+	if (ToolkitPtr.Pin()->IsObjectSelected(Host))
 	{
-		return EnabledButton->IsHovered() ? FEditorStyle::GetBrush(TEXT("Level.UnlockedHighlightIcon16x")) : FEditorStyle::GetBrush(TEXT("Level.UnlockedIcon16x"));
+		return FDisplayClusterConfiguratorStyle::Get().GetColor("DisplayClusterConfigurator.Node.Color.Selected");
 	}
 	else
 	{
-		return EnabledButton->IsHovered() ? FEditorStyle::GetBrush(TEXT("Level.LockedHighlightIcon16x")) : FEditorStyle::GetBrush(TEXT("Level.LockedIcon16x"));
+		return Host->Color;
 	}
 }
 
-FReply FDisplayClusterConfiguratorTreeItemHost::OnEnabledButtonClicked()
+FReply FDisplayClusterConfiguratorTreeItemHost::OnClusterItemGroupClicked(const FGeometry& MyGeometry, const FPointerEvent& MouseEvent)
 {
-	const FScopedTransaction Transaction(LOCTEXT("ToggleHostEnabled", "Set Host Enabled"));
-
 	UDisplayClusterConfigurationHostDisplayData* Host = GetObjectChecked<UDisplayClusterConfigurationHostDisplayData>();
+	TArray<UObject*> Objects { Host };
+	ToolkitPtr.Pin()->SelectObjects(Objects);
 
-	// Use SaveToTransactionBuffer to avoid marking the package as dirty
-	SaveToTransactionBuffer(Host, false);
-	Host->bIsEnabled = !Host->bIsEnabled;
+	return FReply::Handled();
+}
 
-	// Set the enabled state of this host's children to match it.
-	for (TSharedPtr<IDisplayClusterConfiguratorTreeItem> Child : Children)
+FText FDisplayClusterConfiguratorTreeItemHost::GetHostAddress() const
+{
+	UDisplayClusterConfigurationHostDisplayData* Host = GetObjectChecked<UDisplayClusterConfigurationHostDisplayData>();
+	return FText::Format(LOCTEXT("HostAddressFormattedLabel", "({0})"), FText::FromString(FDisplayClusterConfiguratorClusterUtils::GetAddressForHost(Host)));
+}
+
+bool FDisplayClusterConfiguratorTreeItemHost::CanDropClusterNodes(TSharedPtr<FDisplayClusterConfiguratorClusterNodeDragDropOp> ClusterNodeDragDropOp, FText& OutErrorMessage) const
+{
+	if (ClusterNodeDragDropOp.IsValid())
 	{
-		TSharedPtr<FDisplayClusterConfiguratorTreeItemClusterNode> ClusterNodeTreeItem = StaticCastSharedPtr<FDisplayClusterConfiguratorTreeItemClusterNode>(Child);
-		if (ClusterNodeTreeItem.IsValid())
-		{
-			ClusterNodeTreeItem->SetEnabled(Host->bIsEnabled);
-		}
+		UDisplayClusterConfigurationHostDisplayData* Host = GetObjectChecked<UDisplayClusterConfigurationHostDisplayData>();
+		bool bCanAcceptClusterNodes = true;
+
+		return bCanAcceptClusterNodes;
 	}
 
-	return FReply::Handled();
+	return false;
 }
 
 void FDisplayClusterConfiguratorTreeItemHost::OnPostEditChangeChainProperty(const FPropertyChangedChainEvent& PropertyChangedEvent)
diff --git a/Engine/Plugins/Runtime/nDisplay/Source/DisplayClusterConfigurator/Private/Views/TreeViews/Cluster/TreeItems/DisplayClusterConfiguratorTreeItemHost.h b/Engine/Plugins/Runtime/nDisplay/Source/DisplayClusterConfigurator/Private/Views/TreeViews/Cluster/TreeItems/DisplayClusterConfiguratorTreeItemHost.h
--- a/Engine/Plugins/Runtime/nDisplay/Source/DisplayClusterConfigurator/Private/Views/TreeViews/Cluster/TreeItems/DisplayClusterConfiguratorTreeItemHost.h
+++ b/Engine/Plugins/Runtime/nDisplay/Source/DisplayClusterConfigurator/Private/Views/TreeViews/Cluster/TreeItems/DisplayClusterConfiguratorTreeItemHost.h
@@ -21,7 +21,6 @@ public:
 
 	//~ Begin FDisplayClusterConfiguratorTreeItem Interface
 	virtual void Initialize() override;
-	virtual TSharedRef<SWidget> GenerateWidgetForColumn(const FName& ColumnName, TSharedPtr<ITableRow> TableRow, const TAttribute<FText>& FilterText, FIsSelected InIsSelected) override;
 	virtual void DeleteItem() const override;
 
 	virtual void HandleDragEnter(const FDragDropEvent& DragDropEvent) override;
@@ -38,24 +37,25 @@ public:
 protected:
 	virtual void FillItemColumn(TSharedPtr<SHorizontalBox> Box, const TAttribute<FText>& FilterText, FIsSelected InIsSelected) override;
 	virtual void OnDisplayNameCommitted(const FText& NewText, ETextCommit::Type CommitInfo) override;
+	
+	virtual bool CanClusterItemBeHidden() const override { return true; }
+	virtual bool CanClusterItemBeLocked() const override { return true; }
+	virtual bool IsClusterItemGrouped() const override { return true; }
+
+	virtual bool IsClusterItemVisible() const override;
+	virtual bool IsClusterItemUnlocked() const override;
+
+	virtual void ToggleClusterItemVisibility() override;
+	virtual void ToggleClusterItemLock() override;
+
+	virtual FSlateColor GetClusterItemGroupColor() const override;
+	virtual FReply OnClusterItemGroupClicked(const FGeometry& MyGeometry, const FPointerEvent& MouseEvent) override;
 	//~ End FDisplayClusterConfiguratorTreeItem Interface
 
 private:
 	FText GetHostAddress() const;
-	FSlateColor GetHostColor() const;
-	FReply OnHostClicked(const FGeometry& MyGeometry, const FPointerEvent& MouseEvent);
 
 	bool CanDropClusterNodes(TSharedPtr<FDisplayClusterConfiguratorClusterNodeDragDropOp> ClusterNodeDragDropOp, FText& OutErrorMessage) const;
 
-	const FSlateBrush* GetVisibilityButtonBrush() const;
-	FReply OnVisibilityButtonClicked();
-
-	const FSlateBrush* GetEnabledButtonBrush() const;
-	FReply OnEnabledButtonClicked();
-
 	void OnPostEditChangeChainProperty(const FPropertyChangedChainEvent& PropertyChangedEvent);
-
-private:
-	TSharedPtr<SButton> VisibilityButton;
-	TSharedPtr<SButton> EnabledButton;
 };
\ No newline at end of file
diff --git a/Engine/Plugins/Runtime/nDisplay/Source/DisplayClusterConfigurator/Private/Views/TreeViews/Cluster/TreeItems/DisplayClusterConfiguratorTreeItemViewport.cpp b/Engine/Plugins/Runtime/nDisplay/Source/DisplayClusterConfigurator/Private/Views/TreeViews/Cluster/TreeItems/DisplayClusterConfiguratorTreeItemViewport.cpp
--- a/Engine/Plugins/Runtime/nDisplay/Source/DisplayClusterConfigurator/Private/Views/TreeViews/Cluster/TreeItems/DisplayClusterConfiguratorTreeItemViewport.cpp
+++ b/Engine/Plugins/Runtime/nDisplay/Source/DisplayClusterConfigurator/Private/Views/TreeViews/Cluster/TreeItems/DisplayClusterConfiguratorTreeItemViewport.cpp
@@ -41,7 +41,7 @@ void FDisplayClusterConfiguratorTreeItemViewport::SetVisible(bool bIsVisible)
 	PropertyView->GetPropertyHandle()->SetValue(bIsVisible);
 }
 
-void FDisplayClusterConfiguratorTreeItemViewport::SetEnabled(bool bIsEnabled)
+void FDisplayClusterConfiguratorTreeItemViewport::SetUnlocked(bool bIsUnlocked)
 {
 	UDisplayClusterConfigurationViewport* Viewport = GetObjectChecked<UDisplayClusterConfigurationViewport>();
 
@@ -49,9 +49,9 @@ void FDisplayClusterConfiguratorTreeItemViewport::SetEnabled(bool bIsEnabled)
 	SaveToTransactionBuffer(Viewport, false);
 
 	const TSharedPtr<ISinglePropertyView> PropertyView = DisplayClusterConfiguratorPropertyUtils::GetPropertyView(
-		Viewport, GET_MEMBER_NAME_CHECKED(UDisplayClusterConfigurationViewport, bIsEnabled));
+		Viewport, GET_MEMBER_NAME_CHECKED(UDisplayClusterConfigurationViewport, bIsUnlocked));
 	
-	PropertyView->GetPropertyHandle()->SetValue(bIsEnabled);
+	PropertyView->GetPropertyHandle()->SetValue(bIsUnlocked);
 }
 
 void FDisplayClusterConfiguratorTreeItemViewport::OnSelection()
@@ -83,50 +83,6 @@ void FDisplayClusterConfiguratorTreeItemViewport::OnSelection()
 	}
 }
 
-TSharedRef<SWidget> FDisplayClusterConfiguratorTreeItemViewport::GenerateWidgetForColumn(const FName& ColumnName, TSharedPtr<ITableRow> TableRow, const TAttribute<FText>& FilterText, FIsSelected InIsSelected)
-{
-	if (ColumnName == FDisplayClusterConfiguratorViewCluster::Columns::Host)
-	{
-		return SNew(SImage)
-			.ColorAndOpacity(this, &FDisplayClusterConfiguratorTreeItemViewport::GetHostColor)
-			.OnMouseButtonDown(this, &FDisplayClusterConfiguratorTreeItemViewport::OnHostClicked)
-			.Image(FDisplayClusterConfiguratorStyle::GetBrush("DisplayClusterConfigurator.Node.Body"))
-			.Cursor(EMouseCursor::Hand);
-	}
-	else if (ColumnName == FDisplayClusterConfiguratorViewCluster::Columns::Visible)
-	{
-		return SAssignNew(VisibilityButton, SButton)
-			.ContentPadding(0)
-			.ButtonStyle(FEditorStyle::Get(), "NoBorder")
-			.ToolTipText(LOCTEXT("VisibilityButton_Tooltip", "Hides or shows this viewport in the Output Mapping editor"))
-			.HAlign(HAlign_Center)
-			.VAlign(VAlign_Center)
-			.OnClicked(this, &FDisplayClusterConfiguratorTreeItemViewport::OnVisibilityButtonClicked)
-			.Content()
-			[
-				SNew(SImage)
-				.Image(this, &FDisplayClusterConfiguratorTreeItemViewport::GetVisibilityButtonBrush)
-			];
-	}
-	else if (ColumnName == FDisplayClusterConfiguratorViewCluster::Columns::Enabled)
-	{
-		return SAssignNew(EnabledButton, SButton)
-			.ContentPadding(0)
-			.ButtonStyle(FEditorStyle::Get(), "NoBorder")
-			.ToolTipText(LOCTEXT("EnabledButton_Tooltip", "Enables or disables this viewport in the Output Mapping editor"))
-			.HAlign(HAlign_Center)
-			.VAlign(VAlign_Center)
-			.OnClicked(this, &FDisplayClusterConfiguratorTreeItemViewport::OnEnabledButtonClicked)
-			.Content()
-			[
-				SNew(SImage)
-				.Image(this, &FDisplayClusterConfiguratorTreeItemViewport::GetEnabledButtonBrush)
-			];
-	}
-
-	return FDisplayClusterConfiguratorTreeItemCluster::GenerateWidgetForColumn(ColumnName, TableRow, FilterText, InIsSelected);
-}
-
 void FDisplayClusterConfiguratorTreeItemViewport::DeleteItem() const
 {
 	UDisplayClusterConfigurationViewport* Viewport = GetObjectChecked<UDisplayClusterConfigurationViewport>();
@@ -255,7 +211,49 @@ void FDisplayClusterConfiguratorTreeItemViewport::OnDisplayNameCommitted(const F
 	}
 }
 
-FSlateColor FDisplayClusterConfiguratorTreeItemViewport::GetHostColor() const
+bool FDisplayClusterConfiguratorTreeItemViewport::IsClusterItemVisible() const
+{
+	UDisplayClusterConfigurationViewport* Viewport = GetObjectChecked<UDisplayClusterConfigurationViewport>();
+	return Viewport->bIsVisible;
+}
+
+void FDisplayClusterConfiguratorTreeItemViewport::ToggleClusterItemVisibility()
+{
+	const FScopedTransaction Transaction(LOCTEXT("ToggleViewportVisibility", "Toggle Viewport Visibilty"));
+
+	UDisplayClusterConfigurationViewport* Viewport = GetObjectChecked<UDisplayClusterConfigurationViewport>();
+
+	// Use SaveToTransactionBuffer to avoid marking the package as dirty
+	SaveToTransactionBuffer(Viewport, false);
+
+	const TSharedPtr<ISinglePropertyView> PropertyView = DisplayClusterConfiguratorPropertyUtils::GetPropertyView(
+		Viewport, GET_MEMBER_NAME_CHECKED(UDisplayClusterConfigurationViewport, bIsVisible));
+
+	PropertyView->GetPropertyHandle()->SetValue(!Viewport->bIsVisible);
+}
+
+bool FDisplayClusterConfiguratorTreeItemViewport::IsClusterItemUnlocked() const
+{
+	UDisplayClusterConfigurationViewport* Viewport = GetObjectChecked<UDisplayClusterConfigurationViewport>();
+	return Viewport->bIsUnlocked;
+}
+
+void FDisplayClusterConfiguratorTreeItemViewport::ToggleClusterItemLock()
+{
+	const FScopedTransaction Transaction(LOCTEXT("ToggleViewportLock", "Toggle Viewport Lock"));
+
+	UDisplayClusterConfigurationViewport* Viewport = GetObjectChecked<UDisplayClusterConfigurationViewport>();
+
+	// Use SaveToTransactionBuffer to avoid marking the package as dirty
+	SaveToTransactionBuffer(Viewport, false);
+
+	const TSharedPtr<ISinglePropertyView> PropertyView = DisplayClusterConfiguratorPropertyUtils::GetPropertyView(
+		Viewport, GET_MEMBER_NAME_CHECKED(UDisplayClusterConfigurationViewport, bIsUnlocked));
+
+	PropertyView->GetPropertyHandle()->SetValue(!Viewport->bIsUnlocked);
+}
+
+FSlateColor FDisplayClusterConfiguratorTreeItemViewport::GetClusterItemGroupColor() const
 {
 	TSharedPtr<IDisplayClusterConfiguratorTreeItem> ParentItem = Parent.Pin();
 	UDisplayClusterConfigurationHostDisplayData* HostDisplayData = nullptr;
@@ -273,7 +271,7 @@ FSlateColor FDisplayClusterConfiguratorTreeItemViewport::GetHostColor() const
 
 	if (ToolkitPtr.Pin()->IsObjectSelected(HostDisplayData))
 	{
-		return FDisplayClusterConfiguratorStyle::GetColor("DisplayClusterConfigurator.Node.Color.Selected");
+		return FDisplayClusterConfiguratorStyle::Get().GetColor("DisplayClusterConfigurator.Node.Color.Selected");
 	}
 	else
 	{
@@ -281,7 +279,7 @@ FSlateColor FDisplayClusterConfiguratorTreeItemViewport::GetHostColor() const
 	}
 }
 
-FReply FDisplayClusterConfiguratorTreeItemViewport::OnHostClicked(const FGeometry& MyGeometry, const FPointerEvent& MouseEvent)
+FReply FDisplayClusterConfiguratorTreeItemViewport::OnClusterItemGroupClicked(const FGeometry& MyGeometry, const FPointerEvent& MouseEvent)
 {
 	TSharedPtr<IDisplayClusterConfiguratorTreeItem> ParentItem = GetParent();
 
@@ -297,64 +295,4 @@ FReply FDisplayClusterConfiguratorTreeItemViewport::OnHostClicked(const FGeometr
 	return FReply::Unhandled();
 }
 
-const FSlateBrush* FDisplayClusterConfiguratorTreeItemViewport::GetVisibilityButtonBrush() const
-{
-	UDisplayClusterConfigurationViewport* Viewport = GetObjectChecked<UDisplayClusterConfigurationViewport>();
-
-	if (Viewport->bIsVisible)
-	{
-		return VisibilityButton->IsHovered() ? FEditorStyle::GetBrush(TEXT("Level.VisibleHighlightIcon16x")) : FEditorStyle::GetBrush(TEXT("Level.VisibleIcon16x"));
-	}
-	else
-	{
-		return VisibilityButton->IsHovered() ? FEditorStyle::GetBrush(TEXT("Level.NotVisibleHighlightIcon16x")) : FEditorStyle::GetBrush(TEXT("Level.NotVisibleIcon16x"));
-	}
-}
-
-FReply FDisplayClusterConfiguratorTreeItemViewport::OnVisibilityButtonClicked()
-{
-	const FScopedTransaction Transaction(LOCTEXT("ToggleViewportVisibility", "Set Viewport Visibility"));
-
-	UDisplayClusterConfigurationViewport* Viewport = GetObjectChecked<UDisplayClusterConfigurationViewport>();
-
-	// Use SaveToTransactionBuffer to avoid marking the package as dirty
-	SaveToTransactionBuffer(Viewport, false);
-
-	const TSharedPtr<ISinglePropertyView> PropertyView = DisplayClusterConfiguratorPropertyUtils::GetPropertyView(
-		Viewport, GET_MEMBER_NAME_CHECKED(UDisplayClusterConfigurationViewport, bIsVisible));
-
-	PropertyView->GetPropertyHandle()->SetValue(!Viewport->bIsVisible);
-	return FReply::Handled();
-}
-
-const FSlateBrush* FDisplayClusterConfiguratorTreeItemViewport::GetEnabledButtonBrush() const
-{
-	UDisplayClusterConfigurationViewport* Viewport = GetObjectChecked<UDisplayClusterConfigurationViewport>();
-
-	if (Viewport->bIsEnabled)
-	{
-		return EnabledButton->IsHovered() ? FEditorStyle::GetBrush(TEXT("Level.UnlockedHighlightIcon16x")) : FEditorStyle::GetBrush(TEXT("Level.UnlockedIcon16x"));
-	}
-	else
-	{
-		return EnabledButton->IsHovered() ? FEditorStyle::GetBrush(TEXT("Level.LockedHighlightIcon16x")) : FEditorStyle::GetBrush(TEXT("Level.LockedIcon16x"));
-	}
-}
-
-FReply FDisplayClusterConfiguratorTreeItemViewport::OnEnabledButtonClicked()
-{
-	const FScopedTransaction Transaction(LOCTEXT("ToggleViewportEnabled", "Set Viewport Enabled"));
-
-	UDisplayClusterConfigurationViewport* Viewport = GetObjectChecked<UDisplayClusterConfigurationViewport>();
-
-	// Use SaveToTransactionBuffer to avoid marking the package as dirty
-	SaveToTransactionBuffer(Viewport, false);
-
-	const TSharedPtr<ISinglePropertyView> PropertyView = DisplayClusterConfiguratorPropertyUtils::GetPropertyView(
-		Viewport, GET_MEMBER_NAME_CHECKED(UDisplayClusterConfigurationViewport, bIsEnabled));
-
-	PropertyView->GetPropertyHandle()->SetValue(!Viewport->bIsEnabled);
-	return FReply::Handled();
-}
-
 #undef LOCTEXT_NAMESPACE 
\ No newline at end of file
diff --git a/Engine/Plugins/Runtime/nDisplay/Source/DisplayClusterConfigurator/Private/Views/TreeViews/Cluster/TreeItems/DisplayClusterConfiguratorTreeItemViewport.h b/Engine/Plugins/Runtime/nDisplay/Source/DisplayClusterConfigurator/Private/Views/TreeViews/Cluster/TreeItems/DisplayClusterConfiguratorTreeItemViewport.h
--- a/Engine/Plugins/Runtime/nDisplay/Source/DisplayClusterConfigurator/Private/Views/TreeViews/Cluster/TreeItems/DisplayClusterConfiguratorTreeItemViewport.h
+++ b/Engine/Plugins/Runtime/nDisplay/Source/DisplayClusterConfigurator/Private/Views/TreeViews/Cluster/TreeItems/DisplayClusterConfiguratorTreeItemViewport.h
@@ -19,12 +19,11 @@ public:
 			UObject* InObjectToEdit);
 
 	void SetVisible(bool bIsVisible);
-	void SetEnabled(bool bIsEnabled);
+	void SetUnlocked(bool bIsUnlocked);
 
 	//~ Begin FDisplayClusterConfiguratorTreeItem Interface
 	virtual void OnSelection() override;
 
-	virtual TSharedRef<SWidget> GenerateWidgetForColumn(const FName& ColumnName, TSharedPtr<ITableRow> TableRow, const TAttribute<FText>& FilterText, FIsSelected InIsSelected) override;
 	virtual void DeleteItem() const override;
 	virtual bool CanHideItem() const override { return true; }
 	virtual void SetItemHidden(bool bIsHidden) { SetVisible(!bIsHidden); }
@@ -37,19 +36,22 @@ public:
 
 protected:
 	virtual void OnDisplayNameCommitted(const FText& NewText, ETextCommit::Type CommitInfo) override;
-	//~ End FDisplayClusterConfiguratorTreeItem Interface
 
-private:
-	FSlateColor GetHostColor() const;
-	FReply OnHostClicked(const FGeometry& MyGeometry, const FPointerEvent& MouseEvent);
+	virtual bool CanClusterItemBeHidden() const override { return true; }
+	virtual bool CanClusterItemBeLocked() const override { return true; }
+	virtual bool IsClusterItemGrouped() const override { return true; }
 
-	const FSlateBrush* GetVisibilityButtonBrush() const;
-	FReply OnVisibilityButtonClicked();
+	virtual bool IsClusterItemVisible() const override;
+	virtual bool IsClusterItemUnlocked() const override;
 
-	const FSlateBrush* GetEnabledButtonBrush() const;
-	FReply OnEnabledButtonClicked();
+	virtual void ToggleClusterItemVisibility() override;
+	virtual void ToggleClusterItemLock() override;
+
+	virtual FSlateColor GetClusterItemGroupColor() const override;
+	virtual FReply OnClusterItemGroupClicked(const FGeometry& MyGeometry, const FPointerEvent& MouseEvent) override;
+	//~ End FDisplayClusterConfiguratorTreeItem Interface
 
 private:
-	TSharedPtr<SButton> VisibilityButton;
-	TSharedPtr<SButton> EnabledButton;
+	FSlateColor GetHostColor() const;
+	FReply OnHostClicked(const FGeometry& MyGeometry, const FPointerEvent& MouseEvent);
 };
\ No newline at end of file
diff --git a/Engine/Plugins/Runtime/nDisplay/Source/DisplayClusterConfigurator/Private/Views/TreeViews/DisplayClusterConfiguratorTreeItem.cpp b/Engine/Plugins/Runtime/nDisplay/Source/DisplayClusterConfigurator/Private/Views/TreeViews/DisplayClusterConfiguratorTreeItem.cpp
--- a/Engine/Plugins/Runtime/nDisplay/Source/DisplayClusterConfigurator/Private/Views/TreeViews/DisplayClusterConfiguratorTreeItem.cpp
+++ b/Engine/Plugins/Runtime/nDisplay/Source/DisplayClusterConfigurator/Private/Views/TreeViews/DisplayClusterConfiguratorTreeItem.cpp
@@ -26,16 +26,28 @@ TSharedRef<SWidget> FDisplayClusterConfiguratorTreeItem::GenerateWidgetForColumn
 {
 	if (ColumnName == IDisplayClusterConfiguratorViewTree::Columns::Item)
 	{
-		TSharedPtr<SHorizontalBox> RowBox = SNew(SHorizontalBox)
-			+ SHorizontalBox::Slot()
-			.AutoWidth()
-			[
-				SNew(SExpanderArrow, TableRow)
-			];
-
+		TSharedPtr<SHorizontalBox> RowBox = SNew(SHorizontalBox);
 		FillItemColumn(RowBox, FilterText, InIsSelected);
 
-		return RowBox.ToSharedRef();
+		return SNew(SBox)
+			.MinDesiredHeight(20.0f)
+			[
+				SNew(SHorizontalBox)
+
+				+SHorizontalBox::Slot()
+				.AutoWidth()
+				.Padding(6, 0, 0, 0)
+				[
+					SNew(SExpanderArrow, TableRow)
+					.IndentAmount(12)
+				]
+
+				+SHorizontalBox::Slot()
+				.FillWidth(1.0f)
+				[
+					RowBox.ToSharedRef()
+				]
+			];
 	}
 
 	return SNullWidget::NullWidget;
@@ -141,23 +153,26 @@ void FDisplayClusterConfiguratorTreeItem::FillItemColumn(TSharedPtr<SHorizontalB
 {
 	Box->AddSlot()
 		.AutoWidth()
-		.Padding(FMargin(0.0f, 1.0f))
+		.Padding(FMargin(0.f, 1.f, 6.f, 1.f))
 		.VAlign(VAlign_Center)
-		.HAlign(HAlign_Center)
 		[
-			SNew(SImage)
-			.ColorAndOpacity(FSlateColor::UseForeground())
-			.Image(FDisplayClusterConfiguratorStyle::GetBrush(*GetIconStyle()))
+			SNew(SBox)
+			.WidthOverride(16)
+			.HeightOverride(16)
+			[
+				SNew(SImage)
+				.ColorAndOpacity(FSlateColor::UseForeground())
+				.Image(FDisplayClusterConfiguratorStyle::Get().GetBrush(*GetIconStyle()))
+			]
 		];
 
 	Box->AddSlot()
 		.AutoWidth()
-		.Padding(2, 0, 0, 0)
 		.VAlign(VAlign_Center)
 		[
 			SAssignNew(DisplayNameTextBlock, SInlineEditableTextBlock)
 			.Text(this, &FDisplayClusterConfiguratorTreeItem::GetRowItemText)
-			.HighlightText( FilterText )
+			.HighlightText(FilterText)
 			.IsReadOnly(this, &FDisplayClusterConfiguratorTreeItem::IsReadOnly)
 			.OnTextCommitted(this, &FDisplayClusterConfiguratorTreeItem::OnDisplayNameCommitted)
 		];
diff --git a/Engine/Plugins/Runtime/nDisplay/Source/DisplayClusterConfigurator/Private/Views/TreeViews/SDisplayClusterConfiguratorViewTree.cpp b/Engine/Plugins/Runtime/nDisplay/Source/DisplayClusterConfigurator/Private/Views/TreeViews/SDisplayClusterConfiguratorViewTree.cpp
--- a/Engine/Plugins/Runtime/nDisplay/Source/DisplayClusterConfigurator/Private/Views/TreeViews/SDisplayClusterConfiguratorViewTree.cpp
+++ b/Engine/Plugins/Runtime/nDisplay/Source/DisplayClusterConfigurator/Private/Views/TreeViews/SDisplayClusterConfiguratorViewTree.cpp
@@ -67,7 +67,7 @@ void SDisplayClusterConfiguratorViewTree::Construct(const FArguments& InArgs,
 				// Add a border if we are being used as a picker
 				SNew(SBorder)
 				.Visibility_Lambda([this](){ return Mode == EDisplayClusterConfiguratorTreeMode::Picker ? EVisibility::Visible: EVisibility::Collapsed; })
-				.BorderImage(FEditorStyle::Get().GetBrush("Menu.Background"))
+				.BorderImage(FAppStyle::Get().GetBrush("Menu.Background"))
 			]
 			+SOverlay::Slot()
 			[
@@ -79,7 +79,7 @@ void SDisplayClusterConfiguratorViewTree::Construct(const FArguments& InArgs,
 				[
 					SNew(SBorder)
 					.Padding(0)
-					.BorderImage(FEditorStyle::GetBrush("DetailsView.CategoryTop"))
+					.BorderImage(FAppStyle::Get().GetBrush("DetailsView.CategoryTop"))
 					.BorderBackgroundColor(FLinearColor(0.6f, 0.6f, 0.6f, 1.0f))
 					[
 						SNew(SHorizontalBox)
@@ -90,7 +90,7 @@ void SDisplayClusterConfiguratorViewTree::Construct(const FArguments& InArgs,
 						[
 					
 							SNew(SPositiveActionButton)
-							.Text(LOCTEXT("AddNewItem", "Add New"))
+							.Text(LOCTEXT("AddNewItem", "Add"))
 							.Icon(FAppStyle::Get().GetBrush("Icons.Plus"))
 							.ToolTipText(LOCTEXT("AddNewClusterItemTooltip", "Adds a new Cluster Item to the Cluster"))
 							.Visibility(this, &SDisplayClusterConfiguratorViewTree::GetAddNewComboButtonVisibility)
@@ -103,7 +103,7 @@ void SDisplayClusterConfiguratorViewTree::Construct(const FArguments& InArgs,
 						.Padding(3.0f, 3.0f)
 						[
 							SAssignNew(FilterComboButton, SComboButton)
-							.ComboButtonStyle(FEditorStyle::Get(), "GenericFilters.ComboButtonStyle")
+							.ComboButtonStyle(FAppStyle::Get(), "GenericFilters.ComboButtonStyle")
 							.ForegroundColor(FLinearColor::White)
 							.Visibility(this, &SDisplayClusterConfiguratorViewTree::GetFilterOptionsComboButtonVisibility)
 							.ContentPadding(0.0f)
@@ -118,8 +118,8 @@ void SDisplayClusterConfiguratorViewTree::Construct(const FArguments& InArgs,
 								.VAlign(VAlign_Center)
 								[
 									SNew(STextBlock)
-									.TextStyle(FEditorStyle::Get(), "GenericFilters.TextStyle")
-									.Font(FEditorStyle::Get().GetFontStyle("FontAwesome.9"))
+									.TextStyle(FAppStyle::Get(), "GenericFilters.TextStyle")
+									.Font(FAppStyle::Get().GetFontStyle("FontAwesome.9"))
 									.Text(FText::FromString(FString(TEXT("\xf0b0"))) /*fa-filter*/)
 								]
 								+SHorizontalBox::Slot()
@@ -127,8 +127,8 @@ void SDisplayClusterConfiguratorViewTree::Construct(const FArguments& InArgs,
 								.Padding(2, 0, 0, 0)
 								.VAlign(VAlign_Center)
 								[
-									SNew( STextBlock )
-									.TextStyle(FEditorStyle::Get(), "GenericFilters.TextStyle")
+									SNew(STextBlock)
+									.TextStyle(FAppStyle::Get(), "GenericFilters.TextStyle")
 									.Text( LOCTEXT("FilterMenuLabel", "Options") )
 								]
 							]
@@ -177,7 +177,7 @@ void SDisplayClusterConfiguratorViewTree::Construct(const FArguments& InArgs,
 					.AutoWidth()
 					[
 						SNew(SComboButton)
-						.ComboButtonStyle(FEditorStyle::Get(), "GenericFilters.ComboButtonStyle")
+						.ComboButtonStyle(FAppStyle::Get(), "GenericFilters.ComboButtonStyle")
 						.ForegroundColor(FLinearColor::White)
 						.Visibility(this, &SDisplayClusterConfiguratorViewTree::GetViewOptionsComboButtonVisibility)
 						.ContentPadding(0)
@@ -192,7 +192,7 @@ void SDisplayClusterConfiguratorViewTree::Construct(const FArguments& InArgs,
 							.AutoWidth()
 							.VAlign(VAlign_Center)
 							[
-								SNew(SImage).Image(FEditorStyle::GetBrush("GenericViewButton"))
+								SNew(SImage).Image(FAppStyle::Get().GetBrush("GenericViewButton"))
 							]
 
 							+ SHorizontalBox::Slot()
@@ -214,7 +214,7 @@ void SDisplayClusterConfiguratorViewTree::Construct(const FArguments& InArgs,
 			[
 				SNew(STextBlock)
 				.Visibility(EVisibility::HitTestInvisible)
-				.TextStyle(FEditorStyle::Get(), "Graph.CornerText")
+				.TextStyle(FAppStyle::Get(), "Graph.CornerText")
 				.Text(this, &SDisplayClusterConfiguratorViewTree::GetCornerText)
 			]
 		],
@@ -239,9 +239,7 @@ void SDisplayClusterConfiguratorViewTree::Refresh()
 
 void SDisplayClusterConfiguratorViewTree::CreateTreeColumns()
 {
-	TSharedRef<SHeaderRow> TreeHeaderRow =
-		SNew(SHeaderRow)
-		.Visibility(EVisibility::Collapsed);
+	TSharedRef<SHeaderRow> TreeHeaderRow = SNew(SHeaderRow);
 
 	TArray<SHeaderRow::FColumn::FArguments> Columns;
 	ViewTreePtr.Pin()->ConstructColumns(Columns);
diff --git a/Engine/Plugins/Runtime/nDisplay/Source/DisplayClusterConfigurator/Private/Views/Viewport/DisplayClusterConfiguratorSCSEditorViewport.cpp b/Engine/Plugins/Runtime/nDisplay/Source/DisplayClusterConfigurator/Private/Views/Viewport/DisplayClusterConfiguratorSCSEditorViewport.cpp
--- a/Engine/Plugins/Runtime/nDisplay/Source/DisplayClusterConfigurator/Private/Views/Viewport/DisplayClusterConfiguratorSCSEditorViewport.cpp
+++ b/Engine/Plugins/Runtime/nDisplay/Source/DisplayClusterConfigurator/Private/Views/Viewport/DisplayClusterConfiguratorSCSEditorViewport.cpp
@@ -45,27 +45,28 @@ public:
 		EditorViewport = InArgs._EditorViewport;
 
 		static const FName DefaultForegroundName("DefaultForeground");
+		const FMargin ToolbarSlotPadding(4.0f, 1.0f);
 
 		this->ChildSlot
 		[
 			SNew(SBorder)
-			.BorderImage(FEditorStyle::GetBrush("NoBorder"))
-			.ForegroundColor(FEditorStyle::GetSlateColor(DefaultForegroundName))
+			.BorderImage(FAppStyle::Get().GetBrush("EditorViewportToolBar.Background"))
+			.ForegroundColor(FAppStyle::Get().GetSlateColor(DefaultForegroundName))
 			[
 				SNew(SHorizontalBox)
-				+SHorizontalBox::Slot()
+				+ SHorizontalBox::Slot()
 				.AutoWidth()
-				.Padding(2.0f, 2.0f)
+				.Padding(ToolbarSlotPadding)
 				[
 					SNew(SEditorViewportToolbarMenu)
 					.ParentToolBar(SharedThis(this))
 					.Cursor(EMouseCursor::Default)
-					.Image("EditorViewportToolBar.MenuDropdown")
+					.Image("EditorViewportToolBar.OptionsDropdown")
 					.OnGetMenuContent(this, &SDisplayClusterConfiguratorSCSEditorViewportToolBar::GeneratePreviewMenu)
 				]
 				+ SHorizontalBox::Slot()
 				.AutoWidth()
-				.Padding(2.0f, 2.0f)
+				.Padding(ToolbarSlotPadding)
 				[
 					SNew( SEditorViewportToolbarMenu )
 					.ParentToolBar( SharedThis( this ) )
@@ -76,7 +77,7 @@ public:
 				]
 				+ SHorizontalBox::Slot()
 				.AutoWidth()
-				.Padding(2.0f, 2.0f)
+				.Padding(ToolbarSlotPadding)
 				[
 					SNew( SEditorViewportToolbarMenu )
 					.ParentToolBar( SharedThis( this ) )
@@ -87,7 +88,7 @@ public:
 				]
 				+ SHorizontalBox::Slot()
 				.AutoWidth()
-				.Padding(2.0f, 2.0f)
+				.Padding(ToolbarSlotPadding)
 				[
 					SNew( SEditorViewportToolbarMenu )
 					.ParentToolBar( SharedThis( this ) )
@@ -96,8 +97,8 @@ public:
 					.OnGetMenuContent(this, &SDisplayClusterConfiguratorSCSEditorViewportToolBar::GenerateViewportsMenu)
 				]
 				+ SHorizontalBox::Slot()
-				.Padding( 3.0f, 1.0f )
-				.HAlign( HAlign_Right )
+				.Padding(ToolbarSlotPadding)
+				.HAlign(HAlign_Right)
 				[
 					SNew(STransformViewportToolBar)
 					.Viewport(EditorViewport.Pin().ToSharedRef())
@@ -156,7 +157,7 @@ public:
 			return GetCameraMenuLabelIconFromViewportType( EditorViewport.Pin()->GetViewportClient()->GetViewportType() );
 		}
 
-		return FEditorStyle::GetBrush(NAME_None);
+		return FAppStyle::Get().GetBrush(NAME_None);
 	}
 
 	TSharedRef<SWidget> GenerateCameraMenu() const
@@ -231,7 +232,7 @@ public:
 			}
 		}
 
-		return FEditorStyle::GetBrush(Icon);
+		return FAppStyle::Get().GetBrush(Icon);
 	}
 
 	TSharedRef<SWidget> GenerateViewMenu() const
@@ -351,7 +352,7 @@ public:
 		{
 			FToolBarBuilder OnePaneButton(CommandList, FMultiBoxCustomization::None);
 			OnePaneButton.SetLabelVisibility(EVisibility::Collapsed);
-			OnePaneButton.SetStyle(&FEditorStyle::Get(), "ViewportLayoutToolbar");
+			OnePaneButton.SetStyle(&FAppStyle::Get(), "ViewportLayoutToolbar");
 
 			OnePaneButton.AddToolBarButton(FEditorViewportCommands::Get().ViewportConfig_OnePane);
 
@@ -376,7 +377,7 @@ public:
 		{
 			FToolBarBuilder TwoPaneButtons(CommandList, FMultiBoxCustomization::None);
 			TwoPaneButtons.SetLabelVisibility(EVisibility::Collapsed);
-			TwoPaneButtons.SetStyle(&FEditorStyle::Get(), "ViewportLayoutToolbar");
+			TwoPaneButtons.SetStyle(&FAppStyle::Get(), "ViewportLayoutToolbar");
 
 			TwoPaneButtons.AddToolBarButton(FEditorViewportCommands::Get().ViewportConfig_TwoPanesH, NAME_None, FText());
 			TwoPaneButtons.AddToolBarButton(FEditorViewportCommands::Get().ViewportConfig_TwoPanesV, NAME_None, FText());
@@ -402,7 +403,7 @@ public:
 		{
 			FToolBarBuilder ThreePaneButtons(CommandList, FMultiBoxCustomization::None);
 			ThreePaneButtons.SetLabelVisibility(EVisibility::Collapsed);
-			ThreePaneButtons.SetStyle(&FEditorStyle::Get(), "ViewportLayoutToolbar");
+			ThreePaneButtons.SetStyle(&FAppStyle::Get(), "ViewportLayoutToolbar");
 
 			ThreePaneButtons.AddToolBarButton(FEditorViewportCommands::Get().ViewportConfig_ThreePanesLeft, NAME_None, FText());
 			ThreePaneButtons.AddToolBarButton(FEditorViewportCommands::Get().ViewportConfig_ThreePanesRight, NAME_None, FText());
@@ -430,7 +431,7 @@ public:
 		{
 			FToolBarBuilder FourPaneButtons(CommandList, FMultiBoxCustomization::None);
 			FourPaneButtons.SetLabelVisibility(EVisibility::Collapsed);
-			FourPaneButtons.SetStyle(&FEditorStyle::Get(), "ViewportLayoutToolbar");
+			FourPaneButtons.SetStyle(&FAppStyle::Get(), "ViewportLayoutToolbar");
 
 			FourPaneButtons.AddToolBarButton(FEditorViewportCommands::Get().ViewportConfig_FourPanes2x2, NAME_None, FText());
 			FourPaneButtons.AddToolBarButton(FEditorViewportCommands::Get().ViewportConfig_FourPanesLeft, NAME_None, FText());
