nDisplay static string explicit type

diff --git a/Engine/Plugins/Runtime/nDisplay/Source/DisplayCluster/Private/Render/Viewport/DisplayClusterViewportStrings.h b/Engine/Plugins/Runtime/nDisplay/Source/DisplayCluster/Private/Render/Viewport/DisplayClusterViewportStrings.h
--- a/Engine/Plugins/Runtime/nDisplay/Source/DisplayCluster/Private/Render/Viewport/DisplayClusterViewportStrings.h
+++ b/Engine/Plugins/Runtime/nDisplay/Source/DisplayCluster/Private/Render/Viewport/DisplayClusterViewportStrings.h
@@ -8,16 +8,16 @@ namespace DisplayClusterViewportStrings
 {
 	namespace prefix
 	{
-		static constexpr auto projection = TEXT("proj");
+		static constexpr const TCHAR* projection = TEXT("proj");
 	}
 
 	// ICVFXstrings
 	namespace icvfx
 	{
-		static constexpr auto prefix         = TEXT("icvfx");
-		static constexpr auto camera         = TEXT("incamera");
-		static constexpr auto chromakey      = TEXT("chromakey");
-		static constexpr auto lightcard      = TEXT("lightcard");
-		static constexpr auto lightcard_OCIO = TEXT("lightcard_ocio");
+		static constexpr const TCHAR* prefix         = TEXT("icvfx");
+		static constexpr const TCHAR* camera         = TEXT("incamera");
+		static constexpr const TCHAR* chromakey      = TEXT("chromakey");
+		static constexpr const TCHAR* lightcard      = TEXT("lightcard");
+		static constexpr const TCHAR* lightcard_OCIO = TEXT("lightcard_ocio");
 	}
 };
\ No newline at end of file
diff --git a/Engine/Plugins/Runtime/nDisplay/Source/DisplayCluster/Public/IDisplayCluster.h b/Engine/Plugins/Runtime/nDisplay/Source/DisplayCluster/Public/IDisplayCluster.h
--- a/Engine/Plugins/Runtime/nDisplay/Source/DisplayCluster/Public/IDisplayCluster.h
+++ b/Engine/Plugins/Runtime/nDisplay/Source/DisplayCluster/Public/IDisplayCluster.h
@@ -22,7 +22,7 @@ class IDisplayCluster
 	: public IModuleInterface
 {
 public:
-	static constexpr auto ModuleName = TEXT("DisplayCluster");
+	static constexpr const TCHAR* ModuleName = TEXT("DisplayCluster");
 
 public:
 	virtual ~IDisplayCluster() = default;
diff --git a/Engine/Plugins/Runtime/nDisplay/Source/DisplayCluster/Public/Misc/DisplayClusterStrings.h b/Engine/Plugins/Runtime/nDisplay/Source/DisplayCluster/Public/Misc/DisplayClusterStrings.h
--- a/Engine/Plugins/Runtime/nDisplay/Source/DisplayCluster/Public/Misc/DisplayClusterStrings.h
+++ b/Engine/Plugins/Runtime/nDisplay/Source/DisplayCluster/Public/Misc/DisplayClusterStrings.h
@@ -10,61 +10,61 @@ namespace DisplayClusterStrings
 	// Common strings
 	namespace common
 	{
-		static constexpr auto PairSeparator     = TEXT(" ");
-		static constexpr auto KeyValSeparator   = TEXT("=");
-		static constexpr auto ArrayValSeparator = TEXT(",");
+		static constexpr const TCHAR* PairSeparator     = TEXT(" ");
+		static constexpr const TCHAR* KeyValSeparator   = TEXT("=");
+		static constexpr const TCHAR* ArrayValSeparator = TEXT(",");
 	}
 
 	// Command line arguments
 	namespace args
 	{
-		static constexpr auto Cluster    = TEXT("dc_cluster");
-		static constexpr auto Node       = TEXT("dc_node");
-		static constexpr auto Config     = TEXT("dc_cfg");
-		static constexpr auto NodeList   = TEXT("dc_nodes");
-		static constexpr auto NodeRole   = TEXT("dc_role");
+		static constexpr const TCHAR* Cluster    = TEXT("dc_cluster");
+		static constexpr const TCHAR* Node       = TEXT("dc_node");
+		static constexpr const TCHAR* Config     = TEXT("dc_cfg");
+		static constexpr const TCHAR* NodeList   = TEXT("dc_nodes");
+		static constexpr const TCHAR* NodeRole   = TEXT("dc_role");
 
 		// Cluster node roles
 		namespace role
 		{
-			static constexpr auto Primary   = TEXT("primary");
-			static constexpr auto Secondary = TEXT("secondary");
-			static constexpr auto Backup    = TEXT("backup");
+			static constexpr const TCHAR* Primary   = TEXT("primary");
+			static constexpr const TCHAR* Secondary = TEXT("secondary");
+			static constexpr const TCHAR* Backup    = TEXT("backup");
 		}
 
 		// Stereo device types (command line values)
 		namespace dev
 		{
-			static constexpr auto QBS  = TEXT("quad_buffer_stereo");
-			static constexpr auto TB   = TEXT("dc_dev_top_bottom");
-			static constexpr auto SbS  = TEXT("dc_dev_side_by_side");
-			static constexpr auto Mono = TEXT("dc_dev_mono");
+			static constexpr const TCHAR* QBS  = TEXT("quad_buffer_stereo");
+			static constexpr const TCHAR* TB   = TEXT("dc_dev_top_bottom");
+			static constexpr const TCHAR* SbS  = TEXT("dc_dev_side_by_side");
+			static constexpr const TCHAR* Mono = TEXT("dc_dev_mono");
 		}
 	}
 
 	// RHI names
 	namespace rhi
 	{
-		static constexpr auto D3D11  = TEXT("D3D11");
-		static constexpr auto D3D12  = TEXT("D3D12");
-		static constexpr auto Vulkan = TEXT("Vulkan");
+		static constexpr const TCHAR* D3D11  = TEXT("D3D11");
+		static constexpr const TCHAR* D3D12  = TEXT("D3D12");
+		static constexpr const TCHAR* Vulkan = TEXT("Vulkan");
 	}
 
 	// Log strings
 	namespace log
 	{
-		static constexpr auto Found     = TEXT("found");
-		static constexpr auto NotFound  = TEXT("not found");
+		static constexpr const TCHAR* Found     = TEXT("found");
+		static constexpr const TCHAR* NotFound  = TEXT("not found");
 	}
 
 	// System cluster events
 	namespace cluster_events
 	{
 		// Common system events data
-		static constexpr auto EventCategory = TEXT("nDisplay");
-		static constexpr auto EventType     = TEXT("control");
+		static constexpr const TCHAR* EventCategory = TEXT("nDisplay");
+		static constexpr const TCHAR* EventType     = TEXT("control");
 
 		// Specific events
-		static constexpr auto EvtQuitName   = TEXT("quit");
+		static constexpr const TCHAR* EvtQuitName   = TEXT("quit");
 	}
 };
diff --git a/Engine/Plugins/Runtime/nDisplay/Source/DisplayClusterConfiguration/Private/Formats/JSON427/DisplayClusterConfigurationJsonHelpers_427.h b/Engine/Plugins/Runtime/nDisplay/Source/DisplayClusterConfiguration/Private/Formats/JSON427/DisplayClusterConfigurationJsonHelpers_427.h
--- a/Engine/Plugins/Runtime/nDisplay/Source/DisplayClusterConfiguration/Private/Formats/JSON427/DisplayClusterConfigurationJsonHelpers_427.h
+++ b/Engine/Plugins/Runtime/nDisplay/Source/DisplayClusterConfiguration/Private/Formats/JSON427/DisplayClusterConfigurationJsonHelpers_427.h
@@ -18,9 +18,9 @@ namespace JSON427
 			{
 				namespace overscan
 				{
-					static constexpr auto OverscanNone = TEXT("none");
-					static constexpr auto OverscanPixels = TEXT("pixels");
-					static constexpr auto OverscanPercent = TEXT("percent");
+					static constexpr const TCHAR* OverscanNone = TEXT("none");
+					static constexpr const TCHAR* OverscanPixels = TEXT("pixels");
+					static constexpr const TCHAR* OverscanPercent = TEXT("percent");
 				}
 			}
 		}
diff --git a/Engine/Plugins/Runtime/nDisplay/Source/DisplayClusterConfiguration/Private/Formats/JSON500/DisplayClusterConfigurationJsonHelpers_500.h b/Engine/Plugins/Runtime/nDisplay/Source/DisplayClusterConfiguration/Private/Formats/JSON500/DisplayClusterConfigurationJsonHelpers_500.h
--- a/Engine/Plugins/Runtime/nDisplay/Source/DisplayClusterConfiguration/Private/Formats/JSON500/DisplayClusterConfigurationJsonHelpers_500.h
+++ b/Engine/Plugins/Runtime/nDisplay/Source/DisplayClusterConfiguration/Private/Formats/JSON500/DisplayClusterConfigurationJsonHelpers_500.h
@@ -18,9 +18,9 @@ namespace JSON500
 			{
 				namespace overscan
 				{
-					static constexpr auto OverscanNone = TEXT("none");
-					static constexpr auto OverscanPixels = TEXT("pixels");
-					static constexpr auto OverscanPercent = TEXT("percent");
+					static constexpr const TCHAR* OverscanNone = TEXT("none");
+					static constexpr const TCHAR* OverscanPixels = TEXT("pixels");
+					static constexpr const TCHAR* OverscanPercent = TEXT("percent");
 				}
 			}
 		}
diff --git a/Engine/Plugins/Runtime/nDisplay/Source/DisplayClusterConfiguration/Public/DisplayClusterConfigurationStrings.h b/Engine/Plugins/Runtime/nDisplay/Source/DisplayClusterConfiguration/Public/DisplayClusterConfigurationStrings.h
--- a/Engine/Plugins/Runtime/nDisplay/Source/DisplayClusterConfiguration/Public/DisplayClusterConfigurationStrings.h
+++ b/Engine/Plugins/Runtime/nDisplay/Source/DisplayClusterConfiguration/Public/DisplayClusterConfigurationStrings.h
@@ -9,43 +9,43 @@ namespace DisplayClusterConfigurationStrings
 	{
 		namespace preview
 		{
-			static constexpr auto PreviewNodeAll = TEXT("All");
-			static constexpr auto PreviewNodeNone = TEXT("None");
+			static constexpr const TCHAR* PreviewNodeAll = TEXT("All");
+			static constexpr const TCHAR* PreviewNodeNone = TEXT("None");
 		}
 	}
 
 	// Property Categories
 	namespace categories
 	{
-		static constexpr auto DefaultCategory              = TEXT("NDisplay");
+		static constexpr const TCHAR* DefaultCategory              = TEXT("NDisplay");
 
-		static constexpr auto ClusterCategory              = TEXT("NDisplay Cluster");
-		static constexpr auto ClusterConfigurationCategory = TEXT("NDisplay Cluster Configuration");
-		static constexpr auto ClusterPostprocessCategory   = TEXT("Post Process");
-		static constexpr auto ColorGradingCategory         = TEXT("Color Grading");
-		static constexpr auto CameraColorGradingCategory   = TEXT("Inner Frustum Color Grading");
-		static constexpr auto ChromaKeyCategory            = TEXT("Chromakey");
-		static constexpr auto LightcardCategory            = TEXT("Light Cards");
-		static constexpr auto OCIOCategory                 = TEXT("OCIO");
-		static constexpr auto OverrideCategory             = TEXT("Texture Replacement");
-		static constexpr auto ViewportsCategory            = TEXT("Viewports");
-		static constexpr auto ICVFXCategory                = TEXT("In Camera VFX");
-		static constexpr auto ConfigurationCategory        = TEXT("Configuration");
-		static constexpr auto PreviewCategory              = TEXT("Editor Preview");
-		static constexpr auto AdvancedCategory             = TEXT("Advanced");
-		static constexpr auto TextureShareCategory         = TEXT("Texture Share");
+		static constexpr const TCHAR* ClusterCategory              = TEXT("NDisplay Cluster");
+		static constexpr const TCHAR* ClusterConfigurationCategory = TEXT("NDisplay Cluster Configuration");
+		static constexpr const TCHAR* ClusterPostprocessCategory   = TEXT("Post Process");
+		static constexpr const TCHAR* ColorGradingCategory         = TEXT("Color Grading");
+		static constexpr const TCHAR* CameraColorGradingCategory   = TEXT("Inner Frustum Color Grading");
+		static constexpr const TCHAR* ChromaKeyCategory            = TEXT("Chromakey");
+		static constexpr const TCHAR* LightcardCategory            = TEXT("Light Cards");
+		static constexpr const TCHAR* OCIOCategory                 = TEXT("OCIO");
+		static constexpr const TCHAR* OverrideCategory             = TEXT("Texture Replacement");
+		static constexpr const TCHAR* ViewportsCategory            = TEXT("Viewports");
+		static constexpr const TCHAR* ICVFXCategory                = TEXT("In Camera VFX");
+		static constexpr const TCHAR* ConfigurationCategory        = TEXT("Configuration");
+		static constexpr const TCHAR* PreviewCategory              = TEXT("Editor Preview");
+		static constexpr const TCHAR* AdvancedCategory             = TEXT("Advanced");
+		static constexpr const TCHAR* TextureShareCategory         = TEXT("Texture Share");
 	}
 
 	// Command line arguments
 	namespace args
 	{
-		static constexpr auto Gpu    = TEXT("dc_gpu");
+		static constexpr const TCHAR* Gpu    = TEXT("dc_gpu");
 	}
 
 	// Config file extensions
 	namespace file
 	{
-		static constexpr auto FileExtJson = TEXT("ndisplay");
+		static constexpr const TCHAR* FileExtJson = TEXT("ndisplay");
 	}
 
 	namespace config
@@ -54,36 +54,36 @@ namespace DisplayClusterConfigurationStrings
 		{
 			namespace ports
 			{
-				static constexpr auto PortClusterSync         = TEXT("ClusterSync");
-				static constexpr auto PortClusterEventsJson   = TEXT("ClusterEventsJson");
-				static constexpr auto PortClusterEventsBinary = TEXT("ClusterEventsBinary");
+				static constexpr const TCHAR* PortClusterSync         = TEXT("ClusterSync");
+				static constexpr const TCHAR* PortClusterEventsJson   = TEXT("ClusterEventsJson");
+				static constexpr const TCHAR* PortClusterEventsBinary = TEXT("ClusterEventsBinary");
 			}
 
 			namespace network
 			{
-				static constexpr auto NetConnectRetriesAmount     = TEXT("ConnectRetriesAmount");
-				static constexpr auto NetConnectRetryDelay        = TEXT("ConnectRetryDelay");
-				static constexpr auto NetGameStartBarrierTimeout  = TEXT("GameStartBarrierTimeout");
-				static constexpr auto NetFrameStartBarrierTimeout = TEXT("FrameStartBarrierTimeout");
-				static constexpr auto NetFrameEndBarrierTimeout   = TEXT("FrameEndBarrierTimeout");
-				static constexpr auto NetRenderSyncBarrierTimeout = TEXT("RenderSyncBarrierTimeout");
+				static constexpr const TCHAR* NetConnectRetriesAmount     = TEXT("ConnectRetriesAmount");
+				static constexpr const TCHAR* NetConnectRetryDelay        = TEXT("ConnectRetryDelay");
+				static constexpr const TCHAR* NetGameStartBarrierTimeout  = TEXT("GameStartBarrierTimeout");
+				static constexpr const TCHAR* NetFrameStartBarrierTimeout = TEXT("FrameStartBarrierTimeout");
+				static constexpr const TCHAR* NetFrameEndBarrierTimeout   = TEXT("FrameEndBarrierTimeout");
+				static constexpr const TCHAR* NetRenderSyncBarrierTimeout = TEXT("RenderSyncBarrierTimeout");
 			}
 
 			namespace input_sync
 			{
-				static constexpr auto InputSyncPolicyNone             = TEXT("None");
-				static constexpr auto InputSyncPolicyReplicatePrimary = TEXT("ReplicatePrimary");
+				static constexpr const TCHAR* InputSyncPolicyNone             = TEXT("None");
+				static constexpr const TCHAR* InputSyncPolicyReplicatePrimary = TEXT("ReplicatePrimary");
 			}
 
 			namespace render_sync
 			{
-				static constexpr auto None              = TEXT("none");
-				static constexpr auto Ethernet          = TEXT("ethernet");
-				static constexpr auto EthernetBarrier   = TEXT("ethernet_barrier");
+				static constexpr const TCHAR* None              = TEXT("none");
+				static constexpr const TCHAR* Ethernet          = TEXT("ethernet");
+				static constexpr const TCHAR* EthernetBarrier   = TEXT("ethernet_barrier");
 
-				static constexpr auto Nvidia            = TEXT("nvidia");
-				static constexpr auto NvidiaSwapBarrier = TEXT("swap_barrier");
-				static constexpr auto NvidiaSwapGroup   = TEXT("swap_group");
+				static constexpr const TCHAR* Nvidia            = TEXT("nvidia");
+				static constexpr const TCHAR* NvidiaSwapBarrier = TEXT("swap_barrier");
+				static constexpr const TCHAR* NvidiaSwapGroup   = TEXT("swap_group");
 			}
 		}
 
@@ -91,9 +91,9 @@ namespace DisplayClusterConfigurationStrings
 		{
 			namespace camera
 			{
-				static constexpr auto CameraStereoOffsetNone  = TEXT("none");
-				static constexpr auto CameraStereoOffsetLeft  = TEXT("left");
-				static constexpr auto CameraStereoOffsetRight = TEXT("right");
+				static constexpr const TCHAR* CameraStereoOffsetNone  = TEXT("none");
+				static constexpr const TCHAR* CameraStereoOffsetLeft  = TEXT("left");
+				static constexpr const TCHAR* CameraStereoOffsetRight = TEXT("right");
 			}
 		}
 	}
diff --git a/Engine/Plugins/Runtime/nDisplay/Source/DisplayClusterConfiguration/Public/DisplayClusterConfigurationVersion.h b/Engine/Plugins/Runtime/nDisplay/Source/DisplayClusterConfiguration/Public/DisplayClusterConfigurationVersion.h
--- a/Engine/Plugins/Runtime/nDisplay/Source/DisplayClusterConfiguration/Public/DisplayClusterConfigurationVersion.h
+++ b/Engine/Plugins/Runtime/nDisplay/Source/DisplayClusterConfiguration/Public/DisplayClusterConfigurationVersion.h
@@ -12,7 +12,7 @@ enum class EDisplayClusterConfigurationVersion : uint8
 
 namespace DisplayClusterConfiguration
 {
-	static constexpr auto GetCurrentConfigurationSchemeMarker()
+	static constexpr const TCHAR* GetCurrentConfigurationSchemeMarker()
 	{
 		return TEXT("5.00");
 	}
diff --git a/Engine/Plugins/Runtime/nDisplay/Source/DisplayClusterConfiguration/Public/IDisplayClusterConfiguration.h b/Engine/Plugins/Runtime/nDisplay/Source/DisplayClusterConfiguration/Public/IDisplayClusterConfiguration.h
--- a/Engine/Plugins/Runtime/nDisplay/Source/DisplayClusterConfiguration/Public/IDisplayClusterConfiguration.h
+++ b/Engine/Plugins/Runtime/nDisplay/Source/DisplayClusterConfiguration/Public/IDisplayClusterConfiguration.h
@@ -17,7 +17,7 @@ class UDisplayClusterConfigurationData;
 class IDisplayClusterConfiguration : public IModuleInterface
 {
 public:
-	static constexpr auto ModuleName = TEXT("DisplayClusterConfiguration");
+	static constexpr const TCHAR* ModuleName = TEXT("DisplayClusterConfiguration");
 
 public:
 	virtual ~IDisplayClusterConfiguration() = default;
diff --git a/Engine/Plugins/Runtime/nDisplay/Source/DisplayClusterConfigurator/Public/Interfaces/IDisplayClusterConfigurator.h b/Engine/Plugins/Runtime/nDisplay/Source/DisplayClusterConfigurator/Public/Interfaces/IDisplayClusterConfigurator.h
--- a/Engine/Plugins/Runtime/nDisplay/Source/DisplayClusterConfigurator/Public/Interfaces/IDisplayClusterConfigurator.h
+++ b/Engine/Plugins/Runtime/nDisplay/Source/DisplayClusterConfigurator/Public/Interfaces/IDisplayClusterConfigurator.h
@@ -18,7 +18,7 @@ class IDisplayClusterConfigurator :
 	public IModuleInterface
 {
 public:
-	static constexpr auto ModuleName = TEXT("DisplayClusterConfigurator");
+	static constexpr const TCHAR* ModuleName = TEXT("DisplayClusterConfigurator");
 
 public:
 	virtual ~IDisplayClusterConfigurator() = default;
diff --git a/Engine/Plugins/Runtime/nDisplay/Source/DisplayClusterLightCardEditor/Public/IDisplayClusterLightCardEditor.h b/Engine/Plugins/Runtime/nDisplay/Source/DisplayClusterLightCardEditor/Public/IDisplayClusterLightCardEditor.h
--- a/Engine/Plugins/Runtime/nDisplay/Source/DisplayClusterLightCardEditor/Public/IDisplayClusterLightCardEditor.h
+++ b/Engine/Plugins/Runtime/nDisplay/Source/DisplayClusterLightCardEditor/Public/IDisplayClusterLightCardEditor.h
@@ -12,7 +12,7 @@
 class IDisplayClusterLightCardEditor : public IModuleInterface
 {
 public:
-	static constexpr auto ModuleName = TEXT("DisplayClusterLightCardEditor");
+	static constexpr const TCHAR* ModuleName = TEXT("DisplayClusterLightCardEditor");
 
 public:
 	virtual ~IDisplayClusterLightCardEditor() = default;
diff --git a/Engine/Plugins/Runtime/nDisplay/Source/DisplayClusterOperator/Public/IDisplayClusterOperator.h b/Engine/Plugins/Runtime/nDisplay/Source/DisplayClusterOperator/Public/IDisplayClusterOperator.h
--- a/Engine/Plugins/Runtime/nDisplay/Source/DisplayClusterOperator/Public/IDisplayClusterOperator.h
+++ b/Engine/Plugins/Runtime/nDisplay/Source/DisplayClusterOperator/Public/IDisplayClusterOperator.h
@@ -20,7 +20,7 @@ DECLARE_EVENT_OneParam(IDisplayClusterOperator, FOnDetailObjectsChanged, const T
 class IDisplayClusterOperator : public IModuleInterface
 {
 public:
-	static constexpr auto ModuleName = TEXT("DisplayClusterOperator");
+	static constexpr const TCHAR* ModuleName = TEXT("DisplayClusterOperator");
 
 public:
 	virtual ~IDisplayClusterOperator() = default;
diff --git a/Engine/Plugins/Runtime/nDisplay/Source/DisplayClusterPostprocess/Public/DisplayClusterPostprocessStrings.h b/Engine/Plugins/Runtime/nDisplay/Source/DisplayClusterPostprocess/Public/DisplayClusterPostprocessStrings.h
--- a/Engine/Plugins/Runtime/nDisplay/Source/DisplayClusterPostprocess/Public/DisplayClusterPostprocessStrings.h
+++ b/Engine/Plugins/Runtime/nDisplay/Source/DisplayClusterPostprocess/Public/DisplayClusterPostprocessStrings.h
@@ -9,24 +9,24 @@ namespace DisplayClusterPostprocessStrings
 {
 	namespace texture_share
 	{
-		static constexpr auto ShareName = TEXT("share");
+		static constexpr const TCHAR* ShareName = TEXT("share");
 
-		static constexpr auto DestinationViewports = TEXT("destination");
-		static constexpr auto SourceShare          = TEXT("source");
-		static constexpr auto ShareViewports       = TEXT("send");
+		static constexpr const TCHAR* DestinationViewports = TEXT("destination");
+		static constexpr const TCHAR* SourceShare          = TEXT("source");
+		static constexpr const TCHAR* ShareViewports       = TEXT("send");
 
-		static constexpr auto BackbufferTextureId  = TEXT("BackBuffer");
+		static constexpr const TCHAR* BackbufferTextureId  = TEXT("BackBuffer");
 
 		namespace debug
 		{
-			static constexpr auto DuplicateTexture = TEXT("duplicate_texture");
-			static constexpr auto RepeatCopy       = TEXT("repeat_copy");
+			static constexpr const TCHAR* DuplicateTexture = TEXT("duplicate_texture");
+			static constexpr const TCHAR* RepeatCopy       = TEXT("repeat_copy");
 		}
 	}
 
 	namespace postprocess
 	{
-		static constexpr auto TextureShare  = TEXT("TextureShare");
-		static constexpr auto D3D12CrossGPU = TEXT("D3D12CrossGPU");
+		static constexpr const TCHAR* TextureShare  = TEXT("TextureShare");
+		static constexpr const TCHAR* D3D12CrossGPU = TEXT("D3D12CrossGPU");
 	}
 };
diff --git a/Engine/Plugins/Runtime/nDisplay/Source/DisplayClusterProjection/Public/DisplayClusterProjectionStrings.h b/Engine/Plugins/Runtime/nDisplay/Source/DisplayClusterProjection/Public/DisplayClusterProjectionStrings.h
--- a/Engine/Plugins/Runtime/nDisplay/Source/DisplayClusterProjection/Public/DisplayClusterProjectionStrings.h
+++ b/Engine/Plugins/Runtime/nDisplay/Source/DisplayClusterProjection/Public/DisplayClusterProjectionStrings.h
@@ -11,128 +11,128 @@ namespace DisplayClusterProjectionStrings
 	{
 		namespace simple
 		{
-			static constexpr auto Screen = TEXT("screen");
+			static constexpr const TCHAR* Screen = TEXT("screen");
 		}
 
 		namespace easyblend
 		{
-			static constexpr auto File   = TEXT("file");
-			static constexpr auto Origin = TEXT("origin");
-			static constexpr auto Scale  = TEXT("scale");
+			static constexpr const TCHAR* File   = TEXT("file");
+			static constexpr const TCHAR* Origin = TEXT("origin");
+			static constexpr const TCHAR* Scale  = TEXT("scale");
 		}
 
 		namespace VIOSO
 		{
-			static constexpr auto Origin = TEXT("origin");
+			static constexpr const TCHAR* Origin = TEXT("origin");
 
-			static constexpr auto INIFile     = TEXT("inifile");
-			static constexpr auto ChannelName = TEXT("channel");
+			static constexpr const TCHAR* INIFile     = TEXT("inifile");
+			static constexpr const TCHAR* ChannelName = TEXT("channel");
 
-			static constexpr auto File         = TEXT("file");
-			static constexpr auto CalibIndex   = TEXT("index");
-			static constexpr auto CalibAdapter = TEXT("adapter");
+			static constexpr const TCHAR* File         = TEXT("file");
+			static constexpr const TCHAR* CalibIndex   = TEXT("index");
+			static constexpr const TCHAR* CalibAdapter = TEXT("adapter");
 
-			static constexpr auto Gamma = TEXT("gamma");
+			static constexpr const TCHAR* Gamma = TEXT("gamma");
 
-			static constexpr auto BaseMatrix = TEXT("base");
+			static constexpr const TCHAR* BaseMatrix = TEXT("base");
 		}
 
 		namespace manual
 		{
-			static constexpr auto Rendering    = TEXT("ManualRendering");
-			static constexpr auto Type         = TEXT("ManualFrustum");
-			static constexpr auto Rotation     = TEXT("rot");
-
-			static constexpr auto Matrix       = TEXT("matrix");
-			static constexpr auto MatrixLeft   = TEXT("matrix_left");
-			static constexpr auto MatrixRight  = TEXT("matrix_right");
-
-			static constexpr auto Frustum      = TEXT("frustum");
-			static constexpr auto FrustumLeft  = TEXT("frustum_left");
-			static constexpr auto FrustumRight = TEXT("frustum_right");
-
-			static constexpr auto AngleL      = TEXT("l");
-			static constexpr auto AngleR      = TEXT("r");
-			static constexpr auto AngleT      = TEXT("t");
-			static constexpr auto AngleB      = TEXT("b");
+			static constexpr const TCHAR* Rendering    = TEXT("ManualRendering");
+			static constexpr const TCHAR* Type         = TEXT("ManualFrustum");
+			static constexpr const TCHAR* Rotation     = TEXT("rot");
+
+			static constexpr const TCHAR* Matrix       = TEXT("matrix");
+			static constexpr const TCHAR* MatrixLeft   = TEXT("matrix_left");
+			static constexpr const TCHAR* MatrixRight  = TEXT("matrix_right");
+
+			static constexpr const TCHAR* Frustum      = TEXT("frustum");
+			static constexpr const TCHAR* FrustumLeft  = TEXT("frustum_left");
+			static constexpr const TCHAR* FrustumRight = TEXT("frustum_right");
+
+			static constexpr const TCHAR* AngleL      = TEXT("l");
+			static constexpr const TCHAR* AngleR      = TEXT("r");
+			static constexpr const TCHAR* AngleT      = TEXT("t");
+			static constexpr const TCHAR* AngleB      = TEXT("b");
 			namespace FrustumType
 			{
-				static constexpr auto Matrix = TEXT("Matrix");
-				static constexpr auto Angles = TEXT("Angles");
+				static constexpr const TCHAR* Matrix = TEXT("Matrix");
+				static constexpr const TCHAR* Angles = TEXT("Angles");
 			}
 			namespace RenderingType
 			{
-				static constexpr auto Mono = TEXT("Mono");
-				static constexpr auto Stereo = TEXT("Stereo");
-				static constexpr auto MonoStereo = TEXT("Mono & Stereo");
+				static constexpr const TCHAR* Mono = TEXT("Mono");
+				static constexpr const TCHAR* Stereo = TEXT("Stereo");
+				static constexpr const TCHAR* MonoStereo = TEXT("Mono & Stereo");
 			}
 		}
 
 		namespace camera
 		{
-			static constexpr auto Component = TEXT("camera_component");
-			static constexpr auto Native = TEXT("native");
+			static constexpr const TCHAR* Component = TEXT("camera_component");
+			static constexpr const TCHAR* Native = TEXT("native");
 		}
 		
 		namespace mesh
 		{
-			static constexpr auto Component        = TEXT("mesh_component");
+			static constexpr const TCHAR* Component        = TEXT("mesh_component");
 
-			static constexpr auto LODIndex         = TEXT("lod_index");
-			static constexpr auto SectionIndex     = TEXT("section_index");
+			static constexpr const TCHAR* LODIndex         = TEXT("lod_index");
+			static constexpr const TCHAR* SectionIndex     = TEXT("section_index");
 
-			static constexpr auto BaseUVIndex      = TEXT("base_uv_index");
-			static constexpr auto ChromakeyUVIndex = TEXT("chromakey_uv_index");
+			static constexpr const TCHAR* BaseUVIndex      = TEXT("base_uv_index");
+			static constexpr const TCHAR* ChromakeyUVIndex = TEXT("chromakey_uv_index");
 		}
 
 		namespace mpcdi
 		{
-			static constexpr auto File   = TEXT("file");
-			static constexpr auto Buffer = TEXT("buffer");
-			static constexpr auto Region = TEXT("region");
-			static constexpr auto Origin = TEXT("origin");
+			static constexpr const TCHAR* File   = TEXT("file");
+			static constexpr const TCHAR* Buffer = TEXT("buffer");
+			static constexpr const TCHAR* Region = TEXT("region");
+			static constexpr const TCHAR* Origin = TEXT("origin");
 
-			static constexpr auto MPCDIType = TEXT("mpcdi");
+			static constexpr const TCHAR* MPCDIType = TEXT("mpcdi");
 
-			static constexpr auto FilePFM       = TEXT("pfm");
-			static constexpr auto WorldScale    = TEXT("scale");
-			static constexpr auto UseUnrealAxis = TEXT("ue_space");
+			static constexpr const TCHAR* FilePFM       = TEXT("pfm");
+			static constexpr const TCHAR* WorldScale    = TEXT("scale");
+			static constexpr const TCHAR* UseUnrealAxis = TEXT("ue_space");
 
-			static constexpr auto FileAlpha  = TEXT("alpha");
-			static constexpr auto AlphaGamma = TEXT("alpha_gamma");
-			static constexpr auto FileBeta   = TEXT("beta");
+			static constexpr const TCHAR* FileAlpha  = TEXT("alpha");
+			static constexpr const TCHAR* AlphaGamma = TEXT("alpha_gamma");
+			static constexpr const TCHAR* FileBeta   = TEXT("beta");
 
-			static constexpr auto MPCDITypeKey = TEXT("MPCDIType");
-			static constexpr auto TypeMPCDI = TEXT("MPCDI");
-			static constexpr auto TypePFM = TEXT("Explicit PFM");
+			static constexpr const TCHAR* MPCDITypeKey = TEXT("MPCDIType");
+			static constexpr const TCHAR* TypeMPCDI = TEXT("MPCDI");
+			static constexpr const TCHAR* TypePFM = TEXT("Explicit PFM");
 
-			static constexpr auto EnablePreview = TEXT("EnablePreview");
+			static constexpr const TCHAR* EnablePreview = TEXT("EnablePreview");
 		}
 
 		namespace domeprojection
 		{
-			static constexpr auto File    = TEXT("file");
-			static constexpr auto Channel = TEXT("channel");
-			static constexpr auto Origin  = TEXT("origin");
+			static constexpr const TCHAR* File    = TEXT("file");
+			static constexpr const TCHAR* Channel = TEXT("channel");
+			static constexpr const TCHAR* Origin  = TEXT("origin");
 		}
 	}
 
 	namespace projection
 	{
-		static constexpr auto Camera         = TEXT("camera");
-		static constexpr auto Simple         = TEXT("simple");
-		static constexpr auto MPCDI          = TEXT("mpcdi");
-		static constexpr auto Mesh           = TEXT("mesh");
-		static constexpr auto EasyBlend      = TEXT("easyblend");
-		static constexpr auto VIOSO          = TEXT("vioso");
-		static constexpr auto Manual         = TEXT("manual");
-		static constexpr auto Domeprojection = TEXT("domeprojection");
-		static constexpr auto Link           = TEXT("link");
+		static constexpr const TCHAR* Camera         = TEXT("camera");
+		static constexpr const TCHAR* Simple         = TEXT("simple");
+		static constexpr const TCHAR* MPCDI          = TEXT("mpcdi");
+		static constexpr const TCHAR* Mesh           = TEXT("mesh");
+		static constexpr const TCHAR* EasyBlend      = TEXT("easyblend");
+		static constexpr const TCHAR* VIOSO          = TEXT("vioso");
+		static constexpr const TCHAR* Manual         = TEXT("manual");
+		static constexpr const TCHAR* Domeprojection = TEXT("domeprojection");
+		static constexpr const TCHAR* Link           = TEXT("link");
 	}
 
 	namespace rhi
 	{
-		static constexpr auto D3D11  = TEXT("D3D11");
-		static constexpr auto D3D12  = TEXT("D3D12");
+		static constexpr const TCHAR* D3D11  = TEXT("D3D11");
+		static constexpr const TCHAR* D3D12  = TEXT("D3D12");
 	}
 };
\ No newline at end of file
diff --git a/Engine/Plugins/Runtime/nDisplay/Source/DisplayClusterProjection/Public/IDisplayClusterProjection.h b/Engine/Plugins/Runtime/nDisplay/Source/DisplayClusterProjection/Public/IDisplayClusterProjection.h
--- a/Engine/Plugins/Runtime/nDisplay/Source/DisplayClusterProjection/Public/IDisplayClusterProjection.h
+++ b/Engine/Plugins/Runtime/nDisplay/Source/DisplayClusterProjection/Public/IDisplayClusterProjection.h
@@ -14,7 +14,7 @@ class IDisplayClusterProjectionPolicyFactory;
 class IDisplayClusterProjection : public IModuleInterface
 {
 public:
-	static constexpr auto ModuleName = TEXT("DisplayClusterProjection");
+	static constexpr const TCHAR* ModuleName = TEXT("DisplayClusterProjection");
 
 public:
 	virtual ~IDisplayClusterProjection() = default;
diff --git a/Engine/Plugins/Runtime/nDisplay/Source/DisplayClusterShaders/Public/IDisplayClusterShaders.h b/Engine/Plugins/Runtime/nDisplay/Source/DisplayClusterShaders/Public/IDisplayClusterShaders.h
--- a/Engine/Plugins/Runtime/nDisplay/Source/DisplayClusterShaders/Public/IDisplayClusterShaders.h
+++ b/Engine/Plugins/Runtime/nDisplay/Source/DisplayClusterShaders/Public/IDisplayClusterShaders.h
@@ -22,7 +22,7 @@ class IDisplayClusterWarpBlendManager;
 class IDisplayClusterShaders : public IModuleInterface
 {
 public:
-	static constexpr auto ModuleName = TEXT("DisplayClusterShaders");
+	static constexpr const TCHAR* ModuleName = TEXT("DisplayClusterShaders");
 
 public:
 	virtual ~IDisplayClusterShaders() = default;
