ndisplay: fixed color grading behaviour

diff --git a/Engine/Plugins/Runtime/nDisplay/Source/DisplayCluster/Private/Render/Viewport/DisplayClusterViewportStereoscopicPass.h b/Engine/Plugins/Runtime/nDisplay/Source/DisplayCluster/Private/Render/Viewport/DisplayClusterViewportStereoscopicPass.h
--- a/Engine/Plugins/Runtime/nDisplay/Source/DisplayCluster/Private/Render/Viewport/DisplayClusterViewportStereoscopicPass.h
+++ b/Engine/Plugins/Runtime/nDisplay/Source/DisplayCluster/Private/Render/Viewport/DisplayClusterViewportStereoscopicPass.h
@@ -9,24 +9,15 @@ struct FDisplayClusterViewportStereoscopicPass
 {
 	inline static EStereoscopicPass EncodeStereoscopicPass(const uint32 ContextNum, const uint32 ContextAmmount)
 	{
-		EStereoscopicPass EncodedPass = EStereoscopicPass::eSSP_PRIMARY;
-
 		// Monoscopic rendering
 		if (ContextAmmount == 1)
-		{
-			EncodedPass = EStereoscopicPass::eSSP_FULL;
-		}
-		// Stereoscopic rendering
-		else
 		{
 			switch (ContextNum)
 			{
 			case 0:
-				EncodedPass = EStereoscopicPass::eSSP_PRIMARY;
-				break;
+				return EStereoscopicPass::eSSP_PRIMARY;
 			case 1:
-				EncodedPass = EStereoscopicPass::eSSP_SECONDARY;
-				break;
+				return EStereoscopicPass::eSSP_SECONDARY;
 			default:
 				// now stereo only with 2 context
 				check(false);
@@ -34,7 +25,7 @@ struct FDisplayClusterViewportStereoscopicPass
 			}
 		}
 
-		return EncodedPass;
+		return EStereoscopicPass::eSSP_PRIMARY;
 	}
 };
 
