nDisplay: added overscan estimation and fixed naming

diff --git a/Engine/Plugins/Runtime/nDisplay/Source/DisplayCluster/Private/Game/EngineClasses/Scene/DisplayClusterICVFXCameraComponent.cpp b/Engine/Plugins/Runtime/nDisplay/Source/DisplayCluster/Private/Game/EngineClasses/Scene/DisplayClusterICVFXCameraComponent.cpp
--- a/Engine/Plugins/Runtime/nDisplay/Source/DisplayCluster/Private/Game/EngineClasses/Scene/DisplayClusterICVFXCameraComponent.cpp
+++ b/Engine/Plugins/Runtime/nDisplay/Source/DisplayCluster/Private/Game/EngineClasses/Scene/DisplayClusterICVFXCameraComponent.cpp
@@ -4,6 +4,7 @@
 #include "Components/DrawFrustumComponent.h"
 
 #include "Render/Viewport/Containers/DisplayClusterViewport_CameraMotionBlur.h"
+#include "Render/Viewport/Containers/ImplDisplayClusterViewport_CustomFrustum.h"
 #include "Components/DisplayClusterCameraComponent.h"
 #include "DisplayClusterRootActor.h"
 
@@ -45,6 +46,74 @@ TSharedPtr<SWidget> UDisplayClusterICVFXCameraComponent::GetCustomEditorPreviewW
 }
 #endif
 
+void UDisplayClusterICVFXCameraComponent::TickComponent(float DeltaTime, ELevelTick TickType, FActorComponentTickFunction* ThisTickFunction)
+{
+	Super::TickComponent(DeltaTime, TickType, ThisTickFunction);
+
+	UpdateOverscanEstimatedFrameSize();
+}
+
+void UDisplayClusterICVFXCameraComponent::UpdateOverscanEstimatedFrameSize()
+{
+	const ADisplayClusterRootActor* RootActor = Cast<ADisplayClusterRootActor>(GetOwner());
+	if (RootActor == nullptr)
+	{
+		return;
+	}
+
+	float InnerFrustumResolutionWidth = RootActor->GetStageSettings().DefaultFrameSize.Width * CameraSettings.BufferRatio;
+	float InnerFrustumResolutionHeight = RootActor->GetStageSettings().DefaultFrameSize.Height * CameraSettings.BufferRatio;
+
+	float EstimatedOverscanResolutionWidth = InnerFrustumResolutionWidth;
+	float EstimatedOverscanResolutionHeight = InnerFrustumResolutionHeight;
+
+	//calculate estimate
+	{
+		EstimatedOverscanResolutionWidth = InnerFrustumResolutionWidth * CameraSettings.CustomFrustum.FieldOfViewMultiplier;
+		EstimatedOverscanResolutionHeight = InnerFrustumResolutionHeight * CameraSettings.CustomFrustum.FieldOfViewMultiplier;
+
+		FDisplayClusterViewport_CustomFrustumSettings CustomFrustumSettings;
+
+		if (CameraSettings.CustomFrustum.Mode == EDisplayClusterConfigurationViewportCustomFrustumMode::Percent)
+		{
+			const float ConvertToPercent = 0.01;
+			CustomFrustumSettings.CustomFrustumPercent.Left = FDisplayClusterViewport_OverscanSettings::ClampPercent(CameraSettings.CustomFrustum.Left * ConvertToPercent);
+			CustomFrustumSettings.CustomFrustumPercent.Right = FDisplayClusterViewport_OverscanSettings::ClampPercent(CameraSettings.CustomFrustum.Right * ConvertToPercent);
+			CustomFrustumSettings.CustomFrustumPercent.Top = FDisplayClusterViewport_OverscanSettings::ClampPercent(CameraSettings.CustomFrustum.Top * ConvertToPercent);
+			CustomFrustumSettings.CustomFrustumPercent.Bottom = FDisplayClusterViewport_OverscanSettings::ClampPercent(CameraSettings.CustomFrustum.Bottom * ConvertToPercent);
+		}
+		else if (CameraSettings.CustomFrustum.Mode == EDisplayClusterConfigurationViewportCustomFrustumMode::Pixels)
+		{
+			CustomFrustumSettings.CustomFrustumPercent.Left = FDisplayClusterViewport_OverscanSettings::ClampPercent(CameraSettings.CustomFrustum.Left / EstimatedOverscanResolutionWidth);
+			CustomFrustumSettings.CustomFrustumPercent.Right = FDisplayClusterViewport_OverscanSettings::ClampPercent(CameraSettings.CustomFrustum.Right / EstimatedOverscanResolutionWidth);
+			CustomFrustumSettings.CustomFrustumPercent.Top = FDisplayClusterViewport_OverscanSettings::ClampPercent(CameraSettings.CustomFrustum.Top / EstimatedOverscanResolutionHeight);
+			CustomFrustumSettings.CustomFrustumPercent.Bottom = FDisplayClusterViewport_OverscanSettings::ClampPercent(CameraSettings.CustomFrustum.Bottom / EstimatedOverscanResolutionHeight);
+		}
+
+		// Calc pixels from percent
+		CustomFrustumSettings.CustomFrustumPixels.Left = FMath::RoundToInt(EstimatedOverscanResolutionWidth * CustomFrustumSettings.CustomFrustumPercent.Left);
+		CustomFrustumSettings.CustomFrustumPixels.Right = FMath::RoundToInt(EstimatedOverscanResolutionWidth * CustomFrustumSettings.CustomFrustumPercent.Right);
+		CustomFrustumSettings.CustomFrustumPixels.Top = FMath::RoundToInt(EstimatedOverscanResolutionHeight * CustomFrustumSettings.CustomFrustumPercent.Top);
+		CustomFrustumSettings.CustomFrustumPixels.Bottom = FMath::RoundToInt(EstimatedOverscanResolutionHeight * CustomFrustumSettings.CustomFrustumPercent.Bottom);
+
+		const FIntPoint AdjustmentSize = CustomFrustumSettings.CustomFrustumPixels.Size();
+		EstimatedOverscanResolutionWidth += AdjustmentSize.X;
+		EstimatedOverscanResolutionHeight += AdjustmentSize.Y;
+	}
+
+	if (CameraSettings.CustomFrustum.bAdaptResolution)
+	{
+		InnerFrustumResolutionWidth = EstimatedOverscanResolutionWidth;
+		InnerFrustumResolutionHeight = EstimatedOverscanResolutionHeight;
+	}
+
+	CameraSettings.CustomFrustum.InnerFrustumResolution = FIntPoint(InnerFrustumResolutionWidth, InnerFrustumResolutionHeight);
+	CameraSettings.CustomFrustum.EstimatedOverscanResolution = FIntPoint(EstimatedOverscanResolutionWidth, EstimatedOverscanResolutionHeight);
+
+	CameraSettings.CustomFrustum.OverscanPixelsIncrease = ((float)(EstimatedOverscanResolutionWidth * EstimatedOverscanResolutionHeight) / (float)(InnerFrustumResolutionWidth * InnerFrustumResolutionHeight));
+}
+
+
 FDisplayClusterViewport_CameraMotionBlur UDisplayClusterICVFXCameraComponent::GetMotionBlurParameters()
 {
 	FDisplayClusterViewport_CameraMotionBlur OutParameters;
diff --git a/Engine/Plugins/Runtime/nDisplay/Source/DisplayCluster/Public/Components/DisplayClusterICVFXCameraComponent.h b/Engine/Plugins/Runtime/nDisplay/Source/DisplayCluster/Public/Components/DisplayClusterICVFXCameraComponent.h
--- a/Engine/Plugins/Runtime/nDisplay/Source/DisplayCluster/Public/Components/DisplayClusterICVFXCameraComponent.h
+++ b/Engine/Plugins/Runtime/nDisplay/Source/DisplayCluster/Public/Components/DisplayClusterICVFXCameraComponent.h
@@ -40,6 +40,8 @@ public:
 	virtual TSharedPtr<SWidget> GetCustomEditorPreviewWidget() override;
 #endif
 
+	virtual void TickComponent(float DeltaTime, ELevelTick TickType, FActorComponentTickFunction* ThisTickFunction) override;
+
 public:
 	FDisplayClusterViewport_CameraMotionBlur GetMotionBlurParameters();
 
@@ -62,7 +64,10 @@ public:
 
 	// UActorComponent interface
 	virtual void OnRegister() override;
-	
+
+private:
+	void UpdateOverscanEstimatedFrameSize();
+
 //////////////////////////////////////////////////////////////////////////////////////////////
 // Details Panel Property Referencers
 //////////////////////////////////////////////////////////////////////////////////////////////
diff --git a/Engine/Plugins/Runtime/nDisplay/Source/DisplayClusterConfiguration/Public/DisplayClusterConfigurationTypes_Enums.h b/Engine/Plugins/Runtime/nDisplay/Source/DisplayClusterConfiguration/Public/DisplayClusterConfigurationTypes_Enums.h
--- a/Engine/Plugins/Runtime/nDisplay/Source/DisplayClusterConfiguration/Public/DisplayClusterConfigurationTypes_Enums.h
+++ b/Engine/Plugins/Runtime/nDisplay/Source/DisplayClusterConfiguration/Public/DisplayClusterConfigurationTypes_Enums.h
@@ -171,8 +171,8 @@ enum class EDisplayClusterConfigurationRenderMGPUMode : uint8
 UENUM()
 enum class EDisplayClusterConfigurationViewportOverscanMode : uint8
 {
-	Pixels   UMETA(DisplayName = "Pixels values"),
-	Percent  UMETA(DisplayName = "Percent values")
+	Pixels   UMETA(DisplayName = "Pixels"),
+	Percent  UMETA(DisplayName = "Percent")
 };
 
 UENUM()
@@ -186,6 +186,6 @@ enum class EDisplayClusterConfigurationRenderMode : uint8
 UENUM()
 enum class EDisplayClusterConfigurationViewportCustomFrustumMode : uint8
 {
-	Percent  UMETA(DisplayName = "Percent values"),
-	Pixels   UMETA(DisplayName = "Pixels values")
+	Percent  UMETA(DisplayName = "Percentage"),
+	Pixels   UMETA(DisplayName = "Pixels")
 };
diff --git a/Engine/Plugins/Runtime/nDisplay/Source/DisplayClusterConfiguration/Public/DisplayClusterConfigurationTypes_ICVFX.h b/Engine/Plugins/Runtime/nDisplay/Source/DisplayClusterConfiguration/Public/DisplayClusterConfigurationTypes_ICVFX.h
--- a/Engine/Plugins/Runtime/nDisplay/Source/DisplayClusterConfiguration/Public/DisplayClusterConfigurationTypes_ICVFX.h
+++ b/Engine/Plugins/Runtime/nDisplay/Source/DisplayClusterConfiguration/Public/DisplayClusterConfigurationTypes_ICVFX.h
@@ -407,20 +407,29 @@ struct DISPLAYCLUSTERCONFIGURATION_API FDisplayClusterConfigurationICVFX_CameraC
 {
 	GENERATED_BODY()
 
-	/** Multiply the field of view for the ICVFX camera by this value.  This can increase the overall size of the inner frustum to help provide a buffer against latency when moving the camera. */
-	UPROPERTY(EditAnywhere, BlueprintReadWrite, Category = "In Camera VFX", meta = (ClampMin = "0.05", UIMin = "0.05", ClampMax = "5.0", UIMax = "5.0"))
-	float FieldOfViewMultiplier = 1.0f;
+	UPROPERTY(VisibleAnywhere, BlueprintReadOnly, Category = NDisplay)
+	FIntPoint EstimatedOverscanResolution;
+	
+	UPROPERTY(VisibleAnywhere, BlueprintReadOnly, Category = NDisplay)
+	FIntPoint InnerFrustumResolution;
+
+	UPROPERTY(BlueprintReadOnly, Category = NDisplay, meta = (HideInDetailPanel))
+	float OverscanPixelsIncrease;
 
 	/** Enable Custom Frustum Frustum. */
-	UPROPERTY(EditAnywhere, BlueprintReadWrite, Category = NDisplay, meta = (DisplayName = "Enable Custom Frustum"))
+	UPROPERTY(EditAnywhere, BlueprintReadWrite, Category = NDisplay, meta = (DisplayName = "Enable Inner Frustum Overscan"))
 	bool bEnable = false;
 
 	/** Enable Custom Frustum Frustum. */
-	UPROPERTY(EditAnywhere, BlueprintReadWrite, Category = NDisplay, meta = (DisplayName = "Adapt resolution", EditCondition = "bEnable"))
+	UPROPERTY(EditAnywhere, BlueprintReadWrite, Category = NDisplay, meta = (DisplayName = "Adapt Resolution", EditCondition = "bEnable"))
 	bool bAdaptResolution = false;
 
+	/** Multiply the field of view for the ICVFX camera by this value.  This can increase the overall size of the inner frustum to help provide a buffer against latency when moving the camera. */
+	UPROPERTY(EditAnywhere, BlueprintReadWrite, Category = "In Camera VFX", meta = (DisplayName = "Overscan Multiplier", ClampMin = "0.05", UIMin = "0.05", ClampMax = "5.0", UIMax = "5.0", EditCondition = "bEnable"))
+	float FieldOfViewMultiplier = 1.0f;
+
 	/** Enable/disable inner camera custom frustum and specify units as percent or pixel values. */
-	UPROPERTY(EditAnywhere, BlueprintReadWrite, Category = NDisplay, meta = (DisplayName = "Mode", EditCondition = "bEnable"))
+	UPROPERTY(EditAnywhere, BlueprintReadWrite, Category = NDisplay, meta = (DisplayName = "Overscan Units", EditCondition = "bEnable"))
 	EDisplayClusterConfigurationViewportCustomFrustumMode Mode = EDisplayClusterConfigurationViewportCustomFrustumMode::Percent;
 
 	/** Pixel/Percent value to alter the frustum to the left side */
@@ -461,7 +470,7 @@ public:
 	float BufferRatio = 1;
 
 	/** Render a larger or smaller inner frame */
-	UPROPERTY(EditAnywhere, BlueprintReadWrite, Category = "In Camera VFX", meta = (EditCondition = "bEnable"))
+	UPROPERTY(EditAnywhere, BlueprintReadWrite, Category = "In Camera VFX", meta = (DisplayName = "Inner Frustum Overscan", EditCondition = "bEnable"))
 	FDisplayClusterConfigurationICVFX_CameraCustomFrustum CustomFrustum;
 
 	/** Soften the edges of the inner frustum to help avoid hard lines in reflections seen by the live-action camera. */
