nDisplay: Fix CIS camera custom frustum struct initialization.

diff --git a/Engine/Plugins/Runtime/nDisplay/Source/DisplayClusterConfiguration/Private/DisplayClusterConfigurationTypes.cpp b/Engine/Plugins/Runtime/nDisplay/Source/DisplayClusterConfiguration/Private/DisplayClusterConfigurationTypes.cpp
--- a/Engine/Plugins/Runtime/nDisplay/Source/DisplayClusterConfiguration/Private/DisplayClusterConfigurationTypes.cpp
+++ b/Engine/Plugins/Runtime/nDisplay/Source/DisplayClusterConfiguration/Private/DisplayClusterConfigurationTypes.cpp
@@ -396,6 +396,11 @@ FDisplayClusterConfigurationOCIOProfile::FDisplayClusterConfigurationOCIOProfile
 	OCIOConfiguration.bIsEnabled = true;
 }
 
+FDisplayClusterConfigurationICVFX_CameraCustomFrustum::FDisplayClusterConfigurationICVFX_CameraCustomFrustum():
+	EstimatedOverscanResolution(ForceInitToZero), InnerFrustumResolution(ForceInitToZero), OverscanPixelsIncrease(0.f)
+{
+}
+
 FDisplayClusterConfigurationICVFX_CameraSettings::FDisplayClusterConfigurationICVFX_CameraSettings()
 {
 	AllNodesColorGrading.bEnableEntireClusterColorGrading = true;
diff --git a/Engine/Plugins/Runtime/nDisplay/Source/DisplayClusterConfiguration/Public/DisplayClusterConfigurationTypes_ICVFX.h b/Engine/Plugins/Runtime/nDisplay/Source/DisplayClusterConfiguration/Public/DisplayClusterConfigurationTypes_ICVFX.h
--- a/Engine/Plugins/Runtime/nDisplay/Source/DisplayClusterConfiguration/Public/DisplayClusterConfigurationTypes_ICVFX.h
+++ b/Engine/Plugins/Runtime/nDisplay/Source/DisplayClusterConfiguration/Public/DisplayClusterConfigurationTypes_ICVFX.h
@@ -407,6 +407,8 @@ struct DISPLAYCLUSTERCONFIGURATION_API FDisplayClusterConfigurationICVFX_CameraC
 {
 	GENERATED_BODY()
 
+	FDisplayClusterConfigurationICVFX_CameraCustomFrustum();
+	
 	UPROPERTY(VisibleAnywhere, BlueprintReadOnly, Category = NDisplay)
 	FIntPoint EstimatedOverscanResolution;
 	
