Added rotation and flip icons to the Editor's core starship style, and used then in the nDisplay plugin for rotation and flip operations

diff --git a/Engine/Plugins/Runtime/nDisplay/Source/DisplayClusterConfigurator/Private/DisplayClusterConfiguratorStyle.cpp b/Engine/Plugins/Runtime/nDisplay/Source/DisplayClusterConfigurator/Private/DisplayClusterConfiguratorStyle.cpp
--- a/Engine/Plugins/Runtime/nDisplay/Source/DisplayClusterConfigurator/Private/DisplayClusterConfiguratorStyle.cpp
+++ b/Engine/Plugins/Runtime/nDisplay/Source/DisplayClusterConfigurator/Private/DisplayClusterConfiguratorStyle.cpp
@@ -93,7 +93,7 @@ void FDisplayClusterConfiguratorStyle::Initialize()
 		Set("DisplayClusterConfigurator.TreeItems.Viewport", new IMAGE_BRUSH_SVG("Icons/Cluster/Viewport", Icon16x16));
 	}
 
-	// Output Mapping Toolbar
+	// Output Mapping Commands
 	{
 		Set("DisplayClusterConfigurator.OutputMapping.WindowDisplay", new IMAGE_BRUSH_SVG("Icons/OutputMapping/ToggleWindowInfo", Icon16x16));
 		Set("DisplayClusterConfigurator.OutputMapping.ToggleOutsideViewports", new IMAGE_BRUSH_SVG("Icons/OutputMapping/ToggleOutsideViewports", Icon16x16));
@@ -102,6 +102,12 @@ void FDisplayClusterConfiguratorStyle::Initialize()
 		Set("DisplayClusterConfigurator.OutputMapping.Transform", new IMAGE_BRUSH_SVG("Icons/OutputMapping/Transform", Icon16x16));
 		Set("DisplayClusterConfigurator.OutputMapping.Snapping", new IMAGE_BRUSH_SVG("Icons/OutputMapping/Snapping", Icon16x16));
 		Set("DisplayClusterConfigurator.OutputMapping.ResizeAreaHandle", new IMAGE_BRUSH("Icons/OutputMapping/ResizeAreaHandle_20x", Icon20x20));
+
+		Set("DisplayClusterConfigurator.OutputMapping.RotateViewport90CW", new CORE_IMAGE_BRUSH_SVG("Starship/Common/Rotate90Clockwise", Icon16x16));
+		Set("DisplayClusterConfigurator.OutputMapping.RotateViewport90CCW", new CORE_IMAGE_BRUSH_SVG("Starship/Common/Rotate90CounterClockwise", Icon16x16));
+		Set("DisplayClusterConfigurator.OutputMapping.RotateViewport180", new CORE_IMAGE_BRUSH_SVG("Starship/Common/Rotate180", Icon16x16));
+		Set("DisplayClusterConfigurator.OutputMapping.FlipViewportHorizontal", new CORE_IMAGE_BRUSH_SVG("Starship/Common/FlipHorizontal", Icon16x16));
+		Set("DisplayClusterConfigurator.OutputMapping.FlipViewportVertical", new CORE_IMAGE_BRUSH_SVG("Starship/Common/FlipVertical", Icon16x16));
 	}
 
 	// Output Mapping Node Colors
diff --git a/Engine/Source/Runtime/SlateCore/Private/Styling/StarshipCoreStyle.cpp b/Engine/Source/Runtime/SlateCore/Private/Styling/StarshipCoreStyle.cpp
--- a/Engine/Source/Runtime/SlateCore/Private/Styling/StarshipCoreStyle.cpp
+++ b/Engine/Source/Runtime/SlateCore/Private/Styling/StarshipCoreStyle.cpp
@@ -391,6 +391,12 @@ TSharedRef<ISlateStyle> FStarshipCoreStyle::Create()
 		Style->Set("Icons.EyeDropper", new IMAGE_BRUSH_SVG("Starship/Common/EyeDropper", Icon16x16));
 		Style->Set("Icons.C++", new IMAGE_BRUSH_SVG("Starship/Common/CPP", Icon16x16));
 
+		Style->Set("Icons.Rotate90Clockwise", new IMAGE_BRUSH_SVG("Starship/Common/Rotate90Clockwise", Icon16x16));
+		Style->Set("Icons.Rotate90CounterClockwise", new IMAGE_BRUSH_SVG("Starship/Common/Rotate90CounterClockwise", Icon16x16));
+		Style->Set("Icons.Rotate180", new IMAGE_BRUSH_SVG("Starship/Common/Rotate180", Icon16x16));
+		Style->Set("Icons.FlipHorizontal", new IMAGE_BRUSH_SVG("Starship/Common/FlipHorizontal", Icon16x16));
+		Style->Set("Icons.FlipVertical", new IMAGE_BRUSH_SVG("Starship/Common/FlipVertical", Icon16x16));
+
 		// Toolbar Size Icons
 		Style->Set("Icons.Toolbar.Settings", new IMAGE_BRUSH_SVG("Starship/Common/settings", Icon20x20));
 	}
