nDisplay: fixed stereoscopic behaviour

diff --git a/Engine/Plugins/Runtime/nDisplay/Source/DisplayCluster/Private/Render/Viewport/DisplayClusterViewportStereoscopicPass.h b/Engine/Plugins/Runtime/nDisplay/Source/DisplayCluster/Private/Render/Viewport/DisplayClusterViewportStereoscopicPass.h
--- a/Engine/Plugins/Runtime/nDisplay/Source/DisplayCluster/Private/Render/Viewport/DisplayClusterViewportStereoscopicPass.h
+++ b/Engine/Plugins/Runtime/nDisplay/Source/DisplayCluster/Private/Render/Viewport/DisplayClusterViewportStereoscopicPass.h
@@ -9,14 +9,18 @@ struct FDisplayClusterViewportStereoscopicPass
 {
 	inline static EStereoscopicPass EncodeStereoscopicPass(const uint32 ContextNum, const uint32 ContextAmmount)
 	{
-		// Monoscopic rendering
-		if (ContextAmmount == 1)
+		// Stereoscopic rendering
+		if (ContextAmmount > 1)
 		{
 			switch (ContextNum)
 			{
 			case 0:
+				// Left eye
+				// PRIMARY implies the view needs its own pass
 				return EStereoscopicPass::eSSP_PRIMARY;
 			case 1:
+				// Right eye
+				// SECONDARY implies the view can be instanced
 				return EStereoscopicPass::eSSP_SECONDARY;
 			default:
 				// now stereo only with 2 context
@@ -25,6 +29,7 @@ struct FDisplayClusterViewportStereoscopicPass
 			}
 		}
 
+		// Monoscopic rendering
 		return EStereoscopicPass::eSSP_PRIMARY;
 	}
 };
