nDisplay: added icvfx's camera custom frame size support for resolution estimation

diff --git a/Engine/Plugins/Runtime/nDisplay/Source/DisplayCluster/Private/Game/EngineClasses/Scene/DisplayClusterICVFXCameraComponent.cpp b/Engine/Plugins/Runtime/nDisplay/Source/DisplayCluster/Private/Game/EngineClasses/Scene/DisplayClusterICVFXCameraComponent.cpp
--- a/Engine/Plugins/Runtime/nDisplay/Source/DisplayCluster/Private/Game/EngineClasses/Scene/DisplayClusterICVFXCameraComponent.cpp
+++ b/Engine/Plugins/Runtime/nDisplay/Source/DisplayCluster/Private/Game/EngineClasses/Scene/DisplayClusterICVFXCameraComponent.cpp
@@ -64,6 +64,12 @@ void UDisplayClusterICVFXCameraComponent::UpdateOverscanEstimatedFrameSize()
 	float InnerFrustumResolutionWidth = RootActor->GetStageSettings().DefaultFrameSize.Width * CameraSettings.BufferRatio;
 	float InnerFrustumResolutionHeight = RootActor->GetStageSettings().DefaultFrameSize.Height * CameraSettings.BufferRatio;
 
+	if (CameraSettings.RenderSettings.CustomFrameSize.bUseCustomSize)
+	{
+		InnerFrustumResolutionWidth = CameraSettings.RenderSettings.CustomFrameSize.CustomWidth * CameraSettings.BufferRatio;
+		InnerFrustumResolutionHeight = CameraSettings.RenderSettings.CustomFrameSize.CustomHeight * CameraSettings.BufferRatio;
+	}
+
 	float EstimatedOverscanResolutionWidth = InnerFrustumResolutionWidth;
 	float EstimatedOverscanResolutionHeight = InnerFrustumResolutionHeight;
 
