[nDisplay]
Fixed mesh projection policies throwing invalid UV errors

diff --git a/Engine/Plugins/Runtime/nDisplay/Source/DisplayClusterConfigurator/Private/Views/Details/Policies/DisplayClusterConfiguratorPolicyDetailCustomization.cpp b/Engine/Plugins/Runtime/nDisplay/Source/DisplayClusterConfigurator/Private/Views/Details/Policies/DisplayClusterConfiguratorPolicyDetailCustomization.cpp
--- a/Engine/Plugins/Runtime/nDisplay/Source/DisplayClusterConfigurator/Private/Views/Details/Policies/DisplayClusterConfiguratorPolicyDetailCustomization.cpp
+++ b/Engine/Plugins/Runtime/nDisplay/Source/DisplayClusterConfigurator/Private/Views/Details/Policies/DisplayClusterConfiguratorPolicyDetailCustomization.cpp
@@ -464,14 +464,14 @@ void FDisplayClusterConfiguratorProjectionCustomization::CreateMeshPolicy(UDispl
 		"Warp UVs",
 		DisplayClusterProjectionStrings::cfg::mesh::BaseUVIndex,
 		Blueprint,
-		ConfigurationViewports, -1, -1);
+		ConfigurationViewports, INDEX_NONE, INDEX_NONE);
 	BaseUVIndexParameterRef->SetParameterTooltip(LOCTEXT("MeshPolicyBaseUVIndexTooltip", "Define custom UV channel in the original geometry as UV for warp. The default is -1."));
 
 	TSharedPtr<FPolicyParameterInfo> ChromakeyUVIndexParameterRef = MakeShared<FPolicyParameterInfoNumber<int32>>(
 		"Chromakey UVs",
 		DisplayClusterProjectionStrings::cfg::mesh::ChromakeyUVIndex,
 		Blueprint,
-		ConfigurationViewports, -1, -1);
+		ConfigurationViewports, INDEX_NONE, INDEX_NONE);
 	ChromakeyUVIndexParameterRef->SetParameterTooltip(LOCTEXT("MeshPolicyChromakeyUVIndexTooltip", "Define custom UV channel in the original geometry as UV for ChromakeyMarkers. The default is -1."));
 
 	CustomPolicyParameters.Add(SectionIndexParameterRef);
diff --git a/Engine/Plugins/Runtime/nDisplay/Source/DisplayClusterProjection/Private/Policy/Mesh/DisplayClusterProjectionMeshPolicy.cpp b/Engine/Plugins/Runtime/nDisplay/Source/DisplayClusterProjection/Private/Policy/Mesh/DisplayClusterProjectionMeshPolicy.cpp
--- a/Engine/Plugins/Runtime/nDisplay/Source/DisplayClusterProjection/Private/Policy/Mesh/DisplayClusterProjectionMeshPolicy.cpp
+++ b/Engine/Plugins/Runtime/nDisplay/Source/DisplayClusterProjection/Private/Policy/Mesh/DisplayClusterProjectionMeshPolicy.cpp
@@ -147,14 +147,14 @@ bool FDisplayClusterProjectionMeshPolicy::GetWarpMeshConfiguration(IDisplayClust
 		if (CfgBaseUVIndex >= 0)
 		{
 			UE_LOG(LogDisplayClusterProjectionMesh, Verbose, TEXT("Found BaseUVIndex value - '%d'"), CfgBaseUVIndex);
-
-			OutWarpCfg.BaseUVIndex = CfgBaseUVIndex;
 		}
-		else
+		else if(CfgBaseUVIndex != INDEX_NONE)
 		{
 			UE_LOG(LogDisplayClusterProjectionMesh, Error, TEXT("Invalid BaseUVIndex value - '%d'"), CfgBaseUVIndex);
+			CfgBaseUVIndex = INDEX_NONE;
 		}
 	}
+	OutWarpCfg.BaseUVIndex = CfgBaseUVIndex;
 
 	int32 CfgChromakeyUVIndex = INDEX_NONE;
 	if (DisplayClusterHelpers::map::template ExtractValueFromString(GetParameters(), DisplayClusterProjectionStrings::cfg::mesh::ChromakeyUVIndex, CfgChromakeyUVIndex))
@@ -162,14 +162,14 @@ bool FDisplayClusterProjectionMeshPolicy::GetWarpMeshConfiguration(IDisplayClust
 		if (CfgChromakeyUVIndex >= 0)
 		{
 			UE_LOG(LogDisplayClusterProjectionMesh, Verbose, TEXT("Found ChromakeyUVIndex value - '%d'"), CfgChromakeyUVIndex);
-
-			OutWarpCfg.ChromakeyUVIndex = CfgChromakeyUVIndex;
 		}
-		else
+		else if(CfgChromakeyUVIndex != INDEX_NONE)
 		{
 			UE_LOG(LogDisplayClusterProjectionMesh, Error, TEXT("Invalid ChromakeyUVIndex value - '%d'"), CfgChromakeyUVIndex);
+			CfgChromakeyUVIndex = INDEX_NONE;
 		}
 	}
+	OutWarpCfg.ChromakeyUVIndex = CfgChromakeyUVIndex;
 
 	// Get the StaticMeshComponent
 	OutWarpCfg.StaticMeshComponent = Root->GetComponentByName<UStaticMeshComponent>(ComponentId);
