nDisplay: fixes for native input replication
- fixed native input replication and removed one more stdlib dependency
- some data conversion fixes
- some optimization

diff --git a/Engine/Plugins/Runtime/nDisplay/Source/DisplayCluster/Private/Cluster/DisplayClusterClusterManager.cpp b/Engine/Plugins/Runtime/nDisplay/Source/DisplayCluster/Private/Cluster/DisplayClusterClusterManager.cpp
--- a/Engine/Plugins/Runtime/nDisplay/Source/DisplayCluster/Private/Cluster/DisplayClusterClusterManager.cpp
+++ b/Engine/Plugins/Runtime/nDisplay/Source/DisplayCluster/Private/Cluster/DisplayClusterClusterManager.cpp
@@ -758,10 +758,11 @@ void FDisplayClusterClusterManager::ImportNativeInputData(TMap<FString, FString>
 	NativeInputCache = InNativeInputData;
 
 	UE_LOG(LogDisplayClusterCluster, VeryVerbose, TEXT("Native input data cache:"));
+
 	int32 Idx = 0;
-	for (auto It = NativeInputCache.CreateConstIterator(); It; ++It)
+	for (auto It = NativeInputCache.CreateConstIterator(); It; ++It, ++Idx)
 	{
-		UE_LOG(LogDisplayClusterCluster, VeryVerbose, TEXT("Native input data cache: %d - %s - %s"), Idx++, *It->Key, *It->Value);
+		UE_LOG(LogDisplayClusterCluster, VeryVerbose, TEXT("Native input data cache: %d - %s - %s"), Idx, *It->Key, *It->Value);
 	}
 
 	// Notify the data is available
@@ -778,11 +779,6 @@ void FDisplayClusterClusterManager::ExportNativeInputData(TMap<FString, FString>
 		NativeInputCacheReadySignal->Wait();
 		// Export data from cache
 		OutNativeInputData = NativeInputCache;
-
-		for (auto It = OutNativeInputData.CreateConstIterator(); It; ++It)
-		{
-			UE_LOG(LogDisplayClusterCluster, Verbose, TEXT("Downloading native input data..."));
-		}
 	}
 	// Otherwise we need to send network request to download input data
 	else
@@ -791,7 +787,13 @@ void FDisplayClusterClusterManager::ExportNativeInputData(TMap<FString, FString>
 		{
 			UE_LOG(LogDisplayClusterCluster, Verbose, TEXT("Downloading native input data..."));
 			ClusterNodeCtrl->GetNativeInputData(OutNativeInputData);
-			UE_LOG(LogDisplayClusterCluster, Verbose, TEXT("Downloading finished. Available %d records (native input)."), OutNativeInputData.Num());
+			UE_LOG(LogDisplayClusterCluster, Verbose, TEXT("Downloading finished. Available %d native input records."), OutNativeInputData.Num());
+
+			int32 Idx = 0;
+			for (auto It = OutNativeInputData.CreateConstIterator(); It; ++It, ++Idx)
+			{
+				UE_LOG(LogDisplayClusterCluster, VeryVerbose, TEXT("Native input data: %d - %s - %s"), Idx, *It->Key, *It->Value);
+			}
 		}
 	}
 }
diff --git a/Engine/Plugins/Runtime/nDisplay/Source/DisplayCluster/Private/Game/EngineClasses/Basics/DisplayClusterPlayerInput.cpp b/Engine/Plugins/Runtime/nDisplay/Source/DisplayCluster/Private/Game/EngineClasses/Basics/DisplayClusterPlayerInput.cpp
--- a/Engine/Plugins/Runtime/nDisplay/Source/DisplayCluster/Private/Game/EngineClasses/Basics/DisplayClusterPlayerInput.cpp
+++ b/Engine/Plugins/Runtime/nDisplay/Source/DisplayCluster/Private/Game/EngineClasses/Basics/DisplayClusterPlayerInput.cpp
@@ -45,7 +45,7 @@ void UDisplayClusterPlayerInput::ProcessInputStack(const TArray<UInputComponent*
 
 bool UDisplayClusterPlayerInput::SerializeKeyStateMap(TMap<FString, FString>& OutKeyStateMap)
 {
-	TMap<FKey, FKeyState>& StateMap = GetKeyStateMap();
+	const TMap<FKey, FKeyState>& StateMap = GetKeyStateMap();
 	for (auto it = StateMap.CreateConstIterator(); it; ++it)
 	{
 		const FString KeyName = it->Key.ToString();
@@ -53,9 +53,6 @@ bool UDisplayClusterPlayerInput::SerializeKeyStateMap(TMap<FString, FString>& Ou
 		FString StrKeyState;
 		StrKeyState.Reserve(2048);
 
-		float SendValue = GetWorld()->GetRealTimeSeconds() - it->Value.LastUpDownTransitionTime;
-		check(SendValue > 0.f);
-
 		StrKeyState = FString::Printf(TEXT("%s;%s;%s;%s;%s;%s;%s;%s;%s;"),
 			*DisplayClusterTypesConverter::template ToHexString(it->Value.RawValue),
 			*DisplayClusterTypesConverter::template ToHexString(it->Value.Value),
@@ -86,7 +83,7 @@ bool UDisplayClusterPlayerInput::SerializeKeyStateMap(TMap<FString, FString>& Ou
 bool UDisplayClusterPlayerInput::DeserializeKeyStateMap(const TMap<FString, FString>& InKeyStateMap)
 {
 	TMap<FKey, FKeyState>& StateMap = GetKeyStateMap();
-	
+
 	// Reset local key state map
 	StateMap.Reset();
 
diff --git a/Engine/Plugins/Runtime/nDisplay/Source/DisplayCluster/Public/Misc/DisplayClusterHelpers.h b/Engine/Plugins/Runtime/nDisplay/Source/DisplayCluster/Public/Misc/DisplayClusterHelpers.h
--- a/Engine/Plugins/Runtime/nDisplay/Source/DisplayCluster/Public/Misc/DisplayClusterHelpers.h
+++ b/Engine/Plugins/Runtime/nDisplay/Source/DisplayCluster/Public/Misc/DisplayClusterHelpers.h
@@ -78,7 +78,7 @@ namespace DisplayClusterHelpers
 			{
 				if (!bCullEmpty && Item.IsEmpty())
 				{
-					OutData.Add(TVal());
+					OutData.AddDefaulted();
 				}
 				else
 				{
@@ -109,7 +109,7 @@ namespace DisplayClusterHelpers
 
 			if (InSeparator.Len() > 0 && InData.Num() > 0)
 			{
-				ResultStr.RemoveAt(ResultStr.Len() - InSeparator.Len());
+				ResultStr.RemoveAt(ResultStr.Len() - InSeparator.Len(), InSeparator.Len(), !bAddQuotes);
 			}
 
 			if (bAddQuotes)
@@ -178,7 +178,7 @@ namespace DisplayClusterHelpers
 
 			if (InPairSeparator.Len() > 0 && InData.Num() > 0)
 			{
-				ResultStr.RemoveAt(ResultStr.Len() - InPairSeparator.Len());
+				ResultStr.RemoveAt(ResultStr.Len() - InPairSeparator.Len(), InPairSeparator.Len(), !bAddQuoutes);
 			}
 
 			if (bAddQuoutes)
diff --git a/Engine/Plugins/Runtime/nDisplay/Source/DisplayCluster/Public/Misc/DisplayClusterTypesConverter.h b/Engine/Plugins/Runtime/nDisplay/Source/DisplayCluster/Public/Misc/DisplayClusterTypesConverter.h
--- a/Engine/Plugins/Runtime/nDisplay/Source/DisplayCluster/Public/Misc/DisplayClusterTypesConverter.h
+++ b/Engine/Plugins/Runtime/nDisplay/Source/DisplayCluster/Public/Misc/DisplayClusterTypesConverter.h
@@ -9,11 +9,6 @@
 
 #include "DisplayClusterEnums.h"
 
-#include <iomanip>
-#include <iostream>
-#include <limits>
-#include <sstream>
-
 
 /**
  * Auxiliary class with different type conversion functions
@@ -168,10 +163,10 @@ template <> inline FMatrix   DisplayClusterTypesConverter::FromString<>(const FS
 			return ResultMatrix;
 		}
 
-		float PlaneValues[4] = { 0.f };
+		FMatrix::FReal PlaneValues[4] = { 0.f };
 		for (int i = 0; i < 4; ++i)
 		{
-			PlaneValues[i] = FromString<float>(StrPlaneValues[i]);
+			PlaneValues[i] = FromString<FMatrix::FReal>(StrPlaneValues[i]);
 		}
 
 		Planes[PlaneNum] = FPlane(PlaneValues[0], PlaneValues[1], PlaneValues[2], PlaneValues[3]);
@@ -201,7 +196,7 @@ template <> inline FQuat DisplayClusterTypesConverter::FromString<>(const FStrin
 
 template <> inline FIntPoint DisplayClusterTypesConverter::FromString<>(const FString& From)
 {
-	FIntPoint Result;
+	FIntPoint Result = FIntPoint::ZeroValue;
 	const bool bSuccessful
 		=  FParse::Value(*From, TEXT("X="), Result.X)
 		&& FParse::Value(*From, TEXT("Y="), Result.Y);
@@ -296,28 +291,12 @@ template <> inline EDisplayClusterSyncGroup DisplayClusterTypesConverter::FromSt
 //////////////////////////////////////////////////////////////////////////////////////////////
 template <> inline FString DisplayClusterTypesConverter::ToHexString<>(const float& From)
 {
-	static_assert(std::numeric_limits<float>::is_iec559, "Native float must be an IEEE float");
-
-	union { float fval; std::uint32_t ival; };
-	fval = From;
-
-	std::ostringstream stm;
-	stm << std::hex << std::nouppercase << ival;
-
-	return FString(FString(TEXT("0x")) + FString(stm.str().c_str()));
+	return BytesToHex(reinterpret_cast<const uint8*>(&From), sizeof(float));
 }
 
 template <> inline FString DisplayClusterTypesConverter::ToHexString<>(const double& From)
 {
-	static_assert(std::numeric_limits<double>::is_iec559, "Native double must be an IEEE double");
-
-	union { double dval; std::uint64_t ival; };
-	dval = From;
-
-	std::ostringstream stm;
-	stm << std::hex << std::nouppercase << ival;
-
-	return FString(FString(TEXT("0x")) + FString(stm.str().c_str()));
+	return BytesToHex(reinterpret_cast<const uint8*>(&From), sizeof(double));
 }
 
 template <> inline FString DisplayClusterTypesConverter::ToHexString<>(const FVector& From)
@@ -383,39 +362,31 @@ template <> inline FString DisplayClusterTypesConverter::ToHexString<>(const FQu
 //////////////////////////////////////////////////////////////////////////////////////////////
 template <> inline float DisplayClusterTypesConverter::FromHexString<>(const FString& From)
 {
-	static_assert(std::numeric_limits<float>::is_iec559, "Native float must be an IEEE float");
-
-	union { float fval; std::uint32_t ival; };
-
-	std::string str(TCHAR_TO_UTF8(*From));
-	std::stringstream stm(str);
-	stm >> std::hex >> ival;
-	return fval;
+	float Result = 0.f;
+	checkSlow(From.Len() == 2 * sizeof(Result));
+	HexToBytes(From, reinterpret_cast<uint8*>(&Result));
+	return Result;
 }
 
 template <> inline double DisplayClusterTypesConverter::FromHexString<>(const FString& From)
 {
-	static_assert(std::numeric_limits<double>::is_iec559, "Native double must be an IEEE double");
-
-	union { double dval; std::uint64_t ival; };
-
-	std::string str(TCHAR_TO_UTF8(*From));
-	std::stringstream stm(str);
-	stm >> std::hex >> ival;
-	return dval;
+	double Result = 0.f;
+	checkSlow(From.Len() == 2 * sizeof(Result));
+	HexToBytes(From, reinterpret_cast<uint8*>(&Result));
+	return Result;
 }
 
 template <> inline FVector DisplayClusterTypesConverter::FromHexString<>(const FString& From)
 {
 	FString X, Y, Z;
-	FVector Result;
+	FVector Result = FVector::ZeroVector;
 
 	const bool bSuccessful = FParse::Value(*From, TEXT("X="), X) && FParse::Value(*From, TEXT("Y="), Y) && FParse::Value(*From, TEXT("Z="), Z);
 	if (bSuccessful)
 	{
-		Result.X = DisplayClusterTypesConverter::template FromHexString<float>(X);
-		Result.Y = DisplayClusterTypesConverter::template FromHexString<float>(Y);
-		Result.Z = DisplayClusterTypesConverter::template FromHexString<float>(Z);
+		Result.X = DisplayClusterTypesConverter::template FromHexString<FVector::FReal>(X);
+		Result.Y = DisplayClusterTypesConverter::template FromHexString<FVector::FReal>(Y);
+		Result.Z = DisplayClusterTypesConverter::template FromHexString<FVector::FReal>(Z);
 	}
 
 	return Result;
@@ -424,13 +395,13 @@ template <> inline FVector DisplayClusterTypesConverter::FromHexString<>(const F
 template <> inline FVector2D DisplayClusterTypesConverter::FromHexString<>(const FString& From)
 {
 	FString X, Y;
-	FVector2D Result;
+	FVector2D Result = FVector2D::ZeroVector;
 
 	const bool bSuccessful = FParse::Value(*From, TEXT("X="), X) && FParse::Value(*From, TEXT("Y="), Y);
 	if (bSuccessful)
 	{
-		Result.X = DisplayClusterTypesConverter::template FromHexString<float>(X);
-		Result.Y = DisplayClusterTypesConverter::template FromHexString<float>(Y);
+		Result.X = DisplayClusterTypesConverter::template FromHexString<FVector2D::FReal>(X);
+		Result.Y = DisplayClusterTypesConverter::template FromHexString<FVector2D::FReal>(Y);
 	}
 
 	return Result;
@@ -439,14 +410,14 @@ template <> inline FVector2D DisplayClusterTypesConverter::FromHexString<>(const
 template <> inline FRotator DisplayClusterTypesConverter::FromHexString<>(const FString& From)
 {
 	FString P, Y, R;
-	FRotator Result;
+	FRotator Result = FRotator::ZeroRotator;
 
 	const bool bSuccessful = FParse::Value(*From, TEXT("P="), P) && FParse::Value(*From, TEXT("Y="), Y) && FParse::Value(*From, TEXT("R="), R);
 	if (bSuccessful)
 	{
-		Result.Pitch = DisplayClusterTypesConverter::template FromHexString<float>(P);
-		Result.Yaw   = DisplayClusterTypesConverter::template FromHexString<float>(Y);
-		Result.Roll  = DisplayClusterTypesConverter::template FromHexString<float>(R);
+		Result.Pitch = DisplayClusterTypesConverter::template FromHexString<FRotator::FReal>(P);
+		Result.Yaw   = DisplayClusterTypesConverter::template FromHexString<FRotator::FReal>(Y);
+		Result.Roll  = DisplayClusterTypesConverter::template FromHexString<FRotator::FReal>(R);
 	}
 
 	return Result;
@@ -486,10 +457,10 @@ template <> inline FMatrix DisplayClusterTypesConverter::FromHexString<>(const F
 			return ResultMatrix;
 		}
 
-		float PlaneValues[4] = { 0.f };
+		FMatrix::FReal PlaneValues[4] = { 0.f };
 		for (int i = 0; i < 4; ++i)
 		{
-			PlaneValues[i] = DisplayClusterTypesConverter::template FromHexString<float>(StrPlaneValues[i]);
+			PlaneValues[i] = DisplayClusterTypesConverter::template FromHexString<FMatrix::FReal>(StrPlaneValues[i]);
 		}
 
 		Planes[PlaneNum] = FPlane(PlaneValues[0], PlaneValues[1], PlaneValues[2], PlaneValues[3]);
@@ -526,10 +497,10 @@ template <> inline FQuat DisplayClusterTypesConverter::FromHexString<>(const FSt
 	const bool bSuccessful = FParse::Value(*From, TEXT("X="), X) && FParse::Value(*From, TEXT("Y="), Y) && FParse::Value(*From, TEXT("Z="), Z) && FParse::Value(*From, TEXT("W="), W);
 	if (bSuccessful)
 	{
-		Result.X = DisplayClusterTypesConverter::template FromHexString<float>(X);
-		Result.Y = DisplayClusterTypesConverter::template FromHexString<float>(Y);
-		Result.Z = DisplayClusterTypesConverter::template FromHexString<float>(Z);
-		Result.W = DisplayClusterTypesConverter::template FromHexString<float>(W);
+		Result.X = DisplayClusterTypesConverter::template FromHexString<FQuat::FReal>(X);
+		Result.Y = DisplayClusterTypesConverter::template FromHexString<FQuat::FReal>(Y);
+		Result.Z = DisplayClusterTypesConverter::template FromHexString<FQuat::FReal>(Z);
+		Result.W = DisplayClusterTypesConverter::template FromHexString<FQuat::FReal>(W);
 	}
 
 	return Result;
