nDisplay: fixed PIE crash if no p-node set

diff --git a/Engine/Plugins/Runtime/nDisplay/Source/DisplayClusterEditor/Private/DisplayClusterEditorEngine.cpp b/Engine/Plugins/Runtime/nDisplay/Source/DisplayClusterEditor/Private/DisplayClusterEditorEngine.cpp
--- a/Engine/Plugins/Runtime/nDisplay/Source/DisplayClusterEditor/Private/DisplayClusterEditorEngine.cpp
+++ b/Engine/Plugins/Runtime/nDisplay/Source/DisplayClusterEditor/Private/DisplayClusterEditorEngine.cpp
@@ -109,16 +109,19 @@ void UDisplayClusterEditorEngine::StartPlayInEditorSession(FRequestPlaySessionPa
 		// If we found a root actor, start DisplayCluster PIE session
 		if (RootActor)
 		{
-			bIsNDisplayPIE = true;
-
 			// Get current config data from the root actor
 			UDisplayClusterConfigurationData* ConfigData = DuplicateObject<UDisplayClusterConfigurationData>(RootActor->GetConfigData(), this);
 
 			// And start PIE session with that config data
 			if (ConfigData)
 			{
-				if (!DisplayClusterModule->StartSession(ConfigData, ConfigData->Cluster->PrimaryNode.Id))
+				if (DisplayClusterModule->StartSession(ConfigData, ConfigData->Cluster->PrimaryNode.Id))
+				{
+					bIsNDisplayPIE = true;
+				}
+				else
 				{
+					DisplayClusterModule->EndSession();
 					UE_LOG(LogDisplayClusterEditorEngine, Error, TEXT("An error occurred during DisplayCluster session start"));
 				}
 			}
@@ -199,11 +202,14 @@ void UDisplayClusterEditorEngine::OnEndPIE(const bool bSimulate)
 {
 	UE_LOG(LogDisplayClusterEditorEngine, VeryVerbose, TEXT("UDisplayClusterEditorEngine::OnEndPIE"));
 
+	if (bIsNDisplayPIE)
+	{
+		// Notify nDisplay about session end
+		DisplayClusterModule->EndScene();
+		DisplayClusterModule->EndSession();
+	}
+
 	// Reset PIE flags
 	bIsActivePIE   = false;
 	bIsNDisplayPIE = false;
-
-	// Notify nDisplay about session end
-	DisplayClusterModule->EndScene();
-	DisplayClusterModule->EndSession();
 }
