nDisplay: removed spam from logs (StereoViewIndex='-1' not found)

diff --git a/Engine/Plugins/Runtime/nDisplay/Source/DisplayCluster/Private/Render/Device/DisplayClusterDeviceBase.cpp b/Engine/Plugins/Runtime/nDisplay/Source/DisplayCluster/Private/Render/Device/DisplayClusterDeviceBase.cpp
--- a/Engine/Plugins/Runtime/nDisplay/Source/DisplayCluster/Private/Render/Device/DisplayClusterDeviceBase.cpp
+++ b/Engine/Plugins/Runtime/nDisplay/Source/DisplayCluster/Private/Render/Device/DisplayClusterDeviceBase.cpp
@@ -207,6 +207,13 @@ void FDisplayClusterDeviceBase::AdjustViewRect(int32 ViewIndex, int32& X, int32&
 	{
 		return;
 	}
+	// ViewIndex == eSSE_MONOSCOPIC(-1) is a special case called for ISR culling math.
+	// Since nDisplay is not ISR compatible, we ignore this request. This won't be neccessary once
+	// we stop using nDisplay as a stereoscopic rendering device (IStereoRendering).
+	else if (ViewIndex < 0)
+	{
+		return;
+	}
 
 	uint32 ViewportContextNum = 0;
 	IDisplayClusterViewport* ViewportPtr = ViewportManagerPtr->FindViewport(ViewIndex, &ViewportContextNum);
@@ -232,8 +239,17 @@ void FDisplayClusterDeviceBase::CalculateStereoViewOffset(const int32 ViewIndex,
 	check(IsInGameThread());
 	check(WorldToMeters > 0.f);
 
-	if (ViewportManagerPtr ==nullptr || ViewportManagerPtr->IsSceneOpened() == false)
-		{ return; }
+	if (ViewportManagerPtr == nullptr || ViewportManagerPtr->IsSceneOpened() == false)
+	{
+		return;
+	}
+	// ViewIndex == eSSE_MONOSCOPIC(-1) is a special case called for ISR culling math.
+	// Since nDisplay is not ISR compatible, we ignore this request. This won't be neccessary once
+	// we stop using nDisplay as a stereoscopic rendering device (IStereoRendering).
+	else if (ViewIndex < 0)
+	{
+		return;
+	}
 
 	uint32 ViewportContextNum = 0;
 	IDisplayClusterViewport* ViewportPtr = ViewportManagerPtr->FindViewport(ViewIndex, &ViewportContextNum);
@@ -369,7 +385,10 @@ FMatrix FDisplayClusterDeviceBase::GetStereoProjectionMatrix(const int32 ViewInd
 
 	FMatrix PrjMatrix = FMatrix::Identity;
 
-	if (ViewportManagerPtr && ViewportManagerPtr->IsSceneOpened())
+	// ViewIndex == eSSE_MONOSCOPIC(-1) is a special case called for ISR culling math.
+	// Since nDisplay is not ISR compatible, we ignore this request. This won't be neccessary once
+	// we stop using nDisplay as a stereoscopic rendering device (IStereoRendering).
+	if (ViewportManagerPtr && ViewportManagerPtr->IsSceneOpened() && ViewIndex >= 0)
 	{
 		uint32 ViewportContextNum = 0;
 		IDisplayClusterViewport* ViewportPtr = ViewportManagerPtr->FindViewport(ViewIndex, &ViewportContextNum);
