nDisplay: Added light card list to the light card editor, which displays the list of light cards linked to the active root actor

diff --git a/Engine/Plugins/Runtime/nDisplay/Source/DisplayClusterLightCardEditor/Private/SDisplayClusterLightCardEditor.cpp b/Engine/Plugins/Runtime/nDisplay/Source/DisplayClusterLightCardEditor/Private/SDisplayClusterLightCardEditor.cpp
--- a/Engine/Plugins/Runtime/nDisplay/Source/DisplayClusterLightCardEditor/Private/SDisplayClusterLightCardEditor.cpp
+++ b/Engine/Plugins/Runtime/nDisplay/Source/DisplayClusterLightCardEditor/Private/SDisplayClusterLightCardEditor.cpp
@@ -2,6 +2,8 @@
 
 #include "SDisplayClusterLightCardEditor.h"
 
+#include "SDisplayClusterLightCardList.h"
+
 #include "IDisplayClusterOperator.h"
 #include "DisplayClusterRootActor.h"
 
@@ -11,6 +13,7 @@
 #include "Framework/MultiBox/MultiBoxExtender.h"
 #include "Toolkits/AssetEditorToolkit.h"
 #include "Widgets/Docking/SDockTab.h"
+#include "Widgets/Layout/SSplitter.h"
 
 #define LOCTEXT_NAMESPACE "SDisplayClusterLightCardEditor"
 
@@ -64,25 +67,25 @@ void SDisplayClusterLightCardEditor::Construct(const FArguments& InArgs, const T
 {
 	ActiveRootActorChangedHandle = IDisplayClusterOperator::Get().OnActiveRootActorChanged().AddSP(this, &SDisplayClusterLightCardEditor::OnActiveRootActorChanged);
 
-    TabManager = FGlobalTabmanager::Get()->NewTabManager(MajorTabOwner);
-
-	// TODO: Register the various tabs for the light card editor here
-
-	const TSharedRef<FTabManager::FLayout> Layout = FTabManager::NewLayout("nDisplayLightCardEditorLayout")
-		->AddArea
-		(
-			FTabManager::NewPrimaryArea()
-			->SetOrientation(EOrientation::Orient_Vertical)
-			->Split
-			(
-				FTabManager::NewStack()
-				->SetHideTabWell(true)
-			)
-		);
-
-	ChildSlot
+	ChildSlot                     
 	[
-		TabManager->RestoreFrom(Layout, WindowOwner).ToSharedRef()
+		SNew(SSplitter)
+		.Orientation(Orient_Horizontal)
+		
+		+SSplitter::Slot()
+		.Value(0.25f)
+		[
+			// Vertical box for the left hand panel of the editor. Add new slots here as needed for any editor UI controls
+			SNew(SVerticalBox)
+
+			+ SVerticalBox::Slot()
+			[
+				CreateLightCardListWidget()
+			]
+		]
+
+		+SSplitter::Slot()
+		.Value(0.75f)
 	];
 }
 
@@ -90,6 +93,12 @@ void SDisplayClusterLightCardEditor::OnActiveRootActorChanged(ADisplayClusterRoo
 {
 	ActiveRootActor = NewRootActor;
 
-	// TODO: Need to pass a refresh command to any child tabs of this widget to inform them that the active root actor has been changed.
-	// The new root actor pointer could be null, indicating that it was deleted or the user didn't select a valid root actor
-}
\ No newline at end of file
+	LightCardList->SetRootActor(NewRootActor);
+}
+
+TSharedRef<SWidget> SDisplayClusterLightCardEditor::CreateLightCardListWidget()
+{
+	return SAssignNew(LightCardList, SDisplayClusterLightCardList);
+}
+
+#undef LOCTEXT_NAMESPACE
\ No newline at end of file
diff --git a/Engine/Plugins/Runtime/nDisplay/Source/DisplayClusterLightCardEditor/Private/SDisplayClusterLightCardEditor.h b/Engine/Plugins/Runtime/nDisplay/Source/DisplayClusterLightCardEditor/Private/SDisplayClusterLightCardEditor.h
--- a/Engine/Plugins/Runtime/nDisplay/Source/DisplayClusterLightCardEditor/Private/SDisplayClusterLightCardEditor.h
+++ b/Engine/Plugins/Runtime/nDisplay/Source/DisplayClusterLightCardEditor/Private/SDisplayClusterLightCardEditor.h
@@ -11,6 +11,7 @@ class FLayoutExtender;
 class FSpawnTabArgs;
 class FToolBarBuilder;
 class SDockTab;
+class SDisplayClusterLightCardList;
 class ADisplayClusterRootActor;
 
 /** A panel that can be spawned in a tab that contains all the UI elements that make up the 2D light cards editor */
@@ -48,12 +49,16 @@ private:
 	/** Raised when the active Display cluster root actor has been changed in the operator panel */
 	void OnActiveRootActorChanged(ADisplayClusterRootActor* NewRootActor);
 
+	/** Creates the widget used to show the list of light cards associated with the active root actor */
+	TSharedRef<SWidget> CreateLightCardListWidget();
+
 private:
-	/** Holds the tab manager that manages the light card editor's tabs. */
-	TSharedPtr<FTabManager> TabManager;
+	/** The light card list widget */
+	TSharedPtr<SDisplayClusterLightCardList> LightCardList;
 
 	/** A reference to the root actor that is currently being operated on */
 	TWeakObjectPtr<ADisplayClusterRootActor> ActiveRootActor;
 
+	/** Delegate handle for the OnActiveRootActorChanged delegate */
 	FDelegateHandle ActiveRootActorChangedHandle;
 };
\ No newline at end of file
diff --git a/Engine/Plugins/Runtime/nDisplay/Source/DisplayClusterLightCardEditor/Private/SDisplayClusterLightCardList.cpp b/Engine/Plugins/Runtime/nDisplay/Source/DisplayClusterLightCardEditor/Private/SDisplayClusterLightCardList.cpp
new file mode 100644
--- /dev/null
+++ b/Engine/Plugins/Runtime/nDisplay/Source/DisplayClusterLightCardEditor/Private/SDisplayClusterLightCardList.cpp
@@ -0,0 +1,254 @@
+// Copyright Epic Games, Inc. All Rights Reserved.
+
+#include "SDisplayClusterLightCardList.h"
+
+#include "DisplayClusterRootActor.h"
+#include "DisplayClusterConfigurationTypes.h"
+#include "IDisplayClusterOperator.h"
+
+#include "ClassIconFinder.h"
+#include "Widgets/Images/SImage.h"
+#include "Widgets/Views/SListView.h"
+#include "Widgets/Views/STableRow.h"
+#include "Widgets/Views/STreeView.h"
+
+#define LOCTEXT_NAMESPACE "SDisplayClusterLightCardList"
+
+namespace DisplayClusterLightCardListColumnNames
+{
+	const static FName LightCardName(TEXT("LightCardName"));
+};
+
+class SLightCardTreeItemRow : public SMultiColumnTableRow<TSharedPtr<SDisplayClusterLightCardList::FLightCardTreeItem>>
+{
+	SLATE_BEGIN_ARGS(SLightCardTreeItemRow) {}
+	SLATE_END_ARGS()
+
+public:
+
+	void Construct(const FArguments& InArgs, const TSharedRef<STableViewBase>& InOwnerTable, const TSharedPtr<SDisplayClusterLightCardList::FLightCardTreeItem> InLightCardTreeItem)
+	{
+		LightCardTreeItem = InLightCardTreeItem;
+
+		SMultiColumnTableRow<TSharedPtr<SDisplayClusterLightCardList::FLightCardTreeItem>>::Construct(
+			SMultiColumnTableRow<TSharedPtr<SDisplayClusterLightCardList::FLightCardTreeItem>>::FArguments()
+			.Style(&FAppStyle::Get().GetWidgetStyle<FTableRowStyle>("SceneOutliner.TableViewRow"))
+		, InOwnerTable);
+	}
+
+	virtual TSharedRef<SWidget> GenerateWidgetForColumn(const FName& ColumnName) override
+	{
+		if (ColumnName == DisplayClusterLightCardListColumnNames::LightCardName)
+		{
+			return SNew(SHorizontalBox)
+
+			+ SHorizontalBox::Slot()
+			.AutoWidth()
+			.Padding(6, 0, 0, 0)
+			[
+				SNew(SExpanderArrow, SharedThis(this))
+				.IndentAmount(12)
+			]
+
+			+ SHorizontalBox::Slot()
+			.AutoWidth()
+			.Padding(FMargin(0, 1, 6, 1))
+			.VAlign(VAlign_Center)
+			[
+				SNew(SBox)
+				.WidthOverride(16)
+				.HeightOverride(16)
+				[
+					SNew(SImage)
+					.ColorAndOpacity(FSlateColor::UseForeground())
+					.Image(this, &SLightCardTreeItemRow::GetLightCardIcon)
+				]
+			]
+
+			+ SHorizontalBox::Slot()
+			.AutoWidth()
+			.VAlign(VAlign_Center)
+			.HAlign(HAlign_Left)
+			[
+				SNew(STextBlock)
+				.Text(this, &SLightCardTreeItemRow::GetLightCardDisplayName)
+			];
+		}
+
+		return SNullWidget::NullWidget;
+	}
+
+private:
+	FText GetLightCardDisplayName() const
+	{
+		FString ItemName = TEXT("");
+		if (LightCardTreeItem.IsValid())
+		{
+			TSharedPtr<SDisplayClusterLightCardList::FLightCardTreeItem> LightCardTreeItemPin = LightCardTreeItem.Pin();
+
+			if (LightCardTreeItemPin->IsActorLayer())
+			{
+				ItemName = LightCardTreeItemPin->ActorLayer.ToString();
+			}
+			else if (LightCardTreeItemPin->LightCardActor.IsValid())
+			{
+				ItemName = LightCardTreeItemPin->LightCardActor->GetActorNameOrLabel();
+			}
+		}
+
+		return FText::FromString(*ItemName);
+	}
+
+	const FSlateBrush* GetLightCardIcon() const
+	{
+		FString ActorName = TEXT("");
+		if (LightCardTreeItem.IsValid())
+		{
+			TSharedPtr<SDisplayClusterLightCardList::FLightCardTreeItem> LightCardTreeItemPin = LightCardTreeItem.Pin();
+			if (LightCardTreeItemPin->IsActorLayer())
+			{
+				return FAppStyle::Get().GetBrush(TEXT("Layer.Icon16x"));
+			}
+			else if (LightCardTreeItemPin->LightCardActor.IsValid())
+			{
+				return FClassIconFinder::FindIconForActor(LightCardTreeItemPin->LightCardActor.Get());
+			}
+		}
+
+		return nullptr;
+	}
+
+private:
+	TWeakPtr<SDisplayClusterLightCardList::FLightCardTreeItem> LightCardTreeItem;
+};
+
+
+void SDisplayClusterLightCardList::Construct(const FArguments& InArgs)
+{
+	ChildSlot
+	[
+		SAssignNew(LightCardTreeView, STreeView<TSharedPtr<FLightCardTreeItem>>)
+		.TreeItemsSource(&LightCardTree)
+		.ItemHeight(28)
+		.SelectionMode(ESelectionMode::Multi)
+		.OnGenerateRow(this, &SDisplayClusterLightCardList::GenerateTreeItemRow)
+		.OnGetChildren(this, &SDisplayClusterLightCardList::GetChildrenForTreeItem)
+		.OnSelectionChanged(this, &SDisplayClusterLightCardList::OnTreeItemSelected)
+		.HighlightParentNodesForSelection(true)
+		.HeaderRow
+		(
+			SNew(SHeaderRow)
+
+			+ SHeaderRow::Column(DisplayClusterLightCardListColumnNames::LightCardName)
+			.DefaultLabel(LOCTEXT("LightCardName", "Light Card"))
+			.FillWidth(0.8f)
+		)
+	];
+}
+
+void SDisplayClusterLightCardList::SetRootActor(ADisplayClusterRootActor* NewRootActor)
+{
+	RootActor = NewRootActor;
+	FillLightCardList();
+	LightCardTreeView->RebuildList();
+
+	for (const TSharedPtr<FLightCardTreeItem>& LightCardTreeItem : LightCardTree)
+	{
+		if (LightCardTreeItem->IsActorLayer())
+		{
+			LightCardTreeView->SetItemExpansion(LightCardTreeItem, true);
+		}
+	}
+}
+
+void SDisplayClusterLightCardList::FillLightCardList()
+{
+	LightCardTree.Empty();
+	LightCardActors.Empty();
+
+	if (RootActor.IsValid())
+	{
+		FDisplayClusterConfigurationICVFX_VisibilityList& RootActorLightCards = RootActor->GetConfigData()->StageSettings.Lightcard.ShowOnlyList;
+
+		for (const FActorLayer& ActorLayer : RootActorLightCards.ActorLayers)
+		{
+			TSharedPtr<FLightCardTreeItem> LightCardTreeItem = MakeShared<FLightCardTreeItem>();
+			LightCardTreeItem->ActorLayer = ActorLayer.Name;
+			LightCardTree.Add(LightCardTreeItem);
+		}
+
+		for (const TSoftObjectPtr<AActor>& LightCardActor : RootActorLightCards.Actors)
+		{
+			if (LightCardActor.IsValid())
+			{
+				TSharedPtr<FLightCardTreeItem> LightCardTreeItem = MakeShared<FLightCardTreeItem>();
+				LightCardTreeItem->LightCardActor = LightCardActor.Get();
+				LightCardTree.Add(LightCardTreeItem);
+				LightCardActors.Add(LightCardTreeItem);
+			}
+		}
+
+		// If there are any layers that are specified as light card layers, iterate over all actors in the world and 
+		// add any that are members of any of the light card layers to the list. Only add an actor once, even if it is
+		// in multiple layers
+		if (RootActorLightCards.ActorLayers.Num())
+		{
+			if (UWorld* World = RootActor->GetWorld())
+			{
+				for (const TWeakObjectPtr<AActor> WeakActor : FActorRange(World))
+				{
+					if (WeakActor.IsValid())
+					{
+						for (const FActorLayer& ActorLayer : RootActorLightCards.ActorLayers)
+						{
+							if (WeakActor->Layers.Contains(ActorLayer.Name))
+							{
+								TSharedPtr<FLightCardTreeItem> LightCardReference = MakeShared<FLightCardTreeItem>();
+								LightCardReference->LightCardActor = WeakActor.Get();
+								LightCardReference->ActorLayer = ActorLayer.Name;
+								LightCardActors.Add(LightCardReference);
+								break;
+							}
+						}
+					}
+				}
+			}
+		}
+	}
+}
+
+TSharedRef<ITableRow> SDisplayClusterLightCardList::GenerateTreeItemRow(TSharedPtr<FLightCardTreeItem> Item, const TSharedRef<STableViewBase>& OwnerTable)
+{
+	return SNew(SLightCardTreeItemRow, OwnerTable, Item);
+}
+
+void SDisplayClusterLightCardList::GetChildrenForTreeItem(TSharedPtr<FLightCardTreeItem> InItem, TArray<TSharedPtr<FLightCardTreeItem>>& OutChildren)
+{
+	if (InItem.IsValid())
+	{
+		if (InItem->IsActorLayer())
+		{
+			for (const TSharedPtr<FLightCardTreeItem>& LightCardActor : LightCardActors)
+			{
+				if (LightCardActor->IsInActorLayer() && LightCardActor->ActorLayer == InItem->ActorLayer)
+				{
+					OutChildren.Add(LightCardActor);
+				}
+			}
+		}
+	}
+}
+
+void SDisplayClusterLightCardList::OnTreeItemSelected(TSharedPtr<FLightCardTreeItem> InItem, ESelectInfo::Type SelectInfo)
+{
+	if (InItem.IsValid() && InItem->LightCardActor.IsValid())
+	{
+		IDisplayClusterOperator::Get().ShowDetailsForObject(InItem->LightCardActor.Get());
+	}
+	else
+	{
+		IDisplayClusterOperator::Get().ShowDetailsForObject(nullptr);
+	}
+}
+
+#undef LOCTEXT_NAMESPACE
\ No newline at end of file
diff --git a/Engine/Plugins/Runtime/nDisplay/Source/DisplayClusterLightCardEditor/Private/SDisplayClusterLightCardList.h b/Engine/Plugins/Runtime/nDisplay/Source/DisplayClusterLightCardEditor/Private/SDisplayClusterLightCardList.h
new file mode 100644
--- /dev/null
+++ b/Engine/Plugins/Runtime/nDisplay/Source/DisplayClusterLightCardEditor/Private/SDisplayClusterLightCardList.h
@@ -0,0 +1,62 @@
+// Copyright Epic Games, Inc. All Rights Reserved.
+
+#pragma once
+
+#include "CoreMinimal.h"
+#include "Widgets/DeclarativeSyntaxSupport.h"
+#include "Widgets/SCompoundWidget.h"
+
+class ADisplayClusterRootActor;
+class ITableRow;
+class STableViewBase;
+
+template<class T>
+class STreeView;
+
+/** Displays all of the light cards associated with a particular display cluster root actor in a list view */
+class SDisplayClusterLightCardList : public SCompoundWidget
+{
+public:
+	struct FLightCardTreeItem
+	{
+		TWeakObjectPtr<AActor> LightCardActor;
+		FName ActorLayer;
+
+		FLightCardTreeItem() :
+			LightCardActor(nullptr),
+			ActorLayer(NAME_None)
+		{ }
+
+		bool IsInActorLayer() { return !ActorLayer.IsNone() && LightCardActor.IsValid(); }
+		bool IsActorLayer() { return !ActorLayer.IsNone() && !LightCardActor.IsValid(); }
+	};
+
+public:
+	SLATE_BEGIN_ARGS(SDisplayClusterLightCardList)
+	{}
+	SLATE_END_ARGS()
+
+	void Construct(const FArguments& args);
+
+	void SetRootActor(ADisplayClusterRootActor* NewRootActor);
+
+private:
+	void FillLightCardList();
+
+	TSharedRef<ITableRow> GenerateTreeItemRow(TSharedPtr<FLightCardTreeItem> Item, const TSharedRef<STableViewBase>& OwnerTable);
+	void GetChildrenForTreeItem(TSharedPtr<FLightCardTreeItem> InItem, TArray<TSharedPtr<FLightCardTreeItem>>& OutChildren);
+	void OnTreeItemSelected(TSharedPtr<FLightCardTreeItem> InItem, ESelectInfo::Type SelectInfo);
+
+private:
+	/** A hierarchical list of light card tree items to be displayed in the tree view */
+	TArray<TSharedPtr<FLightCardTreeItem>> LightCardTree;
+
+	/** A flattened list of all the light card actors being displayed in the tree view */
+	TArray<TSharedPtr<FLightCardTreeItem>> LightCardActors;
+
+	/** The active root actor whose light cards are being displayed */
+	TWeakObjectPtr<ADisplayClusterRootActor> RootActor;
+
+	/** The tree view widget displaying the light cards */
+	TSharedPtr<STreeView<TSharedPtr<FLightCardTreeItem>>> LightCardTreeView;
+};
\ No newline at end of file
diff --git a/Engine/Plugins/Runtime/nDisplay/Source/DisplayClusterOperator/Private/SDisplayClusterOperatorPanel.cpp b/Engine/Plugins/Runtime/nDisplay/Source/DisplayClusterOperator/Private/SDisplayClusterOperatorPanel.cpp
--- a/Engine/Plugins/Runtime/nDisplay/Source/DisplayClusterOperator/Private/SDisplayClusterOperatorPanel.cpp
+++ b/Engine/Plugins/Runtime/nDisplay/Source/DisplayClusterOperator/Private/SDisplayClusterOperatorPanel.cpp
@@ -84,14 +84,14 @@ void SDisplayClusterOperatorPanel::Construct(const FArguments& InArgs, const TSh
 					(
 						FTabManager::NewStack()
 							->SetExtensionId(TabExtensionId)
-							->SetSizeCoefficient(0.8f)
+							->SetSizeCoefficient(0.67f)
 					)
 					->Split
 					(
 						FTabManager::NewStack()
 							->AddTab(DetailsTabId, ETabState::OpenedTab)
 							->SetHideTabWell(true)
-							->SetSizeCoefficient(0.2f)
+							->SetSizeCoefficient(0.33f)
 					)
 			)
 		);
@@ -106,6 +106,16 @@ void SDisplayClusterOperatorPanel::Construct(const FArguments& InArgs, const TSh
 	[
 		TabManager->RestoreFrom(Layout, WindowOwner).ToSharedRef()
 	];
+
+	// TODO: Move ownership of active root controller to a view model object, instead of letting the toolbar control which root actor is active
+	if (Toolbar.IsValid())
+	{
+		TWeakObjectPtr<ADisplayClusterRootActor> ActiveRootActor = Toolbar->GetActiveRootActor();
+		if (ActiveRootActor.IsValid())
+		{
+			IDisplayClusterOperator::Get().OnActiveRootActorChanged().Broadcast(ActiveRootActor.Get());
+		}
+	}
 }
 
 TSharedRef<SDockTab> SDisplayClusterOperatorPanel::SpawnToolbarTab(const FSpawnTabArgs& Args)
@@ -114,7 +124,7 @@ TSharedRef<SDockTab> SDisplayClusterOperatorPanel::SpawnToolbarTab(const FSpawnT
 		.ShouldAutosize(true)
 		.TabRole(ETabRole::PanelTab)
 		[
-			SNew(SDisplayClusterOperatorToolbar)
+			SAssignNew(Toolbar, SDisplayClusterOperatorToolbar)
 				.CommandList(TSharedPtr<FUICommandList>())
 		];
 }
diff --git a/Engine/Plugins/Runtime/nDisplay/Source/DisplayClusterOperator/Private/SDisplayClusterOperatorPanel.h b/Engine/Plugins/Runtime/nDisplay/Source/DisplayClusterOperator/Private/SDisplayClusterOperatorPanel.h
--- a/Engine/Plugins/Runtime/nDisplay/Source/DisplayClusterOperator/Private/SDisplayClusterOperatorPanel.h
+++ b/Engine/Plugins/Runtime/nDisplay/Source/DisplayClusterOperator/Private/SDisplayClusterOperatorPanel.h
@@ -11,6 +11,7 @@ class FSpawnTabArgs;
 class FMenuBuilder;
 class SDockTab;
 class SWindow;
+class SDisplayClusterOperatorToolbar;
 
 /** The nDisplay operator panel that allows users to edit root actor instances with a variety of tools */
 class SDisplayClusterOperatorPanel : public SCompoundWidget
@@ -63,6 +64,9 @@ private:
 	/** The layout extender used by the operator's layout */
 	TSharedPtr<FLayoutExtender> LayoutExtender;
 
+	/** A reference to the operator panel's toolbar widget */
+	TSharedPtr<SDisplayClusterOperatorToolbar> Toolbar;
+
 	/** A reference to the operator panel's details view */
 	TSharedPtr<class IDetailsView> DetailsView;
 
diff --git a/Engine/Plugins/Runtime/nDisplay/Source/DisplayClusterOperator/Private/SDisplayClusterOperatorToolbar.h b/Engine/Plugins/Runtime/nDisplay/Source/DisplayClusterOperator/Private/SDisplayClusterOperatorToolbar.h
--- a/Engine/Plugins/Runtime/nDisplay/Source/DisplayClusterOperator/Private/SDisplayClusterOperatorToolbar.h
+++ b/Engine/Plugins/Runtime/nDisplay/Source/DisplayClusterOperator/Private/SDisplayClusterOperatorToolbar.h
@@ -26,6 +26,9 @@ public:
 
 	void Construct(const FArguments& InArgs);
 
+	/** Gets the root actor that is currently selected in the toolbar */
+	TWeakObjectPtr<ADisplayClusterRootActor> GetActiveRootActor() const { return ActiveRootActor; }
+
 private:
 	/** 
 	 * Constructs the list of root actors that exist on the current level to use for the root actor picker dropdown.
