nDisplay: Add lightcard editor 3d viewport.

diff --git a/Engine/Plugins/Runtime/nDisplay/Source/DisplayClusterLightCardEditor/DisplayClusterLightCardEditor.Build.cs b/Engine/Plugins/Runtime/nDisplay/Source/DisplayClusterLightCardEditor/DisplayClusterLightCardEditor.Build.cs
--- a/Engine/Plugins/Runtime/nDisplay/Source/DisplayClusterLightCardEditor/DisplayClusterLightCardEditor.Build.cs
+++ b/Engine/Plugins/Runtime/nDisplay/Source/DisplayClusterLightCardEditor/DisplayClusterLightCardEditor.Build.cs
@@ -22,6 +22,7 @@ public class DisplayClusterLightCardEditor : ModuleRules
 				"UnrealEd",
 				"EditorStyle",
 				"WorkspaceMenuStructure",
+				"AdvancedPreviewScene",
 			});
 	}
 }
diff --git a/Engine/Plugins/Runtime/nDisplay/Source/DisplayClusterLightCardEditor/Private/DisplayClusterLightCardEditorCommands.cpp b/Engine/Plugins/Runtime/nDisplay/Source/DisplayClusterLightCardEditor/Private/DisplayClusterLightCardEditorCommands.cpp
new file mode 100644
--- /dev/null
+++ b/Engine/Plugins/Runtime/nDisplay/Source/DisplayClusterLightCardEditor/Private/DisplayClusterLightCardEditorCommands.cpp
@@ -0,0 +1,13 @@
+// Copyright Epic Games, Inc. All Rights Reserved.
+
+#include "DisplayClusterLightCardEditorCommands.h"
+
+#define LOCTEXT_NAMESPACE "DisplayClusterLightCardEditorCommands"
+
+
+void FDisplayClusterLightCardEditorCommands::RegisterCommands()
+{
+	UI_COMMAND(ResetCamera, "Reset Camera", "Resets the camera to focus on the mesh", EUserInterfaceActionType::Button, FInputChord());
+}
+
+#undef LOCTEXT_NAMESPACE
diff --git a/Engine/Plugins/Runtime/nDisplay/Source/DisplayClusterLightCardEditor/Private/DisplayClusterLightCardEditorCommands.h b/Engine/Plugins/Runtime/nDisplay/Source/DisplayClusterLightCardEditor/Private/DisplayClusterLightCardEditorCommands.h
new file mode 100644
--- /dev/null
+++ b/Engine/Plugins/Runtime/nDisplay/Source/DisplayClusterLightCardEditor/Private/DisplayClusterLightCardEditorCommands.h
@@ -0,0 +1,24 @@
+// Copyright Epic Games, Inc. All Rights Reserved.
+
+#pragma once
+
+#include "CoreMinimal.h"
+#include "Framework/Commands/Commands.h"
+#include "EditorStyleSet.h"
+
+
+class FDisplayClusterLightCardEditorCommands 
+	: public TCommands<FDisplayClusterLightCardEditorCommands>
+{
+public:
+	FDisplayClusterLightCardEditorCommands()
+		: TCommands<FDisplayClusterLightCardEditorCommands>(TEXT("DisplayClusterLightCardEditor"), 
+			NSLOCTEXT("Contexts", "DisplayClusterLightCardEditor", "Display Cluster LightCard Editor"), NAME_None, FEditorStyle::GetStyleSetName())
+	{ }
+
+	virtual void RegisterCommands() override;
+
+public:
+	// Viewport commands
+	TSharedPtr<FUICommandInfo> ResetCamera;
+};
diff --git a/Engine/Plugins/Runtime/nDisplay/Source/DisplayClusterLightCardEditor/Private/DisplayClusterLightCardEditorModule.cpp b/Engine/Plugins/Runtime/nDisplay/Source/DisplayClusterLightCardEditor/Private/DisplayClusterLightCardEditorModule.cpp
--- a/Engine/Plugins/Runtime/nDisplay/Source/DisplayClusterLightCardEditor/Private/DisplayClusterLightCardEditorModule.cpp
+++ b/Engine/Plugins/Runtime/nDisplay/Source/DisplayClusterLightCardEditor/Private/DisplayClusterLightCardEditorModule.cpp
@@ -2,6 +2,7 @@
 
 #include "DisplayClusterLightCardEditorModule.h"
 
+#include "DisplayClusterLightCardEditorCommands.h"
 #include "SDisplayClusterLightCardEditor.h"
 
 #define LOCTEXT_NAMESPACE "DisplayClusterLightCardEditor"
@@ -9,11 +10,15 @@
 void FDisplayClusterLightCardEditorModule::StartupModule()
 {
 	RegisterTabSpawners();
+
+	FDisplayClusterLightCardEditorCommands::Register();
 }
 
 void FDisplayClusterLightCardEditorModule::ShutdownModule()
 {
 	UnregisterTabSpawners();
+
+	FDisplayClusterLightCardEditorCommands::Unregister();
 }
 
 void FDisplayClusterLightCardEditorModule::RegisterTabSpawners()
diff --git a/Engine/Plugins/Runtime/nDisplay/Source/DisplayClusterLightCardEditor/Private/SDisplayClusterLightCardEditor.cpp b/Engine/Plugins/Runtime/nDisplay/Source/DisplayClusterLightCardEditor/Private/SDisplayClusterLightCardEditor.cpp
--- a/Engine/Plugins/Runtime/nDisplay/Source/DisplayClusterLightCardEditor/Private/SDisplayClusterLightCardEditor.cpp
+++ b/Engine/Plugins/Runtime/nDisplay/Source/DisplayClusterLightCardEditor/Private/SDisplayClusterLightCardEditor.cpp
@@ -5,6 +5,8 @@
 #include "SDisplayClusterLightCardList.h"
 
 #include "IDisplayClusterOperator.h"
+#include "Viewport/DisplayClusterLightCardEditorViewport.h"
+
 #include "DisplayClusterRootActor.h"
 
 #include "Framework/Docking/TabManager.h"
@@ -55,7 +57,7 @@ TSharedRef<SDockTab> SDisplayClusterLightCardEditor::SpawnInTab(const FSpawnTabA
 
 void SDisplayClusterLightCardEditor::ExtendToolbar(FToolBarBuilder& ToolbarBuilder)
 {
-	// TODO: Any toolbar buttons needed for the lightcards editor can be added to the operator panel's toobar using this toolbar extender
+	// TODO: Any toolbar buttons needed for the lightcards editor can be added to the operator panel's toolbar using this toolbar extender
 }
 
 SDisplayClusterLightCardEditor::~SDisplayClusterLightCardEditor()
@@ -77,23 +79,30 @@ void SDisplayClusterLightCardEditor::Construct(const FArguments& InArgs, const T
 		[
 			// Vertical box for the left hand panel of the editor. Add new slots here as needed for any editor UI controls
 			SNew(SVerticalBox)
-
 			+ SVerticalBox::Slot()
+			.AutoHeight()
 			[
 				CreateLightCardListWidget()
 			]
 		]
-
 		+SSplitter::Slot()
-		.Value(0.75f)
+		[
+			SNew(SVerticalBox)
+			+SVerticalBox::Slot()
+			.FillHeight(0.75f)
+			[
+				CreateViewportWidget()
+			]
+		]
 	];
 }
 
 void SDisplayClusterLightCardEditor::OnActiveRootActorChanged(ADisplayClusterRootActor* NewRootActor)
 {
+	// The new root actor pointer could be null, indicating that it was deleted or the user didn't select a valid root actor
 	ActiveRootActor = NewRootActor;
-
 	LightCardList->SetRootActor(NewRootActor);
+	ViewportView->SetRootActor(NewRootActor);
 }
 
 TSharedRef<SWidget> SDisplayClusterLightCardEditor::CreateLightCardListWidget()
@@ -101,4 +110,9 @@ TSharedRef<SWidget> SDisplayClusterLightCardEditor::CreateLightCardListWidget()
 	return SAssignNew(LightCardList, SDisplayClusterLightCardList);
 }
 
+TSharedRef<SWidget> SDisplayClusterLightCardEditor::CreateViewportWidget()
+{
+	return SAssignNew(ViewportView, SDisplayClusterLightCardEditorViewport, SharedThis(this));
+}
+
 #undef LOCTEXT_NAMESPACE
\ No newline at end of file
diff --git a/Engine/Plugins/Runtime/nDisplay/Source/DisplayClusterLightCardEditor/Private/SDisplayClusterLightCardEditor.h b/Engine/Plugins/Runtime/nDisplay/Source/DisplayClusterLightCardEditor/Private/SDisplayClusterLightCardEditor.h
--- a/Engine/Plugins/Runtime/nDisplay/Source/DisplayClusterLightCardEditor/Private/SDisplayClusterLightCardEditor.h
+++ b/Engine/Plugins/Runtime/nDisplay/Source/DisplayClusterLightCardEditor/Private/SDisplayClusterLightCardEditor.h
@@ -12,6 +12,7 @@ class FSpawnTabArgs;
 class FToolBarBuilder;
 class SDockTab;
 class SDisplayClusterLightCardList;
+class SDisplayClusterLightCardEditorViewport;
 class ADisplayClusterRootActor;
 
 /** A panel that can be spawned in a tab that contains all the UI elements that make up the 2D light cards editor */
@@ -21,7 +22,6 @@ public:
 	/** The name of the tab that the light card editor lives in */
 	static const FName TabName;
 
-
 	/** Registers the light card editor with the global tab manager and adds it to the operator panel's extension tab stack */
 	static void RegisterTabSpawner();
 
@@ -45,6 +45,9 @@ public:
 
 	void Construct(const FArguments& args, const TSharedRef<SDockTab>& MajorTabOwner, const TSharedPtr<SWindow>& WindowOwner);
 
+	/** The current active root actor for this light card editor. */
+	TWeakObjectPtr<ADisplayClusterRootActor> GetActiveRootActor() const { return ActiveRootActor; }
+
 private:
 	/** Raised when the active Display cluster root actor has been changed in the operator panel */
 	void OnActiveRootActorChanged(ADisplayClusterRootActor* NewRootActor);
@@ -52,10 +55,16 @@ private:
 	/** Creates the widget used to show the list of light cards associated with the active root actor */
 	TSharedRef<SWidget> CreateLightCardListWidget();
 
+	/** Create the 3d viewport widget. */
+	TSharedRef<SWidget> CreateViewportWidget();
+
 private:
 	/** The light card list widget */
 	TSharedPtr<SDisplayClusterLightCardList> LightCardList;
 
+	/** The 3d viewport. */
+	TSharedPtr<SDisplayClusterLightCardEditorViewport> ViewportView;
+
 	/** A reference to the root actor that is currently being operated on */
 	TWeakObjectPtr<ADisplayClusterRootActor> ActiveRootActor;
 
diff --git a/Engine/Plugins/Runtime/nDisplay/Source/DisplayClusterLightCardEditor/Private/Viewport/DisplayClusterLightCardEditorViewportClient.cpp b/Engine/Plugins/Runtime/nDisplay/Source/DisplayClusterLightCardEditor/Private/Viewport/DisplayClusterLightCardEditorViewportClient.cpp
new file mode 100644
--- /dev/null
+++ b/Engine/Plugins/Runtime/nDisplay/Source/DisplayClusterLightCardEditor/Private/Viewport/DisplayClusterLightCardEditorViewportClient.cpp
@@ -0,0 +1,283 @@
+// Copyright Epic Games, Inc. All Rights Reserved.
+
+#include "DisplayClusterLightCardEditorViewportClient.h"
+
+#include "DisplayClusterRootActor.h"
+
+#include "AudioDevice.h"
+#include "EngineUtils.h"
+#include "PreviewScene.h"
+#include "UnrealEdGlobals.h"
+#include "ScopedTransaction.h"
+#include "SDisplayClusterLightCardEditor.h"
+#include "UnrealWidget.h"
+#include "Components/DirectionalLightComponent.h"
+#include "Components/LineBatchComponent.h"
+#include "Components/PostProcessComponent.h"
+#include "Components/SkyLightComponent.h"
+#include "Editor/UnrealEdEngine.h"
+#include "Kismet/GameplayStatics.h"
+#include "Slate/SceneViewport.h"
+#include "Widgets/Docking/SDockTab.h"
+
+#define LOCTEXT_NAMESPACE "DisplayClusterLightCardEditorViewportClient"
+
+FDisplayClusterLightCardEditorViewportClient::FDisplayClusterLightCardEditorViewportClient(FAdvancedPreviewScene& InPreviewScene,
+                                                           const TWeakPtr<SEditorViewport>& InEditorViewportWidget,
+                                                           TWeakPtr<SDisplayClusterLightCardEditor> InLightCardEditor) :
+	FEditorViewportClient(nullptr, &InPreviewScene, InEditorViewportWidget),
+	PreviewActorBounds(ForceInitToZero)
+{
+	check (InLightCardEditor.IsValid());
+	
+	LightCardEditorPtr = InLightCardEditor;
+	
+	SelectedActor = nullptr;
+	bDraggingActor = false;
+	ScopedTransaction = nullptr;
+	
+	// Setup defaults for the common draw helper.
+	DrawHelper.bDrawPivot = false;
+	DrawHelper.bDrawWorldBox = false;
+	DrawHelper.bDrawKillZ = false;
+	DrawHelper.bDrawGrid = false;
+	DrawHelper.PerspectiveGridSize = HALF_WORLD_MAX1;
+
+	check(Widget);
+	Widget->SetSnapEnabled(true);
+	
+	ShowWidget(true);
+
+	SetViewMode(VMI_Lit);
+	
+	ViewportType = LVT_Perspective;
+	bSetListenerPosition = false;
+	SetRealtime(true);
+	SetShowStats(true);
+
+	//This seems to be needed to get the correct world time in the preview.
+	SetIsSimulateInEditorViewport(true);
+
+	ResetCamera();
+	
+	UpdatePreviewActor(LightCardEditorPtr.Pin()->GetActiveRootActor().Get());
+}
+
+FDisplayClusterLightCardEditorViewportClient::~FDisplayClusterLightCardEditorViewportClient()
+{
+	EndTransaction();
+}
+
+FLinearColor FDisplayClusterLightCardEditorViewportClient::GetBackgroundColor() const
+{
+	return FLinearColor::Gray;
+}
+
+void FDisplayClusterLightCardEditorViewportClient::Tick(float DeltaSeconds)
+{
+	FEditorViewportClient::Tick(DeltaSeconds);
+
+	// Tick the preview scene world.
+	if (!GIntraFrameDebuggingGameThread)
+	{
+		// Allow full tick only if preview simulation is enabled and we're not currently in an active SIE or PIE session
+		if (GEditor->PlayWorld == nullptr && !GEditor->bIsSimulatingInEditor)
+		{
+			PreviewScene->GetWorld()->Tick(IsRealtime() ? LEVELTICK_All : LEVELTICK_TimeOnly, DeltaSeconds);
+		}
+		else
+		{
+			PreviewScene->GetWorld()->Tick(IsRealtime() ? LEVELTICK_ViewportsOnly : LEVELTICK_TimeOnly, DeltaSeconds);
+		}
+	}
+}
+
+void FDisplayClusterLightCardEditorViewportClient::Draw(FViewport* InViewport, FCanvas* Canvas)
+{
+	FEditorViewportClient::Draw(InViewport, Canvas);
+}
+
+UE::Widget::EWidgetMode FDisplayClusterLightCardEditorViewportClient::GetWidgetMode() const
+{
+	if (IsActorSelected())
+	{
+		return FEditorViewportClient::GetWidgetMode();
+	}
+
+	return UE::Widget::WM_None;
+}
+
+FVector FDisplayClusterLightCardEditorViewportClient::GetWidgetLocation() const
+{
+	if (IsActorSelected())
+	{
+		return SelectedActor->GetActorLocation();
+	}
+
+	return FEditorViewportClient::GetWidgetLocation();
+}
+
+bool FDisplayClusterLightCardEditorViewportClient::InputKey(FViewport* InViewport, int32 ControllerId, FKey Key, EInputEvent Event,
+                                            float AmountDepressed, bool bGamepad)
+{
+	return FEditorViewportClient::InputKey(InViewport, ControllerId, Key, Event, AmountDepressed, bGamepad);
+}
+
+bool FDisplayClusterLightCardEditorViewportClient::InputWidgetDelta(FViewport* InViewport, EAxisList::Type CurrentAxis, FVector& Drag,
+                                                    FRotator& Rot, FVector& Scale)
+{
+	if (IsActorSelected() && bDraggingActor)
+	{
+		GEditor->ApplyDeltaToActor(SelectedActor.Get(), true, &Drag, &Rot, &Scale);
+		return true;
+	}
+
+	return FEditorViewportClient::InputWidgetDelta(InViewport, CurrentAxis, Drag, Rot, Scale);
+}
+
+void FDisplayClusterLightCardEditorViewportClient::TrackingStarted(const FInputEventState& InInputState, bool bIsDraggingWidget,
+	bool bNudge)
+{
+	if (!bDraggingActor && bIsDraggingWidget && InInputState.IsLeftMouseButtonPressed() && IsActorSelected())
+	{
+		GEditor->DisableDeltaModification(true);
+		{
+			// The pivot location won't update properly and the actor will rotate / move around the original selection origin
+			// so update it here to fix that.
+			GUnrealEd->UpdatePivotLocationForSelection();
+			GUnrealEd->SetPivotMovedIndependently(false);
+		}
+
+		BeginTransaction(NSLOCTEXT("LogicDriverPreview", "ModifyPreviewActor", "Modify a Preview Actor"));
+		bDraggingActor = true;
+	}
+	FEditorViewportClient::TrackingStarted(InInputState, bIsDraggingWidget, bNudge);
+}
+
+void FDisplayClusterLightCardEditorViewportClient::TrackingStopped()
+{
+	bDraggingActor = false;
+	EndTransaction();
+
+	if (IsActorSelected())
+	{
+		GEditor->DisableDeltaModification(false);
+	}
+	
+	FEditorViewportClient::TrackingStopped();
+}
+
+void FDisplayClusterLightCardEditorViewportClient::ProcessClick(FSceneView& View, HHitProxy* HitProxy, FKey Key, EInputEvent Event,
+                                                uint32 HitX, uint32 HitY)
+{
+	if (HitProxy && HitProxy->IsA(HActor::StaticGetType()))
+	{
+		HActor* HitActor = static_cast<HActor*>(HitProxy);
+		SelectActor(HitActor->Actor);
+		return;
+	}
+
+	SelectActor(nullptr);
+	
+	FEditorViewportClient::ProcessClick(View, HitProxy, Key, Event, HitX, HitY);
+}
+
+void FDisplayClusterLightCardEditorViewportClient::SetSceneViewport(TSharedPtr<FSceneViewport> InViewport)
+{
+	SceneViewportPtr = InViewport;
+}
+
+void FDisplayClusterLightCardEditorViewportClient::SelectActor(AActor* NewActor)
+{
+}
+
+void FDisplayClusterLightCardEditorViewportClient::ResetSelection()
+{
+	SelectActor(nullptr);
+	SetWidgetMode(UE::Widget::WM_None);
+}
+
+void FDisplayClusterLightCardEditorViewportClient::ResetCamera()
+{
+	ToggleOrbitCamera(true);
+	SetViewLocationForOrbiting(PreviewActorBounds.Origin);
+
+	// TODO: Better define defaults.
+	SetViewLocation(EditorViewportDefs::DefaultPerspectiveViewLocation + FVector(0.f, -400.f, 0.f));
+	SetViewRotation(EditorViewportDefs::DefaultPerspectiveViewRotation + FRotator(0.f, 180.f, 0.f));
+
+	Invalidate();
+}
+
+void FDisplayClusterLightCardEditorViewportClient::UpdatePreviewActor(
+	ADisplayClusterRootActor* RootActor)
+{
+	UWorld* PreviewWorld = PreviewScene->GetWorld();
+	check (PreviewWorld);
+
+	if (SpawnedRootActor.IsValid() && (!RootActor || !PreviewWorld->ContainsActor(RootActor)))
+	{
+		PreviewWorld->EditorDestroyActor(SpawnedRootActor.Get(), false);
+		SpawnedRootActor.Reset();
+	}
+	
+	if (RootActor)
+	{
+		const FVector SpawnLocation = FVector::ZeroVector;
+		const FRotator SpawnRotation = FRotator::ZeroRotator;
+		
+		FActorSpawnParameters SpawnInfo;
+		SpawnInfo.SpawnCollisionHandlingOverride = ESpawnActorCollisionHandlingMethod::AlwaysSpawn;
+		SpawnInfo.bNoFail = true;
+		SpawnInfo.ObjectFlags = RF_Transient|RF_Transactional;
+		SpawnInfo.Template = RootActor;
+		
+		SpawnedRootActor = CastChecked<ADisplayClusterRootActor>(PreviewWorld->SpawnActor(RootActor->GetClass(),
+			&SpawnLocation, &SpawnRotation, MoveTemp(SpawnInfo)));
+
+		// Compute actor bounds as the sum of its visible parts
+		PreviewActorBounds = FBoxSphereBounds(ForceInitToZero);
+		for (UActorComponent* Component : RootActor->GetComponents())
+		{
+			// Aggregate primitive components that either have collision enabled or are otherwise visible components in-game
+			if (UPrimitiveComponent* PrimComp = Cast<UPrimitiveComponent>(Component))
+			{
+				if (PrimComp->IsRegistered() && (!PrimComp->bHiddenInGame || PrimComp->IsCollisionEnabled()) && PrimComp->Bounds.SphereRadius < HALF_WORLD_MAX)
+				{
+					PreviewActorBounds = PreviewActorBounds + PrimComp->Bounds;
+				}
+			}
+		}
+
+		SpawnedRootActor->UpdatePreviewComponents();
+	}
+}
+
+bool FDisplayClusterLightCardEditorViewportClient::GetShowGrid() const
+{
+	return IsSetShowGridChecked();
+}
+
+void FDisplayClusterLightCardEditorViewportClient::ToggleShowGrid()
+{
+	SetShowGrid();
+}
+
+void FDisplayClusterLightCardEditorViewportClient::BeginTransaction(const FText& Description)
+{
+	if (!ScopedTransaction)
+	{
+		ScopedTransaction = new FScopedTransaction(Description);
+	}
+}
+
+void FDisplayClusterLightCardEditorViewportClient::EndTransaction()
+{
+	if (ScopedTransaction)
+	{
+		delete ScopedTransaction;
+		ScopedTransaction = nullptr;
+	}
+}
+
+#undef LOCTEXT_NAMESPACE
diff --git a/Engine/Plugins/Runtime/nDisplay/Source/DisplayClusterLightCardEditor/Private/Viewport/DisplayClusterLightCardEditorViewportClient.h b/Engine/Plugins/Runtime/nDisplay/Source/DisplayClusterLightCardEditor/Private/Viewport/DisplayClusterLightCardEditorViewportClient.h
new file mode 100644
--- /dev/null
+++ b/Engine/Plugins/Runtime/nDisplay/Source/DisplayClusterLightCardEditor/Private/Viewport/DisplayClusterLightCardEditorViewportClient.h
@@ -0,0 +1,79 @@
+// Copyright Epic Games, Inc. All Rights Reserved.
+
+#pragma once
+
+#include "CoreMinimal.h"
+#include "SCommonEditorViewportToolbarBase.h"
+#include "SEditorViewport.h"
+#include "AdvancedPreviewScene.h"
+
+class ADisplayClusterRootActor;
+class SDisplayClusterLightCardEditor;
+class FScopedTransaction;
+
+/** Viewport Client for the preview viewport */
+class FDisplayClusterLightCardEditorViewportClient : public FEditorViewportClient, public TSharedFromThis<FDisplayClusterLightCardEditorViewportClient>
+{
+public:
+	FDisplayClusterLightCardEditorViewportClient(FAdvancedPreviewScene& InPreviewScene, const TWeakPtr<SEditorViewport>& InEditorViewportWidget,
+		TWeakPtr<SDisplayClusterLightCardEditor> InLightCardEditor);
+	virtual ~FDisplayClusterLightCardEditorViewportClient() override;
+	
+	// FEditorViewportClient
+	virtual FLinearColor GetBackgroundColor() const override;
+	virtual void Tick(float DeltaSeconds) override;
+	virtual void Draw(FViewport* InViewport, FCanvas* Canvas) override;
+	virtual bool CanSetWidgetMode(UE::Widget::EWidgetMode NewMode) const override { return true; }
+	virtual bool CanCycleWidgetMode() const override { return true; }
+	virtual UE::Widget::EWidgetMode GetWidgetMode() const override;
+	virtual FVector GetWidgetLocation() const override;
+	virtual bool IsLevelEditorClient() const override { return false; }
+	virtual bool InputKey(FViewport* InViewport, int32 ControllerId, FKey Key, EInputEvent Event, float AmountDepressed, bool bGamepad) override;
+	virtual bool InputWidgetDelta(FViewport* InViewport, EAxisList::Type CurrentAxis, FVector& Drag, FRotator& Rot, FVector& Scale) override;
+	virtual void TrackingStarted(const FInputEventState& InInputState, bool bIsDraggingWidget, bool bNudge) override;
+	virtual void TrackingStopped() override;
+	virtual void ProcessClick(FSceneView& View, HHitProxy* HitProxy, FKey Key, EInputEvent Event, uint32 HitX, uint32 HitY) override;
+	// ~FEditorViewportClient
+
+	void SetSceneViewport(TSharedPtr<FSceneViewport> InViewport);
+	void SelectActor(AActor* NewActor);
+	void ResetSelection();
+	void ResetCamera();
+
+	void UpdatePreviewActor(ADisplayClusterRootActor* RootActor);
+	
+	/**
+	 * Returns true if the grid is currently visible in the viewport
+	 */
+	bool GetShowGrid() const;
+
+	/**
+	 * Will toggle the grid's visibility in the viewport
+	 */
+	void ToggleShowGrid();
+	
+	AActor* GetSelectedActor() const { return SelectedActor.Get(); }
+
+	bool IsActorSelected() const { return SelectedActor.IsValid(); }
+	
+protected:
+	/** Initiates a transaction. */
+	void BeginTransaction(const FText& Description);
+
+	/** Ends the current transaction, if one exists. */
+	void EndTransaction();
+
+private:
+	TWeakPtr<FSceneViewport> SceneViewportPtr;
+	TWeakPtr<SDisplayClusterLightCardEditor> LightCardEditorPtr;
+	TWeakObjectPtr<ADisplayClusterRootActor> SpawnedRootActor;
+	TWeakObjectPtr<AActor> SelectedActor;
+	
+	/** The full bounds of the preview scene (encompasses all visible components) */
+	FBoxSphereBounds PreviewActorBounds;
+	
+	/** The current transaction for undo/redo */
+	FScopedTransaction* ScopedTransaction;
+	
+	bool bDraggingActor;
+};
\ No newline at end of file
diff --git a/Engine/Plugins/Runtime/nDisplay/Source/DisplayClusterLightCardEditor/Private/Viewport/DisplayClusterLightcardEditorViewport.cpp b/Engine/Plugins/Runtime/nDisplay/Source/DisplayClusterLightCardEditor/Private/Viewport/DisplayClusterLightcardEditorViewport.cpp
new file mode 100644
--- /dev/null
+++ b/Engine/Plugins/Runtime/nDisplay/Source/DisplayClusterLightCardEditor/Private/Viewport/DisplayClusterLightcardEditorViewport.cpp
@@ -0,0 +1,292 @@
+// Copyright Epic Games, Inc. All Rights Reserved.
+
+#include "DisplayClusterLightCardEditorViewport.h"
+
+#include "DisplayClusterLightCardEditorViewportClient.h"
+#include "DisplayClusterLightCardEditorCommands.h"
+
+#include "EditorViewportCommands.h"
+#include "SEditorViewportToolBarMenu.h"
+#include "STransformViewportToolbar.h"
+#include "Slate/SceneViewport.h"
+
+#define LOCTEXT_NAMESPACE "DisplayClusterLightcardEditorViewport"
+
+class SDisplayClusterLightCardEditorViewportToolBar : public SViewportToolBar
+{
+public:
+	SLATE_BEGIN_ARGS(SDisplayClusterLightCardEditorViewportToolBar){}
+		SLATE_ARGUMENT(TWeakPtr<SDisplayClusterLightCardEditorViewport>, EditorViewport)
+	SLATE_END_ARGS()
+
+	/** Constructs this widget with the given parameters */
+	void Construct(const FArguments& InArgs)
+	{
+		EditorViewport = InArgs._EditorViewport;
+		static const FName DefaultForegroundName("DefaultForeground");
+
+		this->ChildSlot
+		[
+			SNew(SBorder)
+			.BorderImage(FEditorStyle::GetBrush("NoBorder"))
+			.ForegroundColor(FEditorStyle::GetSlateColor(DefaultForegroundName))
+			[
+				SNew(SHorizontalBox)
+				+SHorizontalBox::Slot()
+				.AutoWidth()
+				.Padding(2.0f, 2.0f)
+				[
+					SNew(SEditorViewportToolbarMenu)
+					.ParentToolBar(SharedThis(this))
+					.Cursor(EMouseCursor::Default)
+					.Image("EditorViewportToolBar.MenuDropdown")
+					.OnGetMenuContent(this, &SDisplayClusterLightCardEditorViewportToolBar::GeneratePreviewMenu)
+				]
+				+ SHorizontalBox::Slot()
+				.AutoWidth()
+				.Padding(2.0f, 2.0f)
+				[
+					SNew( SEditorViewportToolbarMenu )
+					.ParentToolBar( SharedThis( this ) )
+					.Cursor( EMouseCursor::Default )
+					.Label(this, &SDisplayClusterLightCardEditorViewportToolBar::GetCameraMenuLabel)
+					.LabelIcon(this, &SDisplayClusterLightCardEditorViewportToolBar::GetCameraMenuLabelIcon)
+					.OnGetMenuContent(this, &SDisplayClusterLightCardEditorViewportToolBar::GenerateCameraMenu)
+				]
+				+ SHorizontalBox::Slot()
+				.AutoWidth()
+				.Padding(2.0f, 2.0f)
+				[
+					SNew( SEditorViewportToolbarMenu )
+					.ParentToolBar( SharedThis( this ) )
+					.Cursor( EMouseCursor::Default )
+					.Label(this, &SDisplayClusterLightCardEditorViewportToolBar::GetViewMenuLabel)
+					.LabelIcon(this, &SDisplayClusterLightCardEditorViewportToolBar::GetViewMenuLabelIcon)
+					.OnGetMenuContent(this, &SDisplayClusterLightCardEditorViewportToolBar::GenerateViewMenu)
+				]
+				+ SHorizontalBox::Slot()
+				.Padding( 3.0f, 1.0f )
+				.HAlign( HAlign_Right )
+				[
+					SNew(STransformViewportToolBar)
+					.Viewport(EditorViewport.Pin().ToSharedRef())
+					.CommandList(EditorViewport.Pin()->GetCommandList())
+				]
+			]
+		];
+
+		SViewportToolBar::Construct(SViewportToolBar::FArguments());
+	}
+
+	/** Creates the preview menu */
+	TSharedRef<SWidget> GeneratePreviewMenu() const
+	{
+		TSharedPtr<const FUICommandList> CommandList = EditorViewport.IsValid()? EditorViewport.Pin()->GetCommandList(): nullptr;
+
+		const bool bInShouldCloseWindowAfterMenuSelection = true;
+
+		FMenuBuilder PreviewOptionsMenuBuilder(bInShouldCloseWindowAfterMenuSelection, CommandList);
+		{
+			PreviewOptionsMenuBuilder.BeginSection("BlueprintEditorPreviewOptions", NSLOCTEXT("BlueprintEditor", "PreviewOptionsMenuHeader", "Preview Viewport Options"));
+			{
+				PreviewOptionsMenuBuilder.AddMenuEntry(FDisplayClusterLightCardEditorCommands::Get().ResetCamera);
+				PreviewOptionsMenuBuilder.AddMenuEntry(FEditorViewportCommands::Get().ToggleRealTime);
+			}
+			PreviewOptionsMenuBuilder.EndSection();
+		}
+
+		return PreviewOptionsMenuBuilder.MakeWidget();
+	}
+
+	FText GetCameraMenuLabel() const
+	{
+		if(EditorViewport.IsValid())
+		{
+			return GetCameraMenuLabelFromViewportType(EditorViewport.Pin()->GetViewportClient()->GetViewportType());
+		}
+
+		return NSLOCTEXT("BlueprintEditor", "CameraMenuTitle_Default", "Camera");
+	}
+
+	const FSlateBrush* GetCameraMenuLabelIcon() const
+	{
+		if(EditorViewport.IsValid())
+		{
+			return GetCameraMenuLabelIconFromViewportType( EditorViewport.Pin()->GetViewportClient()->GetViewportType() );
+		}
+
+		return FEditorStyle::GetBrush(NAME_None);
+	}
+
+	TSharedRef<SWidget> GenerateCameraMenu() const
+	{
+		TSharedPtr<const FUICommandList> CommandList = EditorViewport.IsValid()? EditorViewport.Pin()->GetCommandList(): nullptr;
+
+		const bool bInShouldCloseWindowAfterMenuSelection = true;
+		FMenuBuilder CameraMenuBuilder(bInShouldCloseWindowAfterMenuSelection, CommandList);
+
+		CameraMenuBuilder.AddMenuEntry(FEditorViewportCommands::Get().Perspective);
+
+		CameraMenuBuilder.BeginSection("LevelViewportCameraType_Ortho", NSLOCTEXT("BlueprintEditor", "CameraTypeHeader_Ortho", "Orthographic"));
+			CameraMenuBuilder.AddMenuEntry(FEditorViewportCommands::Get().Top);
+			CameraMenuBuilder.AddMenuEntry(FEditorViewportCommands::Get().Bottom);
+			CameraMenuBuilder.AddMenuEntry(FEditorViewportCommands::Get().Left);
+			CameraMenuBuilder.AddMenuEntry(FEditorViewportCommands::Get().Right);
+			CameraMenuBuilder.AddMenuEntry(FEditorViewportCommands::Get().Front);
+			CameraMenuBuilder.AddMenuEntry(FEditorViewportCommands::Get().Back);
+		CameraMenuBuilder.EndSection();
+
+		return CameraMenuBuilder.MakeWidget();
+	}
+
+	FText GetViewMenuLabel() const
+	{
+		FText Label = NSLOCTEXT("BlueprintEditor", "ViewMenuTitle_Default", "View");
+
+		if (EditorViewport.IsValid())
+		{
+			switch (EditorViewport.Pin()->GetViewportClient()->GetViewMode())
+			{
+			case VMI_Lit:
+				Label = NSLOCTEXT("BlueprintEditor", "ViewMenuTitle_Lit", "Lit");
+				break;
+
+			case VMI_Unlit:
+				Label = NSLOCTEXT("BlueprintEditor", "ViewMenuTitle_Unlit", "Unlit");
+				break;
+
+			case VMI_BrushWireframe:
+				Label = NSLOCTEXT("BlueprintEditor", "ViewMenuTitle_Wireframe", "Wireframe");
+				break;
+			}
+		}
+
+		return Label;
+	}
+
+	const FSlateBrush* GetViewMenuLabelIcon() const
+	{
+		static FName LitModeIconName("EditorViewport.LitMode");
+		static FName UnlitModeIconName("EditorViewport.UnlitMode");
+		static FName WireframeModeIconName("EditorViewport.WireframeMode");
+
+		FName Icon = NAME_None;
+
+		if (EditorViewport.IsValid())
+		{
+			switch (EditorViewport.Pin()->GetViewportClient()->GetViewMode())
+			{
+			case VMI_Lit:
+				Icon = LitModeIconName;
+				break;
+
+			case VMI_Unlit:
+				Icon = UnlitModeIconName;
+				break;
+
+			case VMI_BrushWireframe:
+				Icon = WireframeModeIconName;
+				break;
+			}
+		}
+
+		return FEditorStyle::GetBrush(Icon);
+	}
+
+	TSharedRef<SWidget> GenerateViewMenu() const
+	{
+		TSharedPtr<const FUICommandList> CommandList = EditorViewport.IsValid() ? EditorViewport.Pin()->GetCommandList() : nullptr;
+
+		const bool bInShouldCloseWindowAfterMenuSelection = true;
+		FMenuBuilder ViewMenuBuilder(bInShouldCloseWindowAfterMenuSelection, CommandList);
+
+		ViewMenuBuilder.AddMenuEntry(FEditorViewportCommands::Get().LitMode, NAME_None, NSLOCTEXT("BlueprintEditor", "LitModeMenuOption", "Lit"));
+		ViewMenuBuilder.AddMenuEntry(FEditorViewportCommands::Get().UnlitMode, NAME_None, NSLOCTEXT("BlueprintEditor", "UnlitModeMenuOption", "Unlit"));
+		ViewMenuBuilder.AddMenuEntry(FEditorViewportCommands::Get().WireframeMode, NAME_None, NSLOCTEXT("BlueprintEditor", "WireframeModeMenuOption", "Wireframe"));
+
+		return ViewMenuBuilder.MakeWidget();
+	}
+
+private:
+	/** Reference to the parent viewport */
+	TWeakPtr<SDisplayClusterLightCardEditorViewport> EditorViewport;
+};
+
+void SDisplayClusterLightCardEditorViewport::Construct(const FArguments& InArgs, TSharedPtr<SDisplayClusterLightCardEditor> InLightCardEditor)
+{
+	LightCardEditorPtr = InLightCardEditor;
+	
+	AdvancedPreviewScene = MakeShared<FAdvancedPreviewScene>(FPreviewScene::ConstructionValues());
+	AdvancedPreviewScene->SetFloorVisibility(true);
+		
+	SEditorViewport::Construct(SEditorViewport::FArguments());
+
+	ViewportClient->SetSceneViewport(SceneViewport);
+}
+
+SDisplayClusterLightCardEditorViewport::~SDisplayClusterLightCardEditorViewport()
+{
+	if (ViewportClient.IsValid())
+	{
+		ViewportClient->Viewport = nullptr;
+		ViewportClient.Reset();
+	}
+}
+
+TSharedRef<SEditorViewport> SDisplayClusterLightCardEditorViewport::GetViewportWidget()
+{
+	return SharedThis(this);
+}
+
+TSharedPtr<FExtender> SDisplayClusterLightCardEditorViewport::GetExtenders() const
+{
+	TSharedPtr<FExtender> Result(MakeShareable(new FExtender));
+	return Result;
+}
+
+void SDisplayClusterLightCardEditorViewport::OnFloatingButtonClicked()
+{
+}
+
+void SDisplayClusterLightCardEditorViewport::SetRootActor(ADisplayClusterRootActor* NewRootActor)
+{
+	if (ViewportClient.IsValid())
+	{
+		ViewportClient->UpdatePreviewActor(NewRootActor);
+	}
+}
+
+TSharedRef<FEditorViewportClient> SDisplayClusterLightCardEditorViewport::MakeEditorViewportClient()
+{
+	check(AdvancedPreviewScene.IsValid());
+	
+	ViewportClient = MakeShareable(new FDisplayClusterLightCardEditorViewportClient(*AdvancedPreviewScene.Get(),
+		SharedThis(this), LightCardEditorPtr));
+	return ViewportClient.ToSharedRef();
+}
+
+TSharedPtr<SWidget> SDisplayClusterLightCardEditorViewport::MakeViewportToolbar()
+{
+	return
+		SNew(SDisplayClusterLightCardEditorViewportToolBar)
+		.EditorViewport(SharedThis(this))
+		.IsEnabled(FSlateApplication::Get().GetNormalExecutionAttribute());
+}
+
+void SDisplayClusterLightCardEditorViewport::PopulateViewportOverlays(TSharedRef<SOverlay> Overlay)
+{
+	SEditorViewport::PopulateViewportOverlays(Overlay);
+}
+
+void SDisplayClusterLightCardEditorViewport::BindCommands()
+{
+	const FDisplayClusterLightCardEditorCommands& Commands = FDisplayClusterLightCardEditorCommands::Get();
+	
+	CommandList->MapAction(
+	Commands.ResetCamera,
+	FExecuteAction::CreateSP(ViewportClient.Get(), &FDisplayClusterLightCardEditorViewportClient::ResetCamera));
+
+	SEditorViewport::BindCommands();
+}
+
+#undef LOCTEXT_NAMESPACE
diff --git a/Engine/Plugins/Runtime/nDisplay/Source/DisplayClusterLightCardEditor/Private/Viewport/DisplayClusterLightcardEditorViewport.h b/Engine/Plugins/Runtime/nDisplay/Source/DisplayClusterLightCardEditor/Private/Viewport/DisplayClusterLightcardEditorViewport.h
new file mode 100644
--- /dev/null
+++ b/Engine/Plugins/Runtime/nDisplay/Source/DisplayClusterLightCardEditor/Private/Viewport/DisplayClusterLightcardEditorViewport.h
@@ -0,0 +1,51 @@
+// Copyright Epic Games, Inc. All Rights Reserved.
+
+#pragma once
+
+#include "CoreMinimal.h"
+#include "AdvancedPreviewScene.h"
+#include "SCommonEditorViewportToolbarBase.h"
+#include "SEditorViewport.h"
+
+class SDisplayClusterLightCardEditor;
+class ADisplayClusterRootActor;
+
+/**
+ * Slate widget which renders our view client.
+ */
+class SDisplayClusterLightCardEditorViewport : public SEditorViewport, public ICommonEditorViewportToolbarInfoProvider
+{
+public:
+	SLATE_BEGIN_ARGS(SDisplayClusterLightCardEditorViewport) {}
+	SLATE_END_ARGS()
+
+	void Construct(const FArguments& InArgs, TSharedPtr<SDisplayClusterLightCardEditor> InLightCardEditor);
+	~SDisplayClusterLightCardEditorViewport();
+
+	// ICommonEditorViewportToolbarInfoProvider
+	virtual TSharedRef<SEditorViewport> GetViewportWidget() override;
+	virtual TSharedPtr<FExtender> GetExtenders() const override;
+	virtual void OnFloatingButtonClicked() override;
+	// ~ICommonEditorViewportToolbarInfoProvider
+
+	TWeakPtr<SDisplayClusterLightCardEditor> GetLightCardEditor() { return LightCardEditorPtr;}
+
+	void SetRootActor(ADisplayClusterRootActor* NewRootActor);
+	
+private:
+	// SEditorViewport
+	virtual TSharedRef<FEditorViewportClient> MakeEditorViewportClient() override;
+	virtual TSharedPtr<SWidget> MakeViewportToolbar() override;
+	virtual void PopulateViewportOverlays(TSharedRef<SOverlay> Overlay) override;
+	virtual void BindCommands() override;
+	// ~SEditorViewport
+
+private:
+	/** Preview Scene - uses advanced preview settings */
+	TSharedPtr<FAdvancedPreviewScene> AdvancedPreviewScene;
+	
+	/** Level viewport client */
+	TSharedPtr<class FDisplayClusterLightCardEditorViewportClient> ViewportClient;
+	TWeakPtr<SDisplayClusterLightCardEditor> LightCardEditorPtr;
+
+};
\ No newline at end of file
