Fix build failure, incorrect case on header

Preflight was not correctly ran which would have caught this

diff --git a/Engine/Plugins/Runtime/nDisplay/Source/DisplayClusterLightCardEditor/Private/SDisplayClusterLightCardEditor.cpp b/Engine/Plugins/Runtime/nDisplay/Source/DisplayClusterLightCardEditor/Private/SDisplayClusterLightCardEditor.cpp
--- a/Engine/Plugins/Runtime/nDisplay/Source/DisplayClusterLightCardEditor/Private/SDisplayClusterLightCardEditor.cpp
+++ b/Engine/Plugins/Runtime/nDisplay/Source/DisplayClusterLightCardEditor/Private/SDisplayClusterLightCardEditor.cpp
@@ -5,7 +5,7 @@
 #include "SDisplayClusterLightCardList.h"
 
 #include "IDisplayClusterOperator.h"
-#include "Viewport/DisplayClusterLightCardEditorViewport.h"
+#include "Viewport/DisplayClusterLightcardEditorViewport.h"
 
 #include "DisplayClusterRootActor.h"
 
@@ -115,4 +115,4 @@ TSharedRef<SWidget> SDisplayClusterLightCardEditor::CreateViewportWidget()
 	return SAssignNew(ViewportView, SDisplayClusterLightCardEditorViewport, SharedThis(this));
 }
 
-#undef LOCTEXT_NAMESPACE
\ No newline at end of file
+#undef LOCTEXT_NAMESPACE
