nDisplay: fixed post processing rendering for icvfx's cinecamera and camera projection policy in runtime/editor

diff --git a/Engine/Plugins/Runtime/nDisplay/Source/DisplayCluster/Private/Render/Viewport/Configuration/DisplayClusterViewportConfigurationProjectionPolicy.cpp b/Engine/Plugins/Runtime/nDisplay/Source/DisplayCluster/Private/Render/Viewport/Configuration/DisplayClusterViewportConfigurationProjectionPolicy.cpp
--- a/Engine/Plugins/Runtime/nDisplay/Source/DisplayCluster/Private/Render/Viewport/Configuration/DisplayClusterViewportConfigurationProjectionPolicy.cpp
+++ b/Engine/Plugins/Runtime/nDisplay/Source/DisplayCluster/Private/Render/Viewport/Configuration/DisplayClusterViewportConfigurationProjectionPolicy.cpp
@@ -7,6 +7,7 @@
 #include "Render/Viewport/DisplayClusterViewportManager.h"
 
 #include "DisplayClusterViewportConfigurationHelpers.h"
+#include "DisplayClusterViewportConfigurationHelpers_Postprocess.h"
 #include "DisplayClusterConfigurationTypes.h"
 
 #include "IDisplayClusterProjection.h"
@@ -92,6 +93,13 @@ bool FDisplayClusterViewportConfigurationProjectionPolicy::UpdateCameraPolicy_Ba
 {
 	check(InCameraComponent);
 
+	// add camera's post processing materials
+	FMinimalViewInfo DesiredView;
+	const float DeltaTime = RootActor.GetWorldDeltaSeconds();
+	InCameraComponent->GetCameraView(DeltaTime, DesiredView);
+
+	DstViewport.CustomPostProcessSettings.AddCustomPostProcess(IDisplayClusterViewport_CustomPostProcessSettings::ERenderPass::Override, DesiredView.PostProcessSettings, 1.0, true);
+
 	FDisplayClusterProjectionCameraPolicySettings PolicyCameraSettings;
 	PolicyCameraSettings.FOVMultiplier = FOVMultiplier;
 
diff --git a/Engine/Plugins/Runtime/nDisplay/Source/DisplayCluster/Private/Render/Viewport/DisplayClusterViewport.cpp b/Engine/Plugins/Runtime/nDisplay/Source/DisplayCluster/Private/Render/Viewport/DisplayClusterViewport.cpp
--- a/Engine/Plugins/Runtime/nDisplay/Source/DisplayCluster/Private/Render/Viewport/DisplayClusterViewport.cpp
+++ b/Engine/Plugins/Runtime/nDisplay/Source/DisplayCluster/Private/Render/Viewport/DisplayClusterViewport.cpp
@@ -133,6 +133,14 @@ void FDisplayClusterViewport::HandleEndScene()
 		UninitializedProjectionPolicy = ProjectionPolicy;
 		ProjectionPolicy.Reset();
 	}
+
+#if WITH_EDITOR
+	// cleanup view states
+	for (int32 ViewIndex = 0; ViewIndex < ViewStates.Num(); ViewIndex++)
+	{
+		ViewStates[ViewIndex].Destroy();
+	}
+#endif //WITH_EDITOR
 }
 
 bool FDisplayClusterViewport::ShouldUseAdditionalTargetableResource() const
diff --git a/Engine/Plugins/Runtime/nDisplay/Source/DisplayCluster/Private/Render/Viewport/DisplayClusterViewport.h b/Engine/Plugins/Runtime/nDisplay/Source/DisplayCluster/Private/Render/Viewport/DisplayClusterViewport.h
--- a/Engine/Plugins/Runtime/nDisplay/Source/DisplayCluster/Private/Render/Viewport/DisplayClusterViewport.h
+++ b/Engine/Plugins/Runtime/nDisplay/Source/DisplayCluster/Private/Render/Viewport/DisplayClusterViewport.h
@@ -114,6 +114,9 @@ public:
 	FSceneView* ImplCalcScenePreview(class FSceneViewFamilyContext& InOutViewFamily, uint32 ContextNum);
 	bool    ImplPreview_CalculateStereoViewOffset(const uint32 InContextNum, FRotator& ViewRotation, const float WorldToMeters, FVector& ViewLocation);
 	FMatrix ImplPreview_GetStereoProjectionMatrix(const uint32 InContextNum);
+
+	TIndirectArray<FSceneViewStateReference> ViewStates;
+	FSceneViewStateInterface* GetViewState(int32 ViewIndex);
 #endif //WITH_EDITOR
 
 	// Get from logic request for additional targetable resource
diff --git a/Engine/Plugins/Runtime/nDisplay/Source/DisplayCluster/Private/Render/Viewport/DisplayClusterViewport_EditorPreview.cpp b/Engine/Plugins/Runtime/nDisplay/Source/DisplayCluster/Private/Render/Viewport/DisplayClusterViewport_EditorPreview.cpp
--- a/Engine/Plugins/Runtime/nDisplay/Source/DisplayCluster/Private/Render/Viewport/DisplayClusterViewport_EditorPreview.cpp
+++ b/Engine/Plugins/Runtime/nDisplay/Source/DisplayCluster/Private/Render/Viewport/DisplayClusterViewport_EditorPreview.cpp
@@ -89,6 +89,9 @@ FSceneView* FDisplayClusterViewport::ImplCalcScenePreview(FSceneViewFamilyContex
 		ViewInitOptions.BackgroundColor = FLinearColor::Black;
 		//ViewInitOptions.OverlayColor = FLinearColor::Black;
 
+		// initialize view states
+		ViewInitOptions.SceneViewStateInterface = GetViewState(InContextNum);
+		
 		FSceneView* View = new FSceneView(ViewInitOptions);
 
 		View->bIsSceneCapture = true;
@@ -146,6 +149,29 @@ FSceneView* FDisplayClusterViewport::ImplCalcScenePreview(FSceneViewFamilyContex
 	return nullptr;
 }
 
+FSceneViewStateInterface* FDisplayClusterViewport::GetViewState(int32 ViewIndex)
+{
+	while (ViewIndex >= ViewStates.Num())
+	{
+		ViewStates.Add(new FSceneViewStateReference());
+	}
+
+	FSceneViewStateInterface* ViewStateInterface = ViewStates[ViewIndex].GetReference();
+	if (ViewStateInterface == NULL)
+	{
+		ERHIFeatureLevel::Type FeatureLevel = GMaxRHIFeatureLevel;
+		if (Owner.GetCurrentWorld() != nullptr)
+		{
+			FeatureLevel = Owner.GetCurrentWorld()->Scene->GetFeatureLevel();
+		}
+
+		ViewStates[ViewIndex].Allocate(FeatureLevel);
+		ViewStateInterface = ViewStates[ViewIndex].GetReference();
+	}
+	
+	return ViewStateInterface;
+}
+
 enum EDisplayClusterEyeType
 {
 	StereoLeft = 0,
