Additional case sensitive fix for Linux. I saw you made a fix to LightCardEditor earlier today but you if you had run the preflight you woudl have spotted this additional issue.

diff --git a/Engine/Plugins/Runtime/nDisplay/Source/DisplayClusterLightCardEditor/Private/Viewport/DisplayClusterLightcardEditorViewport.cpp b/Engine/Plugins/Runtime/nDisplay/Source/DisplayClusterLightCardEditor/Private/Viewport/DisplayClusterLightcardEditorViewport.cpp
--- a/Engine/Plugins/Runtime/nDisplay/Source/DisplayClusterLightCardEditor/Private/Viewport/DisplayClusterLightcardEditorViewport.cpp
+++ b/Engine/Plugins/Runtime/nDisplay/Source/DisplayClusterLightCardEditor/Private/Viewport/DisplayClusterLightcardEditorViewport.cpp
@@ -1,6 +1,6 @@
 // Copyright Epic Games, Inc. All Rights Reserved.
 
-#include "DisplayClusterLightCardEditorViewport.h"
+#include "DisplayClusterLightcardEditorViewport.h"
 
 #include "DisplayClusterLightCardEditorViewportClient.h"
 #include "DisplayClusterLightCardEditorCommands.h"
