[nDisplay]
[ViewportRendering]
refactored viewport render resources (per-node render purpose)

diff --git a/Engine/Plugins/Runtime/nDisplay/Source/DisplayCluster/Private/Render/Viewport/Containers/DisplayClusterViewportProxyData.h b/Engine/Plugins/Runtime/nDisplay/Source/DisplayCluster/Private/Render/Viewport/Containers/DisplayClusterViewportProxyData.h
--- a/Engine/Plugins/Runtime/nDisplay/Source/DisplayCluster/Private/Render/Viewport/Containers/DisplayClusterViewportProxyData.h
+++ b/Engine/Plugins/Runtime/nDisplay/Source/DisplayCluster/Private/Render/Viewport/Containers/DisplayClusterViewportProxyData.h
@@ -8,8 +8,8 @@
 #include "Render/Viewport/Containers/DisplayClusterViewport_OverscanSettings.h"
 
 class IDisplayClusterProjectionPolicy;
-class FDisplayClusterTextureResource;
-class FDisplayClusterRenderTargetResource;
+class FDisplayClusterViewportTextureResource;
+class FDisplayClusterViewportRenderTargetResource;
 class FDisplayClusterViewport;
 class FDisplayClusterViewportProxy;
 class FDisplayClusterViewport_Context;
@@ -45,19 +45,19 @@ private:
 	TArray<FDisplayClusterViewport_Context>      Contexts;
 
 	// View family render to this resources
-	TArray<FDisplayClusterRenderTargetResource*> RenderTargets;
+	TArray<FDisplayClusterViewportRenderTargetResource*> RenderTargets;
 
 	// Projection policy output resources
-	TArray<FDisplayClusterTextureResource*> OutputFrameTargetableResources;
-	TArray<FDisplayClusterTextureResource*> AdditionalFrameTargetableResources;
+	TArray<FDisplayClusterViewportTextureResource*> OutputFrameTargetableResources;
+	TArray<FDisplayClusterViewportTextureResource*> AdditionalFrameTargetableResources;
 
 #if WITH_EDITOR
 	FTextureRHIRef OutputPreviewTargetableResource;
 #endif
 
 	// unique viewport resources
-	TArray<FDisplayClusterTextureResource*> InputShaderResources;
-	TArray<FDisplayClusterTextureResource*> AdditionalTargetableResources;
-	TArray<FDisplayClusterTextureResource*> MipsShaderResources;
+	TArray<FDisplayClusterViewportTextureResource*> InputShaderResources;
+	TArray<FDisplayClusterViewportTextureResource*> AdditionalTargetableResources;
+	TArray<FDisplayClusterViewportTextureResource*> MipsShaderResources;
 };
 
diff --git a/Engine/Plugins/Runtime/nDisplay/Source/DisplayCluster/Private/Render/Viewport/DisplayClusterViewport.h b/Engine/Plugins/Runtime/nDisplay/Source/DisplayCluster/Private/Render/Viewport/DisplayClusterViewport.h
--- a/Engine/Plugins/Runtime/nDisplay/Source/DisplayCluster/Private/Render/Viewport/DisplayClusterViewport.h
+++ b/Engine/Plugins/Runtime/nDisplay/Source/DisplayCluster/Private/Render/Viewport/DisplayClusterViewport.h
@@ -16,8 +16,8 @@
 #include "SceneViewExtensionContext.h"
 #include "OpenColorIODisplayExtension.h"
 
-class FDisplayClusterRenderTargetResource;
-class FDisplayClusterTextureResource;
+class FDisplayClusterViewportRenderTargetResource;
+class FDisplayClusterViewportTextureResource;
 class FDisplayClusterViewportManager;
 class FDisplayClusterRenderTargetManager;
 class FDisplayClusterRenderFrameManager;
@@ -230,20 +230,20 @@ protected:
 	TArray<FDisplayClusterViewport_Context> Contexts;
 
 	// View family render to this resources
-	TArray<FDisplayClusterRenderTargetResource*> RenderTargets;
+	TArray<FDisplayClusterViewportRenderTargetResource*> RenderTargets;
 	
 	// Projection policy output resources
-	TArray<FDisplayClusterTextureResource*> OutputFrameTargetableResources;
-	TArray<FDisplayClusterTextureResource*> AdditionalFrameTargetableResources;
+	TArray<FDisplayClusterViewportTextureResource*> OutputFrameTargetableResources;
+	TArray<FDisplayClusterViewportTextureResource*> AdditionalFrameTargetableResources;
 
 #if WITH_EDITOR
 	FTextureRHIRef OutputPreviewTargetableResource;
 #endif
 
 	// unique viewport resources
-	TArray<FDisplayClusterTextureResource*> InputShaderResources;
-	TArray<FDisplayClusterTextureResource*> AdditionalTargetableResources;
-	TArray<FDisplayClusterTextureResource*> MipsShaderResources;
+	TArray<FDisplayClusterViewportTextureResource*> InputShaderResources;
+	TArray<FDisplayClusterViewportTextureResource*> AdditionalTargetableResources;
+	TArray<FDisplayClusterViewportTextureResource*> MipsShaderResources;
 
 	FDisplayClusterViewportManager& Owner;
 
diff --git a/Engine/Plugins/Runtime/nDisplay/Source/DisplayCluster/Private/Render/Viewport/DisplayClusterViewportManagerProxy.cpp b/Engine/Plugins/Runtime/nDisplay/Source/DisplayCluster/Private/Render/Viewport/DisplayClusterViewportManagerProxy.cpp
--- a/Engine/Plugins/Runtime/nDisplay/Source/DisplayCluster/Private/Render/Viewport/DisplayClusterViewportManagerProxy.cpp
+++ b/Engine/Plugins/Runtime/nDisplay/Source/DisplayCluster/Private/Render/Viewport/DisplayClusterViewportManagerProxy.cpp
@@ -392,7 +392,7 @@ void FDisplayClusterViewportManagerProxy::DoCrossGPUTransfers_RenderThread(FRHIC
 				// Use optimized cross GPU transfer for this context
 
 				FRenderTarget* RenderTarget = ViewportProxy->RenderTargets[ViewportContext.ContextNum];
-				FRHITexture2D* TextureRHI = ViewportProxy->RenderTargets[ViewportContext.ContextNum]->GetRenderTargetTexture();
+				FRHITexture2D* TextureRHI = ViewportProxy->RenderTargets[ViewportContext.ContextNum]->GetViewportRenderTargetResourceRHI();
 
 				FRHIGPUMask RenderTargetGPUMask = (GNumExplicitGPUsForRendering > 1 && RenderTarget) ? RenderTarget->GetGPUMask(RHICmdList) : FRHIGPUMask::GPU0();
 				{
@@ -443,8 +443,8 @@ bool FDisplayClusterViewportManagerProxy::GetFrameTargets_RenderThread(TArray<FR
 	{
 		if (ViewportProxy)
 		{
-			const TArray<FDisplayClusterTextureResource*>& Frames = ViewportProxy->OutputFrameTargetableResources;
-			const TArray<FDisplayClusterTextureResource*>& AdditionalFrames = ViewportProxy->AdditionalFrameTargetableResources;
+			const TArray<FDisplayClusterViewportTextureResource*>& Frames = ViewportProxy->OutputFrameTargetableResources;
+			const TArray<FDisplayClusterViewportTextureResource*>& AdditionalFrames = ViewportProxy->AdditionalFrameTargetableResources;
 
 			if (Frames.Num() > 0)
 			{
@@ -460,12 +460,12 @@ bool FDisplayClusterViewportManagerProxy::GetFrameTargets_RenderThread(TArray<FR
 
 				for (int32 FrameIt = 0; FrameIt < Frames.Num(); FrameIt++)
 				{
-					OutFrameResources.Add(Frames[FrameIt]->GetTextureResource());
+					OutFrameResources.Add(Frames[FrameIt]->GetViewportResourceRHI());
 					OutTargetOffsets.Add(Frames[FrameIt]->BackbufferFrameOffset);
 
 					if (bUseAdditionalFrameResources)
 					{
-						(*OutAdditionalFrameResources)[FrameIt] = AdditionalFrames[FrameIt]->GetTextureResource();
+						(*OutAdditionalFrameResources)[FrameIt] = AdditionalFrames[FrameIt]->GetViewportResourceRHI();
 					}
 				}
 
diff --git a/Engine/Plugins/Runtime/nDisplay/Source/DisplayCluster/Private/Render/Viewport/DisplayClusterViewportProxy.cpp b/Engine/Plugins/Runtime/nDisplay/Source/DisplayCluster/Private/Render/Viewport/DisplayClusterViewportProxy.cpp
--- a/Engine/Plugins/Runtime/nDisplay/Source/DisplayCluster/Private/Render/Viewport/DisplayClusterViewportProxy.cpp
+++ b/Engine/Plugins/Runtime/nDisplay/Source/DisplayCluster/Private/Render/Viewport/DisplayClusterViewportProxy.cpp
@@ -64,7 +64,7 @@ const IDisplayClusterViewportManagerProxy& FDisplayClusterViewportProxy::GetOwne
 	return Owner;
 }
 
-bool ImplGetTextureResources_RenderThread(const TArray<FDisplayClusterTextureResource*>& InResources, TArray<FRHITexture2D*>& OutResources)
+bool ImplGetTextureResources_RenderThread(const TArray<FDisplayClusterViewportTextureResource*>& InResources, TArray<FRHITexture2D*>& OutResources)
 {
 	check(OutResources.Num() == 0);
 
@@ -74,7 +74,7 @@ bool ImplGetTextureResources_RenderThread(const TArray<FDisplayClusterTextureRes
 
 		for (int32 ResourceIndex = 0; ResourceIndex < OutResources.Num(); ResourceIndex++)
 		{
-			OutResources[ResourceIndex] = InResources[ResourceIndex]->GetTextureResource();
+			OutResources[ResourceIndex] = InResources[ResourceIndex]->GetViewportResourceRHI();
 		}
 
 		return true;
@@ -132,14 +132,14 @@ bool FDisplayClusterViewportProxy::GetResources_RenderThread(const EDisplayClust
 				}
 				else
 				{
-					FDisplayClusterRenderTargetResource* Input = RenderTargets[ContextIt];
+					FDisplayClusterViewportRenderTargetResource* Input = RenderTargets[ContextIt];
 					if (Input == nullptr || Input->IsInitialized() == false)
 					{
 						OutResources.Empty();
 						return false;
 					}
 
-					OutResources[ContextIt] = Input->GetRenderTargetTexture();
+					OutResources[ContextIt] = Input->GetViewportRenderTargetResourceRHI();
 				}
 			}
 		}
@@ -213,11 +213,11 @@ void FDisplayClusterViewportProxy::ImplViewportRemap_RenderThread(FRHICommandLis
 
 			for (int32 ContextIt = 0; ContextIt < AdditionalFrameTargetableResources.Num(); ContextIt++)
 			{
-				FDisplayClusterTextureResource* Src = AdditionalFrameTargetableResources[ContextIt];
-				FDisplayClusterTextureResource* Dst = OutputFrameTargetableResources[ContextIt];
+				FDisplayClusterViewportTextureResource* Src = AdditionalFrameTargetableResources[ContextIt];
+				FDisplayClusterViewportTextureResource* Dst = OutputFrameTargetableResources[ContextIt];
 
-				FRHITexture2D* Input = Src ? Src->GetTextureResource() : nullptr;
-				FRHITexture2D* Output = Dst ? Dst->GetTextureResource() : nullptr;
+				FRHITexture2D* Input = Src ? Src->GetViewportResourceRHI() : nullptr;
+				FRHITexture2D* Output = Dst ? Dst->GetViewportResourceRHI() : nullptr;
 
 				if (Input && Output)
 				{
diff --git a/Engine/Plugins/Runtime/nDisplay/Source/DisplayCluster/Private/Render/Viewport/DisplayClusterViewportProxy.h b/Engine/Plugins/Runtime/nDisplay/Source/DisplayCluster/Private/Render/Viewport/DisplayClusterViewportProxy.h
--- a/Engine/Plugins/Runtime/nDisplay/Source/DisplayCluster/Private/Render/Viewport/DisplayClusterViewportProxy.h
+++ b/Engine/Plugins/Runtime/nDisplay/Source/DisplayCluster/Private/Render/Viewport/DisplayClusterViewportProxy.h
@@ -5,8 +5,8 @@
 #include "Render/Viewport/IDisplayClusterViewportProxy.h"
 #include "Render/Viewport/Containers/DisplayClusterViewport_OverscanSettings.h"
 
-class FDisplayClusterRenderTargetResource;
-class FDisplayClusterTextureResource;
+class FDisplayClusterViewportRenderTargetResource;
+class FDisplayClusterViewportTextureResource;
 class FDisplayClusterViewportProxyData;
 class IDisplayClusterViewportManagerProxy;
 class FDisplayClusterViewportManagerProxy;
@@ -129,20 +129,20 @@ protected:
 	TArray<FDisplayClusterViewport_Context> Contexts;
 
 	// View family render to this resources
-	TArray<FDisplayClusterRenderTargetResource*> RenderTargets;
+	TArray<FDisplayClusterViewportRenderTargetResource*> RenderTargets;
 
 	// Projection policy output resources
-	TArray<FDisplayClusterTextureResource*> OutputFrameTargetableResources;
-	TArray<FDisplayClusterTextureResource*> AdditionalFrameTargetableResources;
+	TArray<FDisplayClusterViewportTextureResource*> OutputFrameTargetableResources;
+	TArray<FDisplayClusterViewportTextureResource*> AdditionalFrameTargetableResources;
 
 #if WITH_EDITOR
 	FTextureRHIRef OutputPreviewTargetableResource;
 #endif
 
 	// unique viewport resources
-	TArray<FDisplayClusterTextureResource*> InputShaderResources;
-	TArray<FDisplayClusterTextureResource*> AdditionalTargetableResources;
-	TArray<FDisplayClusterTextureResource*> MipsShaderResources;
+	TArray<FDisplayClusterViewportTextureResource*> InputShaderResources;
+	TArray<FDisplayClusterViewportTextureResource*> AdditionalTargetableResources;
+	TArray<FDisplayClusterViewportTextureResource*> MipsShaderResources;
 
 	const FDisplayClusterViewportManagerProxy& Owner;
 	IDisplayClusterShaders& ShadersAPI;
diff --git a/Engine/Plugins/Runtime/nDisplay/Source/DisplayCluster/Private/Render/Viewport/RenderTarget/DisplayClusterRenderTargetManager.cpp b/Engine/Plugins/Runtime/nDisplay/Source/DisplayCluster/Private/Render/Viewport/RenderTarget/DisplayClusterRenderTargetManager.cpp
--- a/Engine/Plugins/Runtime/nDisplay/Source/DisplayCluster/Private/Render/Viewport/RenderTarget/DisplayClusterRenderTargetManager.cpp
+++ b/Engine/Plugins/Runtime/nDisplay/Source/DisplayCluster/Private/Render/Viewport/RenderTarget/DisplayClusterRenderTargetManager.cpp
@@ -12,6 +12,39 @@
 
 #include "Misc/DisplayClusterLog.h"
 
+////////////////////////////////////////////////////////////////////////////////////////////////////////////
+namespace DisplayClusterRenderTargetManager
+{
+	static EPixelFormat ImplGetCustomFormat(EDisplayClusterViewportCaptureMode CaptureMode)
+	{
+		// Pre-defined formats for targets:
+		switch (CaptureMode)
+		{
+		case EDisplayClusterViewportCaptureMode::Chromakey:
+			return EPixelFormat::PF_B8G8R8A8;
+
+		case EDisplayClusterViewportCaptureMode::Lightcard:
+			return EPixelFormat::PF_FloatRGBA;
+
+		default:
+			break;
+		}
+
+		return EPixelFormat::PF_Unknown;
+	}
+
+	static void ImplViewportTextureResourceLogging(const FDisplayClusterViewport* InViewport, const uint32 ContextNum, const FString& ResourceId, const FDisplayClusterViewportTextureResource* InTextureResource)
+	{
+		if (InTextureResource != nullptr && InTextureResource->GetViewportResourceState(EDisplayClusterViewportResourceState::Initialized) == false)
+		{
+			// Log: New resource created
+			UE_LOG(LogDisplayClusterViewport, Verbose, TEXT("Created new %s resource (%dx%d) for viewport '%s':%d"), *ResourceId, InTextureResource->GetSizeX(), InTextureResource->GetSizeY(), *InViewport->GetId(), ContextNum);
+		}
+	}
+};
+
+using namespace DisplayClusterRenderTargetManager;
+
 ////////////////////////////////////////////////////////////////////////////////////////////////////////////
 /// FDisplayClusterRenderTargetManager
 ////////////////////////////////////////////////////////////////////////////////////////////////////////////
@@ -24,29 +57,11 @@ FDisplayClusterRenderTargetManager::~FDisplayClusterRenderTargetManager()
 {
 }
 
-static EPixelFormat ImplGetCustomFormat(EDisplayClusterViewportCaptureMode CaptureMode)
-{
-	// Pre-defined formats for targets:
-	switch (CaptureMode)
-	{
-	case EDisplayClusterViewportCaptureMode::Chromakey:
-		return EPixelFormat::PF_B8G8R8A8;
-
-	case EDisplayClusterViewportCaptureMode::Lightcard:
-		return EPixelFormat::PF_FloatRGBA;
-
-	default:
-		break;
-	}
-
-	return EPixelFormat::PF_Unknown;
-}
-
-bool FDisplayClusterRenderTargetManager::AllocateRenderFrameResources(class FViewport* InViewport, const FDisplayClusterRenderFrameSettings& InRenderFrameSettings, const TArray<FDisplayClusterViewport*>& InViewports, FDisplayClusterRenderFrame& InOutRenderFrame)
+bool FDisplayClusterRenderTargetManager::AllocateRenderFrameResources(FViewport* InViewport, const FDisplayClusterRenderFrameSettings& InRenderFrameSettings, const TArray<FDisplayClusterViewport*>& InViewports, FDisplayClusterRenderFrame& InOutRenderFrame)
 {
 	bool bResult = true;
 
-	if(ResourcesPool->BeginReallocateRenderTargetResources(InViewport))
+	if(ResourcesPool->BeginReallocateResources(InRenderFrameSettings, InViewport))
 	{
 		// ReAllocate Render targets for all viewports
 		for (FDisplayClusterRenderFrame::FFrameRenderTarget& FrameRenderTargetIt : InOutRenderFrame.RenderTargets)
@@ -54,7 +69,7 @@ bool FDisplayClusterRenderTargetManager::AllocateRenderFrameResources(class FVie
 			if (FrameRenderTargetIt.bShouldUseRenderTarget)
 			{
 				// reallocate
-				FDisplayClusterRenderTargetResource* NewResource = ResourcesPool->AllocateRenderTargetResource(FrameRenderTargetIt.RenderTargetSize, ImplGetCustomFormat(FrameRenderTargetIt.CaptureMode));
+				FDisplayClusterViewportRenderTargetResource* NewResource = ResourcesPool->AllocateRenderTargetResource(FrameRenderTargetIt.RenderTargetSize, ImplGetCustomFormat(FrameRenderTargetIt.CaptureMode));
 				if (NewResource != nullptr)
 				{
 					// Set RenderFrame resource
@@ -77,6 +92,12 @@ bool FDisplayClusterRenderTargetManager::AllocateRenderFrameResources(class FVie
 									if (ViewportPtr->Contexts[ViewIt.ContextNum].bDisableRender == false)
 									{
 										ViewportPtr->RenderTargets[ViewIt.ContextNum] = NewResource;
+
+										if (NewResource->GetViewportResourceState(EDisplayClusterViewportResourceState::Initialized) == false)
+										{
+											// Log: New resource created
+											UE_LOG(LogDisplayClusterViewport, Verbose, TEXT("Created new ViewportRenderTarget resource %08X (%dx%d) for viewport '%s'"), NewResource, NewResource->GetSizeX(), NewResource->GetSizeY(), *ViewportPtr->GetId());
+										}
 									}
 								}
 							}
@@ -86,12 +107,6 @@ bool FDisplayClusterRenderTargetManager::AllocateRenderFrameResources(class FVie
 			}
 		}
 
-		ResourcesPool->FinishReallocateRenderTargetResources();
-	}
-
-	if (ResourcesPool->BeginReallocateTextureResources(InViewport))
-	{
-
 		// Allocate viewport internal resources
 		for (FDisplayClusterViewport* ViewportIt : InViewports)
 		{
@@ -108,16 +123,19 @@ bool FDisplayClusterRenderTargetManager::AllocateRenderFrameResources(class FVie
 
 				check(ViewportIt->InputShaderResources.Num() > 0);
 				ViewportIt->InputShaderResources[ContextNum] = ResourcesPool->AllocateTextureResource(ContextSize, false, CustomFormat);
+				ImplViewportTextureResourceLogging(ViewportIt, ContextNum, TEXT("InputShader"), ViewportIt->InputShaderResources[ContextNum]);
 
 				// Allocate custom resources:
 				if (bShouldUseAdditionalTargetableResource)
 				{
 					ViewportIt->AdditionalTargetableResources[ContextNum] = ResourcesPool->AllocateTextureResource(ContextSize, true, CustomFormat);
+					ImplViewportTextureResourceLogging(ViewportIt, ContextNum, TEXT("AdditionalRTT"), ViewportIt->AdditionalTargetableResources[ContextNum]);
 				}
 
 				if (ContextIt.NumMips > 1)
 				{
 					ViewportIt->MipsShaderResources[ContextNum] = ResourcesPool->AllocateTextureResource(ContextSize, false, CustomFormat, ContextIt.NumMips);
+					ImplViewportTextureResourceLogging(ViewportIt, ContextNum, TEXT("Mips"), ViewportIt->MipsShaderResources[ContextNum]);
 				}
 			}
 		}
@@ -131,7 +149,7 @@ bool FDisplayClusterRenderTargetManager::AllocateRenderFrameResources(class FVie
 			bResult = false;
 		}
 
-		ResourcesPool->FinishReallocateTextureResources();
+		ResourcesPool->FinishReallocateResources();
 	}
 
 	return bResult;
@@ -176,14 +194,20 @@ bool FDisplayClusterRenderTargetManager::AllocateFrameTargets(const FDisplayClus
 	}
 
 	// Reallocate frame target resources
-	TArray<FDisplayClusterTextureResource*> NewFrameTargetResources;
-	TArray<FDisplayClusterTextureResource*> NewAdditionalFrameTargetableResources;
+	TArray<FDisplayClusterViewportTextureResource*> NewFrameTargetResources;
+	TArray<FDisplayClusterViewportTextureResource*> NewAdditionalFrameTargetableResources;
 
 	for (uint32 FrameTargetsIt = 0; FrameTargetsIt < FrameTargetsAmount; FrameTargetsIt++)
 	{
-		FDisplayClusterTextureResource* NewResource = ResourcesPool->AllocateTextureResource(InOutRenderFrame.FrameRect.Size(), true, PF_Unknown);
-		if (NewResource)
+		FDisplayClusterViewportTextureResource* NewResource = ResourcesPool->AllocateTextureResource(InOutRenderFrame.FrameRect.Size(), true, PF_Unknown);
+		if (NewResource != nullptr)
 		{
+			if (NewResource->GetViewportResourceState(EDisplayClusterViewportResourceState::Initialized) == false)
+			{
+				// Log: New resource created
+				UE_LOG(LogDisplayClusterViewport, Verbose, TEXT("Created new RenderFrame resource (%dx%d)"), NewResource->GetSizeX(), NewResource->GetSizeY());
+			}
+
 			// calc and assign backbuffer offset (side_by_side, top_bottom)
 			NewResource->BackbufferFrameOffset = InOutRenderFrame.FrameRect.Min + TargetLocation;
 			TargetLocation += TargetOffset;
@@ -192,9 +216,15 @@ bool FDisplayClusterRenderTargetManager::AllocateFrameTargets(const FDisplayClus
 
 		if (InRenderFrameSettings.bShouldUseAdditionalFrameTargetableResource)
 		{
-			FDisplayClusterTextureResource* NewAdditionalResource = ResourcesPool->AllocateTextureResource(InOutRenderFrame.FrameRect.Size(), true, PF_Unknown);
-			if (NewAdditionalResource)
+			FDisplayClusterViewportTextureResource* NewAdditionalResource = ResourcesPool->AllocateTextureResource(InOutRenderFrame.FrameRect.Size(), true, PF_Unknown);
+			if (NewAdditionalResource != nullptr)
 			{
+				if (NewAdditionalResource->GetViewportResourceState(EDisplayClusterViewportResourceState::Initialized) == false)
+				{
+					// Log: New resource created
+					UE_LOG(LogDisplayClusterViewport, Verbose, TEXT("Created new RenderFrame2 resource (%dx%d)"), NewAdditionalResource->GetSizeX(), NewAdditionalResource->GetSizeY());
+				}
+
 				NewAdditionalFrameTargetableResources.Add(NewAdditionalResource);
 			}
 		}
diff --git a/Engine/Plugins/Runtime/nDisplay/Source/DisplayCluster/Private/Render/Viewport/RenderTarget/DisplayClusterRenderTargetResource.cpp b/Engine/Plugins/Runtime/nDisplay/Source/DisplayCluster/Private/Render/Viewport/RenderTarget/DisplayClusterRenderTargetResource.cpp
--- a/Engine/Plugins/Runtime/nDisplay/Source/DisplayCluster/Private/Render/Viewport/RenderTarget/DisplayClusterRenderTargetResource.cpp
+++ b/Engine/Plugins/Runtime/nDisplay/Source/DisplayCluster/Private/Render/Viewport/RenderTarget/DisplayClusterRenderTargetResource.cpp
@@ -6,68 +6,115 @@
 #include "RenderResource.h"
 #include "UnrealClient.h"
 
+//------------------------------------------------------------------------------------------------
+// FDisplayClusterViewportResource
+//------------------------------------------------------------------------------------------------
+FDisplayClusterViewportResource::FDisplayClusterViewportResource(const FDisplayClusterViewportResourceSettings& InResourceSettings)
+	: ViewportResourceSettings(InResourceSettings)
+{
+	bSRGB = ViewportResourceSettings.bShouldUseSRGB;
+	bGreyScaleFormat = false;
+	bIgnoreGammaConversions = true;
+}
 
-///////////////////////////////////////////////////////////////////
-void FDisplayClusterTextureResource::InitDynamicRHI()
+void FDisplayClusterViewportResource::ImplInitDynamicRHI_RenderTargetResource2D(FTexture2DRHIRef& OutRenderTargetTextureRHI, FTexture2DRHIRef& OutTextureRHI)
 {
 	ETextureCreateFlags CreateFlags = TexCreate_Dynamic;
 
+	// -- we will be manually copying this cross GPU, tell render graph not to
+	CreateFlags |= TexCreate_MultiGPUGraphIgnore;
+
 	// reflect srgb from settings
-	if (InitSettings.bShouldUseSRGB)
+	if (bSRGB)
 	{
 		CreateFlags |= TexCreate_SRGB;
 	}
 
-	if (InitSettings.NumMips > 1)
+	uint32 NumMips = 1;
+	if (ViewportResourceSettings.NumMips > 1)
 	{
 		// Create nummips texture!
 		CreateFlags |= TexCreate_GenerateMipCapable | TexCreate_UAV;
-
-		TRefCountPtr<FRHITexture2D> DummyTexture2DRHI;
-
-		FRHIResourceCreateInfo CreateInfo(TEXT("DisplayClusterTextureMipsResource"), FClearValueBinding::None);
-
-		RHICreateTargetableShaderResource2D(
-			InitSettings.Size.X, 
-			InitSettings.Size.Y, 
-			InitSettings.Format, 
-			InitSettings.NumMips,
-			CreateFlags,
-			TexCreate_RenderTargetable,
-			false,
-			CreateInfo,
-			TextureRHI,
-			DummyTexture2DRHI
-		);
+		NumMips = ViewportResourceSettings.NumMips;
 	}
-	else
+
+	FRHIResourceCreateInfo CreateInfo(TEXT("DisplayClusterViewportRenderTargetResource"), FClearValueBinding::Black);
+
+	RHICreateTargetableShaderResource2D(
+		GetSizeX(),
+		GetSizeY(),
+		ViewportResourceSettings.Format,
+		NumMips,
+		CreateFlags,
+		TexCreate_RenderTargetable,
+		false,
+		CreateInfo,
+		OutRenderTargetTextureRHI,
+		OutTextureRHI
+	);
+}
+
+void FDisplayClusterViewportResource::ImplInitDynamicRHI_TextureResource2D(FTexture2DRHIRef& OutTextureRHI)
+{
+	ETextureCreateFlags CreateFlags = TexCreate_Dynamic | TexCreate_ShaderResource;
+
+	// -- we will be manually copying this cross GPU, tell render graph not to
+	CreateFlags |= TexCreate_MultiGPUGraphIgnore;
+
+	// reflect srgb from settings
+	if (bSRGB)
 	{
-		if (InitSettings.bIsRenderTargetable)
-		{
-			CreateFlags |= TexCreate_RenderTargetable | TexCreate_ShaderResource;
-		}
-		else
-		{
-			// Shader resource usage only
-			CreateFlags |= TexCreate_ResolveTargetable | TexCreate_ShaderResource;
-		}
-
-		FRHIResourceCreateInfo CreateInfo(TEXT("DisplayClusterTextureResource"), FClearValueBinding::None);
-		TextureRHI = RHICreateTexture2D(InitSettings.Size.X, InitSettings.Size.Y, InitSettings.Format, InitSettings.NumMips, 1, CreateFlags, ERHIAccess::SRVMask, CreateInfo);
+		CreateFlags |= TexCreate_SRGB;
 	}
+
+	uint32 NumMips = 1;
+	uint32 NumSamples = 1;
+
+	CreateFlags |= ViewportResourceSettings.bIsRenderTargetable ? TexCreate_RenderTargetable : TexCreate_ResolveTargetable;
+
+	FRHIResourceCreateInfo CreateInfo(TEXT("DisplayClusterViewportTextureResource"), FClearValueBinding::Black);
+
+	OutTextureRHI = RHICreateTexture2D(
+		GetSizeX(),
+		GetSizeY(),
+		ViewportResourceSettings.Format,
+		NumMips,
+		NumSamples,
+		CreateFlags,
+		ERHIAccess::SRVMask,
+		CreateInfo
+	);
 }
 
-void FDisplayClusterTextureResource::ReleaseRHI()
+//------------------------------------------------------------------------------------------------
+// FDisplayClusterViewportTextureResource
+//------------------------------------------------------------------------------------------------
+void FDisplayClusterViewportTextureResource::InitDynamicRHI()
 {
-	TextureRHI.SafeRelease();
-}
+	FTexture2DRHIRef NewTextureRHI;
+
+	if (GetResourceSettingsConstRef().NumMips > 1)
+	{
+		FTexture2DRHIRef DummyTextureRHI;
+		ImplInitDynamicRHI_RenderTargetResource2D(NewTextureRHI, DummyTextureRHI);
+	}
+	else
+	{
+		ImplInitDynamicRHI_TextureResource2D(NewTextureRHI);
+	}
 
+	TextureRHI = (FTextureRHIRef&)NewTextureRHI;
+}
 
-///////////////////////////////////////////////////////////////////
-void FDisplayClusterRenderTargetResource::InitDynamicRHI()
+//------------------------------------------------------------------------------------------------
+// FDisplayClusterViewportRenderTargetResource
+//------------------------------------------------------------------------------------------------
+void FDisplayClusterViewportRenderTargetResource::InitDynamicRHI()
 {
-	// create output render target if necessary -- we will be manually copying this cross GPU, tell render graph not to
-	ETextureCreateFlags CreateFlags = (ResourceSettings.bShouldUseSRGB ? TexCreate_SRGB : TexCreate_None) | TexCreate_MultiGPUGraphIgnore;
+	// Create RTT and shader resources
+	FTexture2DRHIRef NewTextureRHI;
+	ImplInitDynamicRHI_RenderTargetResource2D(RenderTargetTextureRHI, NewTextureRHI);
+	TextureRHI = (FTextureRHIRef&)NewTextureRHI;
 
 	// Create the sampler state RHI resource.
 	FSamplerStateInitializerRHI SamplerStateInitializer
@@ -78,22 +125,4 @@ void FDisplayClusterRenderTargetResource::InitDynamicRHI()
 		AM_Clamp
 	);
 	SamplerStateRHI = GetOrCreateSamplerState(SamplerStateInitializer);
-
-	FTexture2DRHIRef Texture2DRHI;
-	FRHIResourceCreateInfo CreateInfo(TEXT("DisplayClusterRenderTargetResource"), FClearValueBinding(FLinearColor::Black));
-
-	RHICreateTargetableShaderResource2D(
-		GetSizeX(),
-		GetSizeY(),
-		ResourceSettings.Format,
-		1,
-		CreateFlags,
-		TexCreate_RenderTargetable,
-		false,
-		CreateInfo,
-		RenderTargetTextureRHI,
-		Texture2DRHI
-	);
-
-	TextureRHI = (FTextureRHIRef&)Texture2DRHI;
 }
diff --git a/Engine/Plugins/Runtime/nDisplay/Source/DisplayCluster/Private/Render/Viewport/RenderTarget/DisplayClusterRenderTargetResource.h b/Engine/Plugins/Runtime/nDisplay/Source/DisplayCluster/Private/Render/Viewport/RenderTarget/DisplayClusterRenderTargetResource.h
--- a/Engine/Plugins/Runtime/nDisplay/Source/DisplayCluster/Private/Render/Viewport/RenderTarget/DisplayClusterRenderTargetResource.h
+++ b/Engine/Plugins/Runtime/nDisplay/Source/DisplayCluster/Private/Render/Viewport/RenderTarget/DisplayClusterRenderTargetResource.h
@@ -9,129 +9,194 @@
 #include "RHIResources.h"
 #include "RenderResource.h"
 #include "UnrealClient.h"
-
-
-class FViewport;
-class FRHITexture2D;
-
+#include "Containers/StaticBitArray.h"
 
 ///////////////////////////////////////////////////////////////////
 struct FDisplayClusterViewportResourceSettings
 {
-public:
-	FDisplayClusterViewportResourceSettings(FViewport* InViewport);
-	FDisplayClusterViewportResourceSettings(FRHITexture2D* InTexture);
+	FDisplayClusterViewportResourceSettings()
+		: Format(EPixelFormat(0))
+		, bShouldUseSRGB(false)
+		, DisplayGamma(1)
+		, bIsRenderTargetable(false)
+		, NumMips(1)
+	{ }
 
-public:
-	bool operator==(const FDisplayClusterViewportResourceSettings& In) const
+	FDisplayClusterViewportResourceSettings(const FString& InClusterNodeId, const EPixelFormat InFormat, const bool InbShouldUseSRGB, const float InDisplayGamma)
+		: Format(InFormat)
+		, bShouldUseSRGB(InbShouldUseSRGB)
+		, DisplayGamma(InDisplayGamma)
+		, bIsRenderTargetable(false)
+		, NumMips(1)
+		, ClusterNodeId(InClusterNodeId)
+	{ }
+
+	FDisplayClusterViewportResourceSettings(const FDisplayClusterViewportResourceSettings& InBaseSettings, const FIntPoint& InSize, const EPixelFormat InFormat, const bool InbIsRenderTargetable = false, const int32 InNumMips = 1)
+		: Size(InSize)
+		, Format(InFormat == PF_Unknown ? InBaseSettings.Format : InFormat)
+		, bShouldUseSRGB(InBaseSettings.bShouldUseSRGB)
+		, DisplayGamma(InBaseSettings.DisplayGamma)
+		, bIsRenderTargetable(InbIsRenderTargetable)
+		, NumMips(InNumMips)
+		, ClusterNodeId(InBaseSettings.ClusterNodeId)
+	{ }
+
+	FORCEINLINE bool IsResourceSettingsEqual(const FDisplayClusterViewportResourceSettings& In) const
 	{
-		return (In.Size == Size) &&
-			(In.Format == Format) &&
-			(In.bShouldUseSRGB == bShouldUseSRGB) &&
-			(In.DisplayGamma == DisplayGamma) &&
-			(In.bIsRenderTargetable == bIsRenderTargetable) &&
-			(In.NumMips== NumMips);
+		return (In.Size == Size)
+			&& (In.Format == Format)
+			&& (In.bShouldUseSRGB == bShouldUseSRGB)
+			&& (In.DisplayGamma == DisplayGamma)
+			&& (In.bIsRenderTargetable == bIsRenderTargetable)
+			&& (In.NumMips == NumMips);
 	}
 
-public:
-	FIntPoint    Size;
-	EPixelFormat Format = EPixelFormat(0);
-	bool         bShouldUseSRGB = false;
+	FORCEINLINE bool IsClusterNodeNameEqual(const FDisplayClusterViewportResourceSettings& In) const
+	{
+		return ClusterNodeId == In.ClusterNodeId;
+	}
+
+	const FIntPoint    Size;
+	const EPixelFormat Format;
+	const bool         bShouldUseSRGB;
+
+	// Render target params
+	const float DisplayGamma;
 
-	// Render target only params
-	float DisplayGamma = 1;
+	// Texture target params
+	const bool bIsRenderTargetable;
+	const int32 NumMips;
 
-	// Texture target only params
-	bool bIsRenderTargetable = false;
-	int32 NumMips = 1;
+private:
+	// Cluster per-node render
+	const FString ClusterNodeId;
 };
 
+enum class EDisplayClusterViewportResourceState : uint8
+{
+	Initialized = 0,
+	Unused,
+	Deleted,
 
+	COUNT
+};
 ///////////////////////////////////////////////////////////////////
-class FDisplayClusterTextureResource
+class FDisplayClusterViewportResource
+	: public FTexture
 {
 public:
-	FDisplayClusterTextureResource(const FDisplayClusterViewportResourceSettings& InSettings)
-		: InitSettings(InSettings)
-		, BackbufferFrameOffset(EForceInit::ForceInitToZero)
+	FDisplayClusterViewportResource(const FDisplayClusterViewportResourceSettings& InResourceSettings);
+	virtual ~FDisplayClusterViewportResource()
 	{ }
 
 public:
-	void InitDynamicRHI();
-	void ReleaseRHI();
+	FORCEINLINE const FDisplayClusterViewportResourceSettings& GetResourceSettingsConstRef() const
+	{
+		return ViewportResourceSettings;
+	}
+
+	void RaiseViewportResourceState(const EDisplayClusterViewportResourceState InState)
+	{
+		ResourceState[(uint8)InState] = true;
+	}
 
-	bool IsResourceSettingsEqual(const FDisplayClusterViewportResourceSettings& InResourceSettings) const
+	void ClearViewportResourceState(const EDisplayClusterViewportResourceState InState)
 	{
-		return InitSettings == InResourceSettings;
+		ResourceState[(uint8)InState] = false;
 	}
 
-	const FTexture2DRHIRef& GetTextureResource() const
+	bool GetViewportResourceState(const EDisplayClusterViewportResourceState InState) const
+	{
+		return ResourceState[(uint8)InState];
+	}
+
+	virtual const FTexture2DRHIRef& GetViewportResourceRHI() const
 	{
 		return (const FTexture2DRHIRef&)TextureRHI;
 	}
 
-public:
-	const FDisplayClusterViewportResourceSettings InitSettings;
+	virtual const FTexture2DRHIRef& GetViewportRenderTargetResourceRHI() const
+	{
+		return (const FTexture2DRHIRef&)TextureRHI;
+	}
 
-	// OutputFrameTargetableResources frame offset on backbuffer (special for 'side_by_side' and 'top_bottom' DCRenderDevices)
-	FIntPoint BackbufferFrameOffset;
+	/** Returns the width of the texture in pixels. */
+	virtual uint32 GetSizeX() const override
+	{
+		return ViewportResourceSettings.Size.X;
+	}
+
+	/** Returns the height of the texture in pixels. */
+	virtual uint32 GetSizeY() const override
+	{
+		return ViewportResourceSettings.Size.Y;
+	}
+
+protected:
+	void ImplInitDynamicRHI_RenderTargetResource2D(FTexture2DRHIRef& OutRenderTargetTextureRHI, FTexture2DRHIRef& OutTextureRHI);
+	void ImplInitDynamicRHI_TextureResource2D(FTexture2DRHIRef& OutTextureRHI);
 
 private:
-	FTexture2DRHIRef TextureRHI;
-};
+	const FDisplayClusterViewportResourceSettings ViewportResourceSettings;
 
+	// resource state for game thread
+	TStaticBitArray<(uint8)EDisplayClusterViewportResourceState::COUNT> ResourceState;
+};
 
 ///////////////////////////////////////////////////////////////////
-class FDisplayClusterRenderTargetResource
-	: public FTexture
-	, public FRenderTarget
+class FDisplayClusterViewportTextureResource
+	: public FDisplayClusterViewportResource
 {
 public:
-	FDisplayClusterRenderTargetResource(const FDisplayClusterViewportResourceSettings& InResourceSettings)
-		: ResourceSettings(InResourceSettings)
+	FDisplayClusterViewportTextureResource(const FDisplayClusterViewportResourceSettings& InResourceSettings)
+		: FDisplayClusterViewportResource(InResourceSettings)
+		, BackbufferFrameOffset(EForceInit::ForceInitToZero)
+	{ }
+
+	virtual ~FDisplayClusterViewportTextureResource()
 	{ }
 
 public:
 	virtual void InitDynamicRHI() override;
 
-public:
-	virtual const FTexture2DRHIRef& GetRenderTargetTexture() const
+	virtual FString GetFriendlyName() const override
 	{
-		return (const FTexture2DRHIRef&)TextureRHI;
+		return TEXT("DisplayClusterViewportTextureResource");
 	}
 
-	virtual FIntPoint GetSizeXY() const
-	{
-		return ResourceSettings.Size;
-	}
+public:
+	// OutputFrameTargetableResources frame offset on backbuffer (special for 'side_by_side' and 'top_bottom' DCRenderDevices)
+	FIntPoint BackbufferFrameOffset;
+};
 
-	/** Returns the width of the texture in pixels. */
-	virtual uint32 GetSizeX() const
-	{
-		return ResourceSettings.Size.X;
-	}
+///////////////////////////////////////////////////////////////////
+class FDisplayClusterViewportRenderTargetResource
+	: public FDisplayClusterViewportResource
+	, public FRenderTarget
+{
+public:
+	FDisplayClusterViewportRenderTargetResource(const FDisplayClusterViewportResourceSettings& InResourceSettings)
+		: FDisplayClusterViewportResource(InResourceSettings)
+	{ }
 
-	/** Returns the height of the texture in pixels. */
-	virtual uint32 GetSizeY() const
+	virtual ~FDisplayClusterViewportRenderTargetResource()
+	{ }
+
+public:
+	virtual void InitDynamicRHI() override;
+
+	virtual FIntPoint GetSizeXY() const override
 	{
-		return ResourceSettings.Size.Y;
+		return GetResourceSettingsConstRef().Size;
 	}
 
 	virtual float GetDisplayGamma() const 
 	{
-		return ResourceSettings.DisplayGamma;
+		return GetResourceSettingsConstRef().DisplayGamma;
 	}
 
 	virtual FString GetFriendlyName() const override 
 	{
-		return TEXT("FDisplayClusterRenderTargetResource");
+		return TEXT("DisplayClusterViewportRenderTargetResource");
 	}
-
-	bool IsResourceSettingsEqual(const FDisplayClusterViewportResourceSettings& InResourceSettings) const
-	{
-		return ResourceSettings == InResourceSettings;
-	}
-
-private:
-	FDisplayClusterViewportResourceSettings ResourceSettings;
 };
diff --git a/Engine/Plugins/Runtime/nDisplay/Source/DisplayCluster/Private/Render/Viewport/RenderTarget/DisplayClusterRenderTargetResourcesPool.cpp b/Engine/Plugins/Runtime/nDisplay/Source/DisplayCluster/Private/Render/Viewport/RenderTarget/DisplayClusterRenderTargetResourcesPool.cpp
--- a/Engine/Plugins/Runtime/nDisplay/Source/DisplayCluster/Private/Render/Viewport/RenderTarget/DisplayClusterRenderTargetResourcesPool.cpp
+++ b/Engine/Plugins/Runtime/nDisplay/Source/DisplayCluster/Private/Render/Viewport/RenderTarget/DisplayClusterRenderTargetResourcesPool.cpp
@@ -1,52 +1,91 @@
 // Copyright Epic Games, Inc. All Rights Reserved.
 
 #include "DisplayClusterRenderTargetResourcesPool.h"
-
-#include "DisplayClusterRenderTargetResource.h"
+#include "Render/Viewport/RenderFrame/DisplayClusterRenderFrameSettings.h"
+#include "RHICommandList.h"
 #include "Engine/RendererSettings.h"
 #include "RenderingThread.h"
 
-
 ////////////////////////////////////////////////////////////////////////////////////////////////////////////
-/// FDisplayClusterViewportResourceSettings
-////////////////////////////////////////////////////////////////////////////////////////////////////////////
-FDisplayClusterViewportResourceSettings::FDisplayClusterViewportResourceSettings(class FViewport* InViewport)
+namespace DisplayClusterRenderTargetResourcesPool
 {
-	if (InViewport)
+	static void ImplInitializeViewportResourcesRHI(const TArrayView<FDisplayClusterViewportResource*>& InResources)
 	{
-		FRHITexture2D* ViewportTexture = InViewport->GetRenderTargetTexture();
-		if (ViewportTexture)
+		TArray<FDisplayClusterViewportResource*> ResourcesForRenderThread(InResources);
+
+		ENQUEUE_RENDER_COMMAND(DisplayCluster_InitializeViewportResourcesRHI)(
+			[NewResources = std::move(ResourcesForRenderThread)](FRHICommandListImmediate& RHICmdList)
 		{
-			Format = ViewportTexture->GetFormat();
-			bShouldUseSRGB = EnumHasAnyFlags(ViewportTexture->GetFlags(), TexCreate_SRGB);
-			DisplayGamma = InViewport->GetDisplayGamma();
-		}
+			for (FDisplayClusterViewportResource* ResourceIt : NewResources)
+			{
+				ResourceIt->InitResource();
+			}
+		});
 	}
-	else
+
+	static void ImplReleaseViewportResourcesRHI(const TArrayView<FDisplayClusterViewportResource*>& InResources)
 	{
-		// Get settings for preview rendering:
-		static const TConsoleVariableData<int32>* CVarDefaultBackBufferPixelFormat = IConsoleManager::Get().FindTConsoleVariableDataInt(TEXT("r.DefaultBackBufferPixelFormat"));
-		Format = EDefaultBackBufferPixelFormat::Convert2PixelFormat(EDefaultBackBufferPixelFormat::FromInt(CVarDefaultBackBufferPixelFormat->GetValueOnGameThread()));
-
-		// Use default values
-		//@todo: Get gamma and SRGB from editor
-		bShouldUseSRGB = false;
-		DisplayGamma = 2.2f;
+		TArray<FDisplayClusterViewportResource*> ResourcesForRenderThread(InResources);
+
+		ENQUEUE_RENDER_COMMAND(DisplayCluster_ReleaseViewportResourcesRHI)(
+			[ReleasedResources = std::move(ResourcesForRenderThread)](FRHICommandListImmediate& RHICmdList)
+		{
+			for (FDisplayClusterViewportResource* ResourceIt : ReleasedResources)
+			{
+				ResourceIt->ReleaseResource();
+				delete ResourceIt;
+			}
+		});
 	}
-}
 
-FDisplayClusterViewportResourceSettings::FDisplayClusterViewportResourceSettings(class FRHITexture2D* InTexture)
-{
-	if (InTexture)
+	static FDisplayClusterViewportResourceSettings* ImplCreateViewportResourceSettings(const FDisplayClusterRenderFrameSettings& InRenderFrameSettings, FViewport* InViewport)
 	{
-		Size = InTexture->GetSizeXY();
-		Format = InTexture->GetFormat();
-		bShouldUseSRGB = EnumHasAnyFlags(InTexture->GetFlags(), TexCreate_SRGB);
+		FDisplayClusterViewportResourceSettings Result;
 
-		bIsRenderTargetable = EnumHasAnyFlags(InTexture->GetFlags(), TexCreate_RenderTargetable);
-		NumMips = InTexture->GetNumMips();
+		if (InViewport != nullptr)
+		{
+			FRHITexture2D* ViewportTexture = InViewport->GetRenderTargetTexture();
+			if (ViewportTexture != nullptr)
+			{
+				EPixelFormat Format = ViewportTexture->GetFormat();
+				bool bShouldUseSRGB = EnumHasAnyFlags(ViewportTexture->GetFlags(), TexCreate_SRGB);
+				float Gamma = InViewport->GetDisplayGamma();
+
+				return new FDisplayClusterViewportResourceSettings(InRenderFrameSettings.ClusterNodeId, Format, bShouldUseSRGB, Gamma);
+			}
+		}
+
+		// Get default settings for preview rendering:
+		static TConsoleVariableData<int32>* CVarDefaultBackBufferPixelFormat = IConsoleManager::Get().FindTConsoleVariableDataInt(TEXT("r.DefaultBackBufferPixelFormat"));
+		EPixelFormat Format = EDefaultBackBufferPixelFormat::Convert2PixelFormat(EDefaultBackBufferPixelFormat::FromInt(CVarDefaultBackBufferPixelFormat->GetValueOnGameThread()));
+
+		// default gamma for preview in editor
+		float Gamma = 2.2f;
+
+		// By default sRGB disabled (Preview rendering)
+		bool bShouldUseSRGB = false;
+
+		return new FDisplayClusterViewportResourceSettings(InRenderFrameSettings.ClusterNodeId, Format, bShouldUseSRGB, Gamma);
 	}
-}
+
+	static bool IsTextureSizeValidImpl(const FIntPoint& InSize)
+	{
+		static const int32 MaxTextureSize = 1 << (GMaxTextureMipCount - 1);
+
+		// just check the texture size is valid
+		if (InSize.X <= 0 || InSize.Y <= 0 || InSize.X > MaxTextureSize || InSize.Y > MaxTextureSize)
+		{
+			return false;
+		}
+
+
+		//@todo - maybe check free size of video memory before texture allocation
+
+		return true;
+	}
+};
+
+using namespace DisplayClusterRenderTargetResourcesPool;
 
 ////////////////////////////////////////////////////////////////////////////////////////////////////////////
 /// FDisplayClusterRenderTargetResourcesPool
@@ -57,243 +96,168 @@ FDisplayClusterRenderTargetResourcesPool::FDisplayClusterRenderTargetResourcesPo
 
 FDisplayClusterRenderTargetResourcesPool::~FDisplayClusterRenderTargetResourcesPool()
 {
-	ReleaseRenderTargetResources(RenderTargetResources);
-	ReleaseRenderTargetResources(UnusedRenderTargetResources);
-	ReleaseTextureResources(TextureResources);
-	ReleaseTextureResources(UnusedTextureResources);
+	// Release all resources
+	ImplReleaseResources<FDisplayClusterViewportRenderTargetResource>(RenderTargetResources);
+	ImplReleaseResources<FDisplayClusterViewportTextureResource>(TextureResources);
 }
 
-bool FDisplayClusterRenderTargetResourcesPool::IsTextureSizeValid(const FIntPoint& InSize) const
+template <typename TViewportResourceType>
+void FDisplayClusterRenderTargetResourcesPool::ImplBeginReallocateResources(TArray<TViewportResourceType*>& InOutViewportResources)
 {
-	static const int32 MaxTextureSize = 1 << (GMaxTextureMipCount - 1);
+	check(ResourceSettings != nullptr);
 
-	// just check the texture size is valid
-	if (InSize.X <= 0 || InSize.Y <= 0 || InSize.X > MaxTextureSize || InSize.Y > MaxTextureSize)
+	// Mark cluster node resources as unused
+	for (TViewportResourceType* ResourceIt : InOutViewportResources)
 	{
-		return false;
+		if (ResourceIt != nullptr && ResourceIt->GetResourceSettingsConstRef().IsClusterNodeNameEqual(*ResourceSettings))
+		{
+			ResourceIt->RaiseViewportResourceState(EDisplayClusterViewportResourceState::Unused);
+		}
 	}
-
-
-	//@todo - maybe check free size of video memory before texture allocation
-
-	return true;
 }
 
-// --------------------------------------------------------------------------------------------------------------------------------------
-
-bool FDisplayClusterRenderTargetResourcesPool::BeginReallocateRenderTargetResources(class FViewport* InViewport)
+template <typename TViewportResourceType>
+void FDisplayClusterRenderTargetResourcesPool::ImplFinishReallocateResources(TArray<TViewportResourceType*>& InOutViewportResources)
 {
-	if (pRenderTargetResourceSettings == nullptr)
-	{
-		pRenderTargetResourceSettings = new FDisplayClusterViewportResourceSettings(InViewport);
+	// Collect new and unused resources 
+	TArray<TViewportResourceType*> NewResources;
+	TArray<TViewportResourceType*> UnusedResources;
 
-		UnusedRenderTargetResources = RenderTargetResources;
-		RenderTargetResources.Empty();
-		return true;
-	}
-
-	return false;
-}
-
-FDisplayClusterRenderTargetResource* FDisplayClusterRenderTargetResourcesPool::AllocateRenderTargetResource(const FIntPoint& InSize, enum EPixelFormat CustomPixelFormat)
-{
-	if (pRenderTargetResourceSettings == nullptr || !IsTextureSizeValid(InSize))
+	for (TViewportResourceType* ResourceIt : InOutViewportResources)
 	{
-		return nullptr;
-	}
-
-	FDisplayClusterViewportResourceSettings RequiredResourceSettings = *pRenderTargetResourceSettings;
-	RequiredResourceSettings.Size = InSize;
+		if (ResourceIt != nullptr)
+		{
+			check(ResourceIt->GetViewportResourceState(EDisplayClusterViewportResourceState::Deleted) == false);
 
-	if (CustomPixelFormat != PF_Unknown)
-	{
-		RequiredResourceSettings.Format = CustomPixelFormat;
+			if (ResourceIt->GetViewportResourceState(EDisplayClusterViewportResourceState::Initialized) == false)
+			{
+				// Collect new resources
+				NewResources.Add(ResourceIt);
+				ResourceIt->RaiseViewportResourceState(EDisplayClusterViewportResourceState::Initialized);
+			}
+			else if (ResourceIt->GetViewportResourceState(EDisplayClusterViewportResourceState::Unused) == true)
+			{
+				// Collect unused resources
+				UnusedResources.Add(ResourceIt);
+				ResourceIt->RaiseViewportResourceState(EDisplayClusterViewportResourceState::Deleted);
+			}
+		}
 	}
 
-	auto UnusedResourceIndex = UnusedRenderTargetResources.IndexOfByPredicate([RequiredResourceSettings](const FDisplayClusterRenderTargetResource* ResourceIt)
+	// Init RHI for new resources
+	if (NewResources.Num() > 0)
 	{
-		return ResourceIt && ResourceIt->IsResourceSettingsEqual(RequiredResourceSettings);
-	});
-
-	if (UnusedResourceIndex != INDEX_NONE)
-	{
-		// Use exist resource again
-		FDisplayClusterRenderTargetResource* ExistResource = UnusedRenderTargetResources[UnusedResourceIndex];
-		UnusedRenderTargetResources[UnusedResourceIndex] = nullptr;
-
-		RenderTargetResources.Add(ExistResource);
-		return ExistResource;
+		// Send to rendering thread
+		ImplInitializeViewportResourcesRHI(TArrayView<FDisplayClusterViewportResource*>((FDisplayClusterViewportResource**)(NewResources.GetData()), NewResources.Num()));
 	}
 
-	// Create new resource:
-	FDisplayClusterRenderTargetResource* NewResource = new FDisplayClusterRenderTargetResource(RequiredResourceSettings);
-	RenderTargetResources.Add(NewResource);
-	CreatedRenderTargetResources.Add(NewResource);
-
-	return NewResource;
-}
-
-void FDisplayClusterRenderTargetResourcesPool::FinishReallocateRenderTargetResources()
-{
-	if (pRenderTargetResourceSettings != nullptr)
+	// Remove unused resources
+	if (UnusedResources.Num() > 0)
 	{
-		delete pRenderTargetResourceSettings;
-		pRenderTargetResourceSettings = nullptr;
-
-		ReleaseRenderTargetResources(UnusedRenderTargetResources);
-		UnusedRenderTargetResources.Empty();
-
-		if (CreatedRenderTargetResources.Num() > 0)
+		// Remove GameThread resource references
+		for (TViewportResourceType* ResourceIt : UnusedResources)
 		{
-			// Send all new resource to render thread in single cmd
-			ENQUEUE_RENDER_COMMAND(InitRenderTargetCommand)(
-				[NewResources = std::move(CreatedRenderTargetResources)](FRHICommandListImmediate& RHICmdList)
+			int ResourceIndex = INDEX_NONE;
+			while (InOutViewportResources.Find(ResourceIt, ResourceIndex))
 			{
-				for (FDisplayClusterRenderTargetResource* NewResourceIt : NewResources)
-				{
-					NewResourceIt->InitResource();
-				}
-			});
-
-			CreatedRenderTargetResources.Empty();
+				InOutViewportResources.RemoveAt(ResourceIndex);
+			}
 		}
+
+		// Send to released resources to rendering thread
+		ImplReleaseViewportResourcesRHI(TArrayView<FDisplayClusterViewportResource*>((FDisplayClusterViewportResource**)(UnusedResources.GetData()), UnusedResources.Num()));
 	}
 }
 
-// --------------------------------------------------------------------------------------------------------------------------------------
-
-bool FDisplayClusterRenderTargetResourcesPool::BeginReallocateTextureResources(class FViewport* InViewport)
+template <typename TViewportResourceType>
+void FDisplayClusterRenderTargetResourcesPool::ImplReleaseResources(TArray<TViewportResourceType*>& InOutViewportResources)
 {
-	if (pTextureResourceSettings == nullptr)
-	{
-		pTextureResourceSettings = new FDisplayClusterViewportResourceSettings(InViewport);
-
-		UnusedTextureResources = TextureResources;
-		TextureResources.Empty();
+	// Send all resources to rendering thread
+	ImplReleaseViewportResourcesRHI(TArray<FDisplayClusterViewportResource*>((FDisplayClusterViewportResource**)(InOutViewportResources.GetData()), InOutViewportResources.Num()));
 
-		return true;
-	}
-
-	return false;
+	// Reset refs on game thread too
+	InOutViewportResources.Empty();
 }
 
-FDisplayClusterTextureResource* FDisplayClusterRenderTargetResourcesPool::AllocateTextureResource(const FIntPoint& InSize, bool bIsRenderTargetable, enum EPixelFormat CustomPixelFormat, int32 InNumMips)
+template <typename TViewportResourceType>
+TViewportResourceType* FDisplayClusterRenderTargetResourcesPool::ImplAllocateResource(TArray<TViewportResourceType*>& InOutViewportResources, const FDisplayClusterViewportResourceSettings& InSettings)
 {
-	if (pTextureResourceSettings == nullptr || !IsTextureSizeValid(InSize))
+	if (!IsTextureSizeValidImpl(InSettings.Size))
 	{
 		return nullptr;
 	}
 
-	FDisplayClusterViewportResourceSettings RequiredResourceSettings = *pTextureResourceSettings;
-	RequiredResourceSettings.Size = InSize;
-	RequiredResourceSettings.bIsRenderTargetable = bIsRenderTargetable;
-	RequiredResourceSettings.NumMips = InNumMips;
-
-	if (CustomPixelFormat != PF_Unknown)
+	// Unused resources marked for current cluster node
+	auto ExistAndUnusedResourceIndex = InOutViewportResources.IndexOfByPredicate([InSettings](const TViewportResourceType* ResourceIt)
 	{
-		RequiredResourceSettings.Format = CustomPixelFormat;
-	}
+		if (ResourceIt != nullptr && ResourceIt->GetViewportResourceState(EDisplayClusterViewportResourceState::Unused) == true)
+		{
+			return ResourceIt->GetResourceSettingsConstRef().IsResourceSettingsEqual(InSettings);
+		}
 
-	auto UnusedResourceIndex = UnusedTextureResources.IndexOfByPredicate([RequiredResourceSettings](const FDisplayClusterTextureResource* ResourceIt)
-	{
-		return ResourceIt && ResourceIt->IsResourceSettingsEqual(RequiredResourceSettings);
+		return false;
 	});
 
-	if (UnusedResourceIndex != INDEX_NONE)
+	if (ExistAndUnusedResourceIndex != INDEX_NONE)
 	{
 		// Use exist resource again
-		FDisplayClusterTextureResource* ExistResource = UnusedTextureResources[UnusedResourceIndex];
-		UnusedTextureResources[UnusedResourceIndex] = nullptr;
+		TViewportResourceType* ExistResource = InOutViewportResources[ExistAndUnusedResourceIndex];
+
+		// Clear unused state for re-used resource
+		ExistResource->ClearViewportResourceState(EDisplayClusterViewportResourceState::Unused);
 
-		TextureResources.Add(ExistResource);
 		return ExistResource;
 	}
 
 	// Create new resource:
-	FDisplayClusterTextureResource* NewResource = new FDisplayClusterTextureResource(RequiredResourceSettings);
-	TextureResources.Add(NewResource);
-	CreatedTextureResources.Add(NewResource);
+	TViewportResourceType* NewResource = new TViewportResourceType(InSettings);
+	InOutViewportResources.Add(NewResource);
 
 	return NewResource;
 }
 
-void FDisplayClusterRenderTargetResourcesPool::FinishReallocateTextureResources()
+bool FDisplayClusterRenderTargetResourcesPool::BeginReallocateResources(const FDisplayClusterRenderFrameSettings& InRenderFrameSettings, FViewport* InViewport)
 {
-	if (pTextureResourceSettings != nullptr)
-	{
-		delete pTextureResourceSettings;
-		pTextureResourceSettings = nullptr;
+	check(ResourceSettings == nullptr);
 
-		ReleaseTextureResources(UnusedTextureResources);
-		UnusedTextureResources.Empty();
+	// Initialize settings for new render frame
+	ResourceSettings = ImplCreateViewportResourceSettings(InRenderFrameSettings, InViewport);
 
-		// Send all new resource to render thread in single cmd
-		if (CreatedTextureResources.Num() > 0)
-		{
-			ENQUEUE_RENDER_COMMAND(InitTextureResourceCommand)(
-				[NewResources = std::move(CreatedTextureResources)](FRHICommandListImmediate& RHICmdList)
-			{
-				for (FDisplayClusterTextureResource* NewResourceIt : NewResources)
-				{
-					NewResourceIt->InitDynamicRHI();
-				}
-			});
+	if(ResourceSettings != nullptr)
+	{
+		// Begin reallocate resources for current cluster node
+		ImplBeginReallocateResources<FDisplayClusterViewportRenderTargetResource>(RenderTargetResources);
+		ImplBeginReallocateResources<FDisplayClusterViewportTextureResource>(TextureResources);
 
-			CreatedTextureResources.Empty();
-		}
+		return true;
 	}
-}
 
-// --------------------------------------------------------------------------------------------------------------------------------------
+	return false;
+}
 
-void FDisplayClusterRenderTargetResourcesPool::ReleaseRenderTargetResources(TArray<FDisplayClusterRenderTargetResource*>& InOutResources)
+void FDisplayClusterRenderTargetResourcesPool::FinishReallocateResources()
 {
-	TArray<FDisplayClusterRenderTargetResource*> ExistResourcesForRemove;
-	for (FDisplayClusterRenderTargetResource* It : InOutResources)
+	if (ResourceSettings != nullptr)
 	{
-		if (It != nullptr)
-		{
-			ExistResourcesForRemove.Add(It);
-		}
-	}
-	InOutResources.Empty();
+		// Finish reallocate resources for current cluster node
+		ImplFinishReallocateResources<FDisplayClusterViewportRenderTargetResource>(RenderTargetResources);
+		ImplFinishReallocateResources<FDisplayClusterViewportTextureResource>(TextureResources);
 
-	if (ExistResourcesForRemove.Num() > 0)
-	{
-		ENQUEUE_RENDER_COMMAND(ReleaseRenderTargetResourceCommand)(
-			[RemovedResources = std::move(ExistResourcesForRemove)](FRHICommandListImmediate& RHICmdList)
-		{
-			for (FDisplayClusterRenderTargetResource* RemovedResourceIt : RemovedResources)
-			{
-				RemovedResourceIt->ReleaseResource();
-				delete RemovedResourceIt;
-			}
-		});
+		delete ResourceSettings;
+		ResourceSettings = nullptr;
 	}
 }
 
-void FDisplayClusterRenderTargetResourcesPool::ReleaseTextureResources(TArray<FDisplayClusterTextureResource*>& InOutResources)
+FDisplayClusterViewportRenderTargetResource* FDisplayClusterRenderTargetResourcesPool::AllocateRenderTargetResource(const FIntPoint& InSize, enum EPixelFormat CustomPixelFormat)
 {
-	TArray<FDisplayClusterTextureResource*> ExistResourcesForRemove;
-	for (FDisplayClusterTextureResource* It : InOutResources)
-	{
-		if (It != nullptr)
-		{
-			ExistResourcesForRemove.Add(It);
-		}
-	}
-	InOutResources.Empty();
+	check(ResourceSettings != nullptr);
 
-	if (ExistResourcesForRemove.Num() > 0)
-	{
-		ENQUEUE_RENDER_COMMAND(ReleaseTextureResourceCommand)(
-			[RemovedResources = std::move(ExistResourcesForRemove)](FRHICommandListImmediate& RHICmdList)
-		{
-			for (FDisplayClusterTextureResource* RemovedResourceIt : RemovedResources)
-			{
-				RemovedResourceIt->ReleaseRHI();
-				delete RemovedResourceIt;
-			}
-		});
-	}
+	return ImplAllocateResource<FDisplayClusterViewportRenderTargetResource>(RenderTargetResources, FDisplayClusterViewportResourceSettings(*ResourceSettings, InSize, CustomPixelFormat));
+}
+
+FDisplayClusterViewportTextureResource* FDisplayClusterRenderTargetResourcesPool::AllocateTextureResource(const FIntPoint& InSize, bool bIsRenderTargetable, enum EPixelFormat CustomPixelFormat, int32 InNumMips)
+{
+	check(ResourceSettings != nullptr);
+
+	return ImplAllocateResource<FDisplayClusterViewportTextureResource>(TextureResources, FDisplayClusterViewportResourceSettings(*ResourceSettings, InSize, CustomPixelFormat, bIsRenderTargetable, InNumMips));
 }
diff --git a/Engine/Plugins/Runtime/nDisplay/Source/DisplayCluster/Private/Render/Viewport/RenderTarget/DisplayClusterRenderTargetResourcesPool.h b/Engine/Plugins/Runtime/nDisplay/Source/DisplayCluster/Private/Render/Viewport/RenderTarget/DisplayClusterRenderTargetResourcesPool.h
--- a/Engine/Plugins/Runtime/nDisplay/Source/DisplayCluster/Private/Render/Viewport/RenderTarget/DisplayClusterRenderTargetResourcesPool.h
+++ b/Engine/Plugins/Runtime/nDisplay/Source/DisplayCluster/Private/Render/Viewport/RenderTarget/DisplayClusterRenderTargetResourcesPool.h
@@ -3,15 +3,9 @@
 #pragma once
 
 #include "CoreMinimal.h"
-#include "RHI.h"
-#include "RHICommandList.h"
-#include "RHIResources.h"
-
-class FDisplayClusterRenderTargetResource;
-class FDisplayClusterTextureResource;
-struct FDisplayClusterViewportResourceSettings;
-class FViewport;
+#include "DisplayClusterRenderTargetResource.h"
 
+struct FDisplayClusterRenderFrameSettings;
 
 class FDisplayClusterRenderTargetResourcesPool
 {
@@ -20,31 +14,30 @@ public:
 	~FDisplayClusterRenderTargetResourcesPool();
 
 public:
-	bool BeginReallocateRenderTargetResources(FViewport* InViewport);
-	FDisplayClusterRenderTargetResource* AllocateRenderTargetResource(const FIntPoint& InSize, EPixelFormat CustomPixelFormat);
-	void FinishReallocateRenderTargetResources();
+	bool BeginReallocateResources(const FDisplayClusterRenderFrameSettings& InRenderFrameSettings, class FViewport* InViewport);
+	void FinishReallocateResources();
 
-public:
-	bool BeginReallocateTextureResources(FViewport* InViewport);
-	FDisplayClusterTextureResource* AllocateTextureResource(const FIntPoint& InSize, bool bIsRenderTargetable, EPixelFormat CustomPixelFormat, int32 NumMips = 1);
-	void FinishReallocateTextureResources();
+	FDisplayClusterViewportRenderTargetResource* AllocateRenderTargetResource(const FIntPoint& InSize, EPixelFormat CustomPixelFormat);
+	FDisplayClusterViewportTextureResource*      AllocateTextureResource(const FIntPoint& InSize, bool bIsRenderTargetable, EPixelFormat CustomPixelFormat, int32 NumMips = 1);
+
+private:
+	template <typename TViewportResourceType>
+	void ImplBeginReallocateResources(TArray<TViewportResourceType*>& InOutViewportResources);
 
-protected:
-	void ReleaseRenderTargetResources(TArray<FDisplayClusterRenderTargetResource*>& InOutResources);
-	void ReleaseTextureResources(TArray<FDisplayClusterTextureResource*>& InOutResources);
+	template <typename TViewportResourceType>
+	void ImplFinishReallocateResources(TArray<TViewportResourceType*>& InOutViewportResources);
 
-	bool IsTextureSizeValid(const FIntPoint& InSize) const;
+	template <typename TViewportResourceType>
+	TViewportResourceType* ImplAllocateResource(TArray<TViewportResourceType*>& InOutViewportResources, const FDisplayClusterViewportResourceSettings& InSettings);
 
-private:
-	// Current render target settings (initialize from BeginReallocateRenderTargetResources)
-	FDisplayClusterViewportResourceSettings* pRenderTargetResourceSettings = nullptr;
-	FDisplayClusterViewportResourceSettings* pTextureResourceSettings = nullptr;
+	template <typename TViewportResourceType>
+	void ImplReleaseResources(TArray<TViewportResourceType*>& InOutViewportResources);
 
-	TArray<FDisplayClusterRenderTargetResource*> RenderTargetResources;
-	TArray<FDisplayClusterRenderTargetResource*> UnusedRenderTargetResources;
-	TArray<FDisplayClusterRenderTargetResource*> CreatedRenderTargetResources;
+private:
+	// Current render resource settings
+	FDisplayClusterViewportResourceSettings* ResourceSettings = nullptr;
 
-	TArray<FDisplayClusterTextureResource*> TextureResources;
-	TArray<FDisplayClusterTextureResource*> UnusedTextureResources;
-	TArray<FDisplayClusterTextureResource*> CreatedTextureResources;
+	// Viewport render resources
+	TArray<FDisplayClusterViewportRenderTargetResource*> RenderTargetResources;
+	TArray<FDisplayClusterViewportTextureResource*>      TextureResources;
 };
