nDisplay: fixed icvfx feather scaling

diff --git a/Engine/Plugins/Runtime/nDisplay/Shaders/Private/MPCDIOverlayShaders.usf b/Engine/Plugins/Runtime/nDisplay/Shaders/Private/MPCDIOverlayShaders.usf
--- a/Engine/Plugins/Runtime/nDisplay/Shaders/Private/MPCDIOverlayShaders.usf
+++ b/Engine/Plugins/Runtime/nDisplay/Shaders/Private/MPCDIOverlayShaders.usf
@@ -109,9 +109,33 @@ float4 GetInnerCameraColor(float4 WarpedUV, OutputVS IN)
 		//! UV.z <= UV.w &&
 		if (Weight >= 0)
 		{
-			float2 ToSoftEdge = saturate((abs(ToEdge) - 1.0f + InnerCameraSoftEdge.xy) / InnerCameraSoftEdge.xy);
-			float CameraTextureAlpha = Pow2(saturate(1 - Pow2(dot(ToSoftEdge, ToSoftEdge))));
-			CameraColor.w = CameraTextureAlpha;			
+			CameraColor.w = 1;
+			
+			//InnerCameraSoftEdge.xyzw values are used to pass left/top/right/bottom soft edge values
+			if(CameraUV.x < InnerCameraSoftEdge.x) // left 
+			{
+				float CameraTextureAlphaLeft = Pow2(saturate(abs(CameraUV.x) / InnerCameraSoftEdge.x));
+				CameraColor.w *= CameraTextureAlphaLeft;
+			}
+			else
+			if(CameraUV.x > (1 - InnerCameraSoftEdge.z)) // right
+			{
+				float CameraTextureAlphaRight = Pow2(saturate(abs(1 - CameraUV.x) / InnerCameraSoftEdge.z));
+				CameraColor.w *= CameraTextureAlphaRight;
+			}
+
+			
+			if(CameraUV.y < (InnerCameraSoftEdge.y)) // top
+			{
+				float CameraTextureAlphaTop = Pow2(saturate(abs(CameraUV.y) / InnerCameraSoftEdge.y));
+				CameraColor.w *= CameraTextureAlphaTop;
+			}
+			else
+			if(CameraUV.y > (1 - InnerCameraSoftEdge.w)) // bottom
+			{
+				float CameraTextureAlphaBottom = Pow2(saturate(abs(1 - CameraUV.y) / InnerCameraSoftEdge.w));
+				CameraColor.w *= CameraTextureAlphaBottom;
+			}
 		}
 	}
 
diff --git a/Engine/Plugins/Runtime/nDisplay/Source/DisplayCluster/Private/Render/Viewport/Configuration/DisplayClusterViewportConfigurationHelpers_ICVFX.cpp b/Engine/Plugins/Runtime/nDisplay/Source/DisplayCluster/Private/Render/Viewport/Configuration/DisplayClusterViewportConfigurationHelpers_ICVFX.cpp
--- a/Engine/Plugins/Runtime/nDisplay/Source/DisplayCluster/Private/Render/Viewport/Configuration/DisplayClusterViewportConfigurationHelpers_ICVFX.cpp
+++ b/Engine/Plugins/Runtime/nDisplay/Source/DisplayCluster/Private/Render/Viewport/Configuration/DisplayClusterViewportConfigurationHelpers_ICVFX.cpp
@@ -255,20 +255,54 @@ FDisplayClusterShaderParameters_ICVFX::FCameraSettings FDisplayClusterViewportCo
 	Result.Resource.ViewportId = InCameraViewport.GetId();
 	Result.Local2WorldTransform = OriginComp->GetComponentTransform();
 
-	Result.SoftEdge.X = CameraSettings.SoftEdge.Horizontal;
-	Result.SoftEdge.Y = CameraSettings.SoftEdge.Vertical;
-	Result.SoftEdge.Z = 0;
-
-	float RealThicknessScaleValue = 0.1f;
+	const float RealThicknessScaleValue = 0.1f;
 	Result.InnerCameraBorderColor = CameraSettings.Border.Color;
 	Result.InnerCameraBorderThickness = CameraSettings.Border.Enable ? CameraSettings.Border.Thickness * RealThicknessScaleValue : 0.0f;
 	Result.InnerCameraFrameAspectRatio = (float)CameraSettings.RenderSettings.CustomFrameSize.CustomWidth / (float)CameraSettings.RenderSettings.CustomFrameSize.CustomHeight;
 
 	const FString InnerFrustumID = InCameraComponent.GetCameraUniqueId();
 	const int32 CameraRenderOrder = RootActor.GetInnerFrustumPriority(InnerFrustumID);
-
 	Result.RenderOrder = (CameraRenderOrder<0) ? CameraSettings.RenderSettings.RenderOrder : CameraRenderOrder;
 
+	// softedge adjustments
+	const float Overscan = (CameraSettings.CustomFrustum.FieldOfViewMultiplier > 0) ? CameraSettings.CustomFrustum.FieldOfViewMultiplier : 1;
+
+	// remap values from 0-1 GUI range into acceptable 0.0 - 0.25 shader range
+	Result.SoftEdge.X = FMath::GetMappedRangeValueClamped(FVector2D(0.0, 1.0f), FVector2D(0.0, 0.25), CameraSettings.SoftEdge.Horizontal) / Overscan; // Left
+	Result.SoftEdge.Y = FMath::GetMappedRangeValueClamped(FVector2D(0.0, 1.0f), FVector2D(0.0, 0.25), CameraSettings.SoftEdge.Vertical) / Overscan; // Top
+	Result.SoftEdge.Z = FMath::GetMappedRangeValueClamped(FVector2D(0.0, 1.0f), FVector2D(0.0, 0.25), CameraSettings.SoftEdge.Horizontal) / Overscan; // right
+	Result.SoftEdge.W = FMath::GetMappedRangeValueClamped(FVector2D(0.0, 1.0f), FVector2D(0.0, 0.25), CameraSettings.SoftEdge.Vertical) / Overscan; // bottom
+
+	// default - percents
+	const float ConvertToPercent = 0.01f;
+	float Left = FDisplayClusterViewport_OverscanSettings::ClampPercent(CameraSettings.CustomFrustum.Left * ConvertToPercent);
+	float Right = FDisplayClusterViewport_OverscanSettings::ClampPercent(CameraSettings.CustomFrustum.Right * ConvertToPercent);
+	float Top = FDisplayClusterViewport_OverscanSettings::ClampPercent(CameraSettings.CustomFrustum.Top * ConvertToPercent);
+	float Bottom = FDisplayClusterViewport_OverscanSettings::ClampPercent(CameraSettings.CustomFrustum.Bottom * ConvertToPercent);
+
+	if (CameraSettings.CustomFrustum.Mode == EDisplayClusterConfigurationViewportCustomFrustumMode::Pixels)
+	{
+		float FrameWidth = RootActor.GetStageSettings().DefaultFrameSize.Width * CameraSettings.BufferRatio;
+		float FrameHeight = RootActor.GetStageSettings().DefaultFrameSize.Height * CameraSettings.BufferRatio;
+
+		if (CameraSettings.RenderSettings.CustomFrameSize.bUseCustomSize)
+		{
+			FrameWidth = CameraSettings.RenderSettings.CustomFrameSize.CustomWidth * CameraSettings.BufferRatio;
+			FrameHeight = CameraSettings.RenderSettings.CustomFrameSize.CustomHeight * CameraSettings.BufferRatio;
+		}
+
+		Left = FDisplayClusterViewport_OverscanSettings::ClampPercent(CameraSettings.CustomFrustum.Left / FrameWidth);
+		Right = FDisplayClusterViewport_OverscanSettings::ClampPercent(CameraSettings.CustomFrustum.Right / FrameWidth);
+		Top = FDisplayClusterViewport_OverscanSettings::ClampPercent(CameraSettings.CustomFrustum.Top / FrameHeight);
+		Bottom = FDisplayClusterViewport_OverscanSettings::ClampPercent(CameraSettings.CustomFrustum.Bottom / FrameHeight);
+	}
+
+	// recalculate soft edge related offsets based on frustum
+	Result.SoftEdge.X /= (1 + Left + Right);
+	Result.SoftEdge.Y /= (1 + Top + Bottom);
+	Result.SoftEdge.Z /= (1 + Left + Right);
+	Result.SoftEdge.W /= (1 + Top + Bottom);
+
 	return Result;
 }
 
diff --git a/Engine/Plugins/Runtime/nDisplay/Source/DisplayClusterShaders/Public/ShaderParameters/DisplayClusterShaderParameters_ICVFX.h b/Engine/Plugins/Runtime/nDisplay/Source/DisplayClusterShaders/Public/ShaderParameters/DisplayClusterShaderParameters_ICVFX.h
--- a/Engine/Plugins/Runtime/nDisplay/Source/DisplayClusterShaders/Public/ShaderParameters/DisplayClusterShaderParameters_ICVFX.h
+++ b/Engine/Plugins/Runtime/nDisplay/Source/DisplayClusterShaders/Public/ShaderParameters/DisplayClusterShaderParameters_ICVFX.h
@@ -195,7 +195,7 @@ public:
 
 	public:
 		FDisplayClusterShaderParametersICVFX_ViewportResource Resource;
-		FVector SoftEdge;
+		FVector4 SoftEdge;
 
 		FLinearColor InnerCameraBorderColor = FLinearColor::Black;
 		float InnerCameraBorderThickness = 0.1f;
