nDisplay: added GC to handle viewstates

diff --git a/Engine/Plugins/Runtime/nDisplay/Source/DisplayCluster/Private/Game/EngineClasses/Scene/DisplayClusterRootActorEditor.cpp b/Engine/Plugins/Runtime/nDisplay/Source/DisplayCluster/Private/Game/EngineClasses/Scene/DisplayClusterRootActorEditor.cpp
--- a/Engine/Plugins/Runtime/nDisplay/Source/DisplayCluster/Private/Game/EngineClasses/Scene/DisplayClusterRootActorEditor.cpp
+++ b/Engine/Plugins/Runtime/nDisplay/Source/DisplayCluster/Private/Game/EngineClasses/Scene/DisplayClusterRootActorEditor.cpp
@@ -268,6 +268,27 @@ bool ADisplayClusterRootActor::UpdatePreviewConfiguration_Editor(bool bUpdateAll
 	return false;
 }
 
+void ADisplayClusterRootActor::AddReferencedObjects(UObject* InThis, FReferenceCollector& Collector)
+{
+	ADisplayClusterRootActor* This = CastChecked<ADisplayClusterRootActor>(InThis);
+
+	for (IDisplayClusterViewport* ViewportIt : This->ViewportManager->GetViewports())
+	{
+		TIndirectArray<FSceneViewStateReference>& ViewStates = ((FDisplayClusterViewport*)ViewportIt)->ViewStates;
+
+		for (int32 ViewIndex = 0; ViewIndex < ViewStates.Num(); ViewIndex++)
+		{
+			FSceneViewStateInterface* Ref = ViewStates[ViewIndex].GetReference();
+			if (Ref)
+			{
+				Ref->AddReferencedObjects(Collector);
+			}
+		}
+	}
+
+	Super::AddReferencedObjects(This, Collector);
+}
+
 void ADisplayClusterRootActor::RenderPreview_Editor()
 {
 	if (UpdatePreviewConfiguration_Editor(false))
diff --git a/Engine/Plugins/Runtime/nDisplay/Source/DisplayCluster/Private/Render/Viewport/DisplayClusterViewport.cpp b/Engine/Plugins/Runtime/nDisplay/Source/DisplayCluster/Private/Render/Viewport/DisplayClusterViewport.cpp
--- a/Engine/Plugins/Runtime/nDisplay/Source/DisplayCluster/Private/Render/Viewport/DisplayClusterViewport.cpp
+++ b/Engine/Plugins/Runtime/nDisplay/Source/DisplayCluster/Private/Render/Viewport/DisplayClusterViewport.cpp
@@ -133,14 +133,6 @@ void FDisplayClusterViewport::HandleEndScene()
 		UninitializedProjectionPolicy = ProjectionPolicy;
 		ProjectionPolicy.Reset();
 	}
-
-#if WITH_EDITOR
-	// cleanup view states
-	for (int32 ViewIndex = 0; ViewIndex < ViewStates.Num(); ViewIndex++)
-	{
-		ViewStates[ViewIndex].Destroy();
-	}
-#endif //WITH_EDITOR
 }
 
 bool FDisplayClusterViewport::ShouldUseAdditionalTargetableResource() const
diff --git a/Engine/Plugins/Runtime/nDisplay/Source/DisplayCluster/Private/Render/Viewport/DisplayClusterViewport.h b/Engine/Plugins/Runtime/nDisplay/Source/DisplayCluster/Private/Render/Viewport/DisplayClusterViewport.h
--- a/Engine/Plugins/Runtime/nDisplay/Source/DisplayCluster/Private/Render/Viewport/DisplayClusterViewport.h
+++ b/Engine/Plugins/Runtime/nDisplay/Source/DisplayCluster/Private/Render/Viewport/DisplayClusterViewport.h
@@ -115,6 +115,7 @@ public:
 	bool    ImplPreview_CalculateStereoViewOffset(const uint32 InContextNum, FRotator& ViewRotation, const float WorldToMeters, FVector& ViewLocation);
 	FMatrix ImplPreview_GetStereoProjectionMatrix(const uint32 InContextNum);
 
+	// preview rendering view states
 	TIndirectArray<FSceneViewStateReference> ViewStates;
 	FSceneViewStateInterface* GetViewState(int32 ViewIndex);
 #endif //WITH_EDITOR
diff --git a/Engine/Plugins/Runtime/nDisplay/Source/DisplayCluster/Public/DisplayClusterRootActor.h b/Engine/Plugins/Runtime/nDisplay/Source/DisplayCluster/Public/DisplayClusterRootActor.h
--- a/Engine/Plugins/Runtime/nDisplay/Source/DisplayCluster/Public/DisplayClusterRootActor.h
+++ b/Engine/Plugins/Runtime/nDisplay/Source/DisplayCluster/Public/DisplayClusterRootActor.h
@@ -53,7 +53,6 @@ public:
 	~ADisplayClusterRootActor();
 
 public:
-
 	/**
 	 * Initialializes the instance with specified config data
 	 *
@@ -391,5 +390,9 @@ private:
 
 	FOnPreviewUpdated OnPreviewGenerated;
 	FOnPreviewUpdated OnPreviewDestroyed;
+
+private:
+	// ViewStatesReferences needs to be referenced to ReferenceCollector
+	static void AddReferencedObjects(UObject* InThis, FReferenceCollector& Collector);
 #endif
 };
