nDisplay: fixed texture replace for the single viewport

diff --git a/Engine/Plugins/Runtime/nDisplay/Source/DisplayCluster/Private/Game/EngineClasses/Basics/DisplayClusterViewportClient.cpp b/Engine/Plugins/Runtime/nDisplay/Source/DisplayCluster/Private/Game/EngineClasses/Basics/DisplayClusterViewportClient.cpp
--- a/Engine/Plugins/Runtime/nDisplay/Source/DisplayCluster/Private/Game/EngineClasses/Basics/DisplayClusterViewportClient.cpp
+++ b/Engine/Plugins/Runtime/nDisplay/Source/DisplayCluster/Private/Game/EngineClasses/Basics/DisplayClusterViewportClient.cpp
@@ -903,7 +903,7 @@ bool UDisplayClusterViewportClient::Draw_PIE(FViewport* InViewport, FCanvas* Sce
 	}
 
 	FDisplayClusterRenderFrame RenderFrame;
-	if (ViewportManager->BeginNewFrame(InViewport, MyWorld, RenderFrame) == false || RenderFrame.DesiredNumberOfViews < 1)
+	if (ViewportManager->BeginNewFrame(InViewport, MyWorld, RenderFrame) == false)
 	{
 		return false;
 	}
diff --git a/Engine/Plugins/Runtime/nDisplay/Source/DisplayCluster/Private/Game/EngineClasses/Scene/DisplayClusterRootActorEditor.cpp b/Engine/Plugins/Runtime/nDisplay/Source/DisplayCluster/Private/Game/EngineClasses/Scene/DisplayClusterRootActorEditor.cpp
--- a/Engine/Plugins/Runtime/nDisplay/Source/DisplayCluster/Private/Game/EngineClasses/Scene/DisplayClusterRootActorEditor.cpp
+++ b/Engine/Plugins/Runtime/nDisplay/Source/DisplayCluster/Private/Game/EngineClasses/Scene/DisplayClusterRootActorEditor.cpp
@@ -308,7 +308,7 @@ void ADisplayClusterRootActor::RenderPreview_Editor()
 		{
 			// Update preview viewports from settings
 			FDisplayClusterRenderFrame PreviewRenderFrame;
-			if (ViewportManager->BeginNewFrame(nullptr, PreviewWorld, PreviewRenderFrame) && PreviewRenderFrame.DesiredNumberOfViews > 0)
+			if (ViewportManager->BeginNewFrame(nullptr, PreviewWorld, PreviewRenderFrame))
 			{
 				ViewportManager->RenderInEditor(PreviewRenderFrame, nullptr);
 				// Send event about RTT changed
diff --git a/Engine/Plugins/Runtime/nDisplay/Source/DisplayCluster/Private/Render/Device/DisplayClusterDeviceBase.cpp b/Engine/Plugins/Runtime/nDisplay/Source/DisplayCluster/Private/Render/Device/DisplayClusterDeviceBase.cpp
--- a/Engine/Plugins/Runtime/nDisplay/Source/DisplayCluster/Private/Render/Device/DisplayClusterDeviceBase.cpp
+++ b/Engine/Plugins/Runtime/nDisplay/Source/DisplayCluster/Private/Render/Device/DisplayClusterDeviceBase.cpp
@@ -426,17 +426,14 @@ bool FDisplayClusterDeviceBase::BeginNewFrame(FViewport* InViewport, UWorld* InW
 			{
 				if (ViewportManager->BeginNewFrame(InViewport, InWorld, OutRenderFrame))
 				{
-					if (OutRenderFrame.DesiredNumberOfViews > 0)
-					{
-						// Begin use viewport manager for current frame
-						ViewportManagerPtr = ViewportManager;
+					// Begin use viewport manager for current frame
+					ViewportManagerPtr = ViewportManager;
 
-						// Send viewport manager proxy on render thread
-						NewViewportManagerProxy = ViewportManager->GetProxy();
+					// Send viewport manager proxy on render thread
+					NewViewportManagerProxy = ViewportManager->GetProxy();
 
-						// update total number of views for this frame (in multiple families)
-						DesiredNumberOfViews = OutRenderFrame.DesiredNumberOfViews;
-					}
+					// update total number of views for this frame (in multiple families)
+					DesiredNumberOfViews = OutRenderFrame.DesiredNumberOfViews;
 				}
 			}
 		}
diff --git a/Engine/Plugins/Runtime/nDisplay/Source/DisplayCluster/Private/Render/Viewport/DisplayClusterViewportManager_EditorPreview.cpp b/Engine/Plugins/Runtime/nDisplay/Source/DisplayCluster/Private/Render/Viewport/DisplayClusterViewportManager_EditorPreview.cpp
--- a/Engine/Plugins/Runtime/nDisplay/Source/DisplayCluster/Private/Render/Viewport/DisplayClusterViewportManager_EditorPreview.cpp
+++ b/Engine/Plugins/Runtime/nDisplay/Source/DisplayCluster/Private/Render/Viewport/DisplayClusterViewportManager_EditorPreview.cpp
@@ -108,9 +108,6 @@ bool FDisplayClusterViewportManager::RenderInEditor(class FDisplayClusterRenderF
 	}
 
 	FSceneInterface* PreviewScene = CurrentWorld->Scene;
-
-	bool bHasRender = false;
-
 	FEngineShowFlags EngineShowFlags = FEngineShowFlags(EShowFlagInitMode::ESFIM_Game);
 
 	//Experimental code from render team, now always disabled
@@ -176,25 +173,19 @@ bool FDisplayClusterViewportManager::RenderInEditor(class FDisplayClusterRenderF
 					Canvas.Clear(FLinearColor::Black);
 
 					GetRendererModule().BeginRenderingViewFamily(&Canvas, &ViewFamily);
-					bHasRender = true;
 				}
 			}
 		}
 	}
 
-	if (bHasRender)
-	{
-		// Handle special viewports game-thread logic at frame end
-		// custom postprocess single frame flag must be removed at frame end on game thread
-		FinalizeNewFrame();
+	// Handle special viewports game-thread logic at frame end
+	// custom postprocess single frame flag must be removed at frame end on game thread
+	FinalizeNewFrame();
 
-		// After all render target rendered call nDisplay frame rendering:
-		RenderFrame(InViewport);
-
-		return true;
-	}
+	// After all render target rendered call nDisplay frame rendering:
+	RenderFrame(InViewport);
 
-	return false;
+	return true;
 }
 
 #endif
