nDisplay: Fixed dragging a viewport to a mesh triggering a checkslow. Deleting viewports or cluster nodes will now let UE rename the object to the transient package instead of handling it manually.

diff --git a/Engine/Plugins/Runtime/nDisplay/Source/DisplayClusterConfigurator/Private/ClusterConfiguration/DisplayClusterConfiguratorClusterUtils.cpp b/Engine/Plugins/Runtime/nDisplay/Source/DisplayClusterConfigurator/Private/ClusterConfiguration/DisplayClusterConfiguratorClusterUtils.cpp
--- a/Engine/Plugins/Runtime/nDisplay/Source/DisplayClusterConfigurator/Private/ClusterConfiguration/DisplayClusterConfiguratorClusterUtils.cpp
+++ b/Engine/Plugins/Runtime/nDisplay/Source/DisplayClusterConfigurator/Private/ClusterConfiguration/DisplayClusterConfiguratorClusterUtils.cpp
@@ -711,8 +711,8 @@ bool FDisplayClusterConfiguratorClusterUtils::RemoveClusterNodeFromCluster(UDisp
 			}
 			
 			RemoveKeyFromMap(ClusterNodeParent, GET_MEMBER_NAME_CHECKED(UDisplayClusterConfigurationCluster, Nodes), *KeyPtr);
-
-			ClusterNode->Rename(nullptr, GetTransientPackage(), REN_DontCreateRedirectors | REN_ForceNoResetLoaders | REN_DoNotDirty);
+			ensure (ClusterNode->GetPackage() == GetTransientPackage());
+			
 			ClusterNode->SetFlags(RF_Transient);
 
 			if (bRemovingPrimaryNode)
@@ -910,9 +910,10 @@ bool FDisplayClusterConfiguratorClusterUtils::RemoveViewportFromClusterNode(UDis
 			ViewportParent->Modify();
 
 			RemoveKeyFromMap(ViewportParent, GET_MEMBER_NAME_CHECKED(UDisplayClusterConfigurationClusterNode, Viewports), *KeyPtr);
-
-			Viewport->Rename(nullptr, GetTransientPackage(), REN_DontCreateRedirectors | REN_ForceNoResetLoaders | REN_DoNotDirty);
+			ensure (Viewport->GetPackage() == GetTransientPackage());
+			
 			Viewport->SetFlags(RF_Transient);
+			
 			FDisplayClusterConfiguratorUtils::MarkDisplayClusterBlueprintAsModified(ViewportParent, true);
 			
 			return true;
diff --git a/Engine/Plugins/Runtime/nDisplay/Source/DisplayClusterConfigurator/Private/DisplayClusterConfiguratorPropertyUtils.cpp b/Engine/Plugins/Runtime/nDisplay/Source/DisplayClusterConfigurator/Private/DisplayClusterConfiguratorPropertyUtils.cpp
--- a/Engine/Plugins/Runtime/nDisplay/Source/DisplayClusterConfigurator/Private/DisplayClusterConfiguratorPropertyUtils.cpp
+++ b/Engine/Plugins/Runtime/nDisplay/Source/DisplayClusterConfigurator/Private/DisplayClusterConfiguratorPropertyUtils.cpp
@@ -168,11 +168,14 @@ int32 DisplayClusterConfiguratorPropertyUtils::FindMapHandleIndexFromKey(TShared
 		// We have to iterate to find the true index.
 		// There is a discrepency between MapHelper and MapHandle indices.
 		const int32 InternalIndex = MapHelper.FindInternalIndex(Idx);
-		FString* LocalKey = (FString*)MapHelper.GetKeyPtr(InternalIndex);
-		if (LocalKey && *LocalKey == Key)
+		if (InternalIndex != INDEX_NONE)
 		{
-			Index = Idx;
-			break;
+			FString* LocalKey = (FString*)MapHelper.GetKeyPtr(InternalIndex);
+			if (LocalKey && *LocalKey == Key)
+			{
+				Index = Idx;
+				break;
+			}
 		}
 	}
 
