nDisplay: Fix spam and incorrect behavior of replace all viewport textures while in Editor.

The spam message in Editor was:
    NodeId '' not found in ConfigData

It was happening because in Editor the current NodeId is empty.

The new behavior is to update the replace textures flag for all viewports of all nodes when there is no current node id.
When there is one, it only sets the flag of its viewports.

#rnx

diff --git a/Engine/Plugins/Runtime/nDisplay/Source/DisplayCluster/Private/Game/EngineClasses/Scene/DisplayClusterRootActor.cpp b/Engine/Plugins/Runtime/nDisplay/Source/DisplayCluster/Private/Game/EngineClasses/Scene/DisplayClusterRootActor.cpp
--- a/Engine/Plugins/Runtime/nDisplay/Source/DisplayCluster/Private/Game/EngineClasses/Scene/DisplayClusterRootActor.cpp
+++ b/Engine/Plugins/Runtime/nDisplay/Source/DisplayCluster/Private/Game/EngineClasses/Scene/DisplayClusterRootActor.cpp
@@ -957,37 +957,55 @@ bool ADisplayClusterRootActor::SetReplaceTextureFlagForAllViewports(bool bReplac
 		return false;
 	}
 
+	// Only update the viewports of the current node. If there isn't one, which is normal in Editor, we update them all.
+
 	const FString NodeId = Display.GetClusterMgr()->GetNodeId();
-	const UDisplayClusterConfigurationClusterNode* Node = ConfigData->Cluster->GetNode(NodeId);
+	TArray<UDisplayClusterConfigurationClusterNode*, TInlineAllocator<1>> Nodes;
 
-	if (!Node)
+	if (NodeId.IsEmpty())
 	{
-		UE_LOG(LogDisplayClusterGame, Warning, TEXT("NodeId '%s' not found in ConfigData"), *NodeId);
-		return false;
+		ConfigData->Cluster->Nodes.GenerateValueArray(Nodes);
 	}
-	
-	for (const TPair<FString, UDisplayClusterConfigurationViewport*>& ViewportItem : Node->Viewports)
+	else
 	{
-		if (ViewportItem.Value)
+		UDisplayClusterConfigurationClusterNode* Node = ConfigData->Cluster->GetNode(NodeId);
+
+		if (!Node)
 		{
-			ViewportItem.Value->RenderSettings.Replace.bAllowReplace = bReplace;
+			UE_LOG(LogDisplayClusterGame, Warning, TEXT("NodeId '%s' not found in ConfigData"), *NodeId);
+			return false;
 		}
+
+		Nodes.Add(Node);
 	}
-	
-	for (const TPair<FString, UDisplayClusterConfigurationClusterNode*>& NodeItem : ConfigData->Cluster->Nodes)
+
+	for (const UDisplayClusterConfigurationClusterNode* Node: Nodes)
 	{
-		if (!NodeItem.Value)
-		{
-			continue;
-		}
+		check(Node != nullptr);
 
-		for (const TPair<FString, UDisplayClusterConfigurationViewport*>& ViewportItem : NodeItem.Value->Viewports)
+		for (const TPair<FString, UDisplayClusterConfigurationViewport*>& ViewportItem : Node->Viewports)
 		{
 			if (ViewportItem.Value)
 			{
 				ViewportItem.Value->RenderSettings.Replace.bAllowReplace = bReplace;
 			}
 		}
+
+		for (const TPair<FString, UDisplayClusterConfigurationClusterNode*>& NodeItem : ConfigData->Cluster->Nodes)
+		{
+			if (!NodeItem.Value)
+			{
+				continue;
+			}
+
+			for (const TPair<FString, UDisplayClusterConfigurationViewport*>& ViewportItem : NodeItem.Value->Viewports)
+			{
+				if (ViewportItem.Value)
+				{
+					ViewportItem.Value->RenderSettings.Replace.bAllowReplace = bReplace;
+				}
+			}
+		}
 	}
 
 	return true;
