nDisplay: improved ViewState handling

diff --git a/Engine/Plugins/Runtime/nDisplay/Source/DisplayCluster/Private/Game/EngineClasses/Scene/DisplayClusterRootActor.cpp b/Engine/Plugins/Runtime/nDisplay/Source/DisplayCluster/Private/Game/EngineClasses/Scene/DisplayClusterRootActor.cpp
--- a/Engine/Plugins/Runtime/nDisplay/Source/DisplayCluster/Private/Game/EngineClasses/Scene/DisplayClusterRootActor.cpp
+++ b/Engine/Plugins/Runtime/nDisplay/Source/DisplayCluster/Private/Game/EngineClasses/Scene/DisplayClusterRootActor.cpp
@@ -1010,3 +1010,15 @@ bool ADisplayClusterRootActor::SetReplaceTextureFlagForAllViewports(bool bReplac
 
 	return true;
 }
+
+void ADisplayClusterRootActor::AddReferencedObjects(UObject* InThis, FReferenceCollector& Collector)
+{
+	ADisplayClusterRootActor* This = CastChecked<ADisplayClusterRootActor>(InThis);
+
+	if (This && This->ViewportManager.IsValid())
+	{
+		This->ViewportManager->AddReferencedObjects(Collector);
+	}
+
+	Super::AddReferencedObjects(InThis, Collector);
+}
diff --git a/Engine/Plugins/Runtime/nDisplay/Source/DisplayCluster/Private/Game/EngineClasses/Scene/DisplayClusterRootActorEditor.cpp b/Engine/Plugins/Runtime/nDisplay/Source/DisplayCluster/Private/Game/EngineClasses/Scene/DisplayClusterRootActorEditor.cpp
--- a/Engine/Plugins/Runtime/nDisplay/Source/DisplayCluster/Private/Game/EngineClasses/Scene/DisplayClusterRootActorEditor.cpp
+++ b/Engine/Plugins/Runtime/nDisplay/Source/DisplayCluster/Private/Game/EngineClasses/Scene/DisplayClusterRootActorEditor.cpp
@@ -268,27 +268,6 @@ bool ADisplayClusterRootActor::UpdatePreviewConfiguration_Editor(bool bUpdateAll
 	return false;
 }
 
-void ADisplayClusterRootActor::AddReferencedObjects(UObject* InThis, FReferenceCollector& Collector)
-{
-	ADisplayClusterRootActor* This = CastChecked<ADisplayClusterRootActor>(InThis);
-
-	for (IDisplayClusterViewport* ViewportIt : This->ViewportManager->GetViewports())
-	{
-		TIndirectArray<FSceneViewStateReference>& ViewStates = ((FDisplayClusterViewport*)ViewportIt)->ViewStates;
-
-		for (int32 ViewIndex = 0; ViewIndex < ViewStates.Num(); ViewIndex++)
-		{
-			FSceneViewStateInterface* Ref = ViewStates[ViewIndex].GetReference();
-			if (Ref)
-			{
-				Ref->AddReferencedObjects(Collector);
-			}
-		}
-	}
-
-	Super::AddReferencedObjects(This, Collector);
-}
-
 void ADisplayClusterRootActor::RenderPreview_Editor()
 {
 	if (UpdatePreviewConfiguration_Editor(false))
diff --git a/Engine/Plugins/Runtime/nDisplay/Source/DisplayCluster/Private/Render/Viewport/DisplayClusterViewport.cpp b/Engine/Plugins/Runtime/nDisplay/Source/DisplayCluster/Private/Render/Viewport/DisplayClusterViewport.cpp
--- a/Engine/Plugins/Runtime/nDisplay/Source/DisplayCluster/Private/Render/Viewport/DisplayClusterViewport.cpp
+++ b/Engine/Plugins/Runtime/nDisplay/Source/DisplayCluster/Private/Render/Viewport/DisplayClusterViewport.cpp
@@ -133,6 +133,24 @@ void FDisplayClusterViewport::HandleEndScene()
 		UninitializedProjectionPolicy = ProjectionPolicy;
 		ProjectionPolicy.Reset();
 	}
+
+#if WITH_EDITOR
+	CleanupViewState();
+#endif
+}
+
+void FDisplayClusterViewport::AddReferencedObjects(FReferenceCollector& Collector)
+{
+#if WITH_EDITOR
+	for (FSceneViewStateReference& ViewState : ViewStates)
+	{
+		FSceneViewStateInterface* Ref = ViewState.GetReference();
+		if (Ref)
+		{
+			Ref->AddReferencedObjects(Collector);
+		}
+	}
+#endif
 }
 
 bool FDisplayClusterViewport::ShouldUseAdditionalTargetableResource() const
@@ -244,7 +262,7 @@ FIntRect FDisplayClusterViewport::GetValidRect(const FIntRect& InRect, const TCH
 	return OutRect;
 }
 
-bool FDisplayClusterViewport::UpdateFrameContexts(const uint32 InViewPassNum, const FDisplayClusterRenderFrameSettings& InFrameSettings)
+bool FDisplayClusterViewport::UpdateFrameContexts(const uint32 InStereoViewIndex, const FDisplayClusterRenderFrameSettings& InFrameSettings)
 {
 	check(IsInGameThread());
 
@@ -373,7 +391,7 @@ bool FDisplayClusterViewport::UpdateFrameContexts(const uint32 InViewPassNum, co
 	for (uint32 ContextIt = 0; ContextIt < ViewportContextAmmount; ++ContextIt)
 	{
 		const EStereoscopicPass StereoscopicPass = (InFrameSettings.bIsRenderingInEditor) ? EStereoscopicPass::eSSP_FULL : FDisplayClusterViewportStereoscopicPass::EncodeStereoscopicPass(ContextIt, ViewportContextAmmount);
-		const int32 StereoViewIndex = (int32)(InViewPassNum + ContextIt);
+		const int32 StereoViewIndex = (int32)(InStereoViewIndex + ContextIt);
 
 		FDisplayClusterViewport_Context Context(ContextIt, StereoscopicPass, StereoViewIndex);
 
diff --git a/Engine/Plugins/Runtime/nDisplay/Source/DisplayCluster/Private/Render/Viewport/DisplayClusterViewport.h b/Engine/Plugins/Runtime/nDisplay/Source/DisplayCluster/Private/Render/Viewport/DisplayClusterViewport.h
--- a/Engine/Plugins/Runtime/nDisplay/Source/DisplayCluster/Private/Render/Viewport/DisplayClusterViewport.h
+++ b/Engine/Plugins/Runtime/nDisplay/Source/DisplayCluster/Private/Render/Viewport/DisplayClusterViewport.h
@@ -86,7 +86,6 @@ public:
 		return ProjectionPolicy;
 	}
 
-
 	virtual const TArray<FDisplayClusterViewport_Context>& GetContexts() const override
 	{
 		check(IsInGameThread());
@@ -114,10 +113,6 @@ public:
 	FSceneView* ImplCalcScenePreview(class FSceneViewFamilyContext& InOutViewFamily, uint32 ContextNum);
 	bool    ImplPreview_CalculateStereoViewOffset(const uint32 InContextNum, FRotator& ViewRotation, const float WorldToMeters, FVector& ViewLocation);
 	FMatrix ImplPreview_GetStereoProjectionMatrix(const uint32 InContextNum);
-
-	// preview rendering view states
-	TIndirectArray<FSceneViewStateReference> ViewStates;
-	FSceneViewStateInterface* GetViewState(int32 ViewIndex);
 #endif //WITH_EDITOR
 
 	// Get from logic request for additional targetable resource
@@ -149,6 +144,8 @@ public:
 	bool HandleStartScene();
 	void HandleEndScene();
 
+	void AddReferencedObjects(FReferenceCollector& Collector);
+
 	void ResetRuntimeParameters()
 	{
 		// Reset runtim flags from prev frame:
@@ -161,15 +158,24 @@ public:
 		CustomFrustumRendering.ResetConfiguration();
 	}
 
-
 	// Active view extension for this viewport
 	const TArray<FSceneViewExtensionRef> GatherActiveExtensions(FViewport* InViewport) const;
 
-	bool UpdateFrameContexts(const uint32 InViewPassNum, const FDisplayClusterRenderFrameSettings& InFrameSettings);
+	bool UpdateFrameContexts(const uint32 InStereoViewIndex, const FDisplayClusterRenderFrameSettings& InFrameSettings);
 
 public:
 	FIntRect GetValidRect(const FIntRect& InRect, const TCHAR* DbgSourceName);
 
+#if WITH_EDITOR
+	// Support view states for preview
+private:
+	FSceneViewStateInterface* GetViewState(uint32 ViewIndex);
+	void CleanupViewState();
+
+private:
+	TArray<FSceneViewStateReference> ViewStates;
+#endif
+
 public:
 	// Support OCIO:
 	FSceneViewExtensionIsActiveFunctor GetSceneViewExtensionIsActiveFunctor() const;
diff --git a/Engine/Plugins/Runtime/nDisplay/Source/DisplayCluster/Private/Render/Viewport/DisplayClusterViewportManager.cpp b/Engine/Plugins/Runtime/nDisplay/Source/DisplayCluster/Private/Render/Viewport/DisplayClusterViewportManager.cpp
--- a/Engine/Plugins/Runtime/nDisplay/Source/DisplayCluster/Private/Render/Viewport/DisplayClusterViewportManager.cpp
+++ b/Engine/Plugins/Runtime/nDisplay/Source/DisplayCluster/Private/Render/Viewport/DisplayClusterViewportManager.cpp
@@ -325,7 +325,8 @@ bool FDisplayClusterViewportManager::BeginNewFrame(FViewport* InViewport, UWorld
 	}
 
 	// generate unique stereo view index for each frame
-	uint32 ViewPassNum = 0;
+	// Begin from 1, because INDEX_NONE use ViewState[0] in LocalPlayer.cpp:786
+	uint32 StereoViewIndex = 1;
 
 	// Initialize viewports from new render settings, and create new contexts, reset prev frame resources
 	for (FDisplayClusterViewport* Viewport : Viewports)
@@ -333,9 +334,9 @@ bool FDisplayClusterViewportManager::BeginNewFrame(FViewport* InViewport, UWorld
 		// Save orig viewport contexts
 		TArray<FDisplayClusterViewport_Context> PrevContexts;
 		PrevContexts.Append(Viewport->GetContexts());
-		if (Viewport->UpdateFrameContexts(ViewPassNum, GetRenderFrameSettings()))
+		if (Viewport->UpdateFrameContexts(StereoViewIndex, GetRenderFrameSettings()))
 		{
-			ViewPassNum += Viewport->Contexts.Num();
+			StereoViewIndex += Viewport->Contexts.Num();
 		}
 
 		HandleViewportRTTChanges(PrevContexts, Viewport->GetContexts());
@@ -371,7 +372,7 @@ bool FDisplayClusterViewportManager::BeginNewFrame(FViewport* InViewport, UWorld
 	}
 
 	// Update desired views number
-	OutRenderFrame.UpdateDesiredNumberOfViews();
+	UpdateDesiredNumberOfViews(OutRenderFrame);
 
 #if WITH_EDITOR
 	// Get preview resources from root actor
@@ -381,6 +382,34 @@ bool FDisplayClusterViewportManager::BeginNewFrame(FViewport* InViewport, UWorld
 	return true;
 }
 
+void FDisplayClusterViewportManager::UpdateDesiredNumberOfViews(FDisplayClusterRenderFrame& InOutRenderFrame)
+{
+	InOutRenderFrame.DesiredNumberOfViews = 0;
+
+	for (FDisplayClusterRenderFrame::FFrameRenderTarget& RenderTargetIt : InOutRenderFrame.RenderTargets)
+	{
+		for (FDisplayClusterRenderFrame::FFrameViewFamily& ViewFamilyIt : RenderTargetIt.ViewFamilies)
+		{
+			ViewFamilyIt.NumViewsForRender = 0;
+
+			for (FDisplayClusterRenderFrame::FFrameView& ViewIt : ViewFamilyIt.Views)
+			{
+				if (ViewIt.bDisableRender == false)
+				{
+					ViewFamilyIt.NumViewsForRender++;
+
+					if (ViewIt.Viewport)
+					{
+						// Get StereoViewIndex for this viewport
+						const int32 StereoViewIndex = ViewIt.Viewport->GetContexts()[ViewIt.ContextNum].StereoViewIndex;
+						InOutRenderFrame.DesiredNumberOfViews = FMath::Max(StereoViewIndex + 1, InOutRenderFrame.DesiredNumberOfViews);
+					}
+				}
+			}
+		}
+	}
+}
+
 void FDisplayClusterViewportManager::FinalizeNewFrame()
 {
 	check(IsInGameThread());
@@ -682,3 +711,14 @@ void FDisplayClusterViewportManager::MarkComponentGeometryDirty(const FName InCo
 		PostProcessManager->GetOutputRemap()->MarkProceduralMeshComponentGeometryDirty(InComponentName);
 	}
 }
+
+void FDisplayClusterViewportManager::AddReferencedObjects(FReferenceCollector& Collector)
+{
+	for (FDisplayClusterViewport* ViewportIt : Viewports)
+	{
+		if (ViewportIt != nullptr)
+		{
+			ViewportIt->AddReferencedObjects(Collector);
+		}
+	}
+}
diff --git a/Engine/Plugins/Runtime/nDisplay/Source/DisplayCluster/Private/Render/Viewport/DisplayClusterViewportManager.h b/Engine/Plugins/Runtime/nDisplay/Source/DisplayCluster/Private/Render/Viewport/DisplayClusterViewportManager.h
--- a/Engine/Plugins/Runtime/nDisplay/Source/DisplayCluster/Private/Render/Viewport/DisplayClusterViewportManager.h
+++ b/Engine/Plugins/Runtime/nDisplay/Source/DisplayCluster/Private/Render/Viewport/DisplayClusterViewportManager.h
@@ -79,6 +79,8 @@ public:
 
 	virtual void MarkComponentGeometryDirty(const FName InComponentName = NAME_None) override;
 
+	virtual void AddReferencedObjects(FReferenceCollector& Collector) override;
+
 	// internal use only
 	bool CreateViewport(const FString& ViewportId, const class UDisplayClusterConfigurationViewport* ConfigurationViewport);
 	IDisplayClusterViewport* CreateViewport(const FString& ViewportId, const TSharedPtr<IDisplayClusterProjectionPolicy, ESPMode::ThreadSafe>& InProjectionPolicy);
@@ -116,6 +118,7 @@ public:
 	void SetViewportBufferRatio(FDisplayClusterViewport& DstViewport, float InBufferRatio);
 
 private:
+	void UpdateDesiredNumberOfViews(FDisplayClusterRenderFrame& InOutRenderFrame);
 	void ResetSceneRenderTargetSize();
 	void UpdateSceneRenderTargetSize();
 	void HandleViewportRTTChanges(const TArray<FDisplayClusterViewport_Context>& PrevContexts, const TArray<FDisplayClusterViewport_Context>& Contexts);
diff --git a/Engine/Plugins/Runtime/nDisplay/Source/DisplayCluster/Private/Render/Viewport/DisplayClusterViewport_EditorPreview.cpp b/Engine/Plugins/Runtime/nDisplay/Source/DisplayCluster/Private/Render/Viewport/DisplayClusterViewport_EditorPreview.cpp
--- a/Engine/Plugins/Runtime/nDisplay/Source/DisplayCluster/Private/Render/Viewport/DisplayClusterViewport_EditorPreview.cpp
+++ b/Engine/Plugins/Runtime/nDisplay/Source/DisplayCluster/Private/Render/Viewport/DisplayClusterViewport_EditorPreview.cpp
@@ -22,16 +22,44 @@
 #include "Components/DisplayClusterCameraComponent.h"
 #include "Components/DisplayClusterScreenComponent.h"
 
-
 #include "Render/Viewport/DisplayClusterViewportManager.h"
 #include "Render/Viewport/IDisplayClusterViewport.h"
 
 #include "Render/Viewport/Configuration/DisplayClusterViewportConfigurationHelpers_Postprocess.h"
 
-
 ///////////////////////////////////////////////////////////////////////////////////////
 //          FDisplayClusterViewport
 ///////////////////////////////////////////////////////////////////////////////////////
+void FDisplayClusterViewport::CleanupViewState()
+{
+	for (FSceneViewStateReference& ViewState : ViewStates)
+	{
+		FSceneViewStateInterface* Ref = ViewState.GetReference();
+		if (Ref != nullptr)
+		{
+			Ref->ClearMIDPool();
+		}
+	}
+}
+
+FSceneViewStateInterface* FDisplayClusterViewport::GetViewState(uint32 ViewIndex)
+{
+	int32 RequiredAmmount = (int32)ViewIndex - ViewStates.Num() + 1;
+	if(RequiredAmmount > 0)
+	{
+		ViewStates.AddDefaulted(RequiredAmmount);
+	}
+
+	if (ViewStates[ViewIndex].GetReference() == NULL)
+	{
+		const UWorld* CurrentWorld = Owner.GetCurrentWorld();
+		const ERHIFeatureLevel::Type FeatureLevel = CurrentWorld ? CurrentWorld->FeatureLevel.GetValue() : GMaxRHIFeatureLevel;
+
+		ViewStates[ViewIndex].Allocate(FeatureLevel);
+	}
+
+	return ViewStates[ViewIndex].GetReference();
+}
 
 FSceneView* FDisplayClusterViewport::ImplCalcScenePreview(FSceneViewFamilyContext& InOutViewFamily, uint32 InContextNum)
 {
@@ -61,12 +89,12 @@ FSceneView* FDisplayClusterViewport::ImplCalcScenePreview(FSceneViewFamilyContex
 		float StereoIPD = 0.f;
 		FIntRect ViewRect = Contexts[InContextNum].RenderTargetRect;
 
-		FEngineShowFlags ShowFlags = FEngineShowFlags(EShowFlagInitMode::ESFIM_Game);
-
 		FSceneViewInitOptions ViewInitOptions;
 
 		ViewInitOptions.SetViewRectangle(ViewRect);
 		ViewInitOptions.ViewFamily = &InOutViewFamily;
+
+		ViewInitOptions.SceneViewStateInterface = GetViewState(InContextNum);
 		ViewInitOptions.ViewActor = ViewOwner;
 
 		ViewInitOptions.ViewOrigin = ViewLocation;
@@ -87,11 +115,7 @@ FSceneView* FDisplayClusterViewport::ImplCalcScenePreview(FSceneViewFamilyContex
 		ViewInitOptions.StereoIPD = StereoIPD * (ViewInitOptions.WorldToMetersScale / 100.0f);
 
 		ViewInitOptions.BackgroundColor = FLinearColor::Black;
-		//ViewInitOptions.OverlayColor = FLinearColor::Black;
 
-		// initialize view states
-		ViewInitOptions.SceneViewStateInterface = GetViewState(InContextNum);
-		
 		FSceneView* View = new FSceneView(ViewInitOptions);
 
 		View->bIsSceneCapture = true;
@@ -149,29 +173,6 @@ FSceneView* FDisplayClusterViewport::ImplCalcScenePreview(FSceneViewFamilyContex
 	return nullptr;
 }
 
-FSceneViewStateInterface* FDisplayClusterViewport::GetViewState(int32 ViewIndex)
-{
-	while (ViewIndex >= ViewStates.Num())
-	{
-		ViewStates.Add(new FSceneViewStateReference());
-	}
-
-	FSceneViewStateInterface* ViewStateInterface = ViewStates[ViewIndex].GetReference();
-	if (ViewStateInterface == NULL)
-	{
-		ERHIFeatureLevel::Type FeatureLevel = GMaxRHIFeatureLevel;
-		if (Owner.GetCurrentWorld() != nullptr)
-		{
-			FeatureLevel = Owner.GetCurrentWorld()->Scene->GetFeatureLevel();
-		}
-
-		ViewStates[ViewIndex].Allocate(FeatureLevel);
-		ViewStateInterface = ViewStates[ViewIndex].GetReference();
-	}
-
-	return ViewStateInterface;
-}
-
 enum EDisplayClusterEyeType
 {
 	StereoLeft = 0,
diff --git a/Engine/Plugins/Runtime/nDisplay/Source/DisplayCluster/Public/DisplayClusterRootActor.h b/Engine/Plugins/Runtime/nDisplay/Source/DisplayCluster/Public/DisplayClusterRootActor.h
--- a/Engine/Plugins/Runtime/nDisplay/Source/DisplayCluster/Public/DisplayClusterRootActor.h
+++ b/Engine/Plugins/Runtime/nDisplay/Source/DisplayCluster/Public/DisplayClusterRootActor.h
@@ -264,6 +264,10 @@ public:
 	UDisplayClusterConfigurationData* CurrentConfigData;
 
 public:
+	// UObject interface
+	static void AddReferencedObjects(class UObject* InThis, class FReferenceCollector& Collector);
+	// End of UObject interface
+
 	bool IsInnerFrustumEnabled(const FString& InnerFrustumID) const;
 
 	// Return inner frustum priority by InnerFrustum name (from InnerFrustumPriority property)
@@ -391,8 +395,5 @@ private:
 	FOnPreviewUpdated OnPreviewGenerated;
 	FOnPreviewUpdated OnPreviewDestroyed;
 
-private:
-	// ViewStatesReferences needs to be referenced to ReferenceCollector
-	static void AddReferencedObjects(UObject* InThis, FReferenceCollector& Collector);
 #endif
 };
diff --git a/Engine/Plugins/Runtime/nDisplay/Source/DisplayCluster/Public/Render/Viewport/IDisplayClusterViewportManager.h b/Engine/Plugins/Runtime/nDisplay/Source/DisplayCluster/Public/Render/Viewport/IDisplayClusterViewportManager.h
--- a/Engine/Plugins/Runtime/nDisplay/Source/DisplayCluster/Public/Render/Viewport/IDisplayClusterViewportManager.h
+++ b/Engine/Plugins/Runtime/nDisplay/Source/DisplayCluster/Public/Render/Viewport/IDisplayClusterViewportManager.h
@@ -13,6 +13,7 @@ class FSceneViewFamilyContext;
 class ADisplayClusterRootActor;
 class UDisplayClusterConfigurationViewport;
 class IDisplayClusterViewportManagerProxy;
+class FReferenceCollector;
 
 struct FDisplayClusterConfigurationViewportPreview;
 
@@ -80,6 +81,8 @@ public:
 	// Send to render thread
 	virtual void RenderFrame(FViewport* InViewport) = 0;
 
+	virtual void AddReferencedObjects(FReferenceCollector& Collector) = 0;
+
 #if WITH_EDITOR
 	virtual bool UpdatePreviewConfiguration(const FDisplayClusterConfigurationViewportPreview& PreviewConfiguration, ADisplayClusterRootActor* InRootActorPtr) = 0;
 	virtual bool RenderInEditor(class FDisplayClusterRenderFrame& InRenderFrame, FViewport* InViewport) = 0;
diff --git a/Engine/Plugins/Runtime/nDisplay/Source/DisplayCluster/Public/Render/Viewport/RenderFrame/DisplayClusterRenderFrame.h b/Engine/Plugins/Runtime/nDisplay/Source/DisplayCluster/Public/Render/Viewport/RenderFrame/DisplayClusterRenderFrame.h
--- a/Engine/Plugins/Runtime/nDisplay/Source/DisplayCluster/Public/Render/Viewport/RenderFrame/DisplayClusterRenderFrame.h
+++ b/Engine/Plugins/Runtime/nDisplay/Source/DisplayCluster/Public/Render/Viewport/RenderFrame/DisplayClusterRenderFrame.h
@@ -76,29 +76,6 @@ public:
 		TArray<FFrameViewFamily> ViewFamilies;
 	};
 
-public:
-	void UpdateDesiredNumberOfViews()
-	{
-		DesiredNumberOfViews = 0;
-
-		for (FDisplayClusterRenderFrame::FFrameRenderTarget& RenderTargetIt : RenderTargets)
-		{
-			for (FDisplayClusterRenderFrame::FFrameViewFamily& ViewFamilyIt : RenderTargetIt.ViewFamilies)
-			{
-				ViewFamilyIt.NumViewsForRender = 0;
-
-				for (FDisplayClusterRenderFrame::FFrameView& ViewIt : ViewFamilyIt.Views)
-				{
-					if (ViewIt.bDisableRender == false)
-					{
-						ViewFamilyIt.NumViewsForRender++;
-						DesiredNumberOfViews++;
-					}
-				}
-			}
-		}
-	}
-
 public:
 	// Render frame to this targets
 	TArray<FFrameRenderTarget> RenderTargets;
