nDisplay: fixed Inner Frustum Multiplier behaviour when Custom Frustum enabled\disabled

diff --git a/Engine/Plugins/Runtime/nDisplay/Source/DisplayCluster/Private/Render/Viewport/Configuration/DisplayClusterViewportConfigurationHelpers_ICVFX.cpp b/Engine/Plugins/Runtime/nDisplay/Source/DisplayCluster/Private/Render/Viewport/Configuration/DisplayClusterViewportConfigurationHelpers_ICVFX.cpp
--- a/Engine/Plugins/Runtime/nDisplay/Source/DisplayCluster/Private/Render/Viewport/Configuration/DisplayClusterViewportConfigurationHelpers_ICVFX.cpp
+++ b/Engine/Plugins/Runtime/nDisplay/Source/DisplayCluster/Private/Render/Viewport/Configuration/DisplayClusterViewportConfigurationHelpers_ICVFX.cpp
@@ -39,7 +39,7 @@
 static bool ImplUpdateCameraProjectionSettings(TSharedPtr<IDisplayClusterProjectionPolicy, ESPMode::ThreadSafe>& InOutCameraProjection, const FDisplayClusterConfigurationICVFX_CameraSettings& CameraSettings, UCameraComponent* const CameraComponent)
 {
 	FDisplayClusterProjectionCameraPolicySettings PolicyCameraSettings;
-	PolicyCameraSettings.FOVMultiplier = CameraSettings.CustomFrustum.FieldOfViewMultiplier;
+	PolicyCameraSettings.FOVMultiplier = CameraSettings.CustomFrustum.bEnable ? CameraSettings.CustomFrustum.FieldOfViewMultiplier : 1.f;
 
 	// Lens correction
 	PolicyCameraSettings.FrustumRotation = CameraSettings.FrustumRotation;
@@ -264,9 +264,11 @@ FDisplayClusterShaderParameters_ICVFX::FCameraSettings FDisplayClusterViewportCo
 	const int32 CameraRenderOrder = RootActor.GetInnerFrustumPriority(InnerFrustumID);
 	Result.RenderOrder = (CameraRenderOrder<0) ? CameraSettings.RenderSettings.RenderOrder : CameraRenderOrder;
 
-	// softedge adjustments
-	const float Overscan = (CameraSettings.CustomFrustum.FieldOfViewMultiplier > 0) ? CameraSettings.CustomFrustum.FieldOfViewMultiplier : 1;
+	const float FieldOfViewMultiplier = CameraSettings.CustomFrustum.bEnable ? CameraSettings.CustomFrustum.FieldOfViewMultiplier : 1;
 
+	// softedge adjustments	
+	const float Overscan = (FieldOfViewMultiplier > 0) ? FieldOfViewMultiplier : 1;
+	
 	// remap values from 0-1 GUI range into acceptable 0.0 - 0.25 shader range
 	Result.SoftEdge.X = FMath::GetMappedRangeValueClamped(FVector2D(0.0, 1.0f), FVector2D(0.0, 0.25), CameraSettings.SoftEdge.Horizontal) / Overscan; // Left
 	Result.SoftEdge.Y = FMath::GetMappedRangeValueClamped(FVector2D(0.0, 1.0f), FVector2D(0.0, 0.25), CameraSettings.SoftEdge.Vertical) / Overscan; // Top
@@ -675,10 +677,12 @@ void FDisplayClusterViewportConfigurationHelpers_ICVFX::UpdateCameraViewportBuff
 {
 	float BufferRatio = CameraSettings.BufferRatio;
 
+	const float FieldOfViewMultiplier = CameraSettings.CustomFrustum.bEnable ? CameraSettings.CustomFrustum.FieldOfViewMultiplier : 1;
+
 	// adapt resolution should work as a shortcut to improve rendering quality
 	if (CameraSettings.CustomFrustum.bAdaptResolution)
 	{
-		DstViewport.RenderSettings.RenderTargetRatio *= CameraSettings.CustomFrustum.FieldOfViewMultiplier;
+		DstViewport.RenderSettings.RenderTargetRatio *= FieldOfViewMultiplier;
 	}
 
 	DstViewport.Owner.SetViewportBufferRatio(DstViewport, BufferRatio);
