nDisplay: Changed the metadata specifier the property referencer customization used for advanced path edit conditions from EditCondition to EditConditionPath so as not to overlap with the Editor's use of EditCondition, which was causing errors to be logged since EditCondition doesn't support pathed edit conditions

diff --git a/Engine/Plugins/Runtime/nDisplay/Source/DisplayCluster/Public/Components/DisplayClusterICVFXCameraComponent.h b/Engine/Plugins/Runtime/nDisplay/Source/DisplayCluster/Public/Components/DisplayClusterICVFXCameraComponent.h
--- a/Engine/Plugins/Runtime/nDisplay/Source/DisplayCluster/Public/Components/DisplayClusterICVFXCameraComponent.h
+++ b/Engine/Plugins/Runtime/nDisplay/Source/DisplayCluster/Public/Components/DisplayClusterICVFXCameraComponent.h
@@ -112,34 +112,34 @@ private:
 	UPROPERTY(EditAnywhere, Transient, Category = "In Camera VFX", meta = (PropertyPath = "CameraSettings.FrustumOffset"))
 	FDisplayClusterEditorPropertyReference FrustumOffsetRef;
 
-	UPROPERTY(EditAnywhere, Transient, Category = "In Camera VFX", meta = (PropertyPath = "CameraSettings.RenderSettings.GenerateMips", EditCondition = "CameraSettings.bEnable"))
+	UPROPERTY(EditAnywhere, Transient, Category = "In Camera VFX", meta = (PropertyPath = "CameraSettings.RenderSettings.GenerateMips", EditConditionPath = "CameraSettings.bEnable"))
 	FDisplayClusterEditorPropertyReference GenerateMipsRef;
 
-	UPROPERTY(EditAnywhere, Transient, Category = "In Camera VFX", meta = (PropertyPath = "CameraSettings.CameraMotionBlur", EditCondition = "CameraSettings.bEnable"))
+	UPROPERTY(EditAnywhere, Transient, Category = "In Camera VFX", meta = (PropertyPath = "CameraSettings.CameraMotionBlur", EditConditionPath = "CameraSettings.bEnable"))
 	FDisplayClusterEditorPropertyReference CameraMotionBlurRef;
 
-	UPROPERTY(EditAnywhere, Transient, Category = "In Camera VFX", meta = (PropertyPath = "CameraSettings.CameraHideList", EditCondition = "CameraSettings.bEnable"))
+	UPROPERTY(EditAnywhere, Transient, Category = "In Camera VFX", meta = (PropertyPath = "CameraSettings.CameraHideList", EditConditionPath = "CameraSettings.bEnable"))
 	FDisplayClusterEditorPropertyReference CameraHideListRef;
 
-	UPROPERTY(EditAnywhere, Transient, Category = Chromakey, meta = (PropertyPath = "CameraSettings.Chromakey.bEnable", EditCondition = "CameraSettings.bEnable"))
+	UPROPERTY(EditAnywhere, Transient, Category = Chromakey, meta = (PropertyPath = "CameraSettings.Chromakey.bEnable", EditConditionPath = "CameraSettings.bEnable"))
 	FDisplayClusterEditorPropertyReference ChromaKeyEnabledRef;
 
-	UPROPERTY(EditAnywhere, Transient, Category = Chromakey, meta = (PropertyPath = "CameraSettings.Chromakey.ChromakeyColor", EditCondition = "CameraSettings.bEnable && CameraSettings.Chromakey.bEnable"))
+	UPROPERTY(EditAnywhere, Transient, Category = Chromakey, meta = (PropertyPath = "CameraSettings.Chromakey.ChromakeyColor", EditConditionPath = "CameraSettings.bEnable && CameraSettings.Chromakey.bEnable"))
 	FDisplayClusterEditorPropertyReference ChromakeyColorRef;
 
-	UPROPERTY(EditAnywhere, Transient, Category = Chromakey, meta = (PropertyPath = "CameraSettings.Chromakey.ChromakeyMarkers", EditCondition = "CameraSettings.bEnable && CameraSettings.Chromakey.bEnable"))
+	UPROPERTY(EditAnywhere, Transient, Category = Chromakey, meta = (PropertyPath = "CameraSettings.Chromakey.ChromakeyMarkers", EditConditionPath = "CameraSettings.bEnable && CameraSettings.Chromakey.bEnable"))
 	FDisplayClusterEditorPropertyReference ChromakeyMarkersRef;
 
-	UPROPERTY(EditAnywhere, Transient, Category = Chromakey, meta = (PropertyPath = "CameraSettings.Chromakey.ChromakeyRenderTexture", EditCondition = "CameraSettings.bEnable && CameraSettings.Chromakey.bEnable"))
+	UPROPERTY(EditAnywhere, Transient, Category = Chromakey, meta = (PropertyPath = "CameraSettings.Chromakey.ChromakeyRenderTexture", EditConditionPath = "CameraSettings.bEnable && CameraSettings.Chromakey.bEnable"))
 	FDisplayClusterEditorPropertyReference ChromakeyRenderTextureRef;
 
 	UPROPERTY(EditAnywhere, Transient, Category = OCIO, meta = (PropertyPath = "CameraSettings.AllNodesOCIOConfiguration.bIsEnabled", DisplayName = "Enable Inner Frustum OCIO", ToolTip = "Enable the application of an OpenColorIO configuration to all nodes."))
 	FDisplayClusterEditorPropertyReference OCIOConfigurationEnabledRef;
 
-	UPROPERTY(EditAnywhere, Transient, Category = OCIO, meta = (PropertyPath = "CameraSettings.AllNodesOCIOConfiguration.OCIOConfiguration.ColorConfiguration", DisplayName = "All Nodes Color Configuration", ToolTip = "Apply this OpenColorIO configuration to all nodes.", EditCondition = "CameraSettings.AllNodesOCIOConfiguration.bIsEnabled"))
+	UPROPERTY(EditAnywhere, Transient, Category = OCIO, meta = (PropertyPath = "CameraSettings.AllNodesOCIOConfiguration.OCIOConfiguration.ColorConfiguration", DisplayName = "All Nodes Color Configuration", ToolTip = "Apply this OpenColorIO configuration to all nodes.", EditConditionPath = "CameraSettings.AllNodesOCIOConfiguration.bIsEnabled"))
 	FDisplayClusterEditorPropertyReference OCIOColorConfiguratonRef;
 
-	UPROPERTY(EditAnywhere, Transient, Category = OCIO, meta = (PropertyPath = "CameraSettings.PerNodeOCIOProfiles", EditCondition = "CameraSettings.AllNodesOCIOConfiguration.bIsEnabled"))
+	UPROPERTY(EditAnywhere, Transient, Category = OCIO, meta = (PropertyPath = "CameraSettings.PerNodeOCIOProfiles", EditConditionPath = "CameraSettings.AllNodesOCIOConfiguration.bIsEnabled"))
 	FDisplayClusterEditorPropertyReference PerNodeOCIOProfilesRef;
 
 	UPROPERTY(EditAnywhere, Transient, Category = "Inner Frustum Color Grading", meta = (PropertyPath = "CameraSettings.AllNodesColorGrading", DisplayName = "All Nodes"))
@@ -148,34 +148,34 @@ private:
 	UPROPERTY(EditAnywhere, Transient, Category = "Inner Frustum Color Grading", meta = (PropertyPath = "CameraSettings.PerNodeColorGrading"))
 	FDisplayClusterEditorPropertyReference PerNodeColorGradingRef;
 
-	UPROPERTY(EditAnywhere, Transient, Category = "Texture Replacement", meta = (PropertyPath = "CameraSettings.RenderSettings.Replace.bAllowReplace", DisplayName = "Enable Inner Frustum Texture Replacement", ToolTip = "Set to True to replace the entire inner frustum with the specified texture.", EditCondition = "CameraSettings.bEnable"))
+	UPROPERTY(EditAnywhere, Transient, Category = "Texture Replacement", meta = (PropertyPath = "CameraSettings.RenderSettings.Replace.bAllowReplace", DisplayName = "Enable Inner Frustum Texture Replacement", ToolTip = "Set to True to replace the entire inner frustum with the specified texture.", EditConditionPath = "CameraSettings.bEnable"))
 	FDisplayClusterEditorPropertyReference TextureReplacementEnabledRef;
 
-	UPROPERTY(EditAnywhere, Transient, Category = "Texture Replacement", meta = (PropertyPath = "CameraSettings.RenderSettings.Replace.SourceTexture", EditCondition = "CameraSettings.bEnable && CameraSettings.RenderSettings.Replace.bAllowReplace"))
+	UPROPERTY(EditAnywhere, Transient, Category = "Texture Replacement", meta = (PropertyPath = "CameraSettings.RenderSettings.Replace.SourceTexture", EditConditionPath = "CameraSettings.bEnable && CameraSettings.RenderSettings.Replace.bAllowReplace"))
 	FDisplayClusterEditorPropertyReference SourceTextureRef;
 
-	UPROPERTY(EditAnywhere, Transient, Category = "Texture Replacement", meta = (PropertyPath = "CameraSettings.RenderSettings.Replace.bShouldUseTextureRegion", EditCondition = "CameraSettings.bEnable && CameraSettings.RenderSettings.Replace.bAllowReplace"))
+	UPROPERTY(EditAnywhere, Transient, Category = "Texture Replacement", meta = (PropertyPath = "CameraSettings.RenderSettings.Replace.bShouldUseTextureRegion", EditConditionPath = "CameraSettings.bEnable && CameraSettings.RenderSettings.Replace.bAllowReplace"))
 	FDisplayClusterEditorPropertyReference ShouldUseTextureRegionRef;
 
-	UPROPERTY(EditAnywhere, Transient, Category = "Texture Replacement", meta = (PropertyPath = "CameraSettings.RenderSettings.Replace.TextureRegion", EditCondition = "CameraSettings.bEnable && CameraSettings.RenderSettings.Replace.bAllowReplace && CameraSettings.RenderSettings.Replace.bShouldUseTextureRegion"))
+	UPROPERTY(EditAnywhere, Transient, Category = "Texture Replacement", meta = (PropertyPath = "CameraSettings.RenderSettings.Replace.TextureRegion", EditConditionPath = "CameraSettings.bEnable && CameraSettings.RenderSettings.Replace.bAllowReplace && CameraSettings.RenderSettings.Replace.bShouldUseTextureRegion"))
 	FDisplayClusterEditorPropertyReference TextureRegionRef;
 
-	UPROPERTY(EditDefaultsOnly, Transient, Category = Configuration, meta = (PropertyPath = "CameraSettings.RenderSettings.RenderOrder", EditCondition = "CameraSettings.bEnable"))
+	UPROPERTY(EditDefaultsOnly, Transient, Category = Configuration, meta = (PropertyPath = "CameraSettings.RenderSettings.RenderOrder", EditConditionPath = "CameraSettings.bEnable"))
 	FDisplayClusterEditorPropertyReference RenderOrderRef;
 
-	UPROPERTY(EditDefaultsOnly, Transient, Category = Configuration, meta = (PropertyPath = "CameraSettings.RenderSettings.CustomFrameSize", EditCondition = "CameraSettings.bEnable"))
+	UPROPERTY(EditDefaultsOnly, Transient, Category = Configuration, meta = (PropertyPath = "CameraSettings.RenderSettings.CustomFrameSize", EditConditionPath = "CameraSettings.bEnable"))
 	FDisplayClusterEditorPropertyReference CustomFrameSizeRef;
 
-	UPROPERTY(EditDefaultsOnly, Transient, Category = Configuration, meta = (PropertyPath = "CameraSettings.RenderSettings.AdvancedRenderSettings.RenderTargetRatio", EditCondition = "CameraSettings.bEnable"))
+	UPROPERTY(EditDefaultsOnly, Transient, Category = Configuration, meta = (PropertyPath = "CameraSettings.RenderSettings.AdvancedRenderSettings.RenderTargetRatio", EditConditionPath = "CameraSettings.bEnable"))
 	FDisplayClusterEditorPropertyReference RenderTargetRatioRef;
 
-	UPROPERTY(EditDefaultsOnly, Transient, Category = Configuration, meta = (PropertyPath = "CameraSettings.RenderSettings.AdvancedRenderSettings.GPUIndex", EditCondition = "CameraSettings.bEnable"))
+	UPROPERTY(EditDefaultsOnly, Transient, Category = Configuration, meta = (PropertyPath = "CameraSettings.RenderSettings.AdvancedRenderSettings.GPUIndex", EditConditionPath = "CameraSettings.bEnable"))
 	FDisplayClusterEditorPropertyReference GPUIndexRef;
 
-	UPROPERTY(EditDefaultsOnly, Transient, Category = Configuration, meta = (PropertyPath = "CameraSettings.RenderSettings.AdvancedRenderSettings.StereoGPUIndex", EditCondition = "CameraSettings.bEnable"))
+	UPROPERTY(EditDefaultsOnly, Transient, Category = Configuration, meta = (PropertyPath = "CameraSettings.RenderSettings.AdvancedRenderSettings.StereoGPUIndex", EditConditionPath = "CameraSettings.bEnable"))
 	FDisplayClusterEditorPropertyReference StereoGPUIndexRef;
 
-	UPROPERTY(EditDefaultsOnly, Transient, Category = Configuration, meta = (PropertyPath = "CameraSettings.RenderSettings.AdvancedRenderSettings.StereoMode", EditCondition = "CameraSettings.bEnable"))
+	UPROPERTY(EditDefaultsOnly, Transient, Category = Configuration, meta = (PropertyPath = "CameraSettings.RenderSettings.AdvancedRenderSettings.StereoMode", EditConditionPath = "CameraSettings.bEnable"))
 	FDisplayClusterEditorPropertyReference StereoModeRef;
 
 #endif // WITH_EDITORONLY_DATA
diff --git a/Engine/Plugins/Runtime/nDisplay/Source/DisplayCluster/Public/DisplayClusterRootActor.h b/Engine/Plugins/Runtime/nDisplay/Source/DisplayCluster/Public/DisplayClusterRootActor.h
--- a/Engine/Plugins/Runtime/nDisplay/Source/DisplayCluster/Public/DisplayClusterRootActor.h
+++ b/Engine/Plugins/Runtime/nDisplay/Source/DisplayCluster/Public/DisplayClusterRootActor.h
@@ -189,7 +189,7 @@ private:
 	UPROPERTY(EditAnywhere, Transient, Category = "Color Grading", meta = (PropertyPath = "CurrentConfigData.StageSettings.EntireClusterColorGrading.bEnableEntireClusterColorGrading"))
 	FDisplayClusterEditorPropertyReference EnableClusterColorGradingRef;
 
-	UPROPERTY(EditAnywhere, Transient, Category = "Color Grading", meta = (DisplayName = "Entire Cluster", PropertyPath = "CurrentConfigData.StageSettings.EntireClusterColorGrading.ColorGradingSettings", EditConditional = "CurrentConfigData.StageSettings.EntireClusterColorGrading.bEnableEntireClusterColorGrading"))
+	UPROPERTY(EditAnywhere, Transient, Category = "Color Grading", meta = (DisplayName = "Entire Cluster", PropertyPath = "CurrentConfigData.StageSettings.EntireClusterColorGrading.ColorGradingSettings", EditConditionPath = "CurrentConfigData.StageSettings.EntireClusterColorGrading.bEnableEntireClusterColorGrading"))
 	FDisplayClusterEditorPropertyReference ClusterColorGradingRef;
 
 	UPROPERTY(EditAnywhere, Transient, Category = "Color Grading", meta = (PropertyPath = "CurrentConfigData.StageSettings.PerViewportColorGrading"))
@@ -198,7 +198,7 @@ private:
 	UPROPERTY(EditAnywhere, Transient, Category = OCIO, meta = (PropertyPath = "CurrentConfigData.StageSettings.bUseOverallClusterOCIOConfiguration"))
 	FDisplayClusterEditorPropertyReference EnableClusterOCIORef;
 
-	UPROPERTY(EditAnywhere, Transient, Category = OCIO, meta = (DisplayName = "All Viewports Color Configuration", PropertyPath = "CurrentConfigData.StageSettings.AllViewportsOCIOConfiguration.OCIOConfiguration.ColorConfiguration", ToolTip = "Apply this OpenColorIO configuration to all viewports.", EditCondition = "CurrentConfigData.StageSettings.bUseOverallClusterOCIOConfiguration"))
+	UPROPERTY(EditAnywhere, Transient, Category = OCIO, meta = (DisplayName = "All Viewports Color Configuration", PropertyPath = "CurrentConfigData.StageSettings.AllViewportsOCIOConfiguration.OCIOConfiguration.ColorConfiguration", ToolTip = "Apply this OpenColorIO configuration to all viewports.", EditConditionPath = "CurrentConfigData.StageSettings.bUseOverallClusterOCIOConfiguration"))
 	FDisplayClusterEditorPropertyReference ClusterOCIOColorConfigurationRef;
 
 	UPROPERTY(EditAnywhere, Transient, Category = OCIO, meta = (PropertyPath = "CurrentConfigData.StageSettings.PerViewportOCIOProfiles"))
@@ -207,10 +207,10 @@ private:
 	UPROPERTY(EditAnywhere, Transient, Category = "Light Cards", meta = (PropertyPath = "CurrentConfigData.StageSettings.Lightcard.bEnable"))
 	FDisplayClusterEditorPropertyReference EnableLightcardsRef;
 
-	UPROPERTY(EditAnywhere, Transient, Category = "Light Cards", meta = (PropertyPath = "CurrentConfigData.StageSettings.Lightcard.Blendingmode", EditCondition = "CurrentConfigData.StageSettings.Lightcard.bEnable"))
+	UPROPERTY(EditAnywhere, Transient, Category = "Light Cards", meta = (PropertyPath = "CurrentConfigData.StageSettings.Lightcard.Blendingmode", EditConditionPath = "CurrentConfigData.StageSettings.Lightcard.bEnable"))
 	FDisplayClusterEditorPropertyReference LightCardBlendingModeRef;
 
-	UPROPERTY(EditAnywhere, Transient, Category = "Light Cards", meta = (PropertyPath = "CurrentConfigData.StageSettings.Lightcard.ShowOnlyList", EditCondition = "CurrentConfigData.StageSettings.Lightcard.bEnable"))
+	UPROPERTY(EditAnywhere, Transient, Category = "Light Cards", meta = (PropertyPath = "CurrentConfigData.StageSettings.Lightcard.ShowOnlyList", EditConditionPath = "CurrentConfigData.StageSettings.Lightcard.bEnable"))
 	FDisplayClusterEditorPropertyReference LightCardContentRef;
 #endif // WITH_EDITORONLY_DATA
 
diff --git a/Engine/Plugins/Runtime/nDisplay/Source/DisplayClusterConfigurator/Private/Views/Details/DisplayClusterEditorPropertyReferenceTypeCustomization.cpp b/Engine/Plugins/Runtime/nDisplay/Source/DisplayClusterConfigurator/Private/Views/Details/DisplayClusterEditorPropertyReferenceTypeCustomization.cpp
--- a/Engine/Plugins/Runtime/nDisplay/Source/DisplayClusterConfigurator/Private/Views/Details/DisplayClusterEditorPropertyReferenceTypeCustomization.cpp
+++ b/Engine/Plugins/Runtime/nDisplay/Source/DisplayClusterConfigurator/Private/Views/Details/DisplayClusterEditorPropertyReferenceTypeCustomization.cpp
@@ -10,6 +10,7 @@
 #include "Widgets/Text/STextBlock.h"
 
 const FName FDisplayClusterEditorPropertyReferenceTypeCustomization::PropertyPathMetadataKey = TEXT("PropertyPath");
+const FName FDisplayClusterEditorPropertyReferenceTypeCustomization::EditConditionPathMetadataKey = TEXT("EditConditionPath");
 
 TSharedRef<IPropertyTypeCustomization> FDisplayClusterEditorPropertyReferenceTypeCustomization::MakeInstance()
 {
@@ -19,10 +20,10 @@ TSharedRef<IPropertyTypeCustomization> FDisplayClusterEditorPropertyReferenceTyp
 void FDisplayClusterEditorPropertyReferenceTypeCustomization::CustomizeHeader(TSharedRef<IPropertyHandle> InPropertyHandle, FDetailWidgetRow& InHeaderRow, IPropertyTypeCustomizationUtils& CustomizationUtils)
 {
 	check(InPropertyHandle->IsValidHandle());
-	if (InPropertyHandle->HasMetaData(TEXT("EditCondition")))
+	if (InPropertyHandle->HasMetaData(EditConditionPathMetadataKey))
 	{
 		// Check to see if the edit condition is using a property path as well.
-		FString EditCondition = InPropertyHandle->GetMetaData(TEXT("EditCondition"));
+		FString EditCondition = InPropertyHandle->GetMetaData(EditConditionPathMetadataKey);
 		EditCondition.RemoveSpacesInline();
 		if (EditCondition.Contains(TEXT(".")) || EditCondition.Contains(TEXT("->")))
 		{
diff --git a/Engine/Plugins/Runtime/nDisplay/Source/DisplayClusterConfigurator/Private/Views/Details/DisplayClusterEditorPropertyReferenceTypeCustomization.h b/Engine/Plugins/Runtime/nDisplay/Source/DisplayClusterConfigurator/Private/Views/Details/DisplayClusterEditorPropertyReferenceTypeCustomization.h
--- a/Engine/Plugins/Runtime/nDisplay/Source/DisplayClusterConfigurator/Private/Views/Details/DisplayClusterEditorPropertyReferenceTypeCustomization.h
+++ b/Engine/Plugins/Runtime/nDisplay/Source/DisplayClusterConfigurator/Private/Views/Details/DisplayClusterEditorPropertyReferenceTypeCustomization.h
@@ -15,6 +15,10 @@ public:
 	/** Metadata specifier to provide the path to a property to display. */
 	static const FName PropertyPathMetadataKey;
 
+	/** Metadata specifier to provide a path to a property to use as an edit condition. */
+	static const FName EditConditionPathMetadataKey;
+
+
 	static TSharedRef<IPropertyTypeCustomization> MakeInstance();
 
 	//~ IPropertyTypeCustomization interface begin
