nDisplay: failover GUI

diff --git a/Engine/Plugins/Runtime/nDisplay/Source/DisplayCluster/Private/Cluster/DisplayClusterClusterManager.cpp b/Engine/Plugins/Runtime/nDisplay/Source/DisplayCluster/Private/Cluster/DisplayClusterClusterManager.cpp
--- a/Engine/Plugins/Runtime/nDisplay/Source/DisplayCluster/Private/Cluster/DisplayClusterClusterManager.cpp
+++ b/Engine/Plugins/Runtime/nDisplay/Source/DisplayCluster/Private/Cluster/DisplayClusterClusterManager.cpp
@@ -832,16 +832,20 @@ TUniquePtr<IDisplayClusterClusterNodeController> FDisplayClusterClusterManager::
 
 TUniquePtr<IDisplayClusterFailoverNodeController> FDisplayClusterClusterManager::CreateFailoverNodeController() const
 {
-	const EDisplayClusterNodeRole ClusterRole = GetClusterRole();
-	switch (ClusterRole)
+	// Instantiate appropriate controller depending on operation mode and cluster role
+	if (CurrentOperationMode == EDisplayClusterOperationMode::Cluster)
 	{
-	case EDisplayClusterNodeRole::Primary:
-		UE_LOG(LogDisplayClusterCluster, Log, TEXT("Instantiating failover primary controller..."));
-		return MakeUnique<FDisplayClusterFailoverNodeCtrlPrimary>();
+		const EDisplayClusterNodeRole ClusterRole = GetClusterRole();
+		switch (ClusterRole)
+		{
+		case EDisplayClusterNodeRole::Primary:
+			UE_LOG(LogDisplayClusterCluster, Log, TEXT("Instantiating failover primary controller..."));
+			return MakeUnique<FDisplayClusterFailoverNodeCtrlPrimary>();
 
-	case EDisplayClusterNodeRole::Secondary:
-		UE_LOG(LogDisplayClusterCluster, Log, TEXT("Instantiating failover secondary controller..."));
-		return MakeUnique<FDisplayClusterFailoverNodeCtrlSecondary>();
+		case EDisplayClusterNodeRole::Secondary:
+			UE_LOG(LogDisplayClusterCluster, Log, TEXT("Instantiating failover secondary controller..."));
+			return MakeUnique<FDisplayClusterFailoverNodeCtrlSecondary>();
+		}
 	}
 
 	UE_LOG(LogDisplayClusterCluster, Log, TEXT("Instantiating failover null controller..."));
diff --git a/Engine/Plugins/Runtime/nDisplay/Source/DisplayCluster/Private/Cluster/Failover/DisplayClusterFailoverNodeCtrlBase.cpp b/Engine/Plugins/Runtime/nDisplay/Source/DisplayCluster/Private/Cluster/Failover/DisplayClusterFailoverNodeCtrlBase.cpp
--- a/Engine/Plugins/Runtime/nDisplay/Source/DisplayCluster/Private/Cluster/Failover/DisplayClusterFailoverNodeCtrlBase.cpp
+++ b/Engine/Plugins/Runtime/nDisplay/Source/DisplayCluster/Private/Cluster/Failover/DisplayClusterFailoverNodeCtrlBase.cpp
@@ -2,6 +2,10 @@
 
 #include "Cluster/Failover/DisplayClusterFailoverNodeCtrlBase.h"
 
+#include "DisplayClusterConfigurationTypes.h"
+#include "Config/IPDisplayClusterConfigManager.h"
+
+#include "Misc/DisplayClusterGlobals.h"
 #include "Misc/DisplayClusterLog.h"
 #include "Misc/ScopeLock.h"
 
@@ -10,20 +14,10 @@
 #include "HAL/IConsoleManager.h"
 
 
-// Failover settings
-static TAutoConsoleVariable<int32> CVarFailover(
-	TEXT("nDisplay.network.failover"),
-	0,
-	TEXT("Failover settings:\n")
-	TEXT("\t0 : Disabled\n")
-	TEXT("\t1 : Failover v1 - drop secondary nodes on fail\n")
-);
-
-
 FDisplayClusterFailoverNodeCtrlBase::FDisplayClusterFailoverNodeCtrlBase()
-	: FailoverPolicy( FDisplayClusterFailoverNodeCtrlBase::GetFailoverPolicyFromCvarValue(CVarFailover.GetValueOnAnyThread()) )
+	: FailoverPolicy( FDisplayClusterFailoverNodeCtrlBase::GetFailoverPolicyFromConfig() )
 {
-	UE_LOG(LogDisplayClusterCluster, Log, TEXT("Failover: policy is %s"), *FDisplayClusterFailoverNodeCtrlBase::FailoverPolicyAsString(FailoverPolicy));
+	UE_LOG(LogDisplayClusterCluster, Log, TEXT("Failover: policy %u"), static_cast<uint8>(FailoverPolicy));
 }
 
 
@@ -35,11 +29,11 @@ void FDisplayClusterFailoverNodeCtrlBase::HandleCommResult(EDisplayClusterCommRe
 
 	switch (GetFailoverPolicy())
 	{
-	case EDisplayClusterFailoverPolicy::Disabled:
+	case EDisplayClusterConfigurationFailoverPolicy::Disabled:
 		return HandleCommResult_Disabled(CommResult);
 
-	case EDisplayClusterFailoverPolicy::Failover_v1_DropSecondaryNodesOnly:
-		return HandleCommResult_Failover_v1(CommResult);
+	case EDisplayClusterConfigurationFailoverPolicy::DropSecondaryNodesOnly:
+		return HandleCommResult_DropSecondaryNodesOnly(CommResult);
 	}
 }
 
@@ -51,38 +45,23 @@ void FDisplayClusterFailoverNodeCtrlBase::HandleNodeFailed(const FString& NodeId
 
 	switch (GetFailoverPolicy())
 	{
-	case EDisplayClusterFailoverPolicy::Disabled:
+	case EDisplayClusterConfigurationFailoverPolicy::Disabled:
 		return HandleNodeFailed_Disabled(NodeId, NodeFailType);
 
-	case EDisplayClusterFailoverPolicy::Failover_v1_DropSecondaryNodesOnly:
-		return HandleNodeFailed_Failover_v1(NodeId, NodeFailType);
+	case EDisplayClusterConfigurationFailoverPolicy::DropSecondaryNodesOnly:
+		return HandleNodeFailed_DropSecondaryNodesOnly(NodeId, NodeFailType);
 	}
 }
 
-EDisplayClusterFailoverPolicy FDisplayClusterFailoverNodeCtrlBase::GetFailoverPolicyFromCvarValue(int32 FailoverCvarNumber)
+EDisplayClusterConfigurationFailoverPolicy FDisplayClusterFailoverNodeCtrlBase::GetFailoverPolicyFromConfig()
 {
-	switch (FailoverCvarNumber)
+	if (const IPDisplayClusterConfigManager* const ConfigMgr = GDisplayCluster->GetPrivateConfigMgr())
 	{
-	case 1:
-		return EDisplayClusterFailoverPolicy::Failover_v1_DropSecondaryNodesOnly;
-
-	case 0:
-	default:
-		return EDisplayClusterFailoverPolicy::Disabled;
+		if (const UDisplayClusterConfigurationData* const ConfigData = ConfigMgr->GetConfig())
+		{
+			return ConfigData->Cluster->Failover.FailoverPolicy;
+		}
 	}
-}
 
-FString FDisplayClusterFailoverNodeCtrlBase::FailoverPolicyAsString(const EDisplayClusterFailoverPolicy Policy)
-{
-	switch (Policy)
-	{
-	case EDisplayClusterFailoverPolicy::Disabled:
-		return FString("Disabled");
-
-	case EDisplayClusterFailoverPolicy::Failover_v1_DropSecondaryNodesOnly:
-		return FString("Failover_v1");
-
-	default:
-		return FString("Unknown");
-	}
+	return EDisplayClusterConfigurationFailoverPolicy::Disabled;
 }
diff --git a/Engine/Plugins/Runtime/nDisplay/Source/DisplayCluster/Private/Cluster/Failover/DisplayClusterFailoverNodeCtrlBase.h b/Engine/Plugins/Runtime/nDisplay/Source/DisplayCluster/Private/Cluster/Failover/DisplayClusterFailoverNodeCtrlBase.h
--- a/Engine/Plugins/Runtime/nDisplay/Source/DisplayCluster/Private/Cluster/Failover/DisplayClusterFailoverNodeCtrlBase.h
+++ b/Engine/Plugins/Runtime/nDisplay/Source/DisplayCluster/Private/Cluster/Failover/DisplayClusterFailoverNodeCtrlBase.h
@@ -3,6 +3,7 @@
 #pragma once
 
 #include "Cluster/Failover/IDisplayClusterFailoverNodeController.h"
+#include "DisplayClusterConfigurationTypes_Enums.h"
 
 struct FDisplayClusterSessionInfo;
 
@@ -17,7 +18,7 @@ public:
 	FDisplayClusterFailoverNodeCtrlBase();
 
 public:
-	virtual EDisplayClusterFailoverPolicy GetFailoverPolicy() const override final
+	virtual EDisplayClusterConfigurationFailoverPolicy GetFailoverPolicy() const override final
 	{
 		return FailoverPolicy;
 	}
@@ -30,23 +31,21 @@ public:
 	virtual void HandleNodeFailed(const FString& NodeId, IDisplayClusterServer::ENodeFailType NodeFailType) override final;
 
 protected:
-	// EDisplayClusterFailoverPolicy::Disabled
+	// EDisplayClusterConfigurationFailoverPolicy::Disabled
 	virtual void HandleCommResult_Disabled(EDisplayClusterCommResult CommResult) = 0;
 	virtual void HandleNodeFailed_Disabled(const FString& NodeId, IDisplayClusterServer::ENodeFailType NodeFailType) = 0;
 
-	// EDisplayClusterFailoverPolicy::Failover_v1_DropSecondaryNodesOnly
-	virtual void HandleCommResult_Failover_v1(EDisplayClusterCommResult CommResult) = 0;
-	virtual void HandleNodeFailed_Failover_v1(const FString& NodeId, IDisplayClusterServer::ENodeFailType NodeFailType) = 0;
+	// EDisplayClusterConfigurationFailoverPolicy::DropSecondaryNodesOnly
+	virtual void HandleCommResult_DropSecondaryNodesOnly(EDisplayClusterCommResult CommResult) = 0;
+	virtual void HandleNodeFailed_DropSecondaryNodesOnly(const FString& NodeId, IDisplayClusterServer::ENodeFailType NodeFailType) = 0;
 
 private:
 	// Converts failover policy from CVar integer into enum
-	static EDisplayClusterFailoverPolicy GetFailoverPolicyFromCvarValue(int32 FailoverCvarNumber);
-	// Returns failover policy name as readable string
-	static FString FailoverPolicyAsString(const EDisplayClusterFailoverPolicy Policy);
+	static EDisplayClusterConfigurationFailoverPolicy GetFailoverPolicyFromConfig();
 
 private:
 	// Current failover policy
-	const EDisplayClusterFailoverPolicy FailoverPolicy;
+	const EDisplayClusterConfigurationFailoverPolicy FailoverPolicy;
 
 	mutable FCriticalSection InternalsCS;
 };
diff --git a/Engine/Plugins/Runtime/nDisplay/Source/DisplayCluster/Private/Cluster/Failover/DisplayClusterFailoverNodeCtrlNull.h b/Engine/Plugins/Runtime/nDisplay/Source/DisplayCluster/Private/Cluster/Failover/DisplayClusterFailoverNodeCtrlNull.h
--- a/Engine/Plugins/Runtime/nDisplay/Source/DisplayCluster/Private/Cluster/Failover/DisplayClusterFailoverNodeCtrlNull.h
+++ b/Engine/Plugins/Runtime/nDisplay/Source/DisplayCluster/Private/Cluster/Failover/DisplayClusterFailoverNodeCtrlNull.h
@@ -13,9 +13,9 @@ class FDisplayClusterFailoverNodeCtrlNull
 {
 public:
 	// Returns current failover policy
-	virtual EDisplayClusterFailoverPolicy GetFailoverPolicy() const
+	virtual EDisplayClusterConfigurationFailoverPolicy GetFailoverPolicy() const
 	{
-		return EDisplayClusterFailoverPolicy::Disabled;
+		return EDisplayClusterConfigurationFailoverPolicy::Disabled;
 	}
 
 public:
diff --git a/Engine/Plugins/Runtime/nDisplay/Source/DisplayCluster/Private/Cluster/Failover/DisplayClusterFailoverNodeCtrlPrimary.cpp b/Engine/Plugins/Runtime/nDisplay/Source/DisplayCluster/Private/Cluster/Failover/DisplayClusterFailoverNodeCtrlPrimary.cpp
--- a/Engine/Plugins/Runtime/nDisplay/Source/DisplayCluster/Private/Cluster/Failover/DisplayClusterFailoverNodeCtrlPrimary.cpp
+++ b/Engine/Plugins/Runtime/nDisplay/Source/DisplayCluster/Private/Cluster/Failover/DisplayClusterFailoverNodeCtrlPrimary.cpp
@@ -15,7 +15,7 @@
 
 
 ////////////////////////////////////////////////////////////////////////////////////////////////////
-// EDisplayClusterFailoverPolicy::Disabled
+// EDisplayClusterConfigurationFailoverPolicy::Disabled
 ////////////////////////////////////////////////////////////////////////////////////////////////////
 void FDisplayClusterFailoverNodeCtrlPrimary::HandleCommResult_Disabled(EDisplayClusterCommResult CommResult)
 {
@@ -47,9 +47,9 @@ void FDisplayClusterFailoverNodeCtrlPrimary::HandleNodeFailed_Disabled(const FSt
 
 
 ////////////////////////////////////////////////////////////////////////////////////////////////////
-// EDisplayClusterFailoverPolicy::Failover_v1_DropSecondaryNodesOnly
+// EDisplayClusterConfigurationFailoverPolicy::DropSecondaryNodesOnly
 ////////////////////////////////////////////////////////////////////////////////////////////////////
-void FDisplayClusterFailoverNodeCtrlPrimary::HandleCommResult_Failover_v1(EDisplayClusterCommResult CommResult)
+void FDisplayClusterFailoverNodeCtrlPrimary::HandleCommResult_DropSecondaryNodesOnly(EDisplayClusterCommResult CommResult)
 {
 	// Any comm error on the primary node results in node termination
 	if (CommResult != EDisplayClusterCommResult::Ok)
@@ -58,7 +58,7 @@ void FDisplayClusterFailoverNodeCtrlPrimary::HandleCommResult_Failover_v1(EDispl
 	}
 }
 
-void FDisplayClusterFailoverNodeCtrlPrimary::HandleNodeFailed_Failover_v1(const FString& NodeId, IDisplayClusterServer::ENodeFailType NodeFailType)
+void FDisplayClusterFailoverNodeCtrlPrimary::HandleNodeFailed_DropSecondaryNodesOnly(const FString& NodeId, IDisplayClusterServer::ENodeFailType NodeFailType)
 {
 	UE_LOG(LogDisplayClusterCluster, Log, TEXT("Failover: node [%s] failed, fail_type=%u, disconnecting the node and continue working..."), *NodeId, static_cast<uint32>(NodeFailType));
 
diff --git a/Engine/Plugins/Runtime/nDisplay/Source/DisplayCluster/Private/Cluster/Failover/DisplayClusterFailoverNodeCtrlPrimary.h b/Engine/Plugins/Runtime/nDisplay/Source/DisplayCluster/Private/Cluster/Failover/DisplayClusterFailoverNodeCtrlPrimary.h
--- a/Engine/Plugins/Runtime/nDisplay/Source/DisplayCluster/Private/Cluster/Failover/DisplayClusterFailoverNodeCtrlPrimary.h
+++ b/Engine/Plugins/Runtime/nDisplay/Source/DisplayCluster/Private/Cluster/Failover/DisplayClusterFailoverNodeCtrlPrimary.h
@@ -17,11 +17,11 @@ public:
 	{ }
 
 protected:
-	// EDisplayClusterFailoverPolicy::Disabled
+	// EDisplayClusterConfigurationFailoverPolicy::Disabled
 	virtual void HandleCommResult_Disabled(EDisplayClusterCommResult CommResult) override;
 	virtual void HandleNodeFailed_Disabled(const FString& NodeId, IDisplayClusterServer::ENodeFailType NodeFailType) override;
 
-	// EDisplayClusterFailoverPolicy::Failover_v1_DropSecondaryNodesOnly
-	virtual void HandleCommResult_Failover_v1(EDisplayClusterCommResult CommResult) override;
-	virtual void HandleNodeFailed_Failover_v1(const FString& NodeId, IDisplayClusterServer::ENodeFailType NodeFailType) override;
+	// EDisplayClusterConfigurationFailoverPolicy::DropSecondaryNodesOnly
+	virtual void HandleCommResult_DropSecondaryNodesOnly(EDisplayClusterCommResult CommResult) override;
+	virtual void HandleNodeFailed_DropSecondaryNodesOnly(const FString& NodeId, IDisplayClusterServer::ENodeFailType NodeFailType) override;
 };
diff --git a/Engine/Plugins/Runtime/nDisplay/Source/DisplayCluster/Private/Cluster/Failover/DisplayClusterFailoverNodeCtrlSecondary.cpp b/Engine/Plugins/Runtime/nDisplay/Source/DisplayCluster/Private/Cluster/Failover/DisplayClusterFailoverNodeCtrlSecondary.cpp
--- a/Engine/Plugins/Runtime/nDisplay/Source/DisplayCluster/Private/Cluster/Failover/DisplayClusterFailoverNodeCtrlSecondary.cpp
+++ b/Engine/Plugins/Runtime/nDisplay/Source/DisplayCluster/Private/Cluster/Failover/DisplayClusterFailoverNodeCtrlSecondary.cpp
@@ -10,7 +10,7 @@
 
 
 ////////////////////////////////////////////////////////////////////////////////////////////////////
-// EDisplayClusterFailoverPolicy::Disabled
+// EDisplayClusterConfigurationFailoverPolicy::Disabled
 ////////////////////////////////////////////////////////////////////////////////////////////////////
 void FDisplayClusterFailoverNodeCtrlSecondary::HandleCommResult_Disabled(EDisplayClusterCommResult CommResult)
 {
@@ -28,9 +28,9 @@ void FDisplayClusterFailoverNodeCtrlSecondary::HandleNodeFailed_Disabled(const F
 
 
 ////////////////////////////////////////////////////////////////////////////////////////////////////
-// EDisplayClusterFailoverPolicy::Failover_v1_DropSecondaryNodesOnly
+// EDisplayClusterConfigurationFailoverPolicy::DropSecondaryNodesOnly
 ////////////////////////////////////////////////////////////////////////////////////////////////////
-void FDisplayClusterFailoverNodeCtrlSecondary::HandleCommResult_Failover_v1(EDisplayClusterCommResult CommResult)
+void FDisplayClusterFailoverNodeCtrlSecondary::HandleCommResult_DropSecondaryNodesOnly(EDisplayClusterCommResult CommResult)
 {
 	// Any comm error on the primary node results in node termination
 	if (CommResult != EDisplayClusterCommResult::Ok)
@@ -39,7 +39,7 @@ void FDisplayClusterFailoverNodeCtrlSecondary::HandleCommResult_Failover_v1(EDis
 	}
 }
 
-void FDisplayClusterFailoverNodeCtrlSecondary::HandleNodeFailed_Failover_v1(const FString& NodeId, IDisplayClusterServer::ENodeFailType NodeFailType)
+void FDisplayClusterFailoverNodeCtrlSecondary::HandleNodeFailed_DropSecondaryNodesOnly(const FString& NodeId, IDisplayClusterServer::ENodeFailType NodeFailType)
 {
 	// Currently, all comm sessions run on the primary node only. We don't have to do anything over here.
 }
diff --git a/Engine/Plugins/Runtime/nDisplay/Source/DisplayCluster/Private/Cluster/Failover/DisplayClusterFailoverNodeCtrlSecondary.h b/Engine/Plugins/Runtime/nDisplay/Source/DisplayCluster/Private/Cluster/Failover/DisplayClusterFailoverNodeCtrlSecondary.h
--- a/Engine/Plugins/Runtime/nDisplay/Source/DisplayCluster/Private/Cluster/Failover/DisplayClusterFailoverNodeCtrlSecondary.h
+++ b/Engine/Plugins/Runtime/nDisplay/Source/DisplayCluster/Private/Cluster/Failover/DisplayClusterFailoverNodeCtrlSecondary.h
@@ -17,11 +17,11 @@ public:
 	{ }
 
 protected:
-	// EDisplayClusterFailoverPolicy::Disabled
+	// EDisplayClusterConfigurationFailoverPolicy::Disabled
 	virtual void HandleCommResult_Disabled(EDisplayClusterCommResult CommResult) override;
 	virtual void HandleNodeFailed_Disabled(const FString& NodeId, IDisplayClusterServer::ENodeFailType NodeFailType) override;
 
-	// EDisplayClusterFailoverPolicy::Failover_v1_DropSecondaryNodesOnly
-	virtual void HandleCommResult_Failover_v1(EDisplayClusterCommResult CommResult) override;
-	virtual void HandleNodeFailed_Failover_v1(const FString& NodeId, IDisplayClusterServer::ENodeFailType NodeFailType) override;
+	// EDisplayClusterConfigurationFailoverPolicy::DropSecondaryNodesOnly
+	virtual void HandleCommResult_DropSecondaryNodesOnly(EDisplayClusterCommResult CommResult) override;
+	virtual void HandleNodeFailed_DropSecondaryNodesOnly(const FString& NodeId, IDisplayClusterServer::ENodeFailType NodeFailType) override;
 };
diff --git a/Engine/Plugins/Runtime/nDisplay/Source/DisplayCluster/Private/Cluster/Failover/IDisplayClusterFailoverNodeController.h b/Engine/Plugins/Runtime/nDisplay/Source/DisplayCluster/Private/Cluster/Failover/IDisplayClusterFailoverNodeController.h
--- a/Engine/Plugins/Runtime/nDisplay/Source/DisplayCluster/Private/Cluster/Failover/IDisplayClusterFailoverNodeController.h
+++ b/Engine/Plugins/Runtime/nDisplay/Source/DisplayCluster/Private/Cluster/Failover/IDisplayClusterFailoverNodeController.h
@@ -5,19 +5,7 @@
 #include "CoreMinimal.h"
 #include "Network/DisplayClusterNetworkTypes.h"
 #include "Network/IDisplayClusterServer.h"
-#include "DisplayClusterEnums.h"
-
-
-enum class EDisplayClusterFailoverPolicy : uint8
-{
-	// No failover operations performed. The whole cluster gets terminated in case of any error
-	Disabled,
-
-	// First implementation of failover feature. This policy allows to drop any secondary node
-	// out of cluster in case it's failed, and let the others continue working. However, if
-	// primary node fails, the whole cluster will be terminated.
-	Failover_v1_DropSecondaryNodesOnly
-};
+#include "DisplayClusterConfigurationTypes_Enums.h"
 
 
 /**
@@ -30,7 +18,7 @@ public:
 
 public:
 	// Returns current failover policy
-	virtual EDisplayClusterFailoverPolicy GetFailoverPolicy() const = 0;
+	virtual EDisplayClusterConfigurationFailoverPolicy GetFailoverPolicy() const = 0;
 
 public:
 	// Handles in-cluster communication results
diff --git a/Engine/Plugins/Runtime/nDisplay/Source/DisplayCluster/Private/Game/EngineClasses/Scene/DisplayClusterRootActor.cpp b/Engine/Plugins/Runtime/nDisplay/Source/DisplayCluster/Private/Game/EngineClasses/Scene/DisplayClusterRootActor.cpp
--- a/Engine/Plugins/Runtime/nDisplay/Source/DisplayCluster/Private/Game/EngineClasses/Scene/DisplayClusterRootActor.cpp
+++ b/Engine/Plugins/Runtime/nDisplay/Source/DisplayCluster/Private/Game/EngineClasses/Scene/DisplayClusterRootActor.cpp
@@ -175,8 +175,9 @@ void ADisplayClusterRootActor::OverrideFromConfig(UDisplayClusterConfigurationDa
 	if (CurrentConfigData->Cluster)
 	{
 		CurrentConfigData->Cluster->PrimaryNode = ConfigData->Cluster->PrimaryNode;
-		CurrentConfigData->Cluster->Sync = ConfigData->Cluster->Sync;
-		CurrentConfigData->Cluster->Network = ConfigData->Cluster->Network;
+		CurrentConfigData->Cluster->Sync        = ConfigData->Cluster->Sync;
+		CurrentConfigData->Cluster->Network     = ConfigData->Cluster->Network;
+		CurrentConfigData->Cluster->Failover    = ConfigData->Cluster->Failover;
 
 		// Remove nodes in current config that are not in the new config
 
diff --git a/Engine/Plugins/Runtime/nDisplay/Source/DisplayClusterConfiguration/Private/Formats/JSON500/DisplayClusterConfigurationJsonParser_500.cpp b/Engine/Plugins/Runtime/nDisplay/Source/DisplayClusterConfiguration/Private/Formats/JSON500/DisplayClusterConfigurationJsonParser_500.cpp
--- a/Engine/Plugins/Runtime/nDisplay/Source/DisplayClusterConfiguration/Private/Formats/JSON500/DisplayClusterConfigurationJsonParser_500.cpp
+++ b/Engine/Plugins/Runtime/nDisplay/Source/DisplayClusterConfiguration/Private/Formats/JSON500/DisplayClusterConfigurationJsonParser_500.cpp
@@ -205,6 +205,11 @@ namespace JSON500
 					CfgJson.Cluster.Network, FString(DisplayClusterConfigurationStrings::config::cluster::network::NetRenderSyncBarrierTimeout), (uint32)5000);
 			}
 
+			// Failover
+			{
+				Config->Cluster->Failover.FailoverPolicy = CfgJson.Cluster.Failover.FailoverPolicy;
+			}
+
 			// Cluster nodes
 			for (const TPair<FString, FDisplayClusterConfigurationJsonClusterNode_500>& CfgNode : CfgJson.Cluster.Nodes)
 			{
@@ -391,6 +396,11 @@ namespace JSON500
 					DisplayClusterTypesConverter::template ToString(Config->Cluster->Network.RenderSyncBarrierTimeout));
 			}
 
+			// Failover
+			{
+				Json.Cluster.Failover.FailoverPolicy = Config->Cluster->Failover.FailoverPolicy;
+			}
+
 			// Cluster nodes
 			for (const TPair<FString, UDisplayClusterConfigurationClusterNode*>& CfgNode : Config->Cluster->Nodes)
 			{
diff --git a/Engine/Plugins/Runtime/nDisplay/Source/DisplayClusterConfiguration/Private/Formats/JSON500/DisplayClusterConfigurationJsonTypes_500.h b/Engine/Plugins/Runtime/nDisplay/Source/DisplayClusterConfiguration/Private/Formats/JSON500/DisplayClusterConfigurationJsonTypes_500.h
--- a/Engine/Plugins/Runtime/nDisplay/Source/DisplayClusterConfiguration/Private/Formats/JSON500/DisplayClusterConfigurationJsonTypes_500.h
+++ b/Engine/Plugins/Runtime/nDisplay/Source/DisplayClusterConfiguration/Private/Formats/JSON500/DisplayClusterConfigurationJsonTypes_500.h
@@ -4,6 +4,7 @@
 
 #include "CoreMinimal.h"
 #include "DisplayClusterConfigurationStrings.h"
+#include "DisplayClusterConfigurationTypes.h"
 #include "DisplayClusterConfigurationJsonTypes_500.generated.h"
 
 
@@ -342,6 +343,7 @@ public:
 	TMap<FString, uint16> Ports;
 };
 
+
 USTRUCT()
 struct FDisplayClusterConfigurationJsonClusterSyncPolicy_500
 	: public FDisplayClusterConfigurationJsonPolymorphicEntity_500
@@ -349,6 +351,7 @@ struct FDisplayClusterConfigurationJsonClusterSyncPolicy_500
 	GENERATED_BODY()
 };
 
+
 USTRUCT()
 struct FDisplayClusterConfigurationJsonClusterSync_500
 {
@@ -362,6 +365,23 @@ public:
 	FDisplayClusterConfigurationJsonClusterSyncPolicy_500 InputSyncPolicy;
 };
 
+
+USTRUCT()
+struct DISPLAYCLUSTERCONFIGURATION_API FDisplayClusterConfigurationFailoverSettings_500
+{
+	GENERATED_BODY()
+
+public:
+	FDisplayClusterConfigurationFailoverSettings_500()
+		: FailoverPolicy(EDisplayClusterConfigurationFailoverPolicy::Disabled)
+	{ }
+
+public:
+	UPROPERTY()
+	EDisplayClusterConfigurationFailoverPolicy FailoverPolicy;
+};
+
+
 USTRUCT()
 struct FDisplayClusterConfigurationJsonPostprocess_500
 	: public FDisplayClusterConfigurationJsonPolymorphicEntity_500
@@ -369,6 +389,7 @@ struct FDisplayClusterConfigurationJsonPostprocess_500
 	GENERATED_BODY()
 };
 
+
 USTRUCT()
 struct FDisplayClusterConfigurationJsonProjectionPolicy_500
 	: public FDisplayClusterConfigurationJsonPolymorphicEntity_500
@@ -376,6 +397,7 @@ struct FDisplayClusterConfigurationJsonProjectionPolicy_500
 	GENERATED_BODY()
 };
 
+
 USTRUCT()
 struct FDisplayClusterConfigurationJsonOverscan_500
 {
@@ -404,6 +426,7 @@ public:
 	bool Oversize = 0;
 };
 
+
 USTRUCT()
 struct FDisplayClusterConfigurationJsonViewport_500
 {
@@ -435,6 +458,7 @@ public:
 	FDisplayClusterConfigurationJsonProjectionPolicy_500 ProjectionPolicy;
 };
 
+
 USTRUCT()
 struct FDisplayClusterConfigurationFramePostProcess_OutputRemap_500
 {
@@ -454,6 +478,7 @@ public:
 	FString ExternalFile;
 };
 
+
 USTRUCT()
 struct FDisplayClusterConfigurationJsonClusterNode_500
 {
@@ -482,6 +507,7 @@ public:
 	FDisplayClusterConfigurationFramePostProcess_OutputRemap_500 OutputRemap;
 };
 
+
 USTRUCT()
 struct FDisplayClusterConfigurationJsonCluster_500
 {
@@ -497,10 +523,14 @@ public:
 	UPROPERTY()
 	TMap<FString, FString> Network;
 
+	UPROPERTY()
+	FDisplayClusterConfigurationFailoverSettings_500 Failover;
+
 	UPROPERTY()
 	TMap<FString, FDisplayClusterConfigurationJsonClusterNode_500> Nodes;
 };
 
+
 USTRUCT()
 struct FDisplayClusterConfigurationJsonDiagnostics_500
 {
@@ -550,6 +580,7 @@ public:
 	FDisplayClusterConfigurationJsonDiagnostics_500 Diagnostics;
 };
 
+
 // The main nDisplay configuration structure. It's supposed to extract nDisplay related data from a collecting JSON file.
 USTRUCT()
 struct FDisplayClusterConfigurationJsonContainer_500
diff --git a/Engine/Plugins/Runtime/nDisplay/Source/DisplayClusterConfiguration/Public/DisplayClusterConfigurationTypes.h b/Engine/Plugins/Runtime/nDisplay/Source/DisplayClusterConfiguration/Public/DisplayClusterConfigurationTypes.h
--- a/Engine/Plugins/Runtime/nDisplay/Source/DisplayClusterConfiguration/Public/DisplayClusterConfigurationTypes.h
+++ b/Engine/Plugins/Runtime/nDisplay/Source/DisplayClusterConfiguration/Public/DisplayClusterConfigurationTypes.h
@@ -262,6 +262,17 @@ public:
 	int32 RenderSyncBarrierTimeout;
 };
 
+USTRUCT(Blueprintable)
+struct DISPLAYCLUSTERCONFIGURATION_API FDisplayClusterConfigurationFailoverSettings
+{
+	GENERATED_BODY()
+
+public:
+	/** Failover policy */
+	UPROPERTY(EditAnywhere, BlueprintReadWrite, Category = NDisplay)
+	EDisplayClusterConfigurationFailoverPolicy FailoverPolicy = EDisplayClusterConfigurationFailoverPolicy::Disabled;
+};
+
 USTRUCT()
 struct DISPLAYCLUSTERCONFIGURATION_API FDisplayClusterConfigurationExternalImage
 {
@@ -423,6 +434,9 @@ public:
 	UPROPERTY(EditDefaultsOnly, BlueprintReadOnly, Category = "Configuration")
 	FDisplayClusterConfigurationNetworkSettings Network;
 
+	UPROPERTY(EditDefaultsOnly, BlueprintReadOnly, Category = "Configuration")
+	FDisplayClusterConfigurationFailoverSettings Failover;
+
 	UPROPERTY(VisibleInstanceOnly, BlueprintReadOnly, EditFixedSize, Instanced, Category = NDisplay, meta = (DisplayThumbnail = false, HideProperty))
 	TMap<FString, UDisplayClusterConfigurationClusterNode*> Nodes;
 
diff --git a/Engine/Plugins/Runtime/nDisplay/Source/DisplayClusterConfiguration/Public/DisplayClusterConfigurationTypes_Enums.h b/Engine/Plugins/Runtime/nDisplay/Source/DisplayClusterConfiguration/Public/DisplayClusterConfigurationTypes_Enums.h
--- a/Engine/Plugins/Runtime/nDisplay/Source/DisplayClusterConfiguration/Public/DisplayClusterConfigurationTypes_Enums.h
+++ b/Engine/Plugins/Runtime/nDisplay/Source/DisplayClusterConfiguration/Public/DisplayClusterConfigurationTypes_Enums.h
@@ -18,23 +18,14 @@ enum class EDisplayClusterConfigurationDataSource : uint8
 };
 
 UENUM()
-enum class EDisplayClusterConfigurationKeyboardReflectionType : uint8
+enum class EDisplayClusterConfigurationFailoverPolicy : uint8
 {
-	None     UMETA(DisplayName = "No reflection"),
-	nDisplay UMETA(DisplayName = "nDisplay buttons only"),
-	Core     UMETA(DisplayName = "UE core keyboard events only"),
-	All      UMETA(DisplayName = "Both nDisplay and UE core events")
-};
+	// No failover operations performed. The whole cluster gets terminated in case of any error
+	Disabled                          UMETA(DisplayName = "Disabled"),
 
-UENUM()
-enum class EDisplayClusterConfigurationTrackerMapping
-{
-	X    UMETA(DisplayName = "Positive X"),
-	NX   UMETA(DisplayName = "Negative X"),
-	Y    UMETA(DisplayName = "Positive Y"),
-	NY   UMETA(DisplayName = "Negative Y"),
-	Z    UMETA(DisplayName = "Positive Z"),
-	NZ   UMETA(DisplayName = "Negative Z")
+	// This policy allows to drop any secondary node out of cluster in case it's failed,
+	// and let the others continue working. However, the whole cluster will be terminated if primary node fails.
+	DropSecondaryNodesOnly            UMETA(DisplayName = "Drop S-node on fail"),
 };
 
 UENUM()
