nDisplay: fixed issue with DCRA settings overriding

diff --git a/Engine/Plugins/Runtime/nDisplay/Source/DisplayCluster/Private/Game/EngineClasses/Scene/DisplayClusterRootActor.cpp b/Engine/Plugins/Runtime/nDisplay/Source/DisplayCluster/Private/Game/EngineClasses/Scene/DisplayClusterRootActor.cpp
--- a/Engine/Plugins/Runtime/nDisplay/Source/DisplayCluster/Private/Game/EngineClasses/Scene/DisplayClusterRootActor.cpp
+++ b/Engine/Plugins/Runtime/nDisplay/Source/DisplayCluster/Private/Game/EngineClasses/Scene/DisplayClusterRootActor.cpp
@@ -221,45 +221,46 @@ void ADisplayClusterRootActor::OverrideFromConfig(UDisplayClusterConfigurationDa
 			CurrentNode->Postprocess = NewNode->Postprocess;
 			CurrentNode->OutputRemap = NewNode->OutputRemap;
 
-			// Remove viewports in current node that are not in the new node
-
-			for (auto CurrentViewportIt = CurrentNode->Viewports.CreateIterator(); CurrentViewportIt; ++CurrentViewportIt)
+			if (ConfigData->bOverrideViewportsFromExternalConfig)
 			{
-				if (!NewNode->Viewports.Find(CurrentViewportIt->Key))
+				// Remove viewports in current node that are not in the new node
+				for (auto CurrentViewportIt = CurrentNode->Viewports.CreateIterator(); CurrentViewportIt; ++CurrentViewportIt)
 				{
-					CurrentViewportIt.RemoveCurrent();
+					if (!NewNode->Viewports.Find(CurrentViewportIt->Key))
+					{
+						CurrentViewportIt.RemoveCurrent();
+					}
 				}
-			}
-
-			// Go over viewport and override the current ones (or add new ones that didn't exist)
-
-			for (auto NewViewportIt = NewNode->Viewports.CreateConstIterator(); NewViewportIt; ++NewViewportIt)
-			{
-				UDisplayClusterConfigurationViewport** CurrentViewportPtr = CurrentNode->Viewports.Find(NewViewportIt->Key);
 
-				// Add the viewport if it doesn't exist
-				if (!CurrentViewportPtr)
+				// Go over viewport and override the current ones (or add new ones that didn't exist)
+				for (auto NewViewportIt = NewNode->Viewports.CreateConstIterator(); NewViewportIt; ++NewViewportIt)
 				{
-					CurrentNode->Viewports.Add(
-						NewViewportIt->Key,
-						DuplicateObject(NewViewportIt->Value, CurrentNode, NewViewportIt->Value->GetFName())
-					);
-					continue;
-				}
+					UDisplayClusterConfigurationViewport** CurrentViewportPtr = CurrentNode->Viewports.Find(NewViewportIt->Key);
+
+					// Add the viewport if it doesn't exist
+					if (!CurrentViewportPtr)
+					{
+						CurrentNode->Viewports.Add(
+							NewViewportIt->Key,
+							DuplicateObject(NewViewportIt->Value, CurrentNode, NewViewportIt->Value->GetFName())
+						);
+						continue;
+					}
 
-				UDisplayClusterConfigurationViewport* CurrentViewport = *CurrentViewportPtr;
-				check(CurrentViewport);
+					UDisplayClusterConfigurationViewport* CurrentViewport = *CurrentViewportPtr;
+					check(CurrentViewport);
 
-				const UDisplayClusterConfigurationViewport* NewViewport = NewViewportIt->Value;
-				check(NewViewport);
+					const UDisplayClusterConfigurationViewport* NewViewport = NewViewportIt->Value;
+					check(NewViewport);
 
-				// Override viewport settings
+					// Override viewport settings
 
-				CurrentViewport->Camera = NewViewport->Camera;
-				CurrentViewport->RenderSettings.BufferRatio = NewViewport->RenderSettings.BufferRatio;
-				CurrentViewport->GPUIndex = NewViewport->GPUIndex;
-				CurrentViewport->Region = NewViewport->Region;
-				CurrentViewport->ProjectionPolicy = NewViewport->ProjectionPolicy;
+					CurrentViewport->Camera = NewViewport->Camera;
+					CurrentViewport->RenderSettings.BufferRatio = NewViewport->RenderSettings.BufferRatio;
+					CurrentViewport->GPUIndex = NewViewport->GPUIndex;
+					CurrentViewport->Region = NewViewport->Region;
+					CurrentViewport->ProjectionPolicy = NewViewport->ProjectionPolicy;
+				}
 			}
 		}
 	}
diff --git a/Engine/Plugins/Runtime/nDisplay/Source/DisplayClusterConfiguration/Private/Formats/JSON500/DisplayClusterConfigurationJsonParser_500.cpp b/Engine/Plugins/Runtime/nDisplay/Source/DisplayClusterConfiguration/Private/Formats/JSON500/DisplayClusterConfigurationJsonParser_500.cpp
--- a/Engine/Plugins/Runtime/nDisplay/Source/DisplayClusterConfiguration/Private/Formats/JSON500/DisplayClusterConfigurationJsonParser_500.cpp
+++ b/Engine/Plugins/Runtime/nDisplay/Source/DisplayClusterConfiguration/Private/Formats/JSON500/DisplayClusterConfigurationJsonParser_500.cpp
@@ -105,6 +105,7 @@ namespace JSON500
 		// Misc
 		Config->bFollowLocalPlayerCamera = CfgJson.Misc.bFollowLocalPlayerCamera;
 		Config->bExitOnEsc = CfgJson.Misc.bExitOnEsc;
+		Config->bOverrideViewportsFromExternalConfig = CfgJson.Misc.bOverrideViewportsFromExternalConfig;
 
 		// Scene
 		{
@@ -305,6 +306,7 @@ namespace JSON500
 		// Misc
 		Json.Misc.bFollowLocalPlayerCamera = Config->bFollowLocalPlayerCamera;
 		Json.Misc.bExitOnEsc = Config->bExitOnEsc;
+		Json.Misc.bOverrideViewportsFromExternalConfig = Config->bOverrideViewportsFromExternalConfig;
 
 		// Scene
 		{
diff --git a/Engine/Plugins/Runtime/nDisplay/Source/DisplayClusterConfiguration/Private/Formats/JSON500/DisplayClusterConfigurationJsonTypes_500.h b/Engine/Plugins/Runtime/nDisplay/Source/DisplayClusterConfiguration/Private/Formats/JSON500/DisplayClusterConfigurationJsonTypes_500.h
--- a/Engine/Plugins/Runtime/nDisplay/Source/DisplayClusterConfiguration/Private/Formats/JSON500/DisplayClusterConfigurationJsonTypes_500.h
+++ b/Engine/Plugins/Runtime/nDisplay/Source/DisplayClusterConfiguration/Private/Formats/JSON500/DisplayClusterConfigurationJsonTypes_500.h
@@ -214,6 +214,7 @@ public:
 	FDisplayClusterConfigurationJsonMisc_500()
 		: bFollowLocalPlayerCamera(false)
 		, bExitOnEsc(true)
+		, bOverrideViewportsFromExternalConfig(false)
 	{ }
 
 public:
@@ -222,6 +223,9 @@ public:
 
 	UPROPERTY()
 	bool bExitOnEsc;
+
+	UPROPERTY()
+	bool bOverrideViewportsFromExternalConfig;
 };
 
 
diff --git a/Engine/Plugins/Runtime/nDisplay/Source/DisplayClusterConfiguration/Public/DisplayClusterConfigurationTypes.h b/Engine/Plugins/Runtime/nDisplay/Source/DisplayClusterConfiguration/Public/DisplayClusterConfigurationTypes.h
--- a/Engine/Plugins/Runtime/nDisplay/Source/DisplayClusterConfiguration/Public/DisplayClusterConfigurationTypes.h
+++ b/Engine/Plugins/Runtime/nDisplay/Source/DisplayClusterConfiguration/Public/DisplayClusterConfigurationTypes.h
@@ -555,6 +555,9 @@ public:
 	UPROPERTY(EditDefaultsOnly, BlueprintReadWrite, Category = Configuration, meta = (DisplayName = "Exit When ESC Pressed", DisplayAfter = "bFollowLocalPlayerCamera"))
 	bool bExitOnEsc = true;
 
+	UPROPERTY(EditDefaultsOnly, BlueprintReadWrite, Category = Configuration, meta = (DisplayName = "Override Viewports From .ndisplay", DisplayAfter = "bExitOnEsc"))
+	bool bOverrideViewportsFromExternalConfig = false;
+
 	/** Create empty config data. */
 	static UDisplayClusterConfigurationData* CreateNewConfigData(UObject* Owner = nullptr, EObjectFlags ObjectFlags = RF_NoFlags);
 #if WITH_EDITORONLY_DATA
