nDisplay: fixed InnerFrustumResolution values when Inner Frustum disabled

diff --git a/Engine/Plugins/Runtime/nDisplay/Source/DisplayCluster/Private/Game/EngineClasses/Scene/DisplayClusterICVFXCameraComponent.cpp b/Engine/Plugins/Runtime/nDisplay/Source/DisplayCluster/Private/Game/EngineClasses/Scene/DisplayClusterICVFXCameraComponent.cpp
--- a/Engine/Plugins/Runtime/nDisplay/Source/DisplayCluster/Private/Game/EngineClasses/Scene/DisplayClusterICVFXCameraComponent.cpp
+++ b/Engine/Plugins/Runtime/nDisplay/Source/DisplayCluster/Private/Game/EngineClasses/Scene/DisplayClusterICVFXCameraComponent.cpp
@@ -107,7 +107,7 @@ void UDisplayClusterICVFXCameraComponent::UpdateOverscanEstimatedFrameSize()
 		EstimatedOverscanResolutionHeight += AdjustmentSize.Y;
 	}
 
-	if (CameraSettings.CustomFrustum.bAdaptResolution)
+	if (CameraSettings.CustomFrustum.bEnable && CameraSettings.CustomFrustum.bAdaptResolution)
 	{
 		InnerFrustumResolutionWidth = EstimatedOverscanResolutionWidth;
 		InnerFrustumResolutionHeight = EstimatedOverscanResolutionHeight;
