nDisplay: Improve editor performance when adjusting values by sliders. PostEditChangeProperty will not run construction scripts when an interactive change is detected.

diff --git a/Engine/Plugins/Runtime/nDisplay/Source/DisplayCluster/Private/Game/EngineClasses/Scene/DisplayClusterRootActor.cpp b/Engine/Plugins/Runtime/nDisplay/Source/DisplayCluster/Private/Game/EngineClasses/Scene/DisplayClusterRootActor.cpp
--- a/Engine/Plugins/Runtime/nDisplay/Source/DisplayCluster/Private/Game/EngineClasses/Scene/DisplayClusterRootActor.cpp
+++ b/Engine/Plugins/Runtime/nDisplay/Source/DisplayCluster/Private/Game/EngineClasses/Scene/DisplayClusterRootActor.cpp
@@ -916,7 +916,9 @@ static void PropagateDataFromDefaultConfig(UDisplayClusterConfigurationData* InD
 
 void ADisplayClusterRootActor::RerunConstructionScripts()
 {
-	IDisplayClusterConfiguration& Config = IDisplayClusterConfiguration::Get();
+	DECLARE_SCOPE_CYCLE_COUNTER(TEXT("ADisplayClusterRootActor::RerunConstructionScripts"), STAT_RerunConstructionScripts, STATGROUP_NDisplay);
+	
+	const IDisplayClusterConfiguration& Config = IDisplayClusterConfiguration::Get();
 	if (!Config.IsTransactingSnapshot())
 	{
 		Super::RerunConstructionScripts();
diff --git a/Engine/Plugins/Runtime/nDisplay/Source/DisplayCluster/Private/Game/EngineClasses/Scene/DisplayClusterRootActorEditor.cpp b/Engine/Plugins/Runtime/nDisplay/Source/DisplayCluster/Private/Game/EngineClasses/Scene/DisplayClusterRootActorEditor.cpp
--- a/Engine/Plugins/Runtime/nDisplay/Source/DisplayCluster/Private/Game/EngineClasses/Scene/DisplayClusterRootActorEditor.cpp
+++ b/Engine/Plugins/Runtime/nDisplay/Source/DisplayCluster/Private/Game/EngineClasses/Scene/DisplayClusterRootActorEditor.cpp
@@ -48,6 +48,7 @@
 
 #include "Async/Async.h"
 #include "LevelEditor.h"
+#include "EditorSupportDelegates.h"
 
 void ADisplayClusterRootActor::Constructor_Editor()
 {
@@ -98,6 +99,8 @@ void ADisplayClusterRootActor::BeginDestroy_Editor()
 
 void ADisplayClusterRootActor::RerunConstructionScripts_Editor()
 {
+	DECLARE_SCOPE_CYCLE_COUNTER(TEXT("ADisplayClusterRootActor::RerunConstructionScripts_Editor"), STAT_RerunConstructionScripts_Editor, STATGROUP_NDisplay);
+	
 	/* We need to reinitialize since our components are being regenerated here. */
 	InitializeRootActor();
 
@@ -450,13 +453,54 @@ void ADisplayClusterRootActor::RenderPreviewFrustum(const FMatrix ProjectionMatr
 	LineBatcher->DrawLine(PreviewFrustumVerts[6], PreviewFrustumVerts[4], Color, SDPG_World, Thickness, 0.f); // left top to right top
 }
 
+static FName Name_RelativeLocation = USceneComponent::GetRelativeLocationPropertyName();
+static FName Name_RelativeRotation = USceneComponent::GetRelativeRotationPropertyName();
+static FName Name_RelativeScale3D = USceneComponent::GetRelativeScale3DPropertyName();
+
 void ADisplayClusterRootActor::PostEditChangeProperty(FPropertyChangedEvent& PropertyChangedEvent)
 {
 	const FName PropertyName = (PropertyChangedEvent.Property != nullptr) ? PropertyChangedEvent.Property->GetFName() : NAME_None;
+	
+	// The AActor method, simplified and modified to skip construction scripts.
+	// Component registration still needs to occur or the actor will look like it disappeared.
+	auto SuperCallWithoutConstructionScripts = [&]
+	{
+		if (IsPropertyChangedAffectingDataLayers(PropertyChangedEvent))
+		{
+			FixupDataLayers(/*bRevertChangesOnLockedDataLayer*/true);
+		}
+
+		const bool bTransformationChanged = (PropertyName == Name_RelativeLocation || PropertyName == Name_RelativeRotation || PropertyName == Name_RelativeScale3D);
+		if ((GEditor && GEditor->bIsSimulatingInEditor && GetWorld() != nullptr) || ReregisterComponentsWhenModified())
+		{
+			// If a transaction is occurring we rely on the true parent method instead.
+			check (!CurrentTransactionAnnotation.IsValid());
+			
+			UnregisterAllComponents();
+			ReregisterAllComponents();
+		}
+
+		// Let other systems know that an actor was moved
+		if (bTransformationChanged)
+		{
+			GEngine->BroadcastOnActorMoved( this );
+		}
+
+		FEditorSupportDelegates::UpdateUI.Broadcast();
+		UObject::PostEditChangeProperty(PropertyChangedEvent);	
+	};
 
 	bool bReinitializeActor = true;
 
-	Super::PostEditChangeProperty(PropertyChangedEvent);
+	if (PropertyChangedEvent.ChangeType == EPropertyChangeType::Interactive && !CurrentTransactionAnnotation.IsValid())
+	{
+		// Avoid calling construction scripts when the change occurs while the user is dragging a slider.
+		SuperCallWithoutConstructionScripts();
+	}
+	else
+	{
+		Super::PostEditChangeProperty(PropertyChangedEvent);
+	}
 	
 	// Config file has been changed, we should rebuild the nDisplay hierarchy
 	// Cluster node ID has been changed
