nDisplay: Added cvar nDisplay.render.show.visualizationcomponents to control whether to render and show visualization components or not.

Currently there is a disparity with regular non editor renders (e.g. including scene captures) and nDisplay, where the latter was excluding visualization components while the others don't. It will now render them, but the cvar allows the user to go back to the previous behavior.

diff --git a/Engine/Plugins/Runtime/nDisplay/Source/DisplayCluster/Private/Game/EngineClasses/Scene/DisplayClusterRootActor.cpp b/Engine/Plugins/Runtime/nDisplay/Source/DisplayCluster/Private/Game/EngineClasses/Scene/DisplayClusterRootActor.cpp
--- a/Engine/Plugins/Runtime/nDisplay/Source/DisplayCluster/Private/Game/EngineClasses/Scene/DisplayClusterRootActor.cpp
+++ b/Engine/Plugins/Runtime/nDisplay/Source/DisplayCluster/Private/Game/EngineClasses/Scene/DisplayClusterRootActor.cpp
@@ -46,6 +46,14 @@
 #include "IConcertSyncClient.h"
 #endif
 
+static TAutoConsoleVariable<int32> CVarShowVisualizationComponents(
+	TEXT("nDisplay.render.show.visualizationcomponents"),
+	1,
+	TEXT("Whether to show or hide visualization mesh components \n")
+	TEXT("0 : Hides visualization components \n")
+	TEXT("1 : Shows visualization components\n")
+);
+
 ADisplayClusterRootActor::ADisplayClusterRootActor(const FObjectInitializer& ObjectInitializer)
 	: Super(ObjectInitializer)
 	, OperationMode(EDisplayClusterOperationMode::Disabled)
@@ -378,9 +386,15 @@ int ADisplayClusterRootActor::GetInnerFrustumPriority(const FString& InnerFrustu
 }
 
 template <typename TComp>
-void ImplCollectChildrenVisualizationComponent(TSet<FPrimitiveComponentId>& OutPrimitives, TComp* pComp)
+void ImplCollectChildVisualizationOrHiddenComponents(TSet<FPrimitiveComponentId>& OutPrimitives, TComp* pComp, bool bCollectVisualizationComponents = true, bool bCollectHiddenComponents = true)
 {
 #if WITH_EDITOR
+
+	if (!bCollectVisualizationComponents && !bCollectHiddenComponents)
+	{
+		return;
+	}
+
 	USceneComponent* SceneComp = Cast<USceneComponent>(pComp);
 	if (SceneComp)
 	{
@@ -388,8 +402,8 @@ void ImplCollectChildrenVisualizationComponent(TSet<FPrimitiveComponentId>& OutP
 		SceneComp->GetChildrenComponents(false, Childrens);
 		for (USceneComponent* ChildIt : Childrens)
 		{
-			// Hide attached visualization components
-			if (ChildIt->IsVisualizationComponent() || ChildIt->bHiddenInGame)
+			if ((bCollectVisualizationComponents && ChildIt->IsVisualizationComponent()) 
+				|| (bCollectHiddenComponents && ChildIt->bHiddenInGame))
 			{
 				UPrimitiveComponent* PrimComp = Cast<UPrimitiveComponent>(ChildIt);
 				if (PrimComp)
@@ -427,7 +441,7 @@ void ADisplayClusterRootActor::GetTypedPrimitives(TSet<FPrimitiveComponentId>& O
 
 						if (bCollectChildrenVisualizationComponent)
 						{
-							ImplCollectChildrenVisualizationComponent(OutPrimitives, CompIt);
+							ImplCollectChildVisualizationOrHiddenComponents(OutPrimitives, CompIt);
 						}
 						break;
 					}
@@ -443,7 +457,7 @@ void ADisplayClusterRootActor::GetTypedPrimitives(TSet<FPrimitiveComponentId>& O
 
 				if (bCollectChildrenVisualizationComponent)
 				{
-					ImplCollectChildrenVisualizationComponent(OutPrimitives, CompIt);
+					ImplCollectChildVisualizationOrHiddenComponents(OutPrimitives, CompIt);
 				}
 			}
 
@@ -458,7 +472,7 @@ bool ADisplayClusterRootActor::FindPrimitivesByName(const TArray<FString>& InNam
 	return true;
 }
 
-// Gather components no rendered in game
+// Gather components not rendered in game
 bool ADisplayClusterRootActor::GetHiddenInGamePrimitives(TSet<FPrimitiveComponentId>& OutPrimitives)
 {
 	check(IsInGameThread());
@@ -479,47 +493,51 @@ bool ADisplayClusterRootActor::GetHiddenInGamePrimitives(TSet<FPrimitiveComponen
 
 #if WITH_EDITOR
 
-	// Hide all visualization components from RootActor
+	const bool bHideVisualizationComponents = !CVarShowVisualizationComponents.GetValueOnGameThread();
+
+	// Hide visualization and hidden components from RootActor
 	{
 		TArray<UPrimitiveComponent*> PrimitiveComponents;
 		GetComponents<UPrimitiveComponent>(PrimitiveComponents);
+
 		for (UPrimitiveComponent* CompIt : PrimitiveComponents)
 		{
-			if (CompIt->IsVisualizationComponent() || CompIt->bHiddenInGame)
+			if ((bHideVisualizationComponents && CompIt->IsVisualizationComponent()) || CompIt->bHiddenInGame)
 			{
 				OutPrimitives.Add(CompIt->ComponentId);
 			}
 
-			ImplCollectChildrenVisualizationComponent(OutPrimitives, CompIt);
+			ImplCollectChildVisualizationOrHiddenComponents(OutPrimitives, CompIt, bHideVisualizationComponents);
 		}
 	}
 
-	// Hide all visualization components from preview scene
-	UWorld* CurrentWorld = GetWorld();
-	if (CurrentWorld)
+	// Hide visualization and hidden components from scene
+	if (const UWorld* CurrentWorld = GetWorld())
 	{
 		// Iterate over all actors, looking for editor components.
 		for (const TWeakObjectPtr<AActor>& WeakActor : FActorRange(CurrentWorld))
 		{
 			if (AActor* Actor = WeakActor.Get())
 			{
-				// do not render hiiden in game actors on preview
-				bool bActorHideInGame = Actor->IsHidden();
+				// do not render hidden in game actors
+				const bool bActorHideInGame = Actor->IsHidden();
 
 				TArray<UPrimitiveComponent*> PrimitiveComponents;
 				Actor->GetComponents<UPrimitiveComponent>(PrimitiveComponents);
+
 				for (UPrimitiveComponent* PrimComp : PrimitiveComponents)
 				{
-					if (PrimComp->IsVisualizationComponent() || bActorHideInGame || PrimComp->bHiddenInGame)
+					if ((bHideVisualizationComponents && PrimComp->IsVisualizationComponent()) || bActorHideInGame || PrimComp->bHiddenInGame)
 					{
 						OutPrimitives.Add(PrimComp->ComponentId);
 					}
 
-					ImplCollectChildrenVisualizationComponent(OutPrimitives, PrimComp);
+					ImplCollectChildVisualizationOrHiddenComponents(OutPrimitives, PrimComp, bHideVisualizationComponents);
 				}
 			}
 		}
 	}
+
 #endif
 
 	return OutPrimitives.Num() > 0;
