nDisplay: Fix lightcards not showing when they use a translucent material.

The issue was that the EngineShowFlags passed to the ViewFamily was not yet configured, and instead it gets configured after view family construction. In particular, the separate translucency flag has a side effect during view family construction of setting its internal bAllowTranslucencyAfterDOF flag, which would be missed if SetSeparateTranslucency is called later. In such case, translucency goes into a separate render target and is not captured as expected.

So the fix consists in configuring the EngineShowFlags first, and then passing them to the view family constructor via the FSceneViewFamily::ConstructionValues structure.

#rnx

diff --git a/Engine/Plugins/Runtime/nDisplay/Source/DisplayCluster/Private/Game/EngineClasses/Basics/DisplayClusterViewportClient.cpp b/Engine/Plugins/Runtime/nDisplay/Source/DisplayCluster/Private/Game/EngineClasses/Basics/DisplayClusterViewportClient.cpp
--- a/Engine/Plugins/Runtime/nDisplay/Source/DisplayCluster/Private/Game/EngineClasses/Basics/DisplayClusterViewportClient.cpp
+++ b/Engine/Plugins/Runtime/nDisplay/Source/DisplayCluster/Private/Game/EngineClasses/Basics/DisplayClusterViewportClient.cpp
@@ -421,14 +421,17 @@ void UDisplayClusterViewportClient::Draw(FViewport* InViewport, FCanvas* SceneCa
 		for (FDisplayClusterRenderFrame::FFrameViewFamily& DCViewFamily : DCRenderTarget.ViewFamilies)
 		{
 			// Create the view family for rendering the world scene to the viewport's render target
-			FSceneViewFamilyContext ViewFamily(FSceneViewFamily::ConstructionValues(DCRenderTarget.RenderTargetPtr, MyWorld->Scene, EngineShowFlags)
-				.SetRealtimeUpdate(true)
-				.SetAdditionalViewFamily(bAdditionalViewFamily));
+			FSceneViewFamilyContext ViewFamily(RenderFrame.ViewportManager->CreateViewFamilyConstructionValues(
+				DCRenderTarget,
+				MyWorld->Scene,
+				EngineShowFlags,
+				bAdditionalViewFamily
+			));
 
 			// Disable clean op for all next families on this render target
 			bAdditionalViewFamily = true;
 
-			// Configure family flags
+			// Configure family
 			RenderFrame.ViewportManager->ConfigureViewFamily(DCRenderTarget, DCViewFamily, ViewFamily);
 
 #if WITH_EDITOR
diff --git a/Engine/Plugins/Runtime/nDisplay/Source/DisplayCluster/Private/Render/Viewport/DisplayClusterViewportManager.cpp b/Engine/Plugins/Runtime/nDisplay/Source/DisplayCluster/Private/Render/Viewport/DisplayClusterViewportManager.cpp
--- a/Engine/Plugins/Runtime/nDisplay/Source/DisplayCluster/Private/Render/Viewport/DisplayClusterViewportManager.cpp
+++ b/Engine/Plugins/Runtime/nDisplay/Source/DisplayCluster/Private/Render/Viewport/DisplayClusterViewportManager.cpp
@@ -442,60 +442,79 @@ void FDisplayClusterViewportManager::FinalizeNewFrame()
 	PostProcessManager->FinalizeNewFrame();
 }
 
-void FDisplayClusterViewportManager::ConfigureViewFamily(const FDisplayClusterRenderFrame::FFrameRenderTarget& InFrameTarget, const FDisplayClusterRenderFrame::FFrameViewFamily& InFrameViewFamily, FSceneViewFamilyContext& ViewFamily)
+FSceneViewFamily::ConstructionValues FDisplayClusterViewportManager::CreateViewFamilyConstructionValues(
+	const FDisplayClusterRenderFrame::FFrameRenderTarget& InFrameTarget,
+	FSceneInterface* InScene,
+	FEngineShowFlags InEngineShowFlags,
+	const bool bInAdditionalViewFamily) const
 {
-	// Gather Scene View Extensions
-	// Scene View Extension activation with ViewportId granularity only works if you have one ViewFamily per ViewportId
-	{
-		ViewFamily.ViewExtensions = InFrameViewFamily.ViewExtensions;
-		for (FSceneViewExtensionRef& ViewExt : ViewFamily.ViewExtensions)
-		{
-			ViewExt->SetupViewFamily(ViewFamily);
-		}
-	}
 
-	ViewFamily.SceneCaptureCompositeMode = ESceneCaptureCompositeMode::SCCM_Overwrite;
+	bool bResolveScene = true;
 
-	// Setup capture mode:
+	switch (InFrameTarget.CaptureMode)
 	{
-		switch (InFrameTarget.CaptureMode)
-		{
 		case EDisplayClusterViewportCaptureMode::Chromakey:
 		case EDisplayClusterViewportCaptureMode::Lightcard:
-			ViewFamily.SceneCaptureSource = ESceneCaptureSource::SCS_SceneColorHDR;
-			ViewFamily.bResolveScene = false;
+		case EDisplayClusterViewportCaptureMode::Lightcard_OCIO:
+
+			if (InFrameTarget.CaptureMode != EDisplayClusterViewportCaptureMode::Lightcard_OCIO)
+			{
+				bResolveScene = false;
+				InEngineShowFlags.PostProcessing = 0;
+			}
+
+			InEngineShowFlags.SetAtmosphere(0);
+			InEngineShowFlags.SetFog(0);
+			InEngineShowFlags.SetMotionBlur(0); // motion blur doesn't work correctly with scene captures.
+			InEngineShowFlags.SetSeparateTranslucency(0);
+			InEngineShowFlags.SetHMDDistortion(0);
+			InEngineShowFlags.SetOnScreenDebug(0);
 
-			ViewFamily.EngineShowFlags.PostProcessing = 0;
+			InEngineShowFlags.SetNaniteMeshes(0);
+
+			InEngineShowFlags.SetLumenReflections(0);
+			InEngineShowFlags.SetLumenGlobalIllumination(0);
 
-			ViewFamily.EngineShowFlags.SetAtmosphere(0);
-			ViewFamily.EngineShowFlags.SetFog(0);
-			ViewFamily.EngineShowFlags.SetMotionBlur(0); // motion blur doesn't work correctly with scene captures.
-			ViewFamily.EngineShowFlags.SetSeparateTranslucency(0);
-			ViewFamily.EngineShowFlags.SetHMDDistortion(0);
-			ViewFamily.EngineShowFlags.SetOnScreenDebug(0);
 			break;
 
 		default:
 			break;
-		}
 	}
-	// Disable nanite and lumen for lightcards/chromakey
+
+	return FSceneViewFamily::ConstructionValues(InFrameTarget.RenderTargetPtr, InScene, InEngineShowFlags)
+		.SetResolveScene(bResolveScene)
+		.SetRealtimeUpdate(true)
+		.SetAdditionalViewFamily(bInAdditionalViewFamily);
+}
+
+
+void FDisplayClusterViewportManager::ConfigureViewFamily(const FDisplayClusterRenderFrame::FFrameRenderTarget& InFrameTarget, const FDisplayClusterRenderFrame::FFrameViewFamily& InFrameViewFamily, FSceneViewFamilyContext& ViewFamily)
+{
+	// Note: EngineShowFlags should have already been configured in CreateViewFamilyConstructionValues.
+
+	// Gather Scene View Extensions
+	// Scene View Extension activation with ViewportId granularity only works if you have one ViewFamily per ViewportId
 	{
-		switch (InFrameTarget.CaptureMode)
+		ViewFamily.ViewExtensions = InFrameViewFamily.ViewExtensions;
+		for (FSceneViewExtensionRef& ViewExt : ViewFamily.ViewExtensions)
 		{
-		case EDisplayClusterViewportCaptureMode::Chromakey:
-		case EDisplayClusterViewportCaptureMode::Lightcard:
-		case EDisplayClusterViewportCaptureMode::Lightcard_OCIO:
+			ViewExt->SetupViewFamily(ViewFamily);
+		}
+	}
 
-			ViewFamily.EngineShowFlags.SetNaniteMeshes(0);
-			
-			ViewFamily.EngineShowFlags.SetLumenReflections(0);
-			ViewFamily.EngineShowFlags.SetLumenGlobalIllumination(0);
+	// Setup capture mode:
+	{
+		ViewFamily.SceneCaptureCompositeMode = ESceneCaptureCompositeMode::SCCM_Overwrite;
 
-			break;
+		switch (InFrameTarget.CaptureMode)
+		{
+			case EDisplayClusterViewportCaptureMode::Chromakey:
+			case EDisplayClusterViewportCaptureMode::Lightcard:
+				ViewFamily.SceneCaptureSource = ESceneCaptureSource::SCS_SceneColorHDR;
+				break;
 
-		default:
-			break;
+			default:
+				break;
 		}
 	}
 }
diff --git a/Engine/Plugins/Runtime/nDisplay/Source/DisplayCluster/Private/Render/Viewport/DisplayClusterViewportManager.h b/Engine/Plugins/Runtime/nDisplay/Source/DisplayCluster/Private/Render/Viewport/DisplayClusterViewportManager.h
--- a/Engine/Plugins/Runtime/nDisplay/Source/DisplayCluster/Private/Render/Viewport/DisplayClusterViewportManager.h
+++ b/Engine/Plugins/Runtime/nDisplay/Source/DisplayCluster/Private/Render/Viewport/DisplayClusterViewportManager.h
@@ -54,6 +54,13 @@ public:
 	virtual bool BeginNewFrame(FViewport* InViewport, UWorld* InWorld, FDisplayClusterRenderFrame& OutRenderFrame) override;
 	virtual void FinalizeNewFrame() override;
 
+	virtual FSceneViewFamily::ConstructionValues CreateViewFamilyConstructionValues(
+		const FDisplayClusterRenderFrame::FFrameRenderTarget& InFrameTarget,
+		FSceneInterface* InScene,
+		FEngineShowFlags InEngineShowFlags,
+		const bool bInAdditionalViewFamily
+	) const override;
+
 	virtual void ConfigureViewFamily(const FDisplayClusterRenderFrame::FFrameRenderTarget& InFrameTarget, const FDisplayClusterRenderFrame::FFrameViewFamily& InFrameViewFamily, FSceneViewFamilyContext& InOutViewFamily) override;
 	
 	virtual void RenderFrame(FViewport* InViewport) override;
diff --git a/Engine/Plugins/Runtime/nDisplay/Source/DisplayCluster/Private/Render/Viewport/DisplayClusterViewportManager_EditorPreview.cpp b/Engine/Plugins/Runtime/nDisplay/Source/DisplayCluster/Private/Render/Viewport/DisplayClusterViewportManager_EditorPreview.cpp
--- a/Engine/Plugins/Runtime/nDisplay/Source/DisplayCluster/Private/Render/Viewport/DisplayClusterViewportManager_EditorPreview.cpp
+++ b/Engine/Plugins/Runtime/nDisplay/Source/DisplayCluster/Private/Render/Viewport/DisplayClusterViewportManager_EditorPreview.cpp
@@ -116,64 +116,64 @@ bool FDisplayClusterViewportManager::RenderInEditor(class FDisplayClusterRenderF
 	for (FDisplayClusterRenderFrame::FFrameRenderTarget& RenderTargetIt : InRenderFrame.RenderTargets)
 	{
 		// Special flag, allow clear RTT surface only for first family
-		bool AdditionalViewFamily = false;
+		bool bAdditionalViewFamily = false;
 
 		for (FDisplayClusterRenderFrame::FFrameViewFamily& ViewFamiliesIt : RenderTargetIt.ViewFamilies)
 		{
-			{
-				FRenderTarget* DstResource = RenderTargetIt.RenderTargetPtr;
-				// Create the view family for rendering the world scene to the viewport's render target
-				FSceneViewFamilyContext ViewFamily(
-					FSceneViewFamily::ConstructionValues(DstResource, PreviewScene, EngineShowFlags)
-					.SetResolveScene(true)
-					.SetRealtimeUpdate(true)
-					.SetAdditionalViewFamily(AdditionalViewFamily));
+			FRenderTarget* DstResource = RenderTargetIt.RenderTargetPtr;
 
-				ConfigureViewFamily(RenderTargetIt, ViewFamiliesIt, ViewFamily);
+			// Create the view family for rendering the world scene to the viewport's render target
+			FSceneViewFamilyContext ViewFamily(CreateViewFamilyConstructionValues(
+				RenderTargetIt,
+				PreviewScene,
+				EngineShowFlags,
+				bAdditionalViewFamily
+			));
 
-				// Disable clean op for all next families on this render target
-				AdditionalViewFamily = true;
+			ConfigureViewFamily(RenderTargetIt, ViewFamiliesIt, ViewFamily);
 
-				for (FDisplayClusterRenderFrame::FFrameView& ViewIt : ViewFamiliesIt.Views)
-				{
-					FDisplayClusterViewport* ViewportPtr = static_cast<FDisplayClusterViewport*>(ViewIt.Viewport);
+			// Disable clean op for all next families on this render target
+			bAdditionalViewFamily = true;
 
-					check(ViewportPtr != nullptr);
-					check(ViewIt.ContextNum < (uint32)ViewportPtr->Contexts.Num());
+			for (FDisplayClusterRenderFrame::FFrameView& ViewIt : ViewFamiliesIt.Views)
+			{
+				FDisplayClusterViewport* ViewportPtr = static_cast<FDisplayClusterViewport*>(ViewIt.Viewport);
 
-					// Calculate the player's view information.
-					FVector  ViewLocation;
-					FRotator ViewRotation;
-					FSceneView* View = ViewportPtr->ImplCalcScenePreview(ViewFamily, ViewIt.ContextNum);
+				check(ViewportPtr != nullptr);
+				check(ViewIt.ContextNum < (uint32)ViewportPtr->Contexts.Num());
 
-					if (View && ViewIt.bDisableRender)
-					{
-						ViewFamily.Views.Remove(View);
+				// Calculate the player's view information.
+				FVector  ViewLocation;
+				FRotator ViewRotation;
+				FSceneView* View = ViewportPtr->ImplCalcScenePreview(ViewFamily, ViewIt.ContextNum);
 
-						delete View;
-						View = nullptr;
-					}
+				if (View && ViewIt.bDisableRender)
+				{
+					ViewFamily.Views.Remove(View);
 
-					if (View)
-					{
-						// Apply viewport context settings to view (crossGPU, visibility, etc)
-						ViewIt.Viewport->SetupSceneView(ViewIt.ContextNum, PreviewScene->GetWorld(), ViewFamily, *View);
-					}
+					delete View;
+					View = nullptr;
 				}
 
-				if (ViewFamily.Views.Num() > 0)
+				if (View)
 				{
-					// Screen percentage is still not supported in scene capture.
-					ViewFamily.EngineShowFlags.ScreenPercentage = false;
-					ViewFamily.SetScreenPercentageInterface(new FLegacyScreenPercentageDriver(ViewFamily, 1.0f));
+					// Apply viewport context settings to view (crossGPU, visibility, etc)
+					ViewIt.Viewport->SetupSceneView(ViewIt.ContextNum, PreviewScene->GetWorld(), ViewFamily, *View);
+				}
+			}
+
+			if (ViewFamily.Views.Num() > 0)
+			{
+				// Screen percentage is still not supported in scene capture.
+				ViewFamily.EngineShowFlags.ScreenPercentage = false;
+				ViewFamily.SetScreenPercentageInterface(new FLegacyScreenPercentageDriver(ViewFamily, 1.0f));
 
-					ViewFamily.bIsRenderedImmediatelyAfterAnotherViewFamily = bIsRenderedImmediatelyAfterAnotherViewFamily;
+				ViewFamily.bIsRenderedImmediatelyAfterAnotherViewFamily = bIsRenderedImmediatelyAfterAnotherViewFamily;
 
-					FCanvas Canvas(DstResource, nullptr, PreviewScene->GetWorld(), ERHIFeatureLevel::SM5, FCanvas::CDM_DeferDrawing /*FCanvas::CDM_ImmediateDrawing*/, 1.0f);
-					Canvas.Clear(FLinearColor::Black);
+				FCanvas Canvas(DstResource, nullptr, PreviewScene->GetWorld(), ERHIFeatureLevel::SM5, FCanvas::CDM_DeferDrawing /*FCanvas::CDM_ImmediateDrawing*/, 1.0f);
+				Canvas.Clear(FLinearColor::Black);
 
-					GetRendererModule().BeginRenderingViewFamily(&Canvas, &ViewFamily);
-				}
+				GetRendererModule().BeginRenderingViewFamily(&Canvas, &ViewFamily);
 			}
 		}
 	}
diff --git a/Engine/Plugins/Runtime/nDisplay/Source/DisplayCluster/Public/Render/Viewport/IDisplayClusterViewportManager.h b/Engine/Plugins/Runtime/nDisplay/Source/DisplayCluster/Public/Render/Viewport/IDisplayClusterViewportManager.h
--- a/Engine/Plugins/Runtime/nDisplay/Source/DisplayCluster/Public/Render/Viewport/IDisplayClusterViewportManager.h
+++ b/Engine/Plugins/Runtime/nDisplay/Source/DisplayCluster/Public/Render/Viewport/IDisplayClusterViewportManager.h
@@ -68,6 +68,24 @@ public:
 	*/
 	virtual void FinalizeNewFrame() = 0;
 
+	/**
+	* Create constructor values for viewfamily, using rules
+	* [Game thread func]
+	*
+	* @param InFrameTarget           - frame target
+	* @param InScene                 - scene
+	* @param InEngineShowFlags       - starting showflags
+	* @param bInAdditionalViewFamily - flag indicating this is an additional view family
+	* 
+	* @return - The ConstructionValues object that was created, meant to initialize a view family context.
+	*/
+	virtual FSceneViewFamily::ConstructionValues CreateViewFamilyConstructionValues(
+		const FDisplayClusterRenderFrame::FFrameRenderTarget& InFrameTarget,
+		FSceneInterface* InScene,
+		FEngineShowFlags InEngineShowFlags,
+		const bool bInAdditionalViewFamily
+	) const = 0;
+
 	/**
 	* Initialize view family, using rules
 	* [Game thread func]
