nDisplay: Reverted change that removed the ICVFX camera's Configuration property category from level instances of the camera. Now, the Configratuion category shows up anywhere, in level instances and in the blueprint Config Editor

diff --git a/Engine/Plugins/Runtime/nDisplay/Source/DisplayCluster/Public/Components/DisplayClusterICVFXCameraComponent.h b/Engine/Plugins/Runtime/nDisplay/Source/DisplayCluster/Public/Components/DisplayClusterICVFXCameraComponent.h
--- a/Engine/Plugins/Runtime/nDisplay/Source/DisplayCluster/Public/Components/DisplayClusterICVFXCameraComponent.h
+++ b/Engine/Plugins/Runtime/nDisplay/Source/DisplayCluster/Public/Components/DisplayClusterICVFXCameraComponent.h
@@ -160,22 +160,22 @@ private:
 	UPROPERTY(EditAnywhere, Transient, Category = "Texture Replacement", meta = (PropertyPath = "CameraSettings.RenderSettings.Replace.TextureRegion", EditConditionPath = "CameraSettings.bEnable && CameraSettings.RenderSettings.Replace.bAllowReplace && CameraSettings.RenderSettings.Replace.bShouldUseTextureRegion"))
 	FDisplayClusterEditorPropertyReference TextureRegionRef;
 
-	UPROPERTY(EditDefaultsOnly, Transient, Category = Configuration, meta = (PropertyPath = "CameraSettings.RenderSettings.RenderOrder", EditConditionPath = "CameraSettings.bEnable"))
+	UPROPERTY(EditAnywhere, Transient, Category = Configuration, meta = (PropertyPath = "CameraSettings.RenderSettings.RenderOrder", EditConditionPath = "CameraSettings.bEnable"))
 	FDisplayClusterEditorPropertyReference RenderOrderRef;
 
-	UPROPERTY(EditDefaultsOnly, Transient, Category = Configuration, meta = (PropertyPath = "CameraSettings.RenderSettings.CustomFrameSize", EditConditionPath = "CameraSettings.bEnable"))
+	UPROPERTY(EditAnywhere, Transient, Category = Configuration, meta = (PropertyPath = "CameraSettings.RenderSettings.CustomFrameSize", EditConditionPath = "CameraSettings.bEnable"))
 	FDisplayClusterEditorPropertyReference CustomFrameSizeRef;
 
-	UPROPERTY(EditDefaultsOnly, Transient, Category = Configuration, meta = (PropertyPath = "CameraSettings.RenderSettings.AdvancedRenderSettings.RenderTargetRatio", EditConditionPath = "CameraSettings.bEnable"))
+	UPROPERTY(EditAnywhere, Transient, Category = Configuration, meta = (PropertyPath = "CameraSettings.RenderSettings.AdvancedRenderSettings.RenderTargetRatio", EditConditionPath = "CameraSettings.bEnable"))
 	FDisplayClusterEditorPropertyReference RenderTargetRatioRef;
 
-	UPROPERTY(EditDefaultsOnly, Transient, Category = Configuration, meta = (PropertyPath = "CameraSettings.RenderSettings.AdvancedRenderSettings.GPUIndex", EditConditionPath = "CameraSettings.bEnable"))
+	UPROPERTY(EditAnywhere, Transient, Category = Configuration, meta = (PropertyPath = "CameraSettings.RenderSettings.AdvancedRenderSettings.GPUIndex", EditConditionPath = "CameraSettings.bEnable"))
 	FDisplayClusterEditorPropertyReference GPUIndexRef;
 
-	UPROPERTY(EditDefaultsOnly, Transient, Category = Configuration, meta = (PropertyPath = "CameraSettings.RenderSettings.AdvancedRenderSettings.StereoGPUIndex", EditConditionPath = "CameraSettings.bEnable"))
+	UPROPERTY(EditAnywhere, Transient, Category = Configuration, meta = (PropertyPath = "CameraSettings.RenderSettings.AdvancedRenderSettings.StereoGPUIndex", EditConditionPath = "CameraSettings.bEnable"))
 	FDisplayClusterEditorPropertyReference StereoGPUIndexRef;
 
-	UPROPERTY(EditDefaultsOnly, Transient, Category = Configuration, meta = (PropertyPath = "CameraSettings.RenderSettings.AdvancedRenderSettings.StereoMode", EditConditionPath = "CameraSettings.bEnable"))
+	UPROPERTY(EditAnywhere, Transient, Category = Configuration, meta = (PropertyPath = "CameraSettings.RenderSettings.AdvancedRenderSettings.StereoMode", EditConditionPath = "CameraSettings.bEnable"))
 	FDisplayClusterEditorPropertyReference StereoModeRef;
 
 #endif // WITH_EDITORONLY_DATA
