nDisplay: fixed non-Unity build errors

diff --git a/Engine/Plugins/Runtime/nDisplay/Source/DisplayCluster/Public/Render/Viewport/IDisplayClusterViewportManager.h b/Engine/Plugins/Runtime/nDisplay/Source/DisplayCluster/Public/Render/Viewport/IDisplayClusterViewportManager.h
--- a/Engine/Plugins/Runtime/nDisplay/Source/DisplayCluster/Public/Render/Viewport/IDisplayClusterViewportManager.h
+++ b/Engine/Plugins/Runtime/nDisplay/Source/DisplayCluster/Public/Render/Viewport/IDisplayClusterViewportManager.h
@@ -6,6 +6,8 @@
 
 #include "Render/Viewport/IDisplayClusterViewport.h"
 #include "Render/Viewport/RenderFrame/DisplayClusterRenderFrame.h"
+#include "SceneView.h"
+
 
 class UWorld;
 class FViewport;
