Fix for a bug where two nDisplay viewports could end up with the same name after duplicating a cluster node.

diff --git a/Engine/Plugins/Runtime/nDisplay/Source/DisplayClusterConfigurator/Private/ClusterConfiguration/DisplayClusterConfiguratorClusterUtils.cpp b/Engine/Plugins/Runtime/nDisplay/Source/DisplayClusterConfigurator/Private/ClusterConfiguration/DisplayClusterConfiguratorClusterUtils.cpp
--- a/Engine/Plugins/Runtime/nDisplay/Source/DisplayClusterConfigurator/Private/ClusterConfiguration/DisplayClusterConfiguratorClusterUtils.cpp
+++ b/Engine/Plugins/Runtime/nDisplay/Source/DisplayClusterConfigurator/Private/ClusterConfiguration/DisplayClusterConfiguratorClusterUtils.cpp
@@ -510,32 +510,13 @@ bool FDisplayClusterConfiguratorClusterUtils::RemoveUnusedHostDisplayData(UDispl
 
 FString FDisplayClusterConfiguratorClusterUtils::GetUniqueNameForHost(FString InitialName, UDisplayClusterConfigurationCluster* ParentCluster, bool bAddZero)
 {
-	FString UniqueName = InitialName;
-	int NameIndex = 0;
-
-	if (bAddZero)
-	{
-		UniqueName = FString::Printf(TEXT("%s_%d"), *InitialName, NameIndex);
-	}
-
 	TArray<FString> UsedNames;
 	for (TPair<FString, UDisplayClusterConfigurationHostDisplayData*>& HostPair : ParentCluster->HostDisplayData)
 	{
 		UsedNames.Add(HostPair.Value->HostName.ToString());
 	}
 
-	while (UsedNames.Contains(UniqueName))
-	{
-		UniqueName = FString::Printf(TEXT("%s_%d"), *InitialName, ++NameIndex);
-	}
-
-	// If there is already an in-memory object connected to the parent cluster with the name, we need to use a globally unique object name.
-	if (StaticFindObject(nullptr, ParentCluster, *UniqueName, true))
-	{
-		UniqueName = MakeUniqueObjectName(ParentCluster, UDisplayClusterConfigurationClusterNode::StaticClass(), *UniqueName).ToString();
-	}
-
-	return UniqueName;
+	return GetUniqueName(InitialName, UsedNames, UDisplayClusterConfigurationClusterNode::StaticClass(), ParentCluster, bAddZero);
 }
 
 FString FDisplayClusterConfiguratorClusterUtils::GetAddressForHost(UDisplayClusterConfigurationHostDisplayData* HostDisplayData)
@@ -628,26 +609,10 @@ bool FDisplayClusterConfiguratorClusterUtils::IsClusterNodePrimary(UDisplayClust
 
 FString FDisplayClusterConfiguratorClusterUtils::GetUniqueNameForClusterNode(FString InitialName, UDisplayClusterConfigurationCluster* ParentCluster, bool bAddZero)
 {
-	FString UniqueName = InitialName;
-	int NameIndex = 0;
-
-	if (bAddZero)
-	{
-		UniqueName = FString::Printf(TEXT("%s_%d"), *InitialName, NameIndex);
-	}
-
-	while (ParentCluster->Nodes.Contains(UniqueName))
-	{
-		UniqueName = FString::Printf(TEXT("%s_%d"), *InitialName, ++NameIndex);
-	}
-
-	// If there is already an in-memory object connected to the parent cluster with the name, we need to use a globally unique object name.
-	if (StaticFindObject(nullptr, ParentCluster, *UniqueName, true))
-	{
-		UniqueName = MakeUniqueObjectName(ParentCluster, UDisplayClusterConfigurationClusterNode::StaticClass(), *UniqueName).ToString();
-	}
+	TArray<FString> UsedNames;
+	ParentCluster->Nodes.GenerateKeyArray(UsedNames);
 
-	return UniqueName;
+	return GetUniqueName(InitialName, UsedNames, UDisplayClusterConfigurationClusterNode::StaticClass(), ParentCluster, bAddZero);
 }
 
 UDisplayClusterConfigurationClusterNode* FDisplayClusterConfiguratorClusterUtils::AddClusterNodeToCluster(UDisplayClusterConfigurationClusterNode* ClusterNode, UDisplayClusterConfigurationCluster* Cluster, FString NewClusterNodeName)
@@ -808,21 +773,13 @@ FString FDisplayClusterConfiguratorClusterUtils::GetViewportName(UDisplayCluster
 
 FString FDisplayClusterConfiguratorClusterUtils::GetUniqueNameForViewport(FString InitialName, UDisplayClusterConfigurationClusterNode* ParentClusterNode, bool bAddZero)
 {
-	FString UniqueName = InitialName;
-	int NameIndex = 0;
-
-	if (bAddZero)
-	{
-		UniqueName = FString::Printf(TEXT("%s_%d"), *InitialName, NameIndex);
-	}
-
 	// Viewport names must be unique across the entire cluster, not just within its parent cluster nodes. Gather all of the viewport names
 	// in the cluster to check for uniqueness. Add the parent cluster node's viewports first, in case we can't get to the root cluster through
 	// the cluster node's Outer (i.e. the cluster node has not been added to the cluster yet)
-	TSet<FString> UsedViewportNames;
+	TSet<FString> UsedNames;
 	for (const TPair<FString, UDisplayClusterConfigurationViewport*>& ViewportKeyPair : ParentClusterNode->Viewports)
 	{
-		UsedViewportNames.Add(ViewportKeyPair.Key);
+		UsedNames.Add(ViewportKeyPair.Key);
 	}
 
 	if (UDisplayClusterConfigurationCluster* Cluster = Cast<UDisplayClusterConfigurationCluster>(ParentClusterNode->GetOuter()))
@@ -834,24 +791,15 @@ FString FDisplayClusterConfiguratorClusterUtils::GetUniqueNameForViewport(FStrin
 			{
 				for (const TPair<FString, UDisplayClusterConfigurationViewport*>& ViewportKeyPair : ClusterNode->Viewports)
 				{
-					UsedViewportNames.Add(ViewportKeyPair.Key);
+					UsedNames.Add(ViewportKeyPair.Key);
 				}
 			}
 		}
 	}
 
-	while (UsedViewportNames.Contains(UniqueName))
-	{
-		UniqueName = FString::Printf(TEXT("%s_%d"), *InitialName, ++NameIndex);
-	}
-
-	// If there is already an in-memory object connected to the parent cluster node with the name, we need to use a globally unique object name.
-	if (StaticFindObject(nullptr, ParentClusterNode, *UniqueName, true))
-	{
-		UniqueName = MakeUniqueObjectName(ParentClusterNode, UDisplayClusterConfigurationViewport::StaticClass(), *UniqueName).ToString();
-	}
-
-	return UniqueName;
+	TArray<FString> UsedNamesArray = UsedNames.Array();
+	
+	return GetUniqueName(InitialName, UsedNamesArray, UDisplayClusterConfigurationViewport::StaticClass(), ParentClusterNode, bAddZero);
 }
 
 UDisplayClusterConfigurationViewport* FDisplayClusterConfiguratorClusterUtils::AddViewportToClusterNode(UDisplayClusterConfigurationViewport* Viewport, UDisplayClusterConfigurationClusterNode* ClusterNode, FString NewViewportName)
@@ -1191,6 +1139,16 @@ TArray<UObject*> FDisplayClusterConfiguratorClusterUtils::PasteClusterItemsFromC
 				ClusterNodeCopy->WindowRect.Y = NewLocation.Y;
 			}
 
+			// Rename viewports in the new cluster node to ensure they're unique. Copy the viewport pointers into an array
+			// first since renaming would cause the dictionary to change while we're iterating it.
+			TArray<UDisplayClusterConfigurationViewport*> CopiedViewports;
+			ClusterNodeCopy->Viewports.GenerateValueArray(CopiedViewports);
+			
+			for (UDisplayClusterConfigurationViewport* Viewport : CopiedViewports)
+			{
+				RenameViewport(Viewport, *GetUniqueNameForViewport(*Viewport->GetName(), ClusterNodeCopy));
+			}
+
 			CopiedObjects.Add(ClusterNodeCopy);
 		}
 	}
@@ -1247,4 +1205,53 @@ TArray<UObject*> FDisplayClusterConfiguratorClusterUtils::PasteClusterItemsFromC
 	return CopiedObjects;
 }
 
+FString FDisplayClusterConfiguratorClusterUtils::GetUniqueName(const FString& InitialName, const TArray<FString>& UsedNames, const UClass* Class, UObject* Parent, bool bAddZero)
+{
+	if (!bAddZero && !UsedNames.Contains(InitialName))
+	{
+		// Name doesn't need to be modified
+		return InitialName;
+	}
+	
+	int32 Counter = bAddZero ? 0 : 1;
+
+	// Find the start of the existing numeric suffix
+	int32 Index = InitialName.Len();
+	while (Index > 0 && InitialName[Index-1] >= '0' && InitialName[Index-1] <= '9')
+	{
+		--Index;
+	}
+
+	FString BaseName = InitialName;
+	FString NewName;
+	if (Index < BaseName.Len())
+	{
+		// Strip away the suffix and store the value in the counter so we can count up from there
+		FString NumericSuffix = BaseName.RightChop(Index);
+		Counter = FCString::Atoi(*NumericSuffix);
+		NumericSuffix = FString::FromInt(Counter); // Restringify the counter to account for leading 0s that we don't want to remove
+		BaseName.RemoveAt(BaseName.Len() - NumericSuffix.Len(), NumericSuffix.Len(), false);
+	}
+	else
+	{
+		// No existing suffix, so add our underscore separator
+		BaseName += "_";
+	}
+
+	do
+	{
+		NewName = FString::Printf(TEXT("%s%d"), *BaseName, Counter);
+		++Counter;
+	}
+	while (UsedNames.Contains(NewName));
+
+	// If there is already an in-memory object connected to the parent cluster with our generated name, we need to use a globally unique object name.
+	if (StaticFindObject(nullptr, Parent, *NewName, true))
+	{
+		NewName = MakeUniqueObjectName(Parent, Class, *NewName).ToString();
+	}
+
+	return NewName;
+}
+
 #undef LOCTEXT_NAMESPACE
\ No newline at end of file
diff --git a/Engine/Plugins/Runtime/nDisplay/Source/DisplayClusterConfigurator/Private/ClusterConfiguration/DisplayClusterConfiguratorClusterUtils.h b/Engine/Plugins/Runtime/nDisplay/Source/DisplayClusterConfigurator/Private/ClusterConfiguration/DisplayClusterConfiguratorClusterUtils.h
--- a/Engine/Plugins/Runtime/nDisplay/Source/DisplayClusterConfigurator/Private/ClusterConfiguration/DisplayClusterConfiguratorClusterUtils.h
+++ b/Engine/Plugins/Runtime/nDisplay/Source/DisplayClusterConfigurator/Private/ClusterConfiguration/DisplayClusterConfiguratorClusterUtils.h
@@ -240,6 +240,17 @@ public:
 	static TArray<UObject*> PasteClusterItemsFromClipboard(const TArray<UObject*>& TargetClusterItems, TOptional<FVector2D> PasteLocation = TOptional<FVector2D>());
 
 private:
+	/**
+	 * Gets a unique name for something given a list of existing names.
+	 * @param InitialName - The initial name the object wants
+	 * @param UsedNames - The list of names that have already been used
+	 * @param Class - The class to generate a name for
+	 * @param Parent - The parent to check for in-memory objects that collide with the name
+	 * @param bAddZero - Whether to add an "_0" to the initial name if it is unique
+	 * @returns A unique name for the object
+	 */
+	static FString GetUniqueName(const FString& InitialName, const TArray<FString>& UsedNames, const UClass* Class, UObject* Parent, bool bAddZero = false);
+	
 	static const FVector2D NewClusterItemDialogSize;
 	static const FString DefaultNewHostName;
 	static const FString DefaultNewClusterNodeName;
diff --git a/Engine/Plugins/Runtime/nDisplay/Source/DisplayClusterConfigurator/Private/DisplayClusterConfiguratorCompiler.cpp b/Engine/Plugins/Runtime/nDisplay/Source/DisplayClusterConfigurator/Private/DisplayClusterConfiguratorCompiler.cpp
--- a/Engine/Plugins/Runtime/nDisplay/Source/DisplayClusterConfigurator/Private/DisplayClusterConfiguratorCompiler.cpp
+++ b/Engine/Plugins/Runtime/nDisplay/Source/DisplayClusterConfigurator/Private/DisplayClusterConfiguratorCompiler.cpp
@@ -248,6 +248,8 @@ void FDisplayClusterConfiguratorKismetCompilerContext::ValidateConfiguration()
 
 	bool bAtLeastOneViewportFound = false;
 
+	TMap<FString, UDisplayClusterConfigurationViewport*> ViewportsByName;
+	
 	for (const auto& ClusterNodes : BlueprintData->Cluster->Nodes)
 	{
 		if (ClusterNodes.Value->Viewports.Num() > 0)
@@ -257,6 +259,22 @@ void FDisplayClusterConfiguratorKismetCompilerContext::ValidateConfiguration()
 			
 			for (const auto& Viewport : ClusterNodes.Value->Viewports)
 			{
+				const FString ViewportName = Viewport.Value->GetName();
+
+				// Check that no two viewports have the same name
+				if (UDisplayClusterConfigurationViewport** ExistingViewport = ViewportsByName.Find(ViewportName))
+				{
+					MessageLog.Error(
+						*LOCTEXT("DuplicateViewportNameError", "Viewport @@ uses the same name as viewport @@.").ToString(),
+						Viewport.Value,
+						*ExistingViewport
+					);
+				}
+				else
+				{
+					ViewportsByName.Add(ViewportName, Viewport.Value);
+				}
+				
 				if (Viewport.Value->ProjectionPolicy.Type.IsEmpty())
 				{
 					MessageLog.Warning(*LOCTEXT("NoPolicyError", "No projection policy assigned to viewport @@.").ToString(), Viewport.Value);
