nDisplay: Only prevent construction scripts from running on Interactive changes if bRunConstructionScriptOnDrag is false.

diff --git a/Engine/Plugins/Runtime/nDisplay/Source/DisplayCluster/Private/Game/EngineClasses/Scene/DisplayClusterRootActorEditor.cpp b/Engine/Plugins/Runtime/nDisplay/Source/DisplayCluster/Private/Game/EngineClasses/Scene/DisplayClusterRootActorEditor.cpp
--- a/Engine/Plugins/Runtime/nDisplay/Source/DisplayCluster/Private/Game/EngineClasses/Scene/DisplayClusterRootActorEditor.cpp
+++ b/Engine/Plugins/Runtime/nDisplay/Source/DisplayCluster/Private/Game/EngineClasses/Scene/DisplayClusterRootActorEditor.cpp
@@ -491,8 +491,13 @@ void ADisplayClusterRootActor::PostEditChangeProperty(FPropertyChangedEvent& Pro
 	};
 
 	bool bReinitializeActor = true;
-
-	if (PropertyChangedEvent.ChangeType == EPropertyChangeType::Interactive && !CurrentTransactionAnnotation.IsValid())
+	bool bCanSkipConstructionScripts = false;
+	if (const UBlueprint* Blueprint = UBlueprint::GetBlueprintFromClass(GetClass()))
+	{
+		bCanSkipConstructionScripts = !Blueprint->bRunConstructionScriptOnDrag;
+	}
+	
+	if (bCanSkipConstructionScripts && PropertyChangedEvent.ChangeType == EPropertyChangeType::Interactive && !CurrentTransactionAnnotation.IsValid())
 	{
 		// Avoid calling construction scripts when the change occurs while the user is dragging a slider.
 		SuperCallWithoutConstructionScripts();
