nDisplay: Fix crash when overscan sum gave zero width or height.

Also increased overscan limit from 100% to 500% to account for certain use cases.
One use case is when a 2nd camera is shooting inside the inner of the 1st camera, and a large overscan is desired in a particular direction.

#rnx

diff --git a/Engine/Plugins/Runtime/nDisplay/Source/DisplayCluster/Public/Render/Viewport/Containers/DisplayClusterViewport_OverscanSettings.h b/Engine/Plugins/Runtime/nDisplay/Source/DisplayCluster/Public/Render/Viewport/Containers/DisplayClusterViewport_OverscanSettings.h
--- a/Engine/Plugins/Runtime/nDisplay/Source/DisplayCluster/Public/Render/Viewport/Containers/DisplayClusterViewport_OverscanSettings.h
+++ b/Engine/Plugins/Runtime/nDisplay/Source/DisplayCluster/Public/Render/Viewport/Containers/DisplayClusterViewport_OverscanSettings.h
@@ -45,8 +45,7 @@ struct FDisplayClusterViewport_OverscanSettings
 
 	static float ClampPercent(float InValue)
 	{
-		// change max value 1 to lower
-		static const float MaxCustomFrustumValue = 1.f;
+		static const float MaxCustomFrustumValue = 5.f;
 
 		return FMath::Clamp(InValue, -MaxCustomFrustumValue, MaxCustomFrustumValue);
 	}
diff --git a/Engine/Plugins/Runtime/nDisplay/Source/DisplayClusterConfiguration/Public/DisplayClusterConfigurationTypes_ICVFX.h b/Engine/Plugins/Runtime/nDisplay/Source/DisplayClusterConfiguration/Public/DisplayClusterConfigurationTypes_ICVFX.h
--- a/Engine/Plugins/Runtime/nDisplay/Source/DisplayClusterConfiguration/Public/DisplayClusterConfigurationTypes_ICVFX.h
+++ b/Engine/Plugins/Runtime/nDisplay/Source/DisplayClusterConfiguration/Public/DisplayClusterConfigurationTypes_ICVFX.h
@@ -435,19 +435,19 @@ struct DISPLAYCLUSTERCONFIGURATION_API FDisplayClusterConfigurationICVFX_CameraC
 	EDisplayClusterConfigurationViewportCustomFrustumMode Mode = EDisplayClusterConfigurationViewportCustomFrustumMode::Percent;
 
 	/** Pixel/Percent value to alter the frustum to the left side */
-	UPROPERTY(EditAnywhere, BlueprintReadWrite, Category = NDisplay, meta = (DisplayName = "Left", ClampMin = "-100.0", UIMin = "-100.0", ClampMax = "100.0", UIMax = "100.0", EditCondition = "bEnable"))
+	UPROPERTY(EditAnywhere, BlueprintReadWrite, Category = NDisplay, meta = (DisplayName = "Left", ClampMin = "-500.0", UIMin = "-500.0", ClampMax = "500.0", UIMax = "500.0", EditCondition = "bEnable"))
 	float Left = 0;
 
 	/** Pixel/Percent value to alter the frustum to the right side */
-	UPROPERTY(EditAnywhere, BlueprintReadWrite, Category = NDisplay, meta = (DisplayName = "Right", ClampMin = "-100.0", UIMin = "-100.0", ClampMax = "100.0", UIMax = "100.0", EditCondition = "bEnable"))
+	UPROPERTY(EditAnywhere, BlueprintReadWrite, Category = NDisplay, meta = (DisplayName = "Right", ClampMin = "-500.0", UIMin = "-500.0", ClampMax = "500.0", UIMax = "500.0", EditCondition = "bEnable"))
 	float Right = 0;
 
 	/** Pixel/Percent value to alter the frustum to the top*/
-	UPROPERTY(EditAnywhere, BlueprintReadWrite, Category = NDisplay, meta = (DisplayName = "Top", ClampMin = "-100.0", UIMin = "-100.0", ClampMax = "100.0", UIMax = "100.0", EditCondition = "bEnable"))
+	UPROPERTY(EditAnywhere, BlueprintReadWrite, Category = NDisplay, meta = (DisplayName = "Top", ClampMin = "-500.0", UIMin = "-500.0", ClampMax = "500.0", UIMax = "500.0", EditCondition = "bEnable"))
 	float Top = 0;
 
 	/** Pixel/Percent value to alter the frustum to the bottom */
-	UPROPERTY(EditAnywhere, BlueprintReadWrite, Category = NDisplay, meta = (DisplayName = "Bottom", ClampMin = "-100.0", UIMin = "-100.0", ClampMax = "100.0", UIMax = "100.0", EditCondition = "bEnable"))
+	UPROPERTY(EditAnywhere, BlueprintReadWrite, Category = NDisplay, meta = (DisplayName = "Bottom", ClampMin = "-500.0", UIMin = "-500.0", ClampMax = "500.0", UIMax = "500.0", EditCondition = "bEnable"))
 	float Bottom = 0;
 };
 
