nDisplay: fixed resource cleanup for the editor preview

diff --git a/Engine/Plugins/Runtime/nDisplay/Source/DisplayCluster/Private/Game/EngineClasses/Scene/DisplayClusterPreviewComponent.cpp b/Engine/Plugins/Runtime/nDisplay/Source/DisplayCluster/Private/Game/EngineClasses/Scene/DisplayClusterPreviewComponent.cpp
--- a/Engine/Plugins/Runtime/nDisplay/Source/DisplayCluster/Private/Game/EngineClasses/Scene/DisplayClusterPreviewComponent.cpp
+++ b/Engine/Plugins/Runtime/nDisplay/Source/DisplayCluster/Private/Game/EngineClasses/Scene/DisplayClusterPreviewComponent.cpp
@@ -33,35 +33,33 @@ UDisplayClusterPreviewComponent::UDisplayClusterPreviewComponent(const FObjectIn
 {
 #if WITH_EDITOR
 	static ConstructorHelpers::FObjectFinder<UMaterial> PreviewMaterialObj(TEXT("/nDisplay/Materials/Preview/M_ProjPolicyPreview"));
-
 	check(PreviewMaterialObj.Object);
 
-	bWantsInitializeComponent = true;
-
 	PreviewMaterial = PreviewMaterialObj.Object;
-	PreviewMesh = nullptr;
+
+	bWantsInitializeComponent = true;
 #endif
 }
 
 #if WITH_EDITOR
-
-const uint32 UDisplayClusterPreviewComponent::MaxRenderTargetDimension = 2048;
-
 void UDisplayClusterPreviewComponent::OnComponentCreated()
 {
 	DECLARE_SCOPE_CYCLE_COUNTER(TEXT("UDisplayClusterPreviewComponent::OnComponentCreated"), STAT_OnComponentCreated, STATGROUP_NDisplay);
 	
 	Super::OnComponentCreated();
 
-	InitializeInternals();
+	InitializePreviewMaterial();
 }
 
 void UDisplayClusterPreviewComponent::DestroyComponent(bool bPromoteChildren)
 {
 	DECLARE_SCOPE_CYCLE_COUNTER(TEXT("UDisplayClusterPreviewComponent::DestroyComponent"), STAT_DestroyComponent, STATGROUP_NDisplay);
 	
-	UpdatePreviewMesh(true);
-	RemovePreviewTexture();
+	ReleasePreviewMesh();
+	ReleasePreviewMaterial();
+
+	ReleasePreviewTexture();
+	ReleasePreviewRenderTarget();
 
 	Super::DestroyComponent(bPromoteChildren);
 }
@@ -90,7 +88,7 @@ void UDisplayClusterPreviewComponent::UpdatePreviewMeshMaterial(bool bRestoreOri
 	if (bRestoreOriginalMaterial && !bIsRootActorPreviewMesh)
 	{
 		// Forged created meshes, dont restore
-		PreviewMesh = nullptr;
+		ReleasePreviewMesh();
 	}
 
 	if (PreviewMesh)
@@ -139,9 +137,7 @@ bool UDisplayClusterPreviewComponent::UpdatePreviewMesh(bool bRestoreOriginalMat
 		// Screen components are regenerated from construction scripts, but preview components are added in dynamically. This preview component may end up
 		// pointing to invalid data on reconstruction.
 		// TODO: See if we can remove this hack
-		// 
-		//!
-		PreviewMesh = nullptr;
+		ReleasePreviewMesh();
 	}
 
 	// And search for new mesh reference
@@ -154,9 +150,7 @@ bool UDisplayClusterPreviewComponent::UpdatePreviewMesh(bool bRestoreOriginalMat
 			if (Viewport->GetProjectionPolicy()->IsConfigurationChanged(&WarpMeshSavedProjectionPolicy))
 			{
 				UpdatePreviewMeshMaterial(true);
-
-				// Forget old mesh ptr
-				PreviewMesh = nullptr;
+				ReleasePreviewMesh();
 			}
 
 			if (PreviewMesh == nullptr)
@@ -187,6 +181,13 @@ bool UDisplayClusterPreviewComponent::UpdatePreviewMesh(bool bRestoreOriginalMat
 	return false;
 }
 
+void UDisplayClusterPreviewComponent::ReleasePreviewMesh()
+{
+	// Forget old mesh with material
+	PreviewMesh = nullptr;
+	OriginalMaterial = nullptr;
+}
+
 void UDisplayClusterPreviewComponent::UpdatePreviewResources()
 {
 	DECLARE_SCOPE_CYCLE_COUNTER(TEXT("UDisplayClusterPreviewComponent::UpdatePreviewResources"), STAT_UpdatePreviewResources, STATGROUP_NDisplay);
@@ -200,14 +201,39 @@ void UDisplayClusterPreviewComponent::UpdatePreviewResources()
 	UpdatePreviewMeshMaterial();
 }
 
-void UDisplayClusterPreviewComponent::InitializeInternals()
+void UDisplayClusterPreviewComponent::UpdatePreviewMaterial()
+{
+	if (PreviewMaterialInstance != nullptr)
+	{
+		PreviewMaterialInstance->SetTextureParameterValue(TEXT("Preview"), RenderTarget);
+	}
+}
+
+void UDisplayClusterPreviewComponent::InitializePreviewMaterial()
 {
-	if (!PreviewMaterialInstance)
+	if (PreviewMaterial != nullptr && PreviewMaterialInstance == nullptr)
 	{
 		PreviewMaterialInstance = UMaterialInstanceDynamic::Create(PreviewMaterial, this);
 	}
 }
 
+void UDisplayClusterPreviewComponent::ReleasePreviewMaterial()
+{
+	if (PreviewMaterialInstance != nullptr)
+	{
+		PreviewMaterialInstance->SetTextureParameterValue(TEXT("Preview"), nullptr);
+		PreviewMaterialInstance = nullptr;
+	}
+}
+
+void UDisplayClusterPreviewComponent::ReleasePreviewRenderTarget()
+{
+	if (RenderTarget != nullptr)
+	{
+		RenderTarget = nullptr;
+	}
+}
+
 void UDisplayClusterPreviewComponent::UpdatePreviewRenderTarget()
 {
 	DECLARE_SCOPE_CYCLE_COUNTER(TEXT("UDisplayClusterPreviewComponent::UpdatePreviewRenderTarget"), STAT_UpdatePreviewRenderTarget, STATGROUP_NDisplay);
@@ -229,15 +255,8 @@ void UDisplayClusterPreviewComponent::UpdatePreviewRenderTarget()
 
 			RenderTarget->InitCustomFormat(TextureSize.X, TextureSize.Y, SceneTargetFormat, false);
 
-			if (!PreviewMaterialInstance)
-			{
-				PreviewMaterialInstance = UMaterialInstanceDynamic::Create(PreviewMaterial, this);
-			}
-
-			if (PreviewMaterialInstance && RenderTarget)
-			{
-				PreviewMaterialInstance->SetTextureParameterValue(TEXT("Preview"), RenderTarget);
-			}
+			InitializePreviewMaterial();
+			UpdatePreviewMaterial();
 		}
 		// Update exist RTT resource:
 		else
@@ -273,12 +292,6 @@ bool UDisplayClusterPreviewComponent::GetPreviewTextureSettings(FIntPoint& OutSi
 		{
 			OutSize = Contexts[0].FrameTargetRect.Size();
 
-			//! Debug purpose
-			// The int casts above can sometimes cause the OutSize to have a zero in one or both its components, which will cause crashes when
-			// creating the render target on the preview component. Clamp OutSize so that it always has a size of at least 1 in each coordinate
-			static const int32 MaxTextureSize = 1 << (GMaxTextureMipCount - 1);
-			check(OutSize.X <= MaxTextureSize);
-			check(OutSize.Y <= MaxTextureSize);
 			check(OutSize.X > 0);
 			check(OutSize.Y > 0);
 
@@ -295,18 +308,15 @@ bool UDisplayClusterPreviewComponent::GetPreviewTextureSettings(FIntPoint& OutSi
 bool UDisplayClusterPreviewComponent::IsPreviewAvailable() const
 {
 	IDisplayClusterViewport* Viewport = GetCurrentViewport();
-	return (Viewport != nullptr) && Viewport->GetProjectionPolicy().IsValid();//! && Viewport->GetProjectionPolicy()->HasPreviewMesh();
+	return (Viewport != nullptr) && Viewport->GetProjectionPolicy().IsValid();
 }
 
-void UDisplayClusterPreviewComponent::RemovePreviewTexture()
+void UDisplayClusterPreviewComponent::ReleasePreviewTexture()
 {
-#if WITH_EDITOR
-	//! FIXme Add/remove texture for UE resource collection
-
-	//! @todo: add correct RenderTexture delete
-	//! 
-	PreviewTexture = nullptr;
-#endif
+	if (PreviewTexture != nullptr)
+	{
+		PreviewTexture = nullptr;
+	}
 }
 
 bool UDisplayClusterPreviewComponent::UpdatePreviewTexture()
@@ -328,7 +338,7 @@ bool UDisplayClusterPreviewComponent::UpdatePreviewTexture()
 			if (TestEmpty.Num() == 1 && *TestEmpty.CreateConstIterator() == EmptyColor)
 			{
 				// Check rest of the texture -- Texture is blank
-				RemovePreviewTexture();
+				ReleasePreviewTexture();
 				return false;
 			}
 		}
@@ -343,7 +353,7 @@ bool UDisplayClusterPreviewComponent::UpdatePreviewTexture()
 		if (PreviewTexture->GetSizeX() != DstSize.X || PreviewTexture->GetSizeY() != DstSize.Y || PreviewTexture->SRGB != SRGB)
 		{
 			// Size changed, re-create
-			RemovePreviewTexture();
+			ReleasePreviewTexture();
 		}
 	}
 
@@ -381,8 +391,6 @@ void UDisplayClusterPreviewComponent::HandleRenderTargetTextureDeferredUpdate()
 {
 	check(IsInGameThread());
 
-	//! @todo: integrate to configurator logic
-	//! deferred update flag
 	RenderTargetSurfaceChangedCnt = 2;
 }
 
@@ -390,7 +398,7 @@ UTexture2D* UDisplayClusterPreviewComponent::GetOrCreateRenderTexture2D()
 {
 	if (!IsPreviewAvailable())
 	{
-		RemovePreviewTexture();
+		ReleasePreviewTexture();
 	}
 	else
 	if (RenderTarget && RenderTargetSurfaceChangedCnt)
diff --git a/Engine/Plugins/Runtime/nDisplay/Source/DisplayCluster/Public/Components/DisplayClusterPreviewComponent.h b/Engine/Plugins/Runtime/nDisplay/Source/DisplayCluster/Public/Components/DisplayClusterPreviewComponent.h
--- a/Engine/Plugins/Runtime/nDisplay/Source/DisplayCluster/Public/Components/DisplayClusterPreviewComponent.h
+++ b/Engine/Plugins/Runtime/nDisplay/Source/DisplayCluster/Public/Components/DisplayClusterPreviewComponent.h
@@ -73,12 +73,17 @@ protected:
 	bool GetPreviewTextureSettings(FIntPoint& OutSize, float& OutGamma) const;
 
 	void UpdatePreviewRenderTarget();
+	void ReleasePreviewRenderTarget();
+
 	bool UpdatePreviewMesh(bool bRestoreOriginalMaterial = false);
+	void ReleasePreviewMesh();
 
 	bool UpdatePreviewTexture();
-	void RemovePreviewTexture();
+	void ReleasePreviewTexture();
 
-	void InitializeInternals();
+	void InitializePreviewMaterial();
+	void ReleasePreviewMaterial();
+	void UpdatePreviewMaterial();
 
 	void UpdatePreviewMeshMaterial(bool bRestoreOriginalMaterial = false);
 
@@ -86,11 +91,6 @@ protected:
 	// Set to true, when RTT surface updated
 	int32 RenderTargetSurfaceChangedCnt = 0;
 	bool bIsEditingProperty = false;
-
-private:
-	// The maximum dimension that the preview render target can be, which helps ensure good performance even if viewports are very large
-	static const uint32 MaxRenderTargetDimension;
-
 #endif /* WITH_EDITOR */
 
 #if WITH_EDITORONLY_DATA
@@ -128,7 +128,7 @@ private:
 	UMaterialInstanceDynamic* PreviewMaterialInstance = nullptr;
 
 	UPROPERTY(Transient)
-	UTexture2D* PreviewTexture;
+	UTexture2D* PreviewTexture = nullptr;
 
 #endif /*WITH_EDITORONLY_DATA*/
 };
