nDisplay: Fixed not having 'BlueprintReadWrite' specifier to all the DCRA preview properties.

#rnx

diff --git a/Engine/Plugins/Runtime/nDisplay/Source/DisplayCluster/Public/DisplayClusterRootActor.h b/Engine/Plugins/Runtime/nDisplay/Source/DisplayCluster/Public/DisplayClusterRootActor.h
--- a/Engine/Plugins/Runtime/nDisplay/Source/DisplayCluster/Public/DisplayClusterRootActor.h
+++ b/Engine/Plugins/Runtime/nDisplay/Source/DisplayCluster/Public/DisplayClusterRootActor.h
@@ -285,31 +285,31 @@ public:
 #if WITH_EDITORONLY_DATA
 public:
 	/** Render the scene and display it as a preview on the nDisplay root actor in the editor.  This will impact editor performance. */
-	UPROPERTY(EditAnywhere, Category = "Editor Preview", meta = (DisplayName = "Enable Editor Preview"))
+	UPROPERTY(EditAnywhere, BlueprintReadWrite, Category = "Editor Preview", meta = (DisplayName = "Enable Editor Preview"))
 	bool bPreviewEnable = true;
 	
 	/** Selectively preview a specific viewport or show all/none. */
-	UPROPERTY(EditAnywhere, Category = "Editor Preview", meta = (DisplayName = "Preview Node", EditCondition = "bPreviewEnable"))
+	UPROPERTY(EditAnywhere, BlueprintReadWrite, Category = "Editor Preview", meta = (DisplayName = "Preview Node", EditCondition = "bPreviewEnable"))
 	FString PreviewNodeId = DisplayClusterConfigurationStrings::gui::preview::PreviewNodeAll;
 
 	/** Render Mode */
-	UPROPERTY(EditAnywhere, Category = "Editor Preview", meta = (EditCondition = "bPreviewEnable"))
+	UPROPERTY(EditAnywhere, BlueprintReadWrite, Category = "Editor Preview", meta = (EditCondition = "bPreviewEnable"))
 	EDisplayClusterConfigurationRenderMode RenderMode = EDisplayClusterConfigurationRenderMode::Mono;
 
 	/** Tick Per Frame */
-	UPROPERTY(EditAnywhere, Category = "Editor Preview", meta = (ClampMin = "1", UIMin = "1", ClampMax = "200", UIMax = "200", EditCondition = "bPreviewEnable"))
+	UPROPERTY(EditAnywhere, BlueprintReadWrite, Category = "Editor Preview", meta = (ClampMin = "1", UIMin = "1", ClampMax = "200", UIMax = "200", EditCondition = "bPreviewEnable"))
 	int TickPerFrame = 1;
 
 	/** Adjust resolution scaling for the editor preview. */
-	UPROPERTY(EditAnywhere, Category = "Editor Preview", meta = (DisplayName = "Preview Screen Percentage", ClampMin = "0.05", UIMin = "0.05", ClampMax = "1", UIMax = "1", EditCondition = "bPreviewEnable"))
+	UPROPERTY(EditAnywhere, BlueprintReadWrite, Category = "Editor Preview", meta = (DisplayName = "Preview Screen Percentage", ClampMin = "0.05", UIMin = "0.05", ClampMax = "1", UIMax = "1", EditCondition = "bPreviewEnable"))
 	float PreviewRenderTargetRatioMult = 0.25;
 
 	/** Render ICVFX Frustums */
-	UPROPERTY(EditAnywhere, Category = "Editor Preview", meta = (DisplayName = "ICVFX Camera Frustums", EditCondition = "bPreviewEnable"))
+	UPROPERTY(EditAnywhere, BlueprintReadWrite, Category = "Editor Preview", meta = (DisplayName = "ICVFX Camera Frustums", EditCondition = "bPreviewEnable"))
 	bool bPreviewICVFXFrustums = false;
 
 	/** Render ICVFX Frustums */
-	UPROPERTY(EditAnywhere, Category = "Editor Preview", meta = (DisplayName = "ICVFX Camera Frustums Distance", EditCondition = "bPreviewEnable"))
+	UPROPERTY(EditAnywhere, BlueprintReadWrite, Category = "Editor Preview", meta = (DisplayName = "ICVFX Camera Frustums Distance", EditCondition = "bPreviewEnable"))
 	float PreviewICVFXFrustumsFarDistance = 1000.0f;
 
 private:
