nDisplay: Re-enabled Nanite for lightcard and chroma key renders.

They had been previously disabled as a workaround to Nanite not abiding by show/hidden lists. But now they do, so there is no need to disable Nanite, which would allow Nanite meshes to be used as lightcards or chroma key objects.

#rnx

diff --git a/Engine/Plugins/Runtime/nDisplay/Source/DisplayCluster/Private/Render/Viewport/DisplayClusterViewportManager.cpp b/Engine/Plugins/Runtime/nDisplay/Source/DisplayCluster/Private/Render/Viewport/DisplayClusterViewportManager.cpp
--- a/Engine/Plugins/Runtime/nDisplay/Source/DisplayCluster/Private/Render/Viewport/DisplayClusterViewportManager.cpp
+++ b/Engine/Plugins/Runtime/nDisplay/Source/DisplayCluster/Private/Render/Viewport/DisplayClusterViewportManager.cpp
@@ -470,8 +470,6 @@ FSceneViewFamily::ConstructionValues FDisplayClusterViewportManager::CreateViewF
 			InEngineShowFlags.SetHMDDistortion(0);
 			InEngineShowFlags.SetOnScreenDebug(0);
 
-			InEngineShowFlags.SetNaniteMeshes(0);
-
 			InEngineShowFlags.SetLumenReflections(0);
 			InEngineShowFlags.SetLumenGlobalIllumination(0);
 
