nDisplay: fixed 'simple' projection policy math
- compute frustum data in local space to prevent potential precision errors
- use double float according to LWC update

diff --git a/Engine/Plugins/Runtime/nDisplay/Source/DisplayClusterProjection/Private/Policy/Simple/DisplayClusterProjectionSimplePolicy.cpp b/Engine/Plugins/Runtime/nDisplay/Source/DisplayClusterProjection/Private/Policy/Simple/DisplayClusterProjectionSimplePolicy.cpp
--- a/Engine/Plugins/Runtime/nDisplay/Source/DisplayClusterProjection/Private/Policy/Simple/DisplayClusterProjectionSimplePolicy.cpp
+++ b/Engine/Plugins/Runtime/nDisplay/Source/DisplayClusterProjection/Private/Policy/Simple/DisplayClusterProjectionSimplePolicy.cpp
@@ -113,8 +113,8 @@ bool FDisplayClusterProjectionSimplePolicy::GetProjectionMatrix(class IDisplayCl
 		return false;
 	}
 
-	const float n = ViewData[InContextNum].NCP;
-	const float f = ViewData[InContextNum].FCP;
+	const double n = ViewData[InContextNum].NCP;
+	const double f = ViewData[InContextNum].FCP;
 
 	// Half-size
 	//@note The original code that takes the world scale into account needs
@@ -123,26 +123,23 @@ bool FDisplayClusterProjectionSimplePolicy::GetProjectionMatrix(class IDisplayCl
 	// Currently it's not working properly so I commented it out. Once we decide what to
 	// do with the world scale, I will get it back or remove completely.
 #if 0
-	const float hw = ScreenComp->GetScreenSize().X / 2.f * ViewData[InContextNum].WorldToMeters;
-	const float hh = ScreenComp->GetScreenSize().Y / 2.f * ViewData[InContextNum].WorldToMeters;
+	const double hw = ScreenComp->GetScreenSize().X / 2.f * ViewData[InContextNum].WorldToMeters;
+	const double hh = ScreenComp->GetScreenSize().Y / 2.f * ViewData[InContextNum].WorldToMeters;
 #else
-	const float hw = ScreenComp->GetScreenSize().X / 2.f;
-	const float hh = ScreenComp->GetScreenSize().Y / 2.f;
+	const double hw = ScreenComp->GetScreenSize().X / 2.f;
+	const double hh = ScreenComp->GetScreenSize().Y / 2.f;
 #endif
 
-	float lhw = hw;
-	float rhw = hw;
-	float thh = hh;
-	float bhh = hh;
-
 	// Screen data
-	const FVector  ScreenLoc = ScreenComp->GetComponentLocation();
-	const FRotator ScreenRot = ScreenComp->GetComponentRotation();
+	const AActor* const Owner   = ScreenComp->GetOwner();
+	const FTransform LocalSpace = (Owner ? Owner->GetActorTransform() : FTransform::Identity);
+	const FVector    ScreenLoc  = LocalSpace.InverseTransformPosition( ScreenComp->GetComponentLocation() );
+	const FRotator   ScreenRot  = LocalSpace.InverseTransformRotation( ScreenComp->GetComponentRotation().Quaternion() ).Rotator();
 
 	// Screen corners
-	const FVector pa = ScreenLoc + ScreenRot.Quaternion().RotateVector(GetProjectionScreenGeometryLBC(lhw, bhh)); // left bottom corner
-	const FVector pb = ScreenLoc + ScreenRot.Quaternion().RotateVector(GetProjectionScreenGeometryRBC(rhw, bhh)); // right bottom corner
-	const FVector pc = ScreenLoc + ScreenRot.Quaternion().RotateVector(GetProjectionScreenGeometryLTC(lhw, thh)); // left top corner
+	const FVector pa = ScreenLoc + ScreenRot.Quaternion().RotateVector(GetProjectionScreenGeometryLBC(hw, hh)); // left bottom corner
+	const FVector pb = ScreenLoc + ScreenRot.Quaternion().RotateVector(GetProjectionScreenGeometryRBC(hw, hh)); // right bottom corner
+	const FVector pc = ScreenLoc + ScreenRot.Quaternion().RotateVector(GetProjectionScreenGeometryLTC(hw, hh)); // left top corner
 
 	// Screen vectors
 	FVector vr = pb - pa; // lb->rb normalized vector, right axis of projection screen
@@ -152,23 +149,23 @@ bool FDisplayClusterProjectionSimplePolicy::GetProjectionMatrix(class IDisplayCl
 	FVector vn = -FVector::CrossProduct(vr, vu); // Projection plane normal. Use minus because of left-handed coordinate system
 	vn.Normalize();
 
-	const FVector pe = ViewData[InContextNum].ViewLoc; // ViewLocation;
+	const FVector pe = LocalSpace.InverseTransformPosition(ViewData[InContextNum].ViewLoc); // ViewLocation;
 
 	const FVector va = pa - pe; // camera -> lb
 	const FVector vb = pb - pe; // camera -> rb
 	const FVector vc = pc - pe; // camera -> lt
 
-	const float d = -FVector::DotProduct(va, vn); // distance from eye to screen
+	const double d = -FVector::DotProduct(va, vn); // distance from eye to screen
 
-	static const float minScreenDistance = 10; //Minimal distance from eye to screen
-	const float SafeDistance = (fabs(d) < minScreenDistance) ? minScreenDistance : d;
+	static const double minScreenDistance = 10; //Minimal distance from eye to screen
+	const double SafeDistance = (fabs(d) < minScreenDistance) ? minScreenDistance : d;
 
-	const float ndifd = n / SafeDistance;
+	const double ndifd = n / SafeDistance;
 
-	const float l = FVector::DotProduct(vr, va) * ndifd; // distance to left screen edge
-	const float r = FVector::DotProduct(vr, vb) * ndifd; // distance to right screen edge
-	const float b = FVector::DotProduct(vu, va) * ndifd; // distance to bottom screen edge
-	const float t = FVector::DotProduct(vu, vc) * ndifd; // distance to top screen edge
+	const double l = FVector::DotProduct(vr, va) * ndifd; // distance to left screen edge
+	const double r = FVector::DotProduct(vr, vb) * ndifd; // distance to right screen edge
+	const double b = FVector::DotProduct(vu, va) * ndifd; // distance to bottom screen edge
+	const double t = FVector::DotProduct(vu, vc) * ndifd; // distance to top screen edge
 
 	InViewport->CalculateProjectionMatrix(InContextNum, l, r, t, b, n, f, false);
 	OutPrjMatrix = InViewport->GetContexts()[InContextNum].ProjectionMatrix;
