nDisplay: Fix crash when renaming a viewport or cluster node to a previous used name. This was caused from engine changes where removing an instanced map element will rename the object to the transient package for us, causing a naming conflict.

Also fixed an issue where nDisplay's unique name creation wouldn't check the object outer. Not doing so could cause a crash when undoing the creation or rename of a node, then trying to assign another node the name of the node that no longer exists. That object would still be parented to the old outer but marked pending kill, despite proper transaction setup.

diff --git a/Engine/Plugins/Runtime/nDisplay/Source/DisplayClusterConfigurator/Private/ClusterConfiguration/DisplayClusterConfiguratorClusterUtils.cpp b/Engine/Plugins/Runtime/nDisplay/Source/DisplayClusterConfigurator/Private/ClusterConfiguration/DisplayClusterConfiguratorClusterUtils.cpp
--- a/Engine/Plugins/Runtime/nDisplay/Source/DisplayClusterConfigurator/Private/ClusterConfiguration/DisplayClusterConfiguratorClusterUtils.cpp
+++ b/Engine/Plugins/Runtime/nDisplay/Source/DisplayClusterConfigurator/Private/ClusterConfiguration/DisplayClusterConfiguratorClusterUtils.cpp
@@ -720,7 +720,7 @@ bool FDisplayClusterConfiguratorClusterUtils::RenameClusterNode(UDisplayClusterC
 			RemoveKeyFromMap(ClusterNodeParent, FieldName, *KeyPtr);
 
 			// Rename after remove, before add.
-			ClusterNode->Rename(*UniqueName, nullptr, REN_DontCreateRedirectors);
+			ClusterNode->Rename(*UniqueName, ClusterNodeParent, REN_DontCreateRedirectors);
 			
 			UDisplayClusterConfigurationClusterNode* NewClusterNode = CastChecked<UDisplayClusterConfigurationClusterNode>(AddKeyWithInstancedValueToMap(ClusterNodeParent, FieldName, UniqueName, ClusterNode));
 
@@ -895,11 +895,15 @@ bool FDisplayClusterConfiguratorClusterUtils::RenameViewport(UDisplayClusterConf
 			RemoveKeyFromMap(ViewportParent, FieldName, *KeyPtr);
 
 			// Rename after removing. If this is done after adding instances will lose sync with the CDO.
-			Viewport->Rename(*UniqueName, nullptr, REN_DontCreateRedirectors);
+			Viewport->Rename(*UniqueName, ViewportParent, REN_DontCreateRedirectors);
 			
-			AddKeyWithInstancedValueToMap(ViewportParent, FieldName, UniqueName, Viewport);
+			UDisplayClusterConfigurationViewport* NewViewport =
+				CastChecked<UDisplayClusterConfigurationViewport>(AddKeyWithInstancedValueToMap(ViewportParent, FieldName, UniqueName, Viewport));
+			check(Viewport != NewViewport);
+			Viewport->Rename(nullptr, GetTransientPackage(), REN_DontCreateRedirectors | REN_DoNotDirty | REN_ForceNoResetLoaders);
+			Viewport->SetFlags(RF_Transient);
 			
-			FDisplayClusterConfiguratorUtils::MarkDisplayClusterBlueprintAsModified(Viewport, true);
+			FDisplayClusterConfiguratorUtils::MarkDisplayClusterBlueprintAsModified(NewViewport, true);
 			
 			return true;
 		}
@@ -1207,44 +1211,46 @@ TArray<UObject*> FDisplayClusterConfiguratorClusterUtils::PasteClusterItemsFromC
 
 FString FDisplayClusterConfiguratorClusterUtils::GetUniqueName(const FString& InitialName, const TArray<FString>& UsedNames, const UClass* Class, UObject* Parent, bool bAddZero)
 {
+	FString NewName;
 	if (!bAddZero && !UsedNames.Contains(InitialName))
 	{
 		// Name doesn't need to be modified
-		return InitialName;
-	}
-	
-	int32 Counter = bAddZero ? 0 : 1;
-
-	// Find the start of the existing numeric suffix
-	int32 Index = InitialName.Len();
-	while (Index > 0 && InitialName[Index-1] >= '0' && InitialName[Index-1] <= '9')
-	{
-		--Index;
-	}
-
-	FString BaseName = InitialName;
-	FString NewName;
-	if (Index < BaseName.Len())
-	{
-		// Strip away the suffix and store the value in the counter so we can count up from there
-		FString NumericSuffix = BaseName.RightChop(Index);
-		Counter = FCString::Atoi(*NumericSuffix);
-		NumericSuffix = FString::FromInt(Counter); // Restringify the counter to account for leading 0s that we don't want to remove
-		BaseName.RemoveAt(BaseName.Len() - NumericSuffix.Len(), NumericSuffix.Len(), false);
+		NewName = InitialName;
 	}
 	else
 	{
-		// No existing suffix, so add our underscore separator
-		BaseName += "_";
-	}
+		int32 Counter = bAddZero ? 0 : 1;
 
-	do
-	{
-		NewName = FString::Printf(TEXT("%s%d"), *BaseName, Counter);
-		++Counter;
-	}
-	while (UsedNames.Contains(NewName));
+		// Find the start of the existing numeric suffix
+		int32 Index = InitialName.Len();
+		while (Index > 0 && InitialName[Index-1] >= '0' && InitialName[Index-1] <= '9')
+		{
+			--Index;
+		}
 
+		FString BaseName = InitialName;
+		if (Index < BaseName.Len())
+		{
+			// Strip away the suffix and store the value in the counter so we can count up from there
+			FString NumericSuffix = BaseName.RightChop(Index);
+			Counter = FCString::Atoi(*NumericSuffix);
+			NumericSuffix = FString::FromInt(Counter); // Restringify the counter to account for leading 0s that we don't want to remove
+			BaseName.RemoveAt(BaseName.Len() - NumericSuffix.Len(), NumericSuffix.Len(), false);
+		}
+		else
+		{
+			// No existing suffix, so add our underscore separator
+			BaseName += "_";
+		}
+
+		do
+		{
+			NewName = FString::Printf(TEXT("%s%d"), *BaseName, Counter);
+			++Counter;
+		}
+		while (UsedNames.Contains(NewName));
+	}
+	
 	// If there is already an in-memory object connected to the parent cluster with our generated name, we need to use a globally unique object name.
 	if (StaticFindObject(nullptr, Parent, *NewName, true))
 	{
