nDisplay: fixed 'simple' policy frustum math when used with non-1 scale

diff --git a/Engine/Plugins/Runtime/nDisplay/Source/DisplayClusterProjection/Private/Policy/Simple/DisplayClusterProjectionSimplePolicy.cpp b/Engine/Plugins/Runtime/nDisplay/Source/DisplayClusterProjection/Private/Policy/Simple/DisplayClusterProjectionSimplePolicy.cpp
--- a/Engine/Plugins/Runtime/nDisplay/Source/DisplayClusterProjection/Private/Policy/Simple/DisplayClusterProjectionSimplePolicy.cpp
+++ b/Engine/Plugins/Runtime/nDisplay/Source/DisplayClusterProjection/Private/Policy/Simple/DisplayClusterProjectionSimplePolicy.cpp
@@ -133,7 +133,7 @@ bool FDisplayClusterProjectionSimplePolicy::GetProjectionMatrix(class IDisplayCl
 	// Screen data
 	const AActor* const Owner   = ScreenComp->GetOwner();
 	const FTransform LocalSpace = (Owner ? Owner->GetActorTransform() : FTransform::Identity);
-	const FVector    ScreenLoc  = LocalSpace.InverseTransformPosition( ScreenComp->GetComponentLocation() );
+	const FVector    ScreenLoc  = LocalSpace.InverseTransformPositionNoScale( ScreenComp->GetComponentLocation() );
 	const FRotator   ScreenRot  = LocalSpace.InverseTransformRotation( ScreenComp->GetComponentRotation().Quaternion() ).Rotator();
 
 	// Screen corners
@@ -149,7 +149,7 @@ bool FDisplayClusterProjectionSimplePolicy::GetProjectionMatrix(class IDisplayCl
 	FVector vn = -FVector::CrossProduct(vr, vu); // Projection plane normal. Use minus because of left-handed coordinate system
 	vn.Normalize();
 
-	const FVector pe = LocalSpace.InverseTransformPosition(ViewData[InContextNum].ViewLoc); // ViewLocation;
+	const FVector pe = LocalSpace.InverseTransformPositionNoScale(ViewData[InContextNum].ViewLoc); // ViewLocation;
 
 	const FVector va = pa - pe; // camera -> lb
 	const FVector vb = pb - pe; // camera -> rb
