nDisplay: Fix preview performance issue (per-viewport and per-node, multi-GPU)

Added capture every frame flag feature
Added max preview texture size feature

diff --git a/Engine/Plugins/Runtime/nDisplay/Source/DisplayCluster/Private/Game/EngineClasses/Basics/DisplayClusterViewportClient.cpp b/Engine/Plugins/Runtime/nDisplay/Source/DisplayCluster/Private/Game/EngineClasses/Basics/DisplayClusterViewportClient.cpp
--- a/Engine/Plugins/Runtime/nDisplay/Source/DisplayCluster/Private/Game/EngineClasses/Basics/DisplayClusterViewportClient.cpp
+++ b/Engine/Plugins/Runtime/nDisplay/Source/DisplayCluster/Private/Game/EngineClasses/Basics/DisplayClusterViewportClient.cpp
@@ -912,7 +912,8 @@ bool UDisplayClusterViewportClient::Draw_PIE(FViewport* InViewport, FCanvas* Sce
 		return false;
 	}
 
-	if (ViewportManager->RenderInEditor(RenderFrame, InViewport) == false)
+	bool bOutFrameRendered = false;
+	if (ViewportManager->RenderInEditor(RenderFrame, InViewport, 0, -1, bOutFrameRendered) == false)
 	{
 		return false;
 	}
diff --git a/Engine/Plugins/Runtime/nDisplay/Source/DisplayCluster/Private/Game/EngineClasses/Scene/DisplayClusterPreviewComponent.cpp b/Engine/Plugins/Runtime/nDisplay/Source/DisplayCluster/Private/Game/EngineClasses/Scene/DisplayClusterPreviewComponent.cpp
--- a/Engine/Plugins/Runtime/nDisplay/Source/DisplayCluster/Private/Game/EngineClasses/Scene/DisplayClusterPreviewComponent.cpp
+++ b/Engine/Plugins/Runtime/nDisplay/Source/DisplayCluster/Private/Game/EngineClasses/Scene/DisplayClusterPreviewComponent.cpp
@@ -74,64 +74,65 @@ IDisplayClusterViewport* UDisplayClusterPreviewComponent::GetCurrentViewport() c
 	return nullptr;
 }
 
-bool UDisplayClusterPreviewComponent::InitializePreviewComponent(ADisplayClusterRootActor* InRootActor, const FString& InViewportId, UDisplayClusterConfigurationViewport* InViewportConfig)
+bool UDisplayClusterPreviewComponent::InitializePreviewComponent(ADisplayClusterRootActor* InRootActor, const FString& InClusterNodeId, const FString& InViewportId, UDisplayClusterConfigurationViewport* InViewportConfig)
 {
 	RootActor = InRootActor;
 	ViewportId = InViewportId;
+	ClusterNodeId = InClusterNodeId;
 	ViewportConfig = InViewportConfig;
 
 	return true;
 }
 
-void UDisplayClusterPreviewComponent::UpdatePreviewMeshMaterial(bool bRestoreOriginalMaterial)
+bool UDisplayClusterPreviewComponent::IsPreviewEnabled() const
 {
-	if (bRestoreOriginalMaterial && !bIsRootActorPreviewMesh)
+	return (ViewportConfig && RootActor && RootActor->bPreviewEnable);
+}
+
+void UDisplayClusterPreviewComponent::RestorePreviewMeshMaterial()
+{
+	UpdatePreviewMeshReference();
+
+	if (!bIsRootActorPreviewMesh)
 	{
 		// Forged created meshes, dont restore
 		ReleasePreviewMesh();
 	}
 
-	if (PreviewMesh)
+	if (PreviewMesh && OriginalMaterial)
 	{
-		bool bViewportPreviewEnabled = (ViewportConfig && RootActor && RootActor->bPreviewEnable);
+		// Restore
+		PreviewMesh->SetMaterial(0, OriginalMaterial);
+		OriginalMaterial = nullptr;
+	}
+}
+
+void UDisplayClusterPreviewComponent::SetPreviewMeshMaterial()
+{
+	UpdatePreviewMeshReference();
 
-		if (bRestoreOriginalMaterial || !bViewportPreviewEnabled)
+	if (PreviewMesh)
+	{
+		// Save original material
+		if (OriginalMaterial == nullptr)
 		{
-			// Restore
-			if (OriginalMaterial)
+			UMaterialInterface* MatInterface = PreviewMesh->GetMaterial(0);
+			if (MatInterface)
 			{
-				PreviewMesh->SetMaterial(0, OriginalMaterial);
-				OriginalMaterial = nullptr;
+				OriginalMaterial = MatInterface->GetMaterial();
 			}
 		}
-		else
-		{
-			// Save original material
-			if (OriginalMaterial == nullptr)
-			{
-				// Save original mesh material
-				UMaterialInterface* MatInterface = PreviewMesh->GetMaterial(0);
-				if (MatInterface)
-				{
-					OriginalMaterial = MatInterface->GetMaterial();
-				}
-			}
 
-			// Set preview material
-			if (PreviewMaterialInstance)
-			{
-				PreviewMesh->SetMaterial(0, PreviewMaterialInstance);
-			}
+		// Set preview material
+		if (PreviewMaterialInstance)
+		{
+			PreviewMesh->SetMaterial(0, PreviewMaterialInstance);
 		}
 	}
 }
 
-bool UDisplayClusterPreviewComponent::UpdatePreviewMesh(bool bRestoreOriginalMaterial)
+void UDisplayClusterPreviewComponent::UpdatePreviewMeshReference()
 {
-	DECLARE_SCOPE_CYCLE_COUNTER(TEXT("UDisplayClusterPreviewComponent::UpdatePreviewMesh"), STAT_UpdatePreviewMesh, STATGROUP_NDisplay);
-	
-	check(IsInGameThread());
-
 	if (PreviewMesh && PreviewMesh->GetName().Find(TEXT("TRASH_")) != INDEX_NONE)
 	{
 		// Screen components are regenerated from construction scripts, but preview components are added in dynamically. This preview component may end up
@@ -139,43 +140,57 @@ bool UDisplayClusterPreviewComponent::UpdatePreviewMesh(bool bRestoreOriginalMat
 		// TODO: See if we can remove this hack
 		ReleasePreviewMesh();
 	}
+}
+
+bool UDisplayClusterPreviewComponent::UpdatePreviewMesh(bool bRestoreOriginalMaterial)
+{
+	DECLARE_SCOPE_CYCLE_COUNTER(TEXT("UDisplayClusterPreviewComponent::UpdatePreviewMesh"), STAT_UpdatePreviewMesh, STATGROUP_NDisplay);
+	
+	check(IsInGameThread());
+
+	UpdatePreviewMeshReference();
 
 	// And search for new mesh reference
 	IDisplayClusterViewport* Viewport = GetCurrentViewport();
-	if (Viewport != nullptr && Viewport->GetProjectionPolicy().IsValid())
+	if (Viewport != nullptr && Viewport->GetProjectionPolicy().IsValid() && Viewport->GetProjectionPolicy()->HasPreviewMesh())
 	{
-		if (Viewport->GetProjectionPolicy()->HasPreviewMesh())
+		// create warp mesh or update changes
+		if (Viewport->GetProjectionPolicy()->IsConfigurationChanged(&WarpMeshSavedProjectionPolicy))
 		{
-			// create warp mesh or update changes
-			if (Viewport->GetProjectionPolicy()->IsConfigurationChanged(&WarpMeshSavedProjectionPolicy))
+			if (bRestoreOriginalMaterial)
 			{
-				UpdatePreviewMeshMaterial(true);
-				ReleasePreviewMesh();
+				RestorePreviewMeshMaterial();
 			}
+			ReleasePreviewMesh();
+		}
 
-			if (PreviewMesh == nullptr)
+		if (PreviewMesh == nullptr)
+		{
+			// Get new mesh ptr
+			PreviewMesh = Viewport->GetProjectionPolicy()->GetOrCreatePreviewMeshComponent(Viewport, bIsRootActorPreviewMesh);
+
+			// Assign preview material to mesh
+			SetPreviewMeshMaterial();
+
+			if (!ensure(ViewportConfig))
 			{
-				// Get new mesh ptr
-				PreviewMesh = Viewport->GetProjectionPolicy()->GetOrCreatePreviewMeshComponent(Viewport, bIsRootActorPreviewMesh);
-
-				UpdatePreviewMeshMaterial(bRestoreOriginalMaterial);
-					
-					if (!ensure(ViewportConfig))
-					{
-						// Can be null during a reimport.
-						// @TODO reimport: See if we can avoid this during reimport.
-						return false;
-					}
-
-					// Update saved proj policy parameters
-					WarpMeshSavedProjectionPolicy = ViewportConfig->ProjectionPolicy;
-					return true;
-				}
+				// Can be null during a reimport.
+				// @TODO reimport: See if we can avoid this during reimport.
+				return false;
 			}
+
+			// Update saved proj policy parameters
+			WarpMeshSavedProjectionPolicy = ViewportConfig->ProjectionPolicy;
+			return true;
 		}
+	}
 	else
 	{
-		UpdatePreviewMeshMaterial(true);
+		if (bRestoreOriginalMaterial)
+		{
+			RestorePreviewMeshMaterial();
+		}
+		ReleasePreviewMesh();
 	}
 
 	return false;
@@ -197,8 +212,6 @@ void UDisplayClusterPreviewComponent::UpdatePreviewResources()
 		UpdatePreviewRenderTarget();
 		UpdatePreviewMesh();
 	}
-
-	UpdatePreviewMeshMaterial();
 }
 
 void UDisplayClusterPreviewComponent::UpdatePreviewMaterial()
@@ -243,16 +256,16 @@ void UDisplayClusterPreviewComponent::UpdatePreviewRenderTarget()
 
 	if (GetPreviewTextureSettings(TextureSize, TextureGamma))
 	{
-		// Create new RTT
+		// Create new RTT for preview Material in scene
 		if (RenderTarget == nullptr)
 		{
-			RenderTarget = NewObject<UTextureRenderTarget2D>(this);
-			RenderTarget->ClearColor = FLinearColor::Black;
-			RenderTarget->TargetGamma = TextureGamma;
-
 			static const TConsoleVariableData<int32>* CVarDefaultBackBufferPixelFormat = IConsoleManager::Get().FindTConsoleVariableDataInt(TEXT("r.DefaultBackBufferPixelFormat"));
 			static const EPixelFormat SceneTargetFormat = EDefaultBackBufferPixelFormat::Convert2PixelFormat(EDefaultBackBufferPixelFormat::FromInt(CVarDefaultBackBufferPixelFormat->GetValueOnGameThread()));
 
+			RenderTarget = NewObject<UTextureRenderTarget2D>(this);
+			RenderTarget->ClearColor = FLinearColor::Black;
+			// Always use linear color gamma in scene material
+			RenderTarget->TargetGamma = TextureGamma;
 			RenderTarget->InitCustomFormat(TextureSize.X, TextureSize.Y, SceneTargetFormat, false);
 
 			InitializePreviewMaterial();
@@ -370,6 +383,12 @@ bool UDisplayClusterPreviewComponent::UpdatePreviewTexture()
 		PreviewTexture->SRGB = RenderTarget->SRGB;
 	}
 
+	// Alpha channel must non-opaque for preview in OutputMapping
+	for (FColor& PixelIt : SurfData)
+	{
+		PixelIt.A = 255;
+	}
+
 	// Transfer data
 	{
 		void* TextureData = PreviewTexture->GetPlatformData()->Mips[0].BulkData.Lock(LOCK_READ_WRITE);
@@ -387,29 +406,39 @@ bool UDisplayClusterPreviewComponent::UpdatePreviewTexture()
 	return true;
 }
 
-void UDisplayClusterPreviewComponent::HandleRenderTargetTextureDeferredUpdate()
+void UDisplayClusterPreviewComponent::HandleRenderTargetTextureDeferredUpdate(const EDisplayClusterRenderFrameMode InRenderFrameMode)
 {
 	check(IsInGameThread());
 
-	RenderTargetSurfaceChangedCnt = 2;
+	switch (InRenderFrameMode)
+	{
+	case EDisplayClusterRenderFrameMode::PreviewInScene:
+		if (RenderTargetSurfaceChangedCnt < 1)
+		{
+			RenderTargetSurfaceChangedCnt = 2;
+		}
+		break;
+	default:
+		break;
+	}
 }
 
-UTexture2D* UDisplayClusterPreviewComponent::GetOrCreateRenderTexture2D()
+UTexture2D* UDisplayClusterPreviewComponent::GetOrCreateViewportPreviewTexture2D()
 {
-	if (!IsPreviewAvailable())
-	{
-		ReleasePreviewTexture();
-	}
-	else
-	if (RenderTarget && RenderTargetSurfaceChangedCnt)
+	if (IsPreviewAvailable())
 	{
-		if (--RenderTargetSurfaceChangedCnt == 0)
+		if (RenderTarget && RenderTargetSurfaceChangedCnt)
 		{
-			UpdatePreviewTexture();
+			if (--RenderTargetSurfaceChangedCnt == 0)
+			{
+				UpdatePreviewTexture();
+			}
 		}
+
+		return PreviewTexture;
 	}
 
-	return PreviewTexture;
+	return nullptr;
 }
 
 #endif
diff --git a/Engine/Plugins/Runtime/nDisplay/Source/DisplayCluster/Private/Game/EngineClasses/Scene/DisplayClusterRootActor.cpp b/Engine/Plugins/Runtime/nDisplay/Source/DisplayCluster/Private/Game/EngineClasses/Scene/DisplayClusterRootActor.cpp
--- a/Engine/Plugins/Runtime/nDisplay/Source/DisplayCluster/Private/Game/EngineClasses/Scene/DisplayClusterRootActor.cpp
+++ b/Engine/Plugins/Runtime/nDisplay/Source/DisplayCluster/Private/Game/EngineClasses/Scene/DisplayClusterRootActor.cpp
@@ -743,6 +743,11 @@ void ADisplayClusterRootActor::PostLoad()
 void ADisplayClusterRootActor::PostActorCreated()
 {
 	Super::PostActorCreated();
+
+#if WITH_EDITOR
+	PostActorCreated_Editor();
+#endif
+
 	InitializeRootActor();
 }
 
diff --git a/Engine/Plugins/Runtime/nDisplay/Source/DisplayCluster/Private/Game/EngineClasses/Scene/DisplayClusterRootActorEditor.cpp b/Engine/Plugins/Runtime/nDisplay/Source/DisplayCluster/Private/Game/EngineClasses/Scene/DisplayClusterRootActorEditor.cpp
--- a/Engine/Plugins/Runtime/nDisplay/Source/DisplayCluster/Private/Game/EngineClasses/Scene/DisplayClusterRootActorEditor.cpp
+++ b/Engine/Plugins/Runtime/nDisplay/Source/DisplayCluster/Private/Game/EngineClasses/Scene/DisplayClusterRootActorEditor.cpp
@@ -33,7 +33,6 @@
 #include "Render/Viewport/IDisplayClusterViewportManager.h"
 #include "Render/Viewport/IDisplayClusterViewport.h"
 #include "Render/Viewport/IDisplayClusterViewportProxy.h"
-#include "Render/Viewport/RenderFrame/DisplayClusterRenderFrame.h"
 
 #include "Engine/TextureRenderTarget2D.h"
 
@@ -50,10 +49,21 @@
 #include "LevelEditor.h"
 #include "EditorSupportDelegates.h"
 
+void ADisplayClusterRootActor::ResetPreviewInternals()
+{
+	PreviewRenderFrame.Reset();
+
+	TickPerFrameCounter = 0;
+	PreviewClusterNodeIndex = 0;
+	PreviewViewportIndex = 0;
+}
+
 void ADisplayClusterRootActor::Constructor_Editor()
 {
 	// Allow tick in editor for preview rendering
 	PrimaryActorTick.bStartWithTickEnabled = true;
+
+	ResetPreviewInternals();
 }
 
 void ADisplayClusterRootActor::Destructor_Editor()
@@ -64,7 +74,11 @@ void ADisplayClusterRootActor::Destructor_Editor()
 
 void ADisplayClusterRootActor::Tick_Editor(float DeltaSeconds)
 {
-	if (!IsRunningGameOrPIE() && IsPreviewEnabled())
+	if (!IsPreviewEnabled())
+	{
+		ResetPreviewInternals();
+	}
+	else if (!IsRunningGameOrPIE())
 	{
 		if (bDeferPreviewGeneration)
 		{
@@ -73,25 +87,45 @@ void ADisplayClusterRootActor::Tick_Editor(float DeltaSeconds)
 			bDeferPreviewGeneration = false;
 			UpdatePreviewComponents();
 		}
-
-		// Update preview RTTs correspond to 'TickPerFrame' value		
-		if (++TickPerFrameCounter >= TickPerFrame)
+		else
 		{
-			TickPerFrameCounter = 0;
-			// Render preview RTTs
-			RenderPreview_Editor();
+			// Update preview RTTs correspond to 'TickPerFrame' value
+			if (++TickPerFrameCounter >= TickPerFrame)
+			{
+				TickPerFrameCounter = 0;
+
+				// Render viewport for preview material RTTs
+				if (NodesPerFrame > 0)
+				{
+					// Render preview RTTs
+					ImplRenderPreviewForSceneMaterials_Editor(NodesPerFrame);
+				}
+			}
+
+			// preview frustums on each tick
+			RenderPreviewFrustums();
 		}
 	}
 }
 
+void ADisplayClusterRootActor::PostActorCreated_Editor()
+{
+	ResetPreviewInternals();
+}
+
 void ADisplayClusterRootActor::PostLoad_Editor()
 {
 	// Generating the preview on load for instances in the world can't be done on PostLoad, components may not have loaded flags present.
 	bDeferPreviewGeneration = true;
+
+	ResetPreviewInternals();
 }
 void ADisplayClusterRootActor::BeginDestroy_Editor()
 {
+	ResetPreviewInternals();
+
 	ReleasePreviewComponents();
+
 	OnPreviewGenerated.Unbind();
 	OnPreviewDestroyed.Unbind();
 	bDeferPreviewGeneration = true;
@@ -105,6 +139,8 @@ void ADisplayClusterRootActor::RerunConstructionScripts_Editor()
 	InitializeRootActor();
 
 	UpdateInnerFrustumPriority();
+
+	ResetPreviewInternals();
 }
 
 bool ADisplayClusterRootActor::IsPreviewEnabled() const
@@ -124,20 +160,12 @@ bool ADisplayClusterRootActor::IsPreviewEnabled() const
 }
 
 // Return all RTT RHI resources for preview
-void ADisplayClusterRootActor::GetPreviewRenderTargetableTextures(const TArray<FString>& InViewportNames, TArray<FTextureRHIRef>& OutTextures)
+void ADisplayClusterRootActor::GetPreviewRenderTargetableTextures(const EDisplayClusterRenderFrameMode InRenderFrameMode, const TArray<FString>& InViewportNames, TArray<FTextureRHIRef>& OutTextures)
 {
 	check(IsInGameThread());
 
 	if(IsPreviewEnabled())
 	{
-		const bool bPerformDeferredUpdate = PreviewRenderTargetUpdatesRequired > 0;
-		if (bPerformDeferredUpdate)
-		{
-			// Optimization to limit number of deferred updates. This process is expensive
-			// and only needs to happen for a few frames after the BP was modified.
-			--PreviewRenderTargetUpdatesRequired;
-		}
-
 		for(const TPair<FString, UDisplayClusterPreviewComponent*>& PreviewComponentIt : PreviewComponents)
 		{
 			if (PreviewComponentIt.Value)
@@ -146,19 +174,22 @@ void ADisplayClusterRootActor::GetPreviewRenderTargetableTextures(const TArray<F
 				if (OutTextureIndex != INDEX_NONE)
 				{
 					// Add scope for func GetRenderTargetTexture()
-					UTextureRenderTarget2D* RenderTarget2D = PreviewComponentIt.Value->GetRenderTargetTexture();
-					if (RenderTarget2D)
+					UTextureRenderTarget2D* RenderTarget2D = nullptr;
+					switch (InRenderFrameMode)
+					{
+					case EDisplayClusterRenderFrameMode::PreviewInScene:
+						RenderTarget2D = PreviewComponentIt.Value->GetRenderTargetTexture();
+						break;
+					default:
+						break;
+					}
+					
+					if (RenderTarget2D != nullptr)
 					{
 						FTextureRenderTargetResource* DstRenderTarget = RenderTarget2D->GameThread_GetRenderTargetResource();
-						if (DstRenderTarget)
+						if (DstRenderTarget != nullptr)
 						{
 							OutTextures[OutTextureIndex] = DstRenderTarget->TextureRHI;
-
-							if(bPerformDeferredUpdate)
-							{
-								// handle configurator logic: raise deferred update for preview component RTT
-								PreviewComponentIt.Value->HandleRenderTargetTextureDeferredUpdate();
-							}
 						}
 					}
 				}
@@ -250,56 +281,132 @@ IDisplayClusterViewport* ADisplayClusterRootActor::FindPreviewViewport(const FSt
 	return nullptr;
 }
 
-bool ADisplayClusterRootActor::UpdatePreviewConfiguration_Editor(bool bUpdateAllViewports)
+bool ADisplayClusterRootActor::UpdatePreviewConfiguration_Editor(const EDisplayClusterRenderFrameMode InRenderFrameMode, const FString& InClusterNodeId)
 {
-	if (ViewportManager.IsValid())
+	PreviewRenderFrame.Reset();
+
+	if (bPreviewEnable && ViewportManager.IsValid())
 	{
-		//@todo: make correct GUI implementation for preview settings UObject (special widged in configurator bp, etc)
-		// now just copy RootActor properties to UObject::PreviewSettings
-		FDisplayClusterConfigurationViewportPreview PreviewSettings;
-		if(bPreviewEnable)
-		{
-			PreviewSettings.bEnable = bPreviewEnable;
-			PreviewSettings.PreviewNodeId = bUpdateAllViewports ? DisplayClusterConfigurationStrings::gui::preview::PreviewNodeAll : PreviewNodeId;
-			PreviewSettings.TickPerFrame = TickPerFrame;
-			PreviewSettings.PreviewRenderTargetRatioMult = PreviewRenderTargetRatioMult;
+		FDisplayClusterPreviewSettings PreviewSettings;
+		PreviewSettings.PreviewRenderTargetRatioMult = PreviewRenderTargetRatioMult;
+		PreviewSettings.PreviewMaxTextureSize = PreviewMaxTextureSize;
+		PreviewSettings.bAllowMultiGPURenderingInEditor = bAllowMultiGPURendering;
 
-			return ViewportManager->UpdatePreviewConfiguration(PreviewSettings, this);
-		}
+		PreviewSettings.MinGPUIndex = FMath::Max(0, MinGPUIndex);
+		PreviewSettings.MaxGPUIndex = FMath::Max(0, MaxGPUIndex);
+
+		return ViewportManager->UpdateConfiguration(InRenderFrameMode, InClusterNodeId, this, &PreviewSettings);
 	}
 
 	return false;
 }
 
-void ADisplayClusterRootActor::RenderPreview_Editor()
+void ADisplayClusterRootActor::RenderPreviewClusterNode_Editor(const EDisplayClusterRenderFrameMode InRenderFrameMode, const FString& InClusterNodeId)
 {
-	if (UpdatePreviewConfiguration_Editor(false))
+	if (!PreviewRenderFrame.IsValid())
 	{
-		// Update all preview components resources before render
-		for (const TTuple<FString, UDisplayClusterPreviewComponent*>& PreviewKeyVal : PreviewComponents)
+		// Begin render new frame for cluster node
+		if (UpdatePreviewConfiguration_Editor(InRenderFrameMode, InClusterNodeId))
 		{
-			if (UDisplayClusterPreviewComponent* PreviewComp = PreviewKeyVal.Value)
+			// Update all preview components resources before render
+			for (const TTuple<FString, UDisplayClusterPreviewComponent*>& PreviewComponentIt : PreviewComponents)
 			{
-				PreviewComp->UpdatePreviewResources();
+				if (PreviewComponentIt.Value && PreviewComponentIt.Value->GetClusterNodeId() == InClusterNodeId)
+				{
+					PreviewComponentIt.Value->UpdatePreviewResources();
+				}
+			}
+
+			// Now always use RootActor world to preview. 
+			UWorld* PreviewWorld = GetWorld();
+			if (PreviewWorld)
+			{
+				PreviewRenderFrame = MakeUnique<FDisplayClusterRenderFrame>();
+
+				// Update preview viewports from settings
+				if (ViewportManager->BeginNewFrame(nullptr, PreviewWorld, *PreviewRenderFrame))
+				{
+					PreviewViewportIndex = 0;
+				}
 			}
 		}
+	}
+
+	if (PreviewRenderFrame.IsValid() && PreviewViewportIndex >= 0)
+	{
+		bool bFrameRendered = false;
+		ViewportManager->RenderInEditor(*PreviewRenderFrame, nullptr, PreviewViewportIndex, ViewportsPerFrame, bFrameRendered);
 
-		// Now always use RootActor world to preview. 
-		UWorld* PreviewWorld = GetWorld();
-		if (PreviewWorld)
+		PreviewViewportIndex += ViewportsPerFrame;
+
+		if (bFrameRendered)
 		{
-			// Update preview viewports from settings
-			FDisplayClusterRenderFrame PreviewRenderFrame;
-			if (ViewportManager->BeginNewFrame(nullptr, PreviewWorld, PreviewRenderFrame))
+			// Send event about cluster node RTT changed
+			for (const TPair<FString, UDisplayClusterPreviewComponent*>& PreviewComponentIt : PreviewComponents)
 			{
-				ViewportManager->RenderInEditor(PreviewRenderFrame, nullptr);
-				// Send event about RTT changed
-				OnPreviewGenerated.ExecuteIfBound();
+				if (PreviewComponentIt.Value && PreviewComponentIt.Value->GetClusterNodeId() == InClusterNodeId)
+				{
+					PreviewComponentIt.Value->HandleRenderTargetTextureDeferredUpdate(InRenderFrameMode);
+				}
 			}
+
+			PreviewViewportIndex = -1;
+			PreviewRenderFrame.Reset();
+
+			// Send event about RTT changed
+			OnPreviewGenerated.ExecuteIfBound();
 		}
+	}
+}
 
-		// preview frustums
-		RenderPreviewFrustums();
+void ADisplayClusterRootActor::ImplRenderPreviewForSceneMaterials_Editor(const int32 InNumNodesForRender)
+{
+	if (bPreviewRenderEveryFrame && PreviewClusterNodeIndex < 0)
+	{
+		// Allow preview render
+		PreviewClusterNodeIndex = 0;
+	}
+
+	if (bPreviewEnable && ViewportManager.IsValid() && PreviewClusterNodeIndex >= 0)
+	{
+		// per-node render
+		TArray<FString> ExistClusterNodesIDs;
+		CurrentConfigData->Cluster->GetNodeIds(ExistClusterNodesIDs);
+
+		// When viewports count changed, reset render cycle
+		if (PreviewClusterNodeIndex >= ExistClusterNodesIDs.Num())
+		{
+			PreviewClusterNodeIndex = 0;
+		}
+
+		int32 NumNodesForSceneMaterialPreview = FMath::Min(InNumNodesForRender, ExistClusterNodesIDs.Num());
+
+		for (int32 NodeIt = 0; NodeIt < NumNodesForSceneMaterialPreview; NodeIt++)
+		{
+			const FString& ClusterNodeId = ExistClusterNodesIDs[PreviewClusterNodeIndex];
+
+			// Render this cluster node
+			RenderPreviewClusterNode_Editor(EDisplayClusterRenderFrameMode::PreviewInScene, ClusterNodeId);
+
+			if (PreviewRenderFrame.IsValid())
+			{
+				// Cluster node rendere in progress..
+				break;
+			}
+
+			// Loop over cluster nodes
+			PreviewClusterNodeIndex++;
+			if (PreviewClusterNodeIndex >= ExistClusterNodesIDs.Num())
+			{
+				PreviewClusterNodeIndex = bPreviewRenderEveryFrame  ? 0 : -1;
+			}
+
+			if (PreviewClusterNodeIndex < 0)
+			{
+				// Frame captured. stop render
+				break;
+			}
+		}
 	}
 }
 
@@ -329,7 +436,7 @@ void ADisplayClusterRootActor::RenderPreviewFrustums()
 			}
 
 			IDisplayClusterViewport* Viewport = ViewportManager->FindViewport(ViewportConfig.Key);
-			if (Viewport != nullptr)
+			if (Viewport != nullptr && Viewport->GetContexts().Num() > 0 && Viewport->GetRenderSettings().bEnable)
 			{
 				FrustumPreviewViewports.Add(Viewport);
 			}
@@ -422,6 +529,9 @@ void ADisplayClusterRootActor::RenderPreviewFrustum(const FMatrix ProjectionMatr
 	
 	const FVector CameraDirection = ViewMatrix.GetColumn(2);
 
+	// Preview frustum vertices
+	FVector PreviewFrustumVerts[8];
+
 	// Get the 4 points of the camera frustum near plane, in world space
 	PreviewFrustumVerts[0] = ViewOrigin + CameraDirection * NearPlane + StartCameraRightOffset + StartCameraTopOffset;         // 0 Near  Top    Right
 	PreviewFrustumVerts[1] = ViewOrigin + CameraDirection * NearPlane + StartCameraRightOffset + StartCameraBottomOffset;      // 1 Near  Bottom Right
@@ -541,7 +651,10 @@ void ADisplayClusterRootActor::UpdatePreviewComponents()
 		return;
 	}
 
-	UpdatePreviewConfiguration_Editor(true);
+	// Render new preview
+	PreviewClusterNodeIndex = 0;
+
+	UpdatePreviewConfiguration_Editor(EDisplayClusterRenderFrameMode::PreviewInScene, DisplayClusterConfigurationStrings::gui::preview::PreviewNodeAll);
 
 	TArray<UDisplayClusterPreviewComponent*> IteratedPreviewComponents;
 	
@@ -578,7 +691,7 @@ void ADisplayClusterRootActor::UpdatePreviewComponents()
 						}
 
 						// Always reinitialize so changes impact the preview component.
-						PreviewComp->InitializePreviewComponent(this, Viewport.Key, Viewport.Value);
+						PreviewComp->InitializePreviewComponent(this, Node.Key, Viewport.Key, Viewport.Value);
 
 						IteratedPreviewComponents.Add(PreviewComp);
 					}
@@ -601,12 +714,6 @@ void ADisplayClusterRootActor::UpdatePreviewComponents()
 			ExistingComp->DestroyComponent();
 		}
 	}
-
-	// Update render target for output mapping.
-	// Compile - 3 frames
-	// Load    - 2 frames
-	// Modify  - 1 frame
-	PreviewRenderTargetUpdatesRequired = 3;
 }
 
 void ADisplayClusterRootActor::ReleasePreviewComponents()
diff --git a/Engine/Plugins/Runtime/nDisplay/Source/DisplayCluster/Private/Render/Viewport/Configuration/DisplayClusterViewportConfiguration.cpp b/Engine/Plugins/Runtime/nDisplay/Source/DisplayCluster/Private/Render/Viewport/Configuration/DisplayClusterViewportConfiguration.cpp
--- a/Engine/Plugins/Runtime/nDisplay/Source/DisplayCluster/Private/Render/Viewport/Configuration/DisplayClusterViewportConfiguration.cpp
+++ b/Engine/Plugins/Runtime/nDisplay/Source/DisplayCluster/Private/Render/Viewport/Configuration/DisplayClusterViewportConfiguration.cpp
@@ -65,10 +65,9 @@ ADisplayClusterRootActor* FDisplayClusterViewportConfiguration::GetRootActor() c
 	return nullptr;
 }
 
-bool FDisplayClusterViewportConfiguration::UpdateConfiguration(EDisplayClusterRenderFrameMode InRenderMode, const FString& InClusterNodeId)
+bool FDisplayClusterViewportConfiguration::ImplUpdateConfiguration(EDisplayClusterRenderFrameMode InRenderMode, const FString& InClusterNodeId, const FDisplayClusterPreviewSettings* InPreviewSettings)
 {
 	check(IsInGameThread());
-	check(InRenderMode != EDisplayClusterRenderFrameMode::PreviewMono)
 
 	ADisplayClusterRootActor* RootActor = GetRootActor();
 	if (RootActor)
@@ -76,9 +75,6 @@ bool FDisplayClusterViewportConfiguration::UpdateConfiguration(EDisplayClusterRe
 		const UDisplayClusterConfigurationData* ConfigurationData = RootActor->GetConfigData();
 		if (ConfigurationData)
 		{
-			TArray<FString> RenderNodes;
-			RenderNodes.Add(InClusterNodeId);
-
 			FDisplayClusterViewportConfigurationBase ConfigurationBase(ViewportManager, *RootActor, *ConfigurationData);
 			FDisplayClusterViewportConfigurationICVFX ConfigurationICVFX(*RootActor);
 			FDisplayClusterViewportConfigurationProjectionPolicy ConfigurationProjectionPolicy(ViewportManager, *RootActor, *ConfigurationData);
@@ -89,7 +85,24 @@ bool FDisplayClusterViewportConfiguration::UpdateConfiguration(EDisplayClusterRe
 			RenderFrameSettings.RenderMode = InRenderMode;
 			RenderFrameSettings.ClusterNodeId = InClusterNodeId;
 
-			ConfigurationBase.Update(RenderNodes);
+			if (InPreviewSettings != nullptr)
+			{
+				// Downscale resources with PreviewDownscaleRatio
+				RenderFrameSettings.PreviewRenderTargetRatioMult = InPreviewSettings->PreviewRenderTargetRatioMult;
+
+				// Limit preview textures max size
+				RenderFrameSettings.PreviewMaxTextureSize = InPreviewSettings->PreviewMaxTextureSize;
+
+				// Support mGPU for preview rendering
+				RenderFrameSettings.bAllowMultiGPURenderingInEditor = InPreviewSettings->bAllowMultiGPURenderingInEditor;
+				RenderFrameSettings.PreviewMinGPUIndex = InPreviewSettings->MinGPUIndex;
+				RenderFrameSettings.PreviewMaxGPUIndex = InPreviewSettings->MaxGPUIndex;
+
+
+				RenderFrameSettings.bIsRenderingInEditor = true;
+			}
+
+			ConfigurationBase.Update(InClusterNodeId);
 			ConfigurationICVFX.Update();
 			ConfigurationProjectionPolicy.Update();
 			ConfigurationICVFX.PostUpdate();
@@ -97,11 +110,7 @@ bool FDisplayClusterViewportConfiguration::UpdateConfiguration(EDisplayClusterRe
 			ImplUpdateConfigurationVisibility(*RootActor, *ConfigurationData);
 
 			ConfigurationBase.UpdateClusterNodePostProcess(InClusterNodeId);
-			if (!RenderFrameSettings.bIsRenderingInEditor)
-			{
-				// TextureShare not supported in Editor Preview
-				ConfigurationBase.UpdateTextureShare(InClusterNodeId);
-			}
+			ConfigurationBase.UpdateTextureShare(InClusterNodeId);
 
 			ImplPostUpdateRenderFrameConfiguration();
 
@@ -112,6 +121,42 @@ bool FDisplayClusterViewportConfiguration::UpdateConfiguration(EDisplayClusterRe
 	return false;
 }
 
+bool FDisplayClusterViewportConfiguration::UpdateConfiguration(EDisplayClusterRenderFrameMode InRenderMode, const FString& InClusterNodeId)
+{
+	return ImplUpdateConfiguration(InRenderMode, InClusterNodeId, nullptr);
+}
+
+#if WITH_EDITOR
+bool FDisplayClusterViewportConfiguration::UpdatePreviewConfiguration(EDisplayClusterRenderFrameMode InRenderMode, const FString& InClusterNodeId, const FDisplayClusterPreviewSettings& InPreviewSettings)
+{
+	if (InClusterNodeId.Equals(DisplayClusterConfigurationStrings::gui::preview::PreviewNodeAll, ESearchCase::IgnoreCase))
+	{
+		// initialize all nodes
+		ADisplayClusterRootActor* RootActor = GetRootActor();
+		if (RootActor != nullptr)
+		{
+			const UDisplayClusterConfigurationData* ConfigurationData = RootActor->GetConfigData();
+			if (ConfigurationData != nullptr && ConfigurationData->Cluster != nullptr)
+			{
+				TArray<FString> ClusterNodesIDs;
+				ConfigurationData->Cluster->GetNodeIds(ClusterNodesIDs);
+				for (const FString& ClusterNodeIdIt : ClusterNodesIDs)
+				{
+					ImplUpdateConfiguration(InRenderMode, ClusterNodeIdIt, &InPreviewSettings);
+				}
+
+				// all cluster nodes viewports updated
+				return true;
+			}
+		}
+
+		return false;
+	}
+
+	return ImplUpdateConfiguration(InRenderMode, InClusterNodeId, &InPreviewSettings);
+}
+#endif
+
 void FDisplayClusterViewportConfiguration::ImplUpdateConfigurationVisibility(ADisplayClusterRootActor& RootActor, const UDisplayClusterConfigurationData& ConfigurationData)
 {
 	// Hide root actor components for all viewports
@@ -206,57 +251,3 @@ void FDisplayClusterViewportConfiguration::ImplUpdateRenderFrameConfiguration(co
 	}
 #endif /*WITH_EDITOR*/
 }
-
-#if WITH_EDITOR
-bool FDisplayClusterViewportConfiguration::UpdatePreviewConfiguration(const FDisplayClusterConfigurationViewportPreview& InPreviewConfiguration)
-{
-	check(IsInGameThread());
-	check(InPreviewConfiguration.bEnable);
-
-	ADisplayClusterRootActor* RootActor = GetRootActor();
-	if (RootActor)
-	{
-		const UDisplayClusterConfigurationData* ConfigurationData = RootActor->GetConfigData();
-		if (ConfigurationData)
-		{
-			TArray<FString> RenderNodes;
-			if (InPreviewConfiguration.PreviewNodeId == DisplayClusterConfigurationStrings::gui::preview::PreviewNodeAll)
-			{
-				// Collect all nodes from cluster
-				for (const TPair<FString, UDisplayClusterConfigurationClusterNode*>& It : ConfigurationData->Cluster->Nodes)
-				{
-					RenderNodes.Add(It.Key);
-				}
-			}
-			else
-			{
-				RenderNodes.Add(InPreviewConfiguration.PreviewNodeId);
-			}
-
-			FDisplayClusterViewportConfigurationBase ConfigurationBase(ViewportManager, *RootActor, *ConfigurationData);
-			FDisplayClusterViewportConfigurationICVFX ConfigurationICVFX(*RootActor);
-			FDisplayClusterViewportConfigurationProjectionPolicy ConfigurationProjectionPolicy(ViewportManager, *RootActor, *ConfigurationData);
-
-			ConfigurationBase.Update(RenderNodes);
-			ConfigurationICVFX.Update();
-			ConfigurationProjectionPolicy.Update();
-			ConfigurationICVFX.PostUpdate();
-
-			ImplUpdateConfigurationVisibility(*RootActor, *ConfigurationData);
-			ImplUpdateRenderFrameConfiguration(RootActor->GetRenderFrameSettings());
-
-			// Downscale resources with PreviewDownscaleRatio
-			RenderFrameSettings.PreviewRenderTargetRatioMult = InPreviewConfiguration.PreviewRenderTargetRatioMult;
-			RenderFrameSettings.RenderMode = EDisplayClusterRenderFrameMode::PreviewMono;
-
-			RenderFrameSettings.bIsRenderingInEditor = true;
-
-			ImplPostUpdateRenderFrameConfiguration();
-
-			return true;
-		}
-	}
-
-	return false;
-}
-#endif
diff --git a/Engine/Plugins/Runtime/nDisplay/Source/DisplayCluster/Private/Render/Viewport/Configuration/DisplayClusterViewportConfiguration.h b/Engine/Plugins/Runtime/nDisplay/Source/DisplayCluster/Private/Render/Viewport/Configuration/DisplayClusterViewportConfiguration.h
--- a/Engine/Plugins/Runtime/nDisplay/Source/DisplayCluster/Private/Render/Viewport/Configuration/DisplayClusterViewportConfiguration.h
+++ b/Engine/Plugins/Runtime/nDisplay/Source/DisplayCluster/Private/Render/Viewport/Configuration/DisplayClusterViewportConfiguration.h
@@ -7,14 +7,12 @@
 #include "Misc/DisplayClusterObjectRef.h"
 #include "Render/Viewport/Containers/DisplayClusterViewport_Enums.h"
 #include "Render/Viewport/RenderFrame/DisplayClusterRenderFrameSettings.h"
-
-struct FDisplayClusterRenderFrameSettings;
+#include "Render/Viewport/Containers/DisplayClusterPreviewSettings.h"
 
 class FDisplayClusterViewportManager;
 class ADisplayClusterRootActor;
 class UDisplayClusterConfigurationData;
 struct FDisplayClusterConfigurationRenderFrame;
-struct FDisplayClusterConfigurationViewportPreview;
 
 class FDisplayClusterViewportConfiguration
 {
@@ -41,10 +39,12 @@ public:
 	bool UpdateConfiguration(EDisplayClusterRenderFrameMode InRenderMode, const FString& InClusterNodeId);
 
 #if WITH_EDITOR
-	bool UpdatePreviewConfiguration(const FDisplayClusterConfigurationViewportPreview& PreviewConfiguration);
+	bool UpdatePreviewConfiguration(EDisplayClusterRenderFrameMode InRenderMode, const FString& ClusterNodeId, const FDisplayClusterPreviewSettings& InPreviewSettings);
 #endif
 
 private:
+	bool ImplUpdateConfiguration(EDisplayClusterRenderFrameMode InRenderMode, const FString& InClusterNodeId, const FDisplayClusterPreviewSettings* InPreviewSettings);
+
 	void ImplUpdateRenderFrameConfiguration(const FDisplayClusterConfigurationRenderFrame& InRenderFrameConfiguration);
 	void ImplPostUpdateRenderFrameConfiguration();
 	void ImplUpdateConfigurationVisibility(ADisplayClusterRootActor& InRootActor, const UDisplayClusterConfigurationData& InConfigurationData);
diff --git a/Engine/Plugins/Runtime/nDisplay/Source/DisplayCluster/Private/Render/Viewport/Configuration/DisplayClusterViewportConfigurationBase.cpp b/Engine/Plugins/Runtime/nDisplay/Source/DisplayCluster/Private/Render/Viewport/Configuration/DisplayClusterViewportConfigurationBase.cpp
--- a/Engine/Plugins/Runtime/nDisplay/Source/DisplayCluster/Private/Render/Viewport/Configuration/DisplayClusterViewportConfigurationBase.cpp
+++ b/Engine/Plugins/Runtime/nDisplay/Source/DisplayCluster/Private/Render/Viewport/Configuration/DisplayClusterViewportConfigurationBase.cpp
@@ -20,41 +20,54 @@
 // FDisplayClusterViewportConfigurationBase
 ///////////////////////////////////////////////////////////////////
 
-void FDisplayClusterViewportConfigurationBase::Update(const TArray<FString>& InClusterNodeIds)
+void FDisplayClusterViewportConfigurationBase::Update(const FString& ClusterNodeId)
 {
 	TMap<FString, UDisplayClusterConfigurationViewport*> DesiredViewports;
 
-	// Get render viewports
-	for (const FString& NodeIt : InClusterNodeIds)
+	// Get render viewports for cluster node
+	const UDisplayClusterConfigurationClusterNode* ClusterNodeConfiguration = ConfigurationData.Cluster->GetNode(ClusterNodeId);
+	if (ClusterNodeConfiguration)
 	{
-		const UDisplayClusterConfigurationClusterNode* ClusterNode = ConfigurationData.Cluster->GetNode(NodeIt);
-		if (ClusterNode)
+		for (const TPair<FString, UDisplayClusterConfigurationViewport*>& ViewportIt : ClusterNodeConfiguration->Viewports)
 		{
-			for (const TPair<FString, UDisplayClusterConfigurationViewport*>& ViewportIt : ClusterNode->Viewports)
+			if (ViewportIt.Key.Len() && ViewportIt.Value)
 			{
-				if (ViewportIt.Key.Len() && ViewportIt.Value)
-				{
-					DesiredViewports.Add(ViewportIt.Key, ViewportIt.Value);
-				}
+				DesiredViewports.Add(ViewportIt.Key, ViewportIt.Value);
 			}
 		}
 	}
 
 	// Collect unused viewports and delete
 	{
+		TArray<FString> ExistClusterNodesIDs;
+		ConfigurationData.Cluster->GetNodeIds(ExistClusterNodesIDs);
+
 		TArray<FDisplayClusterViewport*> UnusedViewports;
 		for (FDisplayClusterViewport* It : ViewportManager.ImplGetViewports())
 		{
 			// ignore ICVFX internal resources
 			if ((It->GetRenderSettingsICVFX().RuntimeFlags & ViewportRuntime_InternalResource) == 0)
 			{
-				if (!DesiredViewports.Contains(It->GetId()))
+				// Only viewports from cluster node in render
+				if (It->GetClusterNodeId() == ClusterNodeId)
+				{
+					if (!DesiredViewports.Contains(It->GetId()))
+					{
+						UnusedViewports.Add(It);
+					}
+				}
+				else
 				{
-					UnusedViewports.Add(It);
+					if(ExistClusterNodesIDs.Find(It->GetClusterNodeId()) == INDEX_NONE)
+					{
+						// also remove viewports for deleted cluster nodes
+						UnusedViewports.Add(It);
+					}
 				}
 			}
 		}
 
+		// Delete unused viewports
 		for (FDisplayClusterViewport* DeleteIt : UnusedViewports)
 		{
 			ViewportManager.ImplDeleteViewport(DeleteIt);
diff --git a/Engine/Plugins/Runtime/nDisplay/Source/DisplayCluster/Private/Render/Viewport/Configuration/DisplayClusterViewportConfigurationBase.h b/Engine/Plugins/Runtime/nDisplay/Source/DisplayCluster/Private/Render/Viewport/Configuration/DisplayClusterViewportConfigurationBase.h
--- a/Engine/Plugins/Runtime/nDisplay/Source/DisplayCluster/Private/Render/Viewport/Configuration/DisplayClusterViewportConfigurationBase.h
+++ b/Engine/Plugins/Runtime/nDisplay/Source/DisplayCluster/Private/Render/Viewport/Configuration/DisplayClusterViewportConfigurationBase.h
@@ -23,7 +23,7 @@ public:
 	{}
 
 public:
-	void Update(const TArray<FString>& InClusterNodeIds);
+	void Update(const FString& ClusterNodeId);
 	void UpdateClusterNodePostProcess(const FString& ClusterNodeId);
 	void UpdateTextureShare(const FString& ClusterNodeId);
 
diff --git a/Engine/Plugins/Runtime/nDisplay/Source/DisplayCluster/Private/Render/Viewport/Configuration/DisplayClusterViewportConfigurationHelpers_ICVFX.cpp b/Engine/Plugins/Runtime/nDisplay/Source/DisplayCluster/Private/Render/Viewport/Configuration/DisplayClusterViewportConfigurationHelpers_ICVFX.cpp
--- a/Engine/Plugins/Runtime/nDisplay/Source/DisplayCluster/Private/Render/Viewport/Configuration/DisplayClusterViewportConfigurationHelpers_ICVFX.cpp
+++ b/Engine/Plugins/Runtime/nDisplay/Source/DisplayCluster/Private/Render/Viewport/Configuration/DisplayClusterViewportConfigurationHelpers_ICVFX.cpp
@@ -28,11 +28,14 @@
 #include "Render/Viewport/Containers/DisplayClusterViewport_PostRenderSettings.h"
 #include "Render/Viewport/Containers/ImplDisplayClusterViewport_CustomFrustum.h"
 
+#include "Render/Viewport/RenderFrame/DisplayClusterRenderFrameSettings.h"
+
 #include "Containers/DisplayClusterProjectionCameraPolicySettings.h"
 #include "DisplayClusterProjectionStrings.h"
 
 #include "Components/DisplayClusterICVFXCameraComponent.h"
 
+
 #include "Misc/DisplayClusterLog.h"
 
 // Initialize camera policy with camera component and settings
@@ -50,37 +53,46 @@ static bool ImplUpdateCameraProjectionSettings(TSharedPtr<IDisplayClusterProject
 }
 
 // Return unique ICVFX name
-FString ImplGetNameICVFX(const FString& InViewportId, const FString& InResourceId)
+FString ImplGetNameICVFX(const FString& InClusterNodeId, const FString& InViewportId, const FString& InResourceId)
 {
-	return FString::Printf(TEXT("%s_%s_%s"), DisplayClusterViewportStrings::icvfx::prefix, *InViewportId, *InResourceId);
+	return FString::Printf(TEXT("%s_%s_%s_%s"), *InClusterNodeId, DisplayClusterViewportStrings::icvfx::prefix, *InViewportId, *InResourceId);
 }
 
 FDisplayClusterViewport* FDisplayClusterViewportConfigurationHelpers_ICVFX::ImplFindViewport(ADisplayClusterRootActor& RootActor, const FString& InViewportId, const FString& InResourceId)
 {
 	FDisplayClusterViewportManager* ViewportManager = FDisplayClusterViewportConfigurationHelpers_ICVFX::GetViewportManager(RootActor);
-	if (ViewportManager)
+	if (ViewportManager != nullptr)
 	{
-		return ViewportManager->ImplFindViewport(ImplGetNameICVFX(InViewportId, InResourceId));
+		const FDisplayClusterRenderFrameSettings& RenderFrameSettingsConstRef = ViewportManager->GetRenderFrameSettings();
+		return ViewportManager->ImplFindViewport(ImplGetNameICVFX(RenderFrameSettingsConstRef.ClusterNodeId, InViewportId, InResourceId));
 	}
 
 	return nullptr;
 }
 
-static bool ImplCreateProjectionPolicy(const FString& InViewportId, const FString& InResourceId, bool bIsCameraProjection, TSharedPtr<IDisplayClusterProjectionPolicy, ESPMode::ThreadSafe>& OutProjPolicy)
+static bool ImplCreateProjectionPolicy(ADisplayClusterRootActor& RootActor, const FString& InViewportId, const FString& InResourceId, bool bIsCameraProjection, TSharedPtr<IDisplayClusterProjectionPolicy, ESPMode::ThreadSafe>& OutProjPolicy)
 {
-	FDisplayClusterConfigurationProjection CameraProjectionPolicyConfig;
-	CameraProjectionPolicyConfig.Type = bIsCameraProjection ? DisplayClusterProjectionStrings::projection::Camera : DisplayClusterProjectionStrings::projection::Link;
+	FDisplayClusterViewportManager* ViewportManager = FDisplayClusterViewportConfigurationHelpers_ICVFX::GetViewportManager(RootActor);
+	if (ViewportManager != nullptr)
+	{
+		const FDisplayClusterRenderFrameSettings& RenderFrameSettingsConstRef = ViewportManager->GetRenderFrameSettings();
 
-	// Create projection policy for viewport
-	OutProjPolicy = FDisplayClusterViewportManager::CreateProjectionPolicy(ImplGetNameICVFX(InViewportId, InResourceId), &CameraProjectionPolicyConfig);
+		FDisplayClusterConfigurationProjection CameraProjectionPolicyConfig;
+		CameraProjectionPolicyConfig.Type = bIsCameraProjection ? DisplayClusterProjectionStrings::projection::Camera : DisplayClusterProjectionStrings::projection::Link;
 
-	if (!OutProjPolicy.IsValid())
-	{
-		UE_LOG(LogDisplayClusterViewport, Error, TEXT("ICVFX Viewport '%s': projection policy for resource '%s' not created."), *InViewportId, *InResourceId);
-		return false;
+		// Create projection policy for viewport
+		OutProjPolicy = FDisplayClusterViewportManager::CreateProjectionPolicy(ImplGetNameICVFX(RenderFrameSettingsConstRef.ClusterNodeId, InViewportId, InResourceId), &CameraProjectionPolicyConfig);
+
+		if (!OutProjPolicy.IsValid())
+		{
+			UE_LOG(LogDisplayClusterViewport, Error, TEXT("ICVFX Viewport '%s': projection policy for resource '%s' not created for node '%s'."), *InViewportId, *InResourceId, *RenderFrameSettingsConstRef.ClusterNodeId);
+			return false;
+		}
+
+		return true;
 	}
 
-	return true;
+	return false;
 }
 
 //------------------------------------------------------------------------------------------
@@ -92,10 +104,12 @@ FDisplayClusterViewport* FDisplayClusterViewportConfigurationHelpers_ICVFX::Impl
 	check(InProjectionPolicy.IsValid());
 
 	FDisplayClusterViewportManager* ViewportManager = FDisplayClusterViewportConfigurationHelpers_ICVFX::GetViewportManager(RootActor);
-	if (ViewportManager)
+	if (ViewportManager != nullptr)
 	{
+		const FDisplayClusterRenderFrameSettings& RenderFrameSettingsConstRef = ViewportManager->GetRenderFrameSettings();
+
 		// Create viewport for new projection policy
-		FDisplayClusterViewport* NewViewport = ViewportManager->ImplCreateViewport(ImplGetNameICVFX(InViewportId, InResourceId), InProjectionPolicy);
+		FDisplayClusterViewport* NewViewport = ViewportManager->ImplCreateViewport(ImplGetNameICVFX(RenderFrameSettingsConstRef.ClusterNodeId, InViewportId, InResourceId), InProjectionPolicy);
 		if (NewViewport != nullptr)
 		{
 			// Mark as internal resource
@@ -139,7 +153,7 @@ FDisplayClusterViewport* FDisplayClusterViewportConfigurationHelpers_ICVFX::GetO
 	if (CameraViewport == nullptr)
 	{
 		TSharedPtr<IDisplayClusterProjectionPolicy, ESPMode::ThreadSafe> CameraProjectionPolicy;
-		if (!ImplCreateProjectionPolicy(CameraId, DisplayClusterViewportStrings::icvfx::camera, true, CameraProjectionPolicy))
+		if (!ImplCreateProjectionPolicy(RootActor, CameraId, DisplayClusterViewportStrings::icvfx::camera, true, CameraProjectionPolicy))
 		{
 			return nullptr;
 		}
@@ -170,7 +184,7 @@ FDisplayClusterViewport* FDisplayClusterViewportConfigurationHelpers_ICVFX::GetO
 	if (ChromakeyViewport == nullptr)
 	{
 		TSharedPtr<IDisplayClusterProjectionPolicy, ESPMode::ThreadSafe> ChromakeyProjectionPolicy;
-		if (!ImplCreateProjectionPolicy(CameraId, DisplayClusterViewportStrings::icvfx::chromakey, false, ChromakeyProjectionPolicy))
+		if (!ImplCreateProjectionPolicy(RootActor, CameraId, DisplayClusterViewportStrings::icvfx::chromakey, false, ChromakeyProjectionPolicy))
 		{
 			return nullptr;
 		}
@@ -200,7 +214,7 @@ FDisplayClusterViewport* FDisplayClusterViewportConfigurationHelpers_ICVFX::GetO
 	if (LightcardViewport == nullptr)
 	{
 		TSharedPtr<IDisplayClusterProjectionPolicy, ESPMode::ThreadSafe> LightcardProjectionPolicy;
-		if (!ImplCreateProjectionPolicy(BaseViewport.GetId(), ResourceId, false, LightcardProjectionPolicy))
+		if (!ImplCreateProjectionPolicy(RootActor, BaseViewport.GetId(), ResourceId, false, LightcardProjectionPolicy))
 		{
 			return nullptr;
 		}
@@ -308,13 +322,13 @@ FDisplayClusterShaderParameters_ICVFX::FCameraSettings FDisplayClusterViewportCo
 	return Result;
 }
 
-bool FDisplayClusterViewportConfigurationHelpers_ICVFX::GetCameraContext(UDisplayClusterICVFXCameraComponent& InCameraComponent, FCameraContext_ICVFX& OutCameraContext)
+bool FDisplayClusterViewportConfigurationHelpers_ICVFX::GetCameraContext(ADisplayClusterRootActor& RootActor, UDisplayClusterICVFXCameraComponent& InCameraComponent, FCameraContext_ICVFX& OutCameraContext)
 {
 	const FDisplayClusterConfigurationICVFX_CameraSettings& CameraSettings = InCameraComponent.GetCameraSettingsICVFX();
 
 	// Create new camera projection policy for camera viewport
 	TSharedPtr<IDisplayClusterProjectionPolicy, ESPMode::ThreadSafe> CameraProjectionPolicy;
-	if (!ImplCreateProjectionPolicy(InCameraComponent.GetCameraUniqueId(), DisplayClusterViewportStrings::icvfx::camera, true, CameraProjectionPolicy))
+	if (!ImplCreateProjectionPolicy(RootActor, InCameraComponent.GetCameraUniqueId(), DisplayClusterViewportStrings::icvfx::camera, true, CameraProjectionPolicy))
 	{
 		return false;
 	}
diff --git a/Engine/Plugins/Runtime/nDisplay/Source/DisplayCluster/Private/Render/Viewport/Configuration/DisplayClusterViewportConfigurationHelpers_ICVFX.h b/Engine/Plugins/Runtime/nDisplay/Source/DisplayCluster/Private/Render/Viewport/Configuration/DisplayClusterViewportConfigurationHelpers_ICVFX.h
--- a/Engine/Plugins/Runtime/nDisplay/Source/DisplayCluster/Private/Render/Viewport/Configuration/DisplayClusterViewportConfigurationHelpers_ICVFX.h
+++ b/Engine/Plugins/Runtime/nDisplay/Source/DisplayCluster/Private/Render/Viewport/Configuration/DisplayClusterViewportConfigurationHelpers_ICVFX.h
@@ -44,7 +44,7 @@ public:
 	static void UpdateLightcardViewportSetting(FDisplayClusterViewport& DstViewport, FDisplayClusterViewport& BaseViewport, ADisplayClusterRootActor& RootActor, bool bIsOpenColorIOViewportExist);
 
 	static bool IsCameraUsed(UDisplayClusterICVFXCameraComponent& InCameraComponent);
-	static bool GetCameraContext(UDisplayClusterICVFXCameraComponent& InCameraComponent, FCameraContext_ICVFX& OutCameraContext);
+	static bool GetCameraContext(ADisplayClusterRootActor& RootActor, UDisplayClusterICVFXCameraComponent& InCameraComponent, FCameraContext_ICVFX& OutCameraContext);
 
 	static FDisplayClusterShaderParameters_ICVFX::FCameraSettings GetShaderParametersCameraSettings(const FDisplayClusterViewport& InCameraViewport, ADisplayClusterRootActor& RootActor, UDisplayClusterICVFXCameraComponent& InCameraComponent);
 	
diff --git a/Engine/Plugins/Runtime/nDisplay/Source/DisplayCluster/Private/Render/Viewport/Configuration/DisplayClusterViewportConfigurationHelpers_OpenColorIO.cpp b/Engine/Plugins/Runtime/nDisplay/Source/DisplayCluster/Private/Render/Viewport/Configuration/DisplayClusterViewportConfigurationHelpers_OpenColorIO.cpp
--- a/Engine/Plugins/Runtime/nDisplay/Source/DisplayCluster/Private/Render/Viewport/Configuration/DisplayClusterViewportConfigurationHelpers_OpenColorIO.cpp
+++ b/Engine/Plugins/Runtime/nDisplay/Source/DisplayCluster/Private/Render/Viewport/Configuration/DisplayClusterViewportConfigurationHelpers_OpenColorIO.cpp
@@ -84,25 +84,24 @@ bool FDisplayClusterViewportConfigurationHelpers_OpenColorIO::UpdateICVFXCameraV
 	{
 		const FDisplayClusterRenderFrameSettings& RenderFrameSettings = DstViewport.Owner.GetRenderFrameSettings();
 
-		const FString& ClusterNodeId = RenderFrameSettings.ClusterNodeId;
-		if (!ClusterNodeId.IsEmpty())
+		const FString& ClusterNodeId = DstViewport.GetClusterNodeId();
+		check(!ClusterNodeId.IsEmpty());
+
+		for (const FDisplayClusterConfigurationOCIOProfile& OCIOProfileIt : CameraSettings.PerNodeOCIOProfiles)
 		{
-			for (const FDisplayClusterConfigurationOCIOProfile& OCIOProfileIt : CameraSettings.PerNodeOCIOProfiles)
+			if (OCIOProfileIt.bIsEnabled)
 			{
-				if (OCIOProfileIt.bIsEnabled)
+				for (const FString& ClusterNodeIt : OCIOProfileIt.ApplyOCIOToObjects)
 				{
-					for (const FString& ClusterNodeIt : OCIOProfileIt.ApplyOCIOToObjects)
+					if (ClusterNodeId.Compare(ClusterNodeIt, ESearchCase::IgnoreCase) == 0)
 					{
-						if (ClusterNodeId.Compare(ClusterNodeIt, ESearchCase::IgnoreCase) == 0)
+						// Use cluster node OCIO
+						if (FDisplayClusterViewportConfigurationHelpers_OpenColorIO::ImplUpdate(DstViewport, OCIOProfileIt.OCIOConfiguration))
 						{
-							// Use cluster node OCIO
-							if (FDisplayClusterViewportConfigurationHelpers_OpenColorIO::ImplUpdate(DstViewport, OCIOProfileIt.OCIOConfiguration))
-							{
-								return true;
-							}
-
-							break;
+							return true;
 						}
+
+						break;
 					}
 				}
 			}
diff --git a/Engine/Plugins/Runtime/nDisplay/Source/DisplayCluster/Private/Render/Viewport/Configuration/DisplayClusterViewportConfigurationHelpers_Postprocess.cpp b/Engine/Plugins/Runtime/nDisplay/Source/DisplayCluster/Private/Render/Viewport/Configuration/DisplayClusterViewportConfigurationHelpers_Postprocess.cpp
--- a/Engine/Plugins/Runtime/nDisplay/Source/DisplayCluster/Private/Render/Viewport/Configuration/DisplayClusterViewportConfigurationHelpers_Postprocess.cpp
+++ b/Engine/Plugins/Runtime/nDisplay/Source/DisplayCluster/Private/Render/Viewport/Configuration/DisplayClusterViewportConfigurationHelpers_Postprocess.cpp
@@ -172,21 +172,20 @@ bool FDisplayClusterViewportConfigurationHelpers_Postprocess::ImplUpdateInnerFru
 	const FDisplayClusterRenderFrameSettings& RenderFrameSettings = DstViewport.Owner.GetRenderFrameSettings();	
 
 	// per node color grading first (it includes all nodes blending too)
-	const FString& ClusterNodeId = RenderFrameSettings.ClusterNodeId;
-	if (!ClusterNodeId.IsEmpty())
+	const FString& ClusterNodeId = DstViewport.GetClusterNodeId();
+	check(!ClusterNodeId.IsEmpty());
+
+	for (const FDisplayClusterConfigurationViewport_PerNodeColorGrading& ColorGradingProfileIt : CameraSettings.PerNodeColorGrading)
 	{
-		for (const FDisplayClusterConfigurationViewport_PerNodeColorGrading& ColorGradingProfileIt : CameraSettings.PerNodeColorGrading)
+		// Only allowed profiles
+		if (ColorGradingProfileIt.bIsEnabled)
 		{
-			// Only allowed profiles
-			if (ColorGradingProfileIt.bIsEnabled)
+			for (const FString& ClusterNodeIt : ColorGradingProfileIt.ApplyPostProcessToObjects)
 			{
-				for (const FString& ClusterNodeIt : ColorGradingProfileIt.ApplyPostProcessToObjects)
+				if (ClusterNodeId.Compare(ClusterNodeIt, ESearchCase::IgnoreCase) == 0)
 				{
-					if (ClusterNodeId.Compare(ClusterNodeIt, ESearchCase::IgnoreCase) == 0)
-					{
-						// Use cluster node PP
-						return ImplUpdateIncameraPerNodeColorGrading(DstViewport, RootActor, CameraSettings.AllNodesColorGrading, ColorGradingProfileIt);
-					}
+					// Use cluster node PP
+					return ImplUpdateIncameraPerNodeColorGrading(DstViewport, RootActor, CameraSettings.AllNodesColorGrading, ColorGradingProfileIt);
 				}
 			}
 		}
diff --git a/Engine/Plugins/Runtime/nDisplay/Source/DisplayCluster/Private/Render/Viewport/Configuration/DisplayClusterViewportConfigurationICVFX.cpp b/Engine/Plugins/Runtime/nDisplay/Source/DisplayCluster/Private/Render/Viewport/Configuration/DisplayClusterViewportConfigurationICVFX.cpp
--- a/Engine/Plugins/Runtime/nDisplay/Source/DisplayCluster/Private/Render/Viewport/Configuration/DisplayClusterViewportConfigurationICVFX.cpp
+++ b/Engine/Plugins/Runtime/nDisplay/Source/DisplayCluster/Private/Render/Viewport/Configuration/DisplayClusterViewportConfigurationICVFX.cpp
@@ -27,9 +27,9 @@
 class FDisplayClusterViewportConfigurationCameraViewport
 {
 public:
-	bool Initialize(UDisplayClusterICVFXCameraComponent& InCameraComponent)
+	bool Initialize(ADisplayClusterRootActor& InRootActor, UDisplayClusterICVFXCameraComponent& InCameraComponent)
 	{
-		return FDisplayClusterViewportConfigurationHelpers_ICVFX::GetCameraContext(InCameraComponent, CameraContext);
+		return FDisplayClusterViewportConfigurationHelpers_ICVFX::GetCameraContext(InRootActor, InCameraComponent, CameraContext);
 	}
 
 	bool CreateCameraViewport(ADisplayClusterRootActor& InRootActor, UDisplayClusterICVFXCameraComponent& InCameraComponent)
@@ -173,7 +173,7 @@ public:
 
 	bool Initialize()
 	{
-		return CameraConfiguration.Initialize(CameraComponent);
+		return CameraConfiguration.Initialize(RootActor, CameraComponent);
 	}
 
 	bool IsCameraProjectionVisibleOnViewport(FDisplayClusterViewport* TargetViewport)
diff --git a/Engine/Plugins/Runtime/nDisplay/Source/DisplayCluster/Private/Render/Viewport/Containers/DisplayClusterViewportRemap.cpp b/Engine/Plugins/Runtime/nDisplay/Source/DisplayCluster/Private/Render/Viewport/Containers/DisplayClusterViewportRemap.cpp
--- a/Engine/Plugins/Runtime/nDisplay/Source/DisplayCluster/Private/Render/Viewport/Containers/DisplayClusterViewportRemap.cpp
+++ b/Engine/Plugins/Runtime/nDisplay/Source/DisplayCluster/Private/Render/Viewport/Containers/DisplayClusterViewportRemap.cpp
@@ -66,10 +66,14 @@ void FDisplayClusterViewportRemap::Update(const FDisplayClusterViewport& InViewp
 {
 #if WITH_EDITOR
 	const FDisplayClusterRenderFrameSettings& RenderFrameSettings = InViewport.Owner.GetRenderFrameSettings();
-	if (RenderFrameSettings.RenderMode == EDisplayClusterRenderFrameMode::PreviewMono)
+
+	switch(RenderFrameSettings.RenderMode)
 	{
+	case EDisplayClusterRenderFrameMode::PreviewInScene:
 		// Preview in editor not support this feature
 		return;
+	default:
+		break;
 	}
 #endif
 
diff --git a/Engine/Plugins/Runtime/nDisplay/Source/DisplayCluster/Private/Render/Viewport/DisplayClusterViewport.cpp b/Engine/Plugins/Runtime/nDisplay/Source/DisplayCluster/Private/Render/Viewport/DisplayClusterViewport.cpp
--- a/Engine/Plugins/Runtime/nDisplay/Source/DisplayCluster/Private/Render/Viewport/DisplayClusterViewport.cpp
+++ b/Engine/Plugins/Runtime/nDisplay/Source/DisplayCluster/Private/Render/Viewport/DisplayClusterViewport.cpp
@@ -25,9 +25,10 @@
 ///////////////////////////////////////////////////////////////////////////////////////
 //          FDisplayClusterViewport
 ///////////////////////////////////////////////////////////////////////////////////////
-FDisplayClusterViewport::FDisplayClusterViewport(FDisplayClusterViewportManager& InOwner, const FString& InViewportId, const TSharedPtr<IDisplayClusterProjectionPolicy, ESPMode::ThreadSafe>& InProjectionPolicy)
+FDisplayClusterViewport::FDisplayClusterViewport(FDisplayClusterViewportManager& InOwner, const FString& InClusterNodeId, const FString& InViewportId, const TSharedPtr<IDisplayClusterProjectionPolicy, ESPMode::ThreadSafe>& InProjectionPolicy)
 	: UninitializedProjectionPolicy(InProjectionPolicy)
 	, ViewportId(InViewportId)
+	, ClusterNodeId(InClusterNodeId)
 	, Owner(InOwner)
 {
 	check(UninitializedProjectionPolicy.IsValid());
@@ -311,7 +312,7 @@ bool FDisplayClusterViewport::UpdateFrameContexts(const uint32 InStereoViewIndex
 		case EDisplayClusterRenderFrameMode::TopBottom:
 			AdjustRect(DesiredFrameTargetRect, 1.f, 0.5f);
 			break;
-		case EDisplayClusterRenderFrameMode::PreviewMono:
+		case EDisplayClusterRenderFrameMode::PreviewInScene:
 		{
 			// Preview downscale in range 0..1
 			float MultXY = FMath::Clamp(InFrameSettings.PreviewRenderTargetRatioMult, 0.f, 1.f);
@@ -366,6 +367,29 @@ bool FDisplayClusterViewport::UpdateFrameContexts(const uint32 InStereoViewIndex
 	DesiredContextSize.X *= ViewportContextSizeMult;
 	DesiredContextSize.Y *= ViewportContextSizeMult;
 
+#if WITH_EDITOR
+	switch (InFrameSettings.RenderMode)
+	{
+	case EDisplayClusterRenderFrameMode::PreviewInScene:
+		{
+			int32 MaxSize = FMath::Max(DesiredContextSize.X, DesiredContextSize.Y);
+			if (MaxSize > InFrameSettings.PreviewMaxTextureSize)
+			{
+				float RectScale = float(InFrameSettings.PreviewMaxTextureSize) / float(MaxSize);
+
+				const int32 MaxX = FMath::RoundHalfToEven(DesiredContextSize.X * RectScale);
+				const int32 MaxY = FMath::RoundHalfToEven(DesiredContextSize.Y * RectScale);
+
+				// Use limited texture size for preview
+				DesiredContextSize = FIntPoint(MaxX, MaxY);
+			}
+		}
+		break;
+	default:
+		break;
+	}
+#endif
+
 	FIntRect RenderTargetRect = GetValidRect(FIntRect(FIntPoint(0, 0), DesiredContextSize), TEXT("Context RenderTarget"));
 
 	// Exclude zero-size viewports from render
@@ -403,6 +427,29 @@ bool FDisplayClusterViewport::UpdateFrameContexts(const uint32 InStereoViewIndex
 		{
 			// Disable MultiGPU feature
 			Context.GPUIndex = INDEX_NONE;
+
+			if (InFrameSettings.bAllowMultiGPURenderingInEditor)
+			{
+				// Experimental: allow mGPU for preview rendering:
+				if (GNumExplicitGPUsForRendering > 1 && bDisableRender == false)
+				{
+					int32 MaxExplicitGPUIndex = GNumExplicitGPUsForRendering - 1;
+
+					int32 MinGPUIndex = FMath::Min(InFrameSettings.PreviewMinGPUIndex, MaxExplicitGPUIndex);
+					int32 MaxGPUIndex = FMath::Min(InFrameSettings.PreviewMaxGPUIndex, MaxExplicitGPUIndex);
+
+					static int32 PreviewGPUIndex = MinGPUIndex;
+					Context.GPUIndex = PreviewGPUIndex++;
+
+					if (PreviewGPUIndex > MaxGPUIndex)
+					{
+						PreviewGPUIndex = MinGPUIndex;
+					}
+
+					Context.bAllowGPUTransferOptimization = true;
+					Context.bEnabledGPUTransferLockSteps = false;
+				}
+			}
 		}
 		else
 		{
diff --git a/Engine/Plugins/Runtime/nDisplay/Source/DisplayCluster/Private/Render/Viewport/DisplayClusterViewport.h b/Engine/Plugins/Runtime/nDisplay/Source/DisplayCluster/Private/Render/Viewport/DisplayClusterViewport.h
--- a/Engine/Plugins/Runtime/nDisplay/Source/DisplayCluster/Private/Render/Viewport/DisplayClusterViewport.h
+++ b/Engine/Plugins/Runtime/nDisplay/Source/DisplayCluster/Private/Render/Viewport/DisplayClusterViewport.h
@@ -43,7 +43,7 @@ class FDisplayClusterViewport
 	: public IDisplayClusterViewport
 {
 public:
-	FDisplayClusterViewport(FDisplayClusterViewportManager& Owner, const FString& ViewportId, const TSharedPtr<IDisplayClusterProjectionPolicy, ESPMode::ThreadSafe>& InProjectionPolicy);
+	FDisplayClusterViewport(FDisplayClusterViewportManager& Owner, const FString& ClusterNodeId, const FString& ViewportId, const TSharedPtr<IDisplayClusterProjectionPolicy, ESPMode::ThreadSafe>& InProjectionPolicy);
 	
 	virtual ~FDisplayClusterViewport();
 
@@ -57,6 +57,11 @@ public:
 		return ViewportId; 
 	}
 
+	virtual FString GetClusterNodeId() const override
+	{
+		check(IsInGameThread());
+		return ClusterNodeId;
+	}
 	virtual const FDisplayClusterViewport_RenderSettings& GetRenderSettings() const override
 	{
 		check(IsInGameThread());
@@ -228,6 +233,9 @@ protected:
 	// Unique viewport name
 	const FString ViewportId;
 
+	// Owner cluster node name
+	const FString ClusterNodeId;
+
 	// Viewport render params
 	FDisplayClusterViewport_RenderSettings       RenderSettings;
 	FDisplayClusterViewport_RenderSettingsICVFX  RenderSettingsICVFX;
diff --git a/Engine/Plugins/Runtime/nDisplay/Source/DisplayCluster/Private/Render/Viewport/DisplayClusterViewportManager.cpp b/Engine/Plugins/Runtime/nDisplay/Source/DisplayCluster/Private/Render/Viewport/DisplayClusterViewportManager.cpp
--- a/Engine/Plugins/Runtime/nDisplay/Source/DisplayCluster/Private/Render/Viewport/DisplayClusterViewportManager.cpp
+++ b/Engine/Plugins/Runtime/nDisplay/Source/DisplayCluster/Private/Render/Viewport/DisplayClusterViewportManager.cpp
@@ -240,7 +240,7 @@ bool FDisplayClusterViewportManager::ShouldUseFullSizeFrameTargetableResource()
 		return true;
 	}
 
-	for (const FDisplayClusterViewport* Viewport : Viewports)
+	for (const FDisplayClusterViewport* Viewport : ClusterNodeViewports)
 	{
 		if (Viewport && Viewport->ShouldUseFullSizeFrameTargetableResource())
 		{
@@ -260,7 +260,7 @@ bool FDisplayClusterViewportManager::ShouldUseAdditionalFrameTargetableResource(
 		return true;
 	}
 
-	for (const FDisplayClusterViewport* Viewport : Viewports)
+	for (const FDisplayClusterViewport* Viewport : ClusterNodeViewports)
 	{
 		if (Viewport && Viewport->ShouldUseAdditionalFrameTargetableResource())
 		{
@@ -271,8 +271,10 @@ bool FDisplayClusterViewportManager::ShouldUseAdditionalFrameTargetableResource(
 	return false;
 }
 
-bool FDisplayClusterViewportManager::UpdateConfiguration(EDisplayClusterRenderFrameMode InRenderMode, const FString& InClusterNodeId, ADisplayClusterRootActor* InRootActorPtr)
+bool FDisplayClusterViewportManager::UpdateConfiguration(EDisplayClusterRenderFrameMode InRenderMode, const FString& InClusterNodeId, ADisplayClusterRootActor* InRootActorPtr, const FDisplayClusterPreviewSettings* InPreviewSettings)
 {
+	ClusterNodeViewports.Empty();
+
 	if (InRootActorPtr)
 	{
 		bool bIsRootActorChanged = Configuration->SetRootActor(InRootActorPtr);
@@ -283,7 +285,26 @@ bool FDisplayClusterViewportManager::UpdateConfiguration(EDisplayClusterRenderFr
 			ResetScene();
 		}
 
-		return Configuration->UpdateConfiguration(InRenderMode, InClusterNodeId);
+		if (InPreviewSettings == nullptr)
+		{
+			if (Configuration->UpdateConfiguration(InRenderMode, InClusterNodeId))
+			{
+				ImplUpdateClusterNodeViewports();
+
+				return true;
+			}
+		}
+#if WITH_EDITOR
+		else
+		{
+			if (Configuration->UpdatePreviewConfiguration(InRenderMode, InClusterNodeId, *InPreviewSettings))
+			{
+				ImplUpdateClusterNodeViewports();
+
+				return true;
+			}
+		}
+#endif
 	}
 
 	return false;
@@ -296,6 +317,51 @@ const FDisplayClusterRenderFrameSettings& FDisplayClusterViewportManager::GetRen
 	return Configuration->GetRenderFrameSettings();
 }
 
+void FDisplayClusterViewportManager::ImplUpdateClusterNodeViewports()
+{
+	const FDisplayClusterRenderFrameSettings& RenderFrameSettings = GetRenderFrameSettings();
+
+	switch (RenderFrameSettings.RenderMode)
+	{
+	case EDisplayClusterRenderFrameMode::PreviewInScene:
+		{
+			// Per-viewport preview render required special render order for viewports
+			// Sort viewports: move linked and overrided viewports to last
+			TArray<FDisplayClusterViewport*> LinkedViewports;
+			TArray<FDisplayClusterViewport*> OverridedViewports;
+			for (FDisplayClusterViewport* Viewport : Viewports)
+			{
+				if (Viewport && Viewport->GetClusterNodeId() == RenderFrameSettings.ClusterNodeId)
+				{
+					bool bHasParentViewport = Viewport->RenderSettings.GetParentViewportId().IsEmpty() == false;
+					bool bHasViewportOverride = Viewport->RenderSettings.OverrideViewportId.IsEmpty() == false;
+
+					if (bHasViewportOverride)
+					{
+						OverridedViewports.Add(Viewport);
+					}
+					else if (bHasParentViewport)
+					{
+						LinkedViewports.Add(Viewport);
+					}
+					else
+					{
+						ClusterNodeViewports.Add(Viewport);
+					}
+				}
+			}
+
+			ClusterNodeViewports.Append(LinkedViewports);
+			ClusterNodeViewports.Append(OverridedViewports);
+		}
+		break;
+	default:
+		// Use all exist viewports for render
+		ClusterNodeViewports.Append(Viewports);
+		break;
+	}
+}
+
 bool FDisplayClusterViewportManager::BeginNewFrame(FViewport* InViewport, UWorld* InWorld, FDisplayClusterRenderFrame& OutRenderFrame)
 {
 	check(IsInGameThread());
@@ -328,13 +394,15 @@ bool FDisplayClusterViewportManager::BeginNewFrame(FViewport* InViewport, UWorld
 	// Begin from 1, because INDEX_NONE use ViewState[0] in LocalPlayer.cpp:786
 	uint32 StereoViewIndex = 1;
 
+	const FDisplayClusterRenderFrameSettings& RenderFrameSettings = GetRenderFrameSettings();
+
 	// Initialize viewports from new render settings, and create new contexts, reset prev frame resources
-	for (FDisplayClusterViewport* Viewport : Viewports)
+	for (FDisplayClusterViewport* Viewport : ClusterNodeViewports)
 	{
 		// Save orig viewport contexts
 		TArray<FDisplayClusterViewport_Context> PrevContexts;
 		PrevContexts.Append(Viewport->GetContexts());
-		if (Viewport->UpdateFrameContexts(StereoViewIndex, GetRenderFrameSettings()))
+		if (Viewport->UpdateFrameContexts(StereoViewIndex, RenderFrameSettings))
 		{
 			StereoViewIndex += Viewport->Contexts.Num();
 		}
@@ -346,20 +414,20 @@ bool FDisplayClusterViewportManager::BeginNewFrame(FViewport* InViewport, UWorld
 	UpdateSceneRenderTargetSize();
 
 	// Build new frame structure
-	if (!RenderFrameManager->BuildRenderFrame(InViewport, GetRenderFrameSettings(), Viewports, OutRenderFrame))
+	if (!RenderFrameManager->BuildRenderFrame(InViewport, RenderFrameSettings, ClusterNodeViewports, OutRenderFrame))
 	{
 		return false;
 	}
 
 	// Allocate resources for frame
-	if (!RenderTargetManager->AllocateRenderFrameResources(InViewport, GetRenderFrameSettings(), Viewports, OutRenderFrame))
+	if (!RenderTargetManager->AllocateRenderFrameResources(InViewport, RenderFrameSettings, ClusterNodeViewports, OutRenderFrame))
 	{
 		return false;
 	}
 
 	const FIntPoint RenderFrameSize = OutRenderFrame.FrameRect.Size();
 
-	for (FDisplayClusterViewport* Viewport : Viewports)
+	for (FDisplayClusterViewport* Viewport : ClusterNodeViewports)
 	{
 		if (Viewport)
 		{
@@ -385,6 +453,7 @@ bool FDisplayClusterViewportManager::BeginNewFrame(FViewport* InViewport, UWorld
 void FDisplayClusterViewportManager::UpdateDesiredNumberOfViews(FDisplayClusterRenderFrame& InOutRenderFrame)
 {
 	InOutRenderFrame.DesiredNumberOfViews = 0;
+	InOutRenderFrame.ViewportsAmmount = 0;
 
 	for (FDisplayClusterRenderFrame::FFrameRenderTarget& RenderTargetIt : InOutRenderFrame.RenderTargets)
 	{
@@ -394,6 +463,8 @@ void FDisplayClusterViewportManager::UpdateDesiredNumberOfViews(FDisplayClusterR
 
 			for (FDisplayClusterRenderFrame::FFrameView& ViewIt : ViewFamilyIt.Views)
 			{
+				InOutRenderFrame.ViewportsAmmount++;
+
 				if (ViewIt.bDisableRender == false)
 				{
 					ViewFamilyIt.NumViewsForRender++;
@@ -415,7 +486,7 @@ void FDisplayClusterViewportManager::FinalizeNewFrame()
 	check(IsInGameThread());
 
 	// When all viewports processed, we remove all single frame custom postprocess
-	for (FDisplayClusterViewport* Viewport : Viewports)
+	for (FDisplayClusterViewport* Viewport : ClusterNodeViewports)
 	{
 		if (Viewport)
 		{
@@ -433,7 +504,7 @@ void FDisplayClusterViewportManager::FinalizeNewFrame()
 	ViewportManagerProxy->ImplUpdateRenderFrameSettings(GetRenderFrameSettings());
 
 	// Send updated viewports data to render thread proxy
-	ViewportManagerProxy->ImplUpdateViewports(Viewports);
+	ViewportManagerProxy->ImplUpdateViewports(ClusterNodeViewports);
 
 	// Update postprocess data from game thread
 	PostProcessManager->Tick();
@@ -607,10 +678,13 @@ FDisplayClusterViewport* FDisplayClusterViewportManager::ImplCreateViewport(cons
 {
 	check(IsInGameThread());
 
-	check(InProjectionPolicy.IsValid())
+	check(InProjectionPolicy.IsValid());
+
+	// Create viewport for cluster node used for rendering
+	const FString& ClusterNodeId = GetRenderFrameSettings().ClusterNodeId;
 
 	// Create viewport
-	FDisplayClusterViewport* NewViewport = new FDisplayClusterViewport(*this, ViewportId, InProjectionPolicy);
+	FDisplayClusterViewport* NewViewport = new FDisplayClusterViewport(*this, ClusterNodeId, ViewportId, InProjectionPolicy);
 
 	// Add viewport on gamethread
 	Viewports.Add(NewViewport);
diff --git a/Engine/Plugins/Runtime/nDisplay/Source/DisplayCluster/Private/Render/Viewport/DisplayClusterViewportManager.h b/Engine/Plugins/Runtime/nDisplay/Source/DisplayCluster/Private/Render/Viewport/DisplayClusterViewportManager.h
--- a/Engine/Plugins/Runtime/nDisplay/Source/DisplayCluster/Private/Render/Viewport/DisplayClusterViewportManager.h
+++ b/Engine/Plugins/Runtime/nDisplay/Source/DisplayCluster/Private/Render/Viewport/DisplayClusterViewportManager.h
@@ -49,7 +49,7 @@ public:
 
 	virtual bool IsSceneOpened() const override;
 
-	virtual bool UpdateConfiguration(EDisplayClusterRenderFrameMode InRenderMode, const FString& InClusterNodeId, class ADisplayClusterRootActor* InRootActorPtr) override;
+	virtual bool UpdateConfiguration(EDisplayClusterRenderFrameMode InRenderMode, const FString& InClusterNodeId, class ADisplayClusterRootActor* InRootActorPtr, const FDisplayClusterPreviewSettings* InPreviewSettings = nullptr) override;
 
 	virtual bool BeginNewFrame(FViewport* InViewport, UWorld* InWorld, FDisplayClusterRenderFrame& OutRenderFrame) override;
 	virtual void FinalizeNewFrame() override;
@@ -66,8 +66,7 @@ public:
 	virtual void RenderFrame(FViewport* InViewport) override;
 
 #if WITH_EDITOR
-	virtual bool UpdatePreviewConfiguration(const FDisplayClusterConfigurationViewportPreview& PreviewConfiguration, class ADisplayClusterRootActor* InRootActorPtr) override;
-	virtual bool RenderInEditor(class FDisplayClusterRenderFrame& InRenderFrame, FViewport* InViewport) override;
+	virtual bool RenderInEditor(class FDisplayClusterRenderFrame& InRenderFrame, FViewport* InViewport, const uint32 InFirstViewportNum, const int32 InViewportsAmmount, bool& bOutFrameRendered) override;
 	
 	void ImplUpdatePreviewRTTResources();
 #endif
@@ -129,7 +128,7 @@ private:
 	void ResetSceneRenderTargetSize();
 	void UpdateSceneRenderTargetSize();
 	void HandleViewportRTTChanges(const TArray<FDisplayClusterViewport_Context>& PrevContexts, const TArray<FDisplayClusterViewport_Context>& Contexts);
-
+	void ImplUpdateClusterNodeViewports();
 
 protected:
 	friend FDisplayClusterViewportManagerProxy;
@@ -144,7 +143,8 @@ public:
 private:
 	TUniquePtr<FDisplayClusterRenderFrameManager>  RenderFrameManager;
 
-	TArray<FDisplayClusterViewport*>      Viewports;
+	TArray<FDisplayClusterViewport*> Viewports;
+	TArray<FDisplayClusterViewport*> ClusterNodeViewports;
 
 	/** Render thread proxy manager. Deleted on render thread */
 	FDisplayClusterViewportManagerProxy* ViewportManagerProxy = nullptr;
diff --git a/Engine/Plugins/Runtime/nDisplay/Source/DisplayCluster/Private/Render/Viewport/DisplayClusterViewportManagerProxy.cpp b/Engine/Plugins/Runtime/nDisplay/Source/DisplayCluster/Private/Render/Viewport/DisplayClusterViewportManagerProxy.cpp
--- a/Engine/Plugins/Runtime/nDisplay/Source/DisplayCluster/Private/Render/Viewport/DisplayClusterViewportManagerProxy.cpp
+++ b/Engine/Plugins/Runtime/nDisplay/Source/DisplayCluster/Private/Render/Viewport/DisplayClusterViewportManagerProxy.cpp
@@ -73,6 +73,20 @@ FDisplayClusterViewportManagerProxy::~FDisplayClusterViewportManagerProxy()
 	}
 
 	ViewportProxies.Empty();
+	ImplUpdateClusterNodeViewportProxies();
+}
+
+void FDisplayClusterViewportManagerProxy::ImplUpdateClusterNodeViewportProxies()
+{
+	ClusterNodeViewportProxies.Empty();
+	// Collect viewport proxies for rendered cluster node
+	for (FDisplayClusterViewportProxy* ViewportProxyIt : ViewportProxies)
+	{
+		if (ViewportProxyIt && ViewportProxyIt->GetClusterNodeId() == RenderFrameSettings.ClusterNodeId)
+		{
+			ClusterNodeViewportProxies.Add(ViewportProxyIt);
+		}
+	}
 }
 
 void FDisplayClusterViewportManagerProxy::ImplSafeRelease()
@@ -97,6 +111,7 @@ void FDisplayClusterViewportManagerProxy::ImplCreateViewport(FDisplayClusterView
 			[ViewportManagerProxy = this, ViewportProxy = InViewportProxy](FRHICommandListImmediate& RHICmdList)
 		{
 			ViewportManagerProxy->ViewportProxies.Add(ViewportProxy);
+			ViewportManagerProxy->ImplUpdateClusterNodeViewportProxies();
 		});
 	}
 }
@@ -115,6 +130,7 @@ void FDisplayClusterViewportManagerProxy::ImplDeleteViewport(FDisplayClusterView
 		{
 			ViewportManagerProxy->ViewportProxies[ViewportProxyIndex] = nullptr;
 			ViewportManagerProxy->ViewportProxies.RemoveAt(ViewportProxyIndex);
+			ViewportManagerProxy->ImplUpdateClusterNodeViewportProxies();
 		}
 
 		delete ViewportProxy;
@@ -132,6 +148,9 @@ void FDisplayClusterViewportManagerProxy::ImplUpdateRenderFrameSettings(const FD
 	{
 		ViewportManagerProxy->RenderFrameSettings = *Settings;
 		delete Settings;
+
+		// After updated settings we need update cluster node viewports
+		ViewportManagerProxy->ImplUpdateClusterNodeViewportProxies();
 	});
 }
 
@@ -171,7 +190,7 @@ void FDisplayClusterViewportManagerProxy::ImplRenderFrame(FViewport* InViewport)
 		bool bWarpBlendEnabled = ViewportManagerProxy->RenderFrameSettings.bAllowWarpBlend && CVarWarpBlendEnabled.GetValueOnRenderThread() != 0;
 
 		// mGPU not used for in-editor rendering
-		if(ViewportManagerProxy->RenderFrameSettings.bIsRenderingInEditor == false)
+		if(ViewportManagerProxy->RenderFrameSettings.bIsRenderingInEditor == false || ViewportManagerProxy->RenderFrameSettings.bAllowMultiGPURenderingInEditor)
 		{
 			// Move all render target cross gpu
 			ViewportManagerProxy->DoCrossGPUTransfers_RenderThread(RHICmdList);
@@ -212,9 +231,9 @@ void FDisplayClusterViewportManagerProxy::UpdateDeferredResources_RenderThread(F
 	check(IsInRenderingThread());
 
 	TArray<FDisplayClusterViewportProxy*> OverridedViewports;
-	OverridedViewports.Reserve(ViewportProxies.Num());
+	OverridedViewports.Reserve(ClusterNodeViewportProxies.Num());
 
-	for (FDisplayClusterViewportProxy* ViewportProxy : ViewportProxies)
+	for (FDisplayClusterViewportProxy* ViewportProxy : ClusterNodeViewportProxies)
 	{
 		if (ViewportProxy->RenderSettings.OverrideViewportId.IsEmpty())
 		{
@@ -277,7 +296,7 @@ void FDisplayClusterViewportManagerProxy::UpdateFrameResources_RenderThread(FRHI
 	PostProcessManager->PerformPostProcessBeforeWarpBlend_RenderThread(RHICmdList, this);
 
 	// Support viewport overlap order sorting:
-	TArray<FDisplayClusterViewportProxy*> SortedViewportProxy = ImplGetViewportProxies_RenderThread();
+	TArray<FDisplayClusterViewportProxy*> SortedViewportProxy = ClusterNodeViewportProxies;
 	SortedViewportProxy.Sort(
 		[](const FDisplayClusterViewportProxy& VP1, const FDisplayClusterViewportProxy& VP2)
 		{
@@ -383,7 +402,7 @@ void FDisplayClusterViewportManagerProxy::DoCrossGPUTransfers_RenderThread(FRHIC
 	// Copy the view render results to all GPUs that are native to the viewport.
 	TArray<FTransferResourceParams> TransferResources;
 
-	for (FDisplayClusterViewportProxy* ViewportProxy : ViewportProxies)
+	for (FDisplayClusterViewportProxy* ViewportProxy : ClusterNodeViewportProxies)
 	{
 		for (FDisplayClusterViewport_Context& ViewportContext : ViewportProxy->Contexts)
 		{
@@ -439,7 +458,7 @@ bool FDisplayClusterViewportManagerProxy::GetFrameTargets_RenderThread(TArray<FR
 	check(IsInRenderingThread());
 
 	// Get any defined frame targets from first visible viewport
-	for (FDisplayClusterViewportProxy* ViewportProxy : ViewportProxies)
+	for (FDisplayClusterViewportProxy* ViewportProxy : ClusterNodeViewportProxies)
 	{
 		if (ViewportProxy)
 		{
diff --git a/Engine/Plugins/Runtime/nDisplay/Source/DisplayCluster/Private/Render/Viewport/DisplayClusterViewportManagerProxy.h b/Engine/Plugins/Runtime/nDisplay/Source/DisplayCluster/Private/Render/Viewport/DisplayClusterViewportManagerProxy.h
--- a/Engine/Plugins/Runtime/nDisplay/Source/DisplayCluster/Private/Render/Viewport/DisplayClusterViewportManagerProxy.h
+++ b/Engine/Plugins/Runtime/nDisplay/Source/DisplayCluster/Private/Render/Viewport/DisplayClusterViewportManagerProxy.h
@@ -40,7 +40,7 @@ public:
 
 	virtual const TArrayView<IDisplayClusterViewportProxy*> GetViewports_RenderThread() const override
 	{
-		return TArrayView<IDisplayClusterViewportProxy*>((IDisplayClusterViewportProxy**)(ViewportProxies.GetData()), ViewportProxies.Num());
+		return TArrayView<IDisplayClusterViewportProxy*>((IDisplayClusterViewportProxy**)(ClusterNodeViewportProxies.GetData()), ClusterNodeViewportProxies.Num());
 	}
 
 	virtual bool GetFrameTargets_RenderThread(TArray<FRHITexture2D*>& OutFrameResources, TArray<FIntPoint>& OutTargetOffsets, TArray<FRHITexture2D*>* OutAdditionalFrameResources = nullptr) const override;
@@ -58,7 +58,7 @@ public:
 	{
 		check(IsInRenderingThread());
 
-		return ViewportProxies;
+		return ClusterNodeViewportProxies;
 	}
 
 	FDisplayClusterViewportProxy* ImplFindViewport_RenderThread(const FString& InViewportId) const;
@@ -72,6 +72,9 @@ public:
 	void ImplRenderFrame(FViewport* InViewport);
 	void ImplClearFrameTargets_RenderThread(FRHICommandListImmediate& RHICmdList) const;
 
+private:
+	void ImplUpdateClusterNodeViewportProxies();
+
 private:
 	TSharedPtr<FDisplayClusterRenderTargetManager, ESPMode::ThreadSafe>        RenderTargetManager;
 	TSharedPtr<FDisplayClusterViewportPostProcessManager, ESPMode::ThreadSafe> PostProcessManager;
@@ -79,4 +82,5 @@ private:
 	FDisplayClusterRenderFrameSettings RenderFrameSettings;
 
 	TArray<FDisplayClusterViewportProxy*> ViewportProxies;
+	TArray<FDisplayClusterViewportProxy*> ClusterNodeViewportProxies;
 };
diff --git a/Engine/Plugins/Runtime/nDisplay/Source/DisplayCluster/Private/Render/Viewport/DisplayClusterViewportManager_EditorPreview.cpp b/Engine/Plugins/Runtime/nDisplay/Source/DisplayCluster/Private/Render/Viewport/DisplayClusterViewportManager_EditorPreview.cpp
--- a/Engine/Plugins/Runtime/nDisplay/Source/DisplayCluster/Private/Render/Viewport/DisplayClusterViewportManager_EditorPreview.cpp
+++ b/Engine/Plugins/Runtime/nDisplay/Source/DisplayCluster/Private/Render/Viewport/DisplayClusterViewportManager_EditorPreview.cpp
@@ -37,10 +37,12 @@ void FDisplayClusterViewportManager::ImplUpdatePreviewRTTResources()
 {
 	check(IsInGameThread());
 
+	const EDisplayClusterRenderFrameMode RenderMode = GetRenderFrameSettings().RenderMode;
+
 	// Only for preview modes:
-	switch (GetRenderFrameSettings().RenderMode)
+	switch (RenderMode)
 	{
-	case EDisplayClusterRenderFrameMode::PreviewMono:
+	case EDisplayClusterRenderFrameMode::PreviewInScene:
 		break;
 	default:
 		return;
@@ -53,6 +55,8 @@ void FDisplayClusterViewportManager::ImplUpdatePreviewRTTResources()
 		return;
 	}
 
+	const FString& ClusterNodeId = GetRenderFrameSettings().ClusterNodeId;
+
 	TArray<FString>        PreviewViewportNames;
 	TArray<FTextureRHIRef> PreviewRenderTargetableTextures;
 
@@ -62,21 +66,25 @@ void FDisplayClusterViewportManager::ImplUpdatePreviewRTTResources()
 	// Collect visible and enabled viewports for preview rendering:
 	for (FDisplayClusterViewport* const ViewportIt : Viewports)
 	{
-		if (ViewportIt->RenderSettings.bEnable && ViewportIt->RenderSettings.bVisible)
+		// update only current cluster node
+		if (ViewportIt->GetClusterNodeId() == ClusterNodeId)
 		{
-			PreviewViewportNames.Add(ViewportIt->GetId());
-			PreviewRenderTargetableTextures.AddDefaulted();
+			if (ViewportIt->RenderSettings.bEnable && ViewportIt->RenderSettings.bVisible)
+			{
+				PreviewViewportNames.Add(ViewportIt->GetId());
+				PreviewRenderTargetableTextures.AddDefaulted();
+			}
 		}
 	}
 
 	// Get all supported preview rtt resources from root actor
-	RootActor->GetPreviewRenderTargetableTextures(PreviewViewportNames, PreviewRenderTargetableTextures);
+	RootActor->GetPreviewRenderTargetableTextures(RenderMode, PreviewViewportNames, PreviewRenderTargetableTextures);
 
 	// Configure preview RTT to viwports:
 	for (int32 ViewportIndex = 0; ViewportIndex < PreviewViewportNames.Num(); ViewportIndex++)
 	{
 		FDisplayClusterViewport* DesiredViewport = ImplFindViewport(PreviewViewportNames[ViewportIndex]);
-		if (DesiredViewport)
+		if (DesiredViewport != nullptr)
 		{
 			FTextureRHIRef& PreviewRTT = PreviewRenderTargetableTextures[ViewportIndex];
 			if (PreviewRTT.IsValid())
@@ -93,14 +101,10 @@ void FDisplayClusterViewportManager::ImplUpdatePreviewRTTResources()
 	}
 }
 
-bool FDisplayClusterViewportManager::UpdatePreviewConfiguration(const FDisplayClusterConfigurationViewportPreview& PreviewConfiguration, class ADisplayClusterRootActor* InRootActorPtr)
+bool FDisplayClusterViewportManager::RenderInEditor(class FDisplayClusterRenderFrame& InRenderFrame, FViewport* InViewport, const uint32 InFirstViewportNum, const int32 InViewportsAmmount, bool& bOutFrameRendered)
 {
-	Configuration->SetRootActor(InRootActorPtr);
-	return Configuration->UpdatePreviewConfiguration(PreviewConfiguration);
-}
+	bOutFrameRendered = false;
 
-bool FDisplayClusterViewportManager::RenderInEditor(class FDisplayClusterRenderFrame& InRenderFrame, FViewport* InViewport)
-{
 	UWorld* CurrentWorld = GetCurrentWorld();
 	if (CurrentWorld == nullptr)
 	{
@@ -113,14 +117,26 @@ bool FDisplayClusterViewportManager::RenderInEditor(class FDisplayClusterRenderF
 	//Experimental code from render team, now always disabled
 	const bool bIsRenderedImmediatelyAfterAnotherViewFamily = false;
 
+	int32 ViewportIndex = 0;
+	int32 RenderedViewportsAmmount = 0;
+	bool bViewportsRenderPassDone = false;
+
 	for (FDisplayClusterRenderFrame::FFrameRenderTarget& RenderTargetIt : InRenderFrame.RenderTargets)
 	{
+		if (bViewportsRenderPassDone)
+		{
+			break;
+		}
+
 		// Special flag, allow clear RTT surface only for first family
 		bool bAdditionalViewFamily = false;
 
 		for (FDisplayClusterRenderFrame::FFrameViewFamily& ViewFamiliesIt : RenderTargetIt.ViewFamilies)
 		{
-			FRenderTarget* DstResource = RenderTargetIt.RenderTargetPtr;
+			if (bViewportsRenderPassDone)
+			{
+				break;
+			}
 
 			// Create the view family for rendering the world scene to the viewport's render target
 			FSceneViewFamilyContext ViewFamily(CreateViewFamilyConstructionValues(
@@ -132,33 +148,43 @@ bool FDisplayClusterViewportManager::RenderInEditor(class FDisplayClusterRenderF
 
 			ConfigureViewFamily(RenderTargetIt, ViewFamiliesIt, ViewFamily);
 
-			// Disable clean op for all next families on this render target
-			bAdditionalViewFamily = true;
-
 			for (FDisplayClusterRenderFrame::FFrameView& ViewIt : ViewFamiliesIt.Views)
 			{
-				FDisplayClusterViewport* ViewportPtr = static_cast<FDisplayClusterViewport*>(ViewIt.Viewport);
+				bool bViewportAlreadyRendered = ViewportIndex < (int32)InFirstViewportNum;
+				ViewportIndex++;
 
-				check(ViewportPtr != nullptr);
-				check(ViewIt.ContextNum < (uint32)ViewportPtr->Contexts.Num());
-
-				// Calculate the player's view information.
-				FVector  ViewLocation;
-				FRotator ViewRotation;
-				FSceneView* View = ViewportPtr->ImplCalcScenePreview(ViewFamily, ViewIt.ContextNum);
-
-				if (View && ViewIt.bDisableRender)
-				{
-					ViewFamily.Views.Remove(View);
-
-					delete View;
-					View = nullptr;
-				}
-
-				if (View)
+				if (!bViewportAlreadyRendered)
 				{
-					// Apply viewport context settings to view (crossGPU, visibility, etc)
-					ViewIt.Viewport->SetupSceneView(ViewIt.ContextNum, PreviewScene->GetWorld(), ViewFamily, *View);
+					FDisplayClusterViewport* ViewportPtr = static_cast<FDisplayClusterViewport*>(ViewIt.Viewport);
+
+					check(ViewportPtr != nullptr);
+					check(ViewIt.ContextNum < (uint32)ViewportPtr->Contexts.Num());
+
+					// Calculate the player's view information.
+					FVector  ViewLocation;
+					FRotator ViewRotation;
+					FSceneView* View = ViewportPtr->ImplCalcScenePreview(ViewFamily, ViewIt.ContextNum);
+
+					if (View && ViewIt.bDisableRender)
+					{
+						ViewFamily.Views.Remove(View);
+
+						delete View;
+						View = nullptr;
+					}
+
+					if (View)
+					{
+						// Apply viewport context settings to view (crossGPU, visibility, etc)
+						ViewIt.Viewport->SetupSceneView(ViewIt.ContextNum, PreviewScene->GetWorld(), ViewFamily, *View);
+
+						RenderedViewportsAmmount++;
+						if (InViewportsAmmount > 0 && RenderedViewportsAmmount >= InViewportsAmmount)
+						{
+							bViewportsRenderPassDone = true;
+							break;
+						}
+					}
 				}
 			}
 
@@ -170,20 +196,43 @@ bool FDisplayClusterViewportManager::RenderInEditor(class FDisplayClusterRenderF
 
 				ViewFamily.bIsRenderedImmediatelyAfterAnotherViewFamily = bIsRenderedImmediatelyAfterAnotherViewFamily;
 
-				FCanvas Canvas(DstResource, nullptr, PreviewScene->GetWorld(), ERHIFeatureLevel::SM5, FCanvas::CDM_DeferDrawing /*FCanvas::CDM_ImmediateDrawing*/, 1.0f);
+				FCanvas Canvas(RenderTargetIt.RenderTargetPtr, nullptr, PreviewScene->GetWorld(), ERHIFeatureLevel::SM5, FCanvas::CDM_DeferDrawing /*FCanvas::CDM_ImmediateDrawing*/, 1.0f);
 				Canvas.Clear(FLinearColor::Black);
 
 				GetRendererModule().BeginRenderingViewFamily(&Canvas, &ViewFamily);
+
+				if (GetRenderFrameSettings().bAllowMultiGPURenderingInEditor)
+				{
+					if (GNumExplicitGPUsForRendering > 1)
+					{
+						const FRHIGPUMask SubmitGPUMask = ViewFamily.Views.Num() == 1 ? ViewFamily.Views[0]->GPUMask : FRHIGPUMask::All();
+						ENQUEUE_RENDER_COMMAND(UDisplayClusterViewportClient_SubmitCommandList)(
+							[SubmitGPUMask](FRHICommandListImmediate& RHICmdList)
+						{
+							SCOPED_GPU_MASK(RHICmdList, SubmitGPUMask);
+							RHICmdList.SubmitCommandsHint();
+						});
+					}
+				}
 			}
 		}
 	}
 
-	// Handle special viewports game-thread logic at frame end
-	// custom postprocess single frame flag must be removed at frame end on game thread
-	FinalizeNewFrame();
+	// bViewportsRenderPassDone means all viewports for this cluster node have been rendered
+	// and we can start frame composing(creating mips, warps, composing icvfx, OutputRemap, etc.).
+	// The main idea is to move the frame composition to a separate render pass.
+	// This will reduce the load on the CPU + GPU and increase the FPS for the preview.
+	if (ViewportIndex == InRenderFrame.ViewportsAmmount && !bViewportsRenderPassDone)
+	{
+		// Handle special viewports game-thread logic at frame end
+		// custom postprocess single frame flag must be removed at frame end on game thread
+		FinalizeNewFrame();
+
+		// After all render target rendered call nDisplay frame rendering:
+		RenderFrame(InViewport);
 
-	// After all render target rendered call nDisplay frame rendering:
-	RenderFrame(InViewport);
+		bOutFrameRendered = true;
+	}
 
 	return true;
 }
diff --git a/Engine/Plugins/Runtime/nDisplay/Source/DisplayCluster/Private/Render/Viewport/DisplayClusterViewportProxy.cpp b/Engine/Plugins/Runtime/nDisplay/Source/DisplayCluster/Private/Render/Viewport/DisplayClusterViewportProxy.cpp
--- a/Engine/Plugins/Runtime/nDisplay/Source/DisplayCluster/Private/Render/Viewport/DisplayClusterViewportProxy.cpp
+++ b/Engine/Plugins/Runtime/nDisplay/Source/DisplayCluster/Private/Render/Viewport/DisplayClusterViewportProxy.cpp
@@ -45,6 +45,7 @@ static TAutoConsoleVariable<int32> CVarDisplayClusterRenderOverscanResolve(
 
 FDisplayClusterViewportProxy::FDisplayClusterViewportProxy(const FDisplayClusterViewportManagerProxy& InOwner, const FDisplayClusterViewport& RenderViewport)
 	: ViewportId(RenderViewport.ViewportId)
+	, ClusterNodeId(RenderViewport.ClusterNodeId)
 	, RenderSettings(RenderViewport.RenderSettings)
 	, ProjectionPolicy(RenderViewport.UninitializedProjectionPolicy)
 	, Owner(InOwner)
@@ -193,10 +194,13 @@ void FDisplayClusterViewportProxy::ImplViewportRemap_RenderThread(FRHICommandLis
 	check(IsInRenderingThread());
 
 #if WITH_EDITOR
-	if (Owner.GetRenderFrameSettings_RenderThread().RenderMode == EDisplayClusterRenderFrameMode::PreviewMono)
+	switch(Owner.GetRenderFrameSettings_RenderThread().RenderMode)
 	{
+	case EDisplayClusterRenderFrameMode::PreviewInScene:
 		// Preview in editor not support this feature
 		return;
+	default:
+		break;
 	}
 #endif
 
@@ -233,9 +237,13 @@ EDisplayClusterViewportResourceType FDisplayClusterViewportProxy::GetOutputResou
 	check(IsInRenderingThread());
 
 #if WITH_EDITOR
-	if (Owner.GetRenderFrameSettings_RenderThread().RenderMode == EDisplayClusterRenderFrameMode::PreviewMono)
+	switch(Owner.GetRenderFrameSettings_RenderThread().RenderMode)
 	{
+	case EDisplayClusterRenderFrameMode::PreviewInScene:
 		return EDisplayClusterViewportResourceType::OutputPreviewTargetableResource;
+
+	default:
+		break;
 	}
 #endif
 
diff --git a/Engine/Plugins/Runtime/nDisplay/Source/DisplayCluster/Private/Render/Viewport/DisplayClusterViewportProxy.h b/Engine/Plugins/Runtime/nDisplay/Source/DisplayCluster/Private/Render/Viewport/DisplayClusterViewportProxy.h
--- a/Engine/Plugins/Runtime/nDisplay/Source/DisplayCluster/Private/Render/Viewport/DisplayClusterViewportProxy.h
+++ b/Engine/Plugins/Runtime/nDisplay/Source/DisplayCluster/Private/Render/Viewport/DisplayClusterViewportProxy.h
@@ -32,6 +32,12 @@ public:
 		return ViewportId;
 	}
 
+	virtual FString GetClusterNodeId() const override
+	{
+		check(IsInRenderingThread());
+		return ClusterNodeId;
+	}
+
 	virtual const FDisplayClusterViewport_RenderSettings& GetRenderSettings_RenderThread() const override
 	{
 		check(IsInRenderingThread());
@@ -112,6 +118,8 @@ protected:
 	// Unique viewport name
 	const FString ViewportId;
 
+	const FString ClusterNodeId;
+
 	// Viewport render params
 	FDisplayClusterViewport_RenderSettings       RenderSettings;
 	FDisplayClusterViewport_RenderSettingsICVFX  RenderSettingsICVFX;
diff --git a/Engine/Plugins/Runtime/nDisplay/Source/DisplayCluster/Private/Render/Viewport/RenderFrame/DisplayClusterRenderFrameManager.cpp b/Engine/Plugins/Runtime/nDisplay/Source/DisplayCluster/Private/Render/Viewport/RenderFrame/DisplayClusterRenderFrameManager.cpp
--- a/Engine/Plugins/Runtime/nDisplay/Source/DisplayCluster/Private/Render/Viewport/RenderFrame/DisplayClusterRenderFrameManager.cpp
+++ b/Engine/Plugins/Runtime/nDisplay/Source/DisplayCluster/Private/Render/Viewport/RenderFrame/DisplayClusterRenderFrameManager.cpp
@@ -16,7 +16,7 @@ bool FDisplayClusterRenderFrameManager::BuildRenderFrame(FViewport* InViewport,
 
 	switch (InRenderFrameSettings.RenderMode)
 	{
-	case EDisplayClusterRenderFrameMode::PreviewMono:
+	case EDisplayClusterRenderFrameMode::PreviewInScene:
 		// Dont use render frame for preview
 		break;
 	default:
diff --git a/Engine/Plugins/Runtime/nDisplay/Source/DisplayCluster/Private/Render/Viewport/RenderFrame/DisplayClusterRenderFrameSettings.h b/Engine/Plugins/Runtime/nDisplay/Source/DisplayCluster/Private/Render/Viewport/RenderFrame/DisplayClusterRenderFrameSettings.h
--- a/Engine/Plugins/Runtime/nDisplay/Source/DisplayCluster/Private/Render/Viewport/RenderFrame/DisplayClusterRenderFrameSettings.h
+++ b/Engine/Plugins/Runtime/nDisplay/Source/DisplayCluster/Private/Render/Viewport/RenderFrame/DisplayClusterRenderFrameSettings.h
@@ -51,6 +51,14 @@ struct FDisplayClusterRenderFrameSettings
 	// Render in Editor mode
 	bool bIsRenderingInEditor = false;
 
+	// Allow mGPU in editor mode
+	bool bAllowMultiGPURenderingInEditor = false;
+	int32 PreviewMinGPUIndex = 0;
+	int32 PreviewMaxGPUIndex = 0;
+
+	// The maximum dimension of any texture for preview
+	int32 PreviewMaxTextureSize = 2048;
+
 	// Performance: Allow merge multiple viewports on single RTT with atlasing (required for bAllowViewFamilyMergeOptimization)
 	bool bAllowRenderTargetAtlasing = false;
 
diff --git a/Engine/Plugins/Runtime/nDisplay/Source/DisplayCluster/Private/Render/Viewport/RenderTarget/DisplayClusterRenderTargetManager.cpp b/Engine/Plugins/Runtime/nDisplay/Source/DisplayCluster/Private/Render/Viewport/RenderTarget/DisplayClusterRenderTargetManager.cpp
--- a/Engine/Plugins/Runtime/nDisplay/Source/DisplayCluster/Private/Render/Viewport/RenderTarget/DisplayClusterRenderTargetManager.cpp
+++ b/Engine/Plugins/Runtime/nDisplay/Source/DisplayCluster/Private/Render/Viewport/RenderTarget/DisplayClusterRenderTargetManager.cpp
@@ -166,7 +166,7 @@ bool FDisplayClusterRenderTargetManager::AllocateFrameTargets(const FDisplayClus
 	//Re-allocate frame targets
 	switch (InRenderFrameSettings.RenderMode)
 	{
-	case EDisplayClusterRenderFrameMode::PreviewMono:
+	case EDisplayClusterRenderFrameMode::PreviewInScene:
 		// Preview model: render to external RTT resource
 		return true;
 
diff --git a/Engine/Plugins/Runtime/nDisplay/Source/DisplayCluster/Private/Render/Viewport/RenderTarget/DisplayClusterRenderTargetResourcesPool.cpp b/Engine/Plugins/Runtime/nDisplay/Source/DisplayCluster/Private/Render/Viewport/RenderTarget/DisplayClusterRenderTargetResourcesPool.cpp
--- a/Engine/Plugins/Runtime/nDisplay/Source/DisplayCluster/Private/Render/Viewport/RenderTarget/DisplayClusterRenderTargetResourcesPool.cpp
+++ b/Engine/Plugins/Runtime/nDisplay/Source/DisplayCluster/Private/Render/Viewport/RenderTarget/DisplayClusterRenderTargetResourcesPool.cpp
@@ -61,6 +61,14 @@ namespace DisplayClusterRenderTargetResourcesPool
 
 		// default gamma for preview in editor
 		float Gamma = 2.2f;
+		switch (InRenderFrameSettings.RenderMode)
+		{
+		case EDisplayClusterRenderFrameMode::PreviewInScene:
+			Gamma = 1.f;
+			break;
+		default:
+			break;
+		}
 
 		// By default sRGB disabled (Preview rendering)
 		bool bShouldUseSRGB = false;
diff --git a/Engine/Plugins/Runtime/nDisplay/Source/DisplayCluster/Public/Components/DisplayClusterPreviewComponent.h b/Engine/Plugins/Runtime/nDisplay/Source/DisplayCluster/Public/Components/DisplayClusterPreviewComponent.h
--- a/Engine/Plugins/Runtime/nDisplay/Source/DisplayCluster/Public/Components/DisplayClusterPreviewComponent.h
+++ b/Engine/Plugins/Runtime/nDisplay/Source/DisplayCluster/Public/Components/DisplayClusterPreviewComponent.h
@@ -6,6 +6,8 @@
 
 #include "DisplayClusterConfigurationTypes_Base.h"
 
+#include "Render/Viewport/IDisplayClusterViewportManager.h"
+
 #include "DisplayClusterPreviewComponent.generated.h"
 
 class IDisplayClusterProjectionPolicy;
@@ -38,7 +40,7 @@ public:
 	virtual void DestroyComponent(bool bPromoteChildren = false) override;
 
 public:
-	bool InitializePreviewComponent(ADisplayClusterRootActor* RootActor, const FString& ViewportId, UDisplayClusterConfigurationViewport* ViewportConfig);
+	bool InitializePreviewComponent(ADisplayClusterRootActor* RootActor, const FString& ClusterNodeId, const FString& ViewportId, UDisplayClusterConfigurationViewport* ViewportConfig);
 
 	bool IsPreviewAvailable() const;
 
@@ -47,6 +49,11 @@ public:
 		return ViewportId; 
 	}
 
+	const FString& GetClusterNodeId() const
+	{
+		return ClusterNodeId;
+	}
+
 	UDisplayClusterConfigurationViewport* GetViewportConfig() const
 	{ 
 		return ViewportConfig;
@@ -58,7 +65,7 @@ public:
 	}
 
 	void UpdatePreviewResources();
-	void HandleRenderTargetTextureDeferredUpdate();
+	void HandleRenderTargetTextureDeferredUpdate(const EDisplayClusterRenderFrameMode InRenderFrameMode);
 
 	UMeshComponent* GetPreviewMesh() const
 	{
@@ -66,7 +73,7 @@ public:
 	}
 
 	/** Create and retrieve a render texture 2d from the render target. */
-	UTexture2D* GetOrCreateRenderTexture2D();
+	UTexture2D* GetOrCreateViewportPreviewTexture2D();
 
 protected:
 	class IDisplayClusterViewport* GetCurrentViewport() const;
@@ -77,6 +84,7 @@ protected:
 
 	bool UpdatePreviewMesh(bool bRestoreOriginalMaterial = false);
 	void ReleasePreviewMesh();
+	void UpdatePreviewMeshReference();
 
 	bool UpdatePreviewTexture();
 	void ReleasePreviewTexture();
@@ -85,7 +93,10 @@ protected:
 	void ReleasePreviewMaterial();
 	void UpdatePreviewMaterial();
 
-	void UpdatePreviewMeshMaterial(bool bRestoreOriginalMaterial = false);
+	void RestorePreviewMeshMaterial();
+	void SetPreviewMeshMaterial();
+
+	bool IsPreviewEnabled() const;
 
 protected:
 	// Set to true, when RTT surface updated
@@ -95,6 +106,7 @@ protected:
 
 #if WITH_EDITORONLY_DATA
 protected:
+	// Texture for preview material
 	UPROPERTY(VisibleAnywhere, BlueprintReadOnly, Category = "Preview", meta = (DisplayName = "Render Target"))
 	UTextureRenderTarget2D* RenderTarget;
 
@@ -109,6 +121,9 @@ private:
 	UPROPERTY()
 	FString ViewportId;
 
+	UPROPERTY()
+	FString ClusterNodeId;
+
 	UPROPERTY()
 	UDisplayClusterConfigurationViewport* ViewportConfig = nullptr;
 
diff --git a/Engine/Plugins/Runtime/nDisplay/Source/DisplayCluster/Public/DisplayClusterRootActor.h b/Engine/Plugins/Runtime/nDisplay/Source/DisplayCluster/Public/DisplayClusterRootActor.h
--- a/Engine/Plugins/Runtime/nDisplay/Source/DisplayCluster/Public/DisplayClusterRootActor.h
+++ b/Engine/Plugins/Runtime/nDisplay/Source/DisplayCluster/Public/DisplayClusterRootActor.h
@@ -21,6 +21,7 @@
 #include "DisplayClusterEditorPropertyReference.h"
 
 #include "Render/Viewport/IDisplayClusterViewportManager.h"
+#include "Render/Viewport/RenderFrame/DisplayClusterRenderFrame.h"
 
 #include "DisplayClusterRootActor.generated.h"
 
@@ -290,6 +291,10 @@ public:
 	UPROPERTY(EditAnywhere, BlueprintReadWrite, Category = "Editor Preview", meta = (DisplayName = "Enable Editor Preview"))
 	bool bPreviewEnable = true;
 	
+	/** render preview every frame */
+	UPROPERTY(EditAnywhere, BlueprintReadWrite, Category = "Editor Preview", meta = (DisplayName = "Capture Every Frame", EditCondition = "bPreviewEnable"))
+	bool bPreviewRenderEveryFrame = true;
+
 	/** Selectively preview a specific viewport or show all/none. */
 	UPROPERTY(EditAnywhere, BlueprintReadWrite, Category = "Editor Preview", meta = (DisplayName = "Preview Node", EditCondition = "bPreviewEnable"))
 	FString PreviewNodeId = DisplayClusterConfigurationStrings::gui::preview::PreviewNodeAll;
@@ -298,10 +303,28 @@ public:
 	UPROPERTY(EditAnywhere, BlueprintReadWrite, Category = "Editor Preview", meta = (EditCondition = "bPreviewEnable"))
 	EDisplayClusterConfigurationRenderMode RenderMode = EDisplayClusterConfigurationRenderMode::Mono;
 
+	/** Render with mGPU */
+	UPROPERTY(EditAnywhere, BlueprintReadWrite, Category = "Editor Preview", meta = (EditCondition = "bPreviewEnable"))
+	bool bAllowMultiGPURendering = true;
+
+	UPROPERTY(EditAnywhere, BlueprintReadWrite, Category = "Editor Preview", meta = (EditCondition = "bPreviewEnable"))
+	int MinGPUIndex = 1;
+
+	UPROPERTY(EditAnywhere, BlueprintReadWrite, Category = "Editor Preview", meta = (EditCondition = "bPreviewEnable"))
+	int MaxGPUIndex = 1;
+
 	/** Tick Per Frame */
 	UPROPERTY(EditAnywhere, BlueprintReadWrite, Category = "Editor Preview", meta = (ClampMin = "1", UIMin = "1", ClampMax = "200", UIMax = "200", EditCondition = "bPreviewEnable"))
 	int TickPerFrame = 1;
 
+	/** Nodes Per Frame */
+	UPROPERTY(EditAnywhere, BlueprintReadWrite, Category = "Editor Preview", meta = (ClampMin = "1", UIMin = "1", ClampMax = "200", UIMax = "200", EditCondition = "bPreviewEnable"))
+	int NodesPerFrame = 1;
+
+	/** Viewports Per Frame */
+	UPROPERTY(EditAnywhere, BlueprintReadWrite, Category = "Editor Preview", meta = (ClampMin = "1", UIMin = "1", ClampMax = "200", UIMax = "200", EditCondition = "bPreviewEnable"))
+	int ViewportsPerFrame = 1;
+
 	/** Adjust resolution scaling for the editor preview. */
 	UPROPERTY(EditAnywhere, BlueprintReadWrite, Category = "Editor Preview", meta = (DisplayName = "Preview Screen Percentage", ClampMin = "0.05", UIMin = "0.05", ClampMax = "1", UIMax = "1", EditCondition = "bPreviewEnable"))
 	float PreviewRenderTargetRatioMult = 0.25;
@@ -314,6 +337,10 @@ public:
 	UPROPERTY(EditAnywhere, BlueprintReadWrite, Category = "Editor Preview", meta = (DisplayName = "ICVFX Camera Frustums Distance", EditCondition = "bPreviewEnable"))
 	float PreviewICVFXFrustumsFarDistance = 1000.0f;
 
+	/** The maximum dimension of any internal texture for preview. Use less memory for large preview viewports */
+	UPROPERTY(EditAnywhere, BlueprintReadWrite, Category = "Editor Preview", meta = (DisplayName = "Preview Texture Max Size", ClampMin = "64", UIMin = "64", ClampMax = "4096", UIMax = "4096", EditCondition = "bPreviewEnable"))
+	int PreviewMaxTextureSize = 2048;
+
 private:
 	UPROPERTY(Transient)
 	TMap<FString, UDisplayClusterPreviewComponent*> PreviewComponents;
@@ -341,6 +368,7 @@ public:
 
 	void Tick_Editor(float DeltaSeconds);
 	void PostLoad_Editor();
+	void PostActorCreated_Editor();
 	void BeginDestroy_Editor();
 	void RerunConstructionScripts_Editor();
 
@@ -356,7 +384,7 @@ public:
 
 	IDisplayClusterViewport* FindPreviewViewport(const FString& InViewportId) const;
 
-	void GetPreviewRenderTargetableTextures(const TArray<FString>& InViewportNames, TArray<FTextureRHIRef>& OutTextures);
+	void GetPreviewRenderTargetableTextures(const EDisplayClusterRenderFrameMode InRenderFrameMode, const TArray<FString>& InViewportNames, TArray<FTextureRHIRef>& OutTextures);
 
 	void UpdateInnerFrustumPriority();
 	void ResetInnerFrustumPriority();
@@ -372,8 +400,13 @@ public:
 
 protected:
 	FString GeneratePreviewComponentName(const FString& NodeId, const FString& ViewportId) const;
-	void RenderPreview_Editor();
-	bool UpdatePreviewConfiguration_Editor(bool bUpdateAllViewports);
+
+	void ImplRenderPreviewForSceneMaterials_Editor(const int32 InNumNodesForRender);
+
+	void RenderPreviewClusterNode_Editor(const EDisplayClusterRenderFrameMode InRenderFrameMode, const FString& InClusterNodeId);
+	bool UpdatePreviewConfiguration_Editor(const EDisplayClusterRenderFrameMode InRenderFrameMode, const FString& InClusterNodeId);
+
+	void ResetPreviewInternals();
 	
 	void RenderPreviewFrustums();
 	void RenderPreviewFrustum(const FMatrix ProjectionMatrix, const FMatrix ViewMatrix, const FVector ViewOrigin);
@@ -382,18 +415,17 @@ protected:
 	virtual void PostEditMove(bool bFinished) override;
 
 private:
-	/** The number of times to update the render target via deferred update. */
-	int32 PreviewRenderTargetUpdatesRequired = 0;
 	bool bIsSelectedInEditor = false;
-
-	// Preview frustum vertices
-	FVector PreviewFrustumVerts[8];
 	
 private:
 	TWeakPtr<IDisplayClusterConfiguratorBlueprintEditor> ToolkitPtr;
 
 	int32 TickPerFrameCounter = 0;
 
+	int32 PreviewClusterNodeIndex = 0;
+	int32 PreviewViewportIndex = 0;
+	TUniquePtr<FDisplayClusterRenderFrame> PreviewRenderFrame;
+
 	FOnPreviewUpdated OnPreviewGenerated;
 	FOnPreviewUpdated OnPreviewDestroyed;
 
diff --git a/Engine/Plugins/Runtime/nDisplay/Source/DisplayCluster/Public/Render/Viewport/Containers/DisplayClusterPreviewSettings.h b/Engine/Plugins/Runtime/nDisplay/Source/DisplayCluster/Public/Render/Viewport/Containers/DisplayClusterPreviewSettings.h
new file mode 100644
--- /dev/null
+++ b/Engine/Plugins/Runtime/nDisplay/Source/DisplayCluster/Public/Render/Viewport/Containers/DisplayClusterPreviewSettings.h
@@ -0,0 +1,21 @@
+// Copyright Epic Games, Inc. All Rights Reserved.
+
+#pragma once
+
+#include "CoreMinimal.h"
+
+struct FDisplayClusterPreviewSettings
+{
+	// Preview RTT size multiplier
+	float PreviewRenderTargetRatioMult = 1.f;
+
+	// The maximum dimension of any texture for preview
+	int32 PreviewMaxTextureSize = 2048;
+
+	// Allow mGPU in editor mode
+	bool bAllowMultiGPURenderingInEditor = false;
+
+	int32 MinGPUIndex = 1;
+	int32 MaxGPUIndex = 1;
+};
+
diff --git a/Engine/Plugins/Runtime/nDisplay/Source/DisplayCluster/Public/Render/Viewport/IDisplayClusterViewport.h b/Engine/Plugins/Runtime/nDisplay/Source/DisplayCluster/Public/Render/Viewport/IDisplayClusterViewport.h
--- a/Engine/Plugins/Runtime/nDisplay/Source/DisplayCluster/Public/Render/Viewport/IDisplayClusterViewport.h
+++ b/Engine/Plugins/Runtime/nDisplay/Source/DisplayCluster/Public/Render/Viewport/IDisplayClusterViewport.h
@@ -24,6 +24,7 @@ public:
 
 public:
 	virtual FString GetId() const = 0;
+	virtual FString GetClusterNodeId() const = 0;
 
 	virtual const FDisplayClusterViewport_RenderSettings&      GetRenderSettings() const = 0;
 	virtual const FDisplayClusterViewport_RenderSettingsICVFX& GetRenderSettingsICVFX() const = 0;
diff --git a/Engine/Plugins/Runtime/nDisplay/Source/DisplayCluster/Public/Render/Viewport/IDisplayClusterViewportManager.h b/Engine/Plugins/Runtime/nDisplay/Source/DisplayCluster/Public/Render/Viewport/IDisplayClusterViewportManager.h
--- a/Engine/Plugins/Runtime/nDisplay/Source/DisplayCluster/Public/Render/Viewport/IDisplayClusterViewportManager.h
+++ b/Engine/Plugins/Runtime/nDisplay/Source/DisplayCluster/Public/Render/Viewport/IDisplayClusterViewportManager.h
@@ -6,9 +6,9 @@
 
 #include "Render/Viewport/IDisplayClusterViewport.h"
 #include "Render/Viewport/RenderFrame/DisplayClusterRenderFrame.h"
+#include "Render/Viewport/Containers/DisplayClusterPreviewSettings.h"
 #include "SceneView.h"
 
-
 class UWorld;
 class FViewport;
 class FSceneViewFamilyContext;
@@ -17,8 +17,6 @@ class UDisplayClusterConfigurationViewport;
 class IDisplayClusterViewportManagerProxy;
 class FReferenceCollector;
 
-struct FDisplayClusterConfigurationViewportPreview;
-
 class DISPLAYCLUSTER_API IDisplayClusterViewportManager
 {
 public:
@@ -46,10 +44,11 @@ public:
 	* @param InRenderMode     - Render mode
 	* @param InClusterNodeId  - cluster node for rendering
 	* @param InRootActorPtr   - reference to RootActor with actual configuration inside
+	* @param InPreviewSettings - support preview rendering
 	*
 	* @return - true, if success
 	*/
-	virtual bool UpdateConfiguration(EDisplayClusterRenderFrameMode InRenderMode, const FString& InClusterNodeId, class ADisplayClusterRootActor* InRootActorPtr) = 0;
+	virtual bool UpdateConfiguration(EDisplayClusterRenderFrameMode InRenderMode, const FString& InClusterNodeId, class ADisplayClusterRootActor* InRootActorPtr, const FDisplayClusterPreviewSettings* InPreviewSettings = nullptr) = 0;
 
 	/**
 	* Initialize new frame for all viewports on game thread, and update context, render resources with viewport new settings
@@ -104,8 +103,7 @@ public:
 	virtual void AddReferencedObjects(FReferenceCollector& Collector) = 0;
 
 #if WITH_EDITOR
-	virtual bool UpdatePreviewConfiguration(const FDisplayClusterConfigurationViewportPreview& PreviewConfiguration, ADisplayClusterRootActor* InRootActorPtr) = 0;
-	virtual bool RenderInEditor(class FDisplayClusterRenderFrame& InRenderFrame, FViewport* InViewport) = 0;
+	virtual bool RenderInEditor(class FDisplayClusterRenderFrame& InRenderFrame, FViewport* InViewport, const uint32 InFirstViewportNum, const int32 InViewportsAmmount, bool& bOutFrameRendered) = 0;
 #endif
 
 	///////////////////////////////////////////////////////////////////////////////////////////////////////////////////////////////////////////////////////////////////////////
diff --git a/Engine/Plugins/Runtime/nDisplay/Source/DisplayCluster/Public/Render/Viewport/IDisplayClusterViewportProxy.h b/Engine/Plugins/Runtime/nDisplay/Source/DisplayCluster/Public/Render/Viewport/IDisplayClusterViewportProxy.h
--- a/Engine/Plugins/Runtime/nDisplay/Source/DisplayCluster/Public/Render/Viewport/IDisplayClusterViewportProxy.h
+++ b/Engine/Plugins/Runtime/nDisplay/Source/DisplayCluster/Public/Render/Viewport/IDisplayClusterViewportProxy.h
@@ -22,6 +22,7 @@ public:
 
 public:
 	virtual FString GetId() const = 0;
+	virtual FString GetClusterNodeId() const = 0;
 
 	virtual const TSharedPtr<class IDisplayClusterProjectionPolicy, ESPMode::ThreadSafe>& GetProjectionPolicy_RenderThread() const = 0;
 
diff --git a/Engine/Plugins/Runtime/nDisplay/Source/DisplayCluster/Public/Render/Viewport/RenderFrame/DisplayClusterRenderFrame.h b/Engine/Plugins/Runtime/nDisplay/Source/DisplayCluster/Public/Render/Viewport/RenderFrame/DisplayClusterRenderFrame.h
--- a/Engine/Plugins/Runtime/nDisplay/Source/DisplayCluster/Public/Render/Viewport/RenderFrame/DisplayClusterRenderFrame.h
+++ b/Engine/Plugins/Runtime/nDisplay/Source/DisplayCluster/Public/Render/Viewport/RenderFrame/DisplayClusterRenderFrame.h
@@ -84,6 +84,7 @@ public:
 	FIntRect FrameRect;
 
 	int32 DesiredNumberOfViews = 0;
+	int32 ViewportsAmmount = 0;
 
 	IDisplayClusterViewportManager* ViewportManager = nullptr;
 };
diff --git a/Engine/Plugins/Runtime/nDisplay/Source/DisplayCluster/Public/Render/Viewport/RenderFrame/DisplayClusterRenderFrameEnums.h b/Engine/Plugins/Runtime/nDisplay/Source/DisplayCluster/Public/Render/Viewport/RenderFrame/DisplayClusterRenderFrameEnums.h
--- a/Engine/Plugins/Runtime/nDisplay/Source/DisplayCluster/Public/Render/Viewport/RenderFrame/DisplayClusterRenderFrameEnums.h
+++ b/Engine/Plugins/Runtime/nDisplay/Source/DisplayCluster/Public/Render/Viewport/RenderFrame/DisplayClusterRenderFrameEnums.h
@@ -18,11 +18,8 @@ enum class EDisplayClusterRenderFrameMode: uint8
 	// Special render: render to two separated backbuffer textures, with half size Y
 	TopBottom,
 
-	// Special preview render:
-	// 1. Render in mono
-	// 2. For each viewport created separated frametarget
-	// 3. Applied preview settings from FDisplayClusterRenderFrameSettings::FPreviewSttings
-	PreviewMono
+	// Special render for preview in scene
+	PreviewInScene,
 };
 
 // (experimental, not implemented, reserved)
diff --git a/Engine/Plugins/Runtime/nDisplay/Source/DisplayClusterConfiguration/Public/DisplayClusterConfigurationTypes_Viewport.h b/Engine/Plugins/Runtime/nDisplay/Source/DisplayClusterConfiguration/Public/DisplayClusterConfigurationTypes_Viewport.h
--- a/Engine/Plugins/Runtime/nDisplay/Source/DisplayClusterConfiguration/Public/DisplayClusterConfigurationTypes_Viewport.h
+++ b/Engine/Plugins/Runtime/nDisplay/Source/DisplayClusterConfiguration/Public/DisplayClusterConfigurationTypes_Viewport.h
@@ -266,27 +266,3 @@ public:
 	UPROPERTY(EditAnywhere, BlueprintReadWrite, Category = "Configuration")
 	bool bAllowWarpBlend = true;
 };
-
-USTRUCT()
-struct DISPLAYCLUSTERCONFIGURATION_API FDisplayClusterConfigurationViewportPreview
-{
-	GENERATED_BODY()
-
-public:
-	// Allow preview render
-	UPROPERTY()
-	bool bEnable = true;
-
-	// Render single node preview or whole cluster
-	UPROPERTY()
-	FString PreviewNodeId = DisplayClusterConfigurationStrings::gui::preview::PreviewNodeAll;
-
-	// Update preview texture period in tick
-	UPROPERTY()
-	int TickPerFrame = 1;
-
-	// Preview texture size get from viewport, and scaled by this value
-	UPROPERTY()
-	float PreviewRenderTargetRatioMult = 0.25;
-};
-
diff --git a/Engine/Plugins/Runtime/nDisplay/Source/DisplayClusterConfigurator/Private/Views/OutputMapping/EdNodes/DisplayClusterConfiguratorViewportNode.cpp b/Engine/Plugins/Runtime/nDisplay/Source/DisplayClusterConfigurator/Private/Views/OutputMapping/EdNodes/DisplayClusterConfiguratorViewportNode.cpp
--- a/Engine/Plugins/Runtime/nDisplay/Source/DisplayClusterConfigurator/Private/Views/OutputMapping/EdNodes/DisplayClusterConfiguratorViewportNode.cpp
+++ b/Engine/Plugins/Runtime/nDisplay/Source/DisplayClusterConfigurator/Private/Views/OutputMapping/EdNodes/DisplayClusterConfiguratorViewportNode.cpp
@@ -228,7 +228,7 @@ UTexture* UDisplayClusterConfiguratorViewportNode::GetPreviewTexture() const
 		UDisplayClusterConfiguratorWindowNode* ParentWindow = GetParentChecked<UDisplayClusterConfiguratorWindowNode>();
 		if (UDisplayClusterPreviewComponent* PreviewComp = RootActor->GetPreviewComponent(ParentWindow->GetNodeName(), GetNodeName()))
 		{
-			if (UTexture2D* Texture = PreviewComp->GetOrCreateRenderTexture2D())
+			if (UTexture2D* Texture = PreviewComp->GetOrCreateViewportPreviewTexture2D())
 			{
 				return Texture;
 			}
diff --git a/Engine/Plugins/Runtime/nDisplay/Source/DisplayClusterProjection/Private/Policy/Link/DisplayClusterProjectionLinkPolicy.cpp b/Engine/Plugins/Runtime/nDisplay/Source/DisplayClusterProjection/Private/Policy/Link/DisplayClusterProjectionLinkPolicy.cpp
--- a/Engine/Plugins/Runtime/nDisplay/Source/DisplayClusterProjection/Private/Policy/Link/DisplayClusterProjectionLinkPolicy.cpp
+++ b/Engine/Plugins/Runtime/nDisplay/Source/DisplayClusterProjection/Private/Policy/Link/DisplayClusterProjectionLinkPolicy.cpp
@@ -34,17 +34,21 @@ bool FDisplayClusterProjectionLinkPolicy::CalculateView(IDisplayClusterViewport*
 	IDisplayClusterViewport* ParentViewport = InViewport->GetOwner().FindViewport(InViewport->GetRenderSettings().GetParentViewportId());
 	if (ParentViewport != nullptr)
 	{
-		if (ParentViewport->GetContexts()[InContextNum].bDisableRender && ParentViewport->GetProjectionPolicy().IsValid())
+		const TArray<FDisplayClusterViewport_Context>& ParentContexts = ParentViewport->GetContexts();
+		if (InContextNum < (uint32)ParentContexts.Num())
 		{
-			return ParentViewport->GetProjectionPolicy()->CalculateView(ParentViewport, InContextNum, InOutViewLocation, InOutViewRotation, ViewOffset, WorldToMeters, NCP, FCP);
-		}
-		else
-		{
-			const FDisplayClusterViewport_Context& ParentContext = ParentViewport->GetContexts()[InContextNum];
-			InOutViewLocation = ParentContext.ViewLocation;
-			InOutViewRotation = ParentContext.ViewRotation;
-
-			return true;
+			if (ParentContexts[InContextNum].bDisableRender && ParentViewport->GetProjectionPolicy().IsValid())
+			{
+				return ParentViewport->GetProjectionPolicy()->CalculateView(ParentViewport, InContextNum, InOutViewLocation, InOutViewRotation, ViewOffset, WorldToMeters, NCP, FCP);
+			}
+			else
+			{
+				const FDisplayClusterViewport_Context& ParentContext = ParentContexts[InContextNum];
+				InOutViewLocation = ParentContext.ViewLocation;
+				InOutViewRotation = ParentContext.ViewRotation;
+
+				return true;
+			}
 		}
 	}
 
@@ -59,16 +63,20 @@ bool FDisplayClusterProjectionLinkPolicy::GetProjectionMatrix(class IDisplayClus
 	IDisplayClusterViewport* ParentViewport = InViewport->GetOwner().FindViewport(InViewport->GetRenderSettings().GetParentViewportId());
 	if (ParentViewport != nullptr)
 	{
-		if (ParentViewport->GetContexts()[InContextNum].bDisableRender && ParentViewport->GetProjectionPolicy().IsValid())
+		const TArray<FDisplayClusterViewport_Context>& ParentContexts = ParentViewport->GetContexts();
+		if (InContextNum < (uint32)ParentContexts.Num())
 		{
-			return ParentViewport->GetProjectionPolicy()->GetProjectionMatrix(ParentViewport, InContextNum, OutPrjMatrix);
-		}
-		else
-		{
-			const FDisplayClusterViewport_Context& ParentContext = ParentViewport->GetContexts()[InContextNum];
-			OutPrjMatrix = ParentContext.ProjectionMatrix;
-
-			return true;
+			if (ParentContexts[InContextNum].bDisableRender && ParentViewport->GetProjectionPolicy().IsValid())
+			{
+				return ParentViewport->GetProjectionPolicy()->GetProjectionMatrix(ParentViewport, InContextNum, OutPrjMatrix);
+			}
+			else
+			{
+				const FDisplayClusterViewport_Context& ParentContext = ParentContexts[InContextNum];
+				OutPrjMatrix = ParentContext.ProjectionMatrix;
+
+				return true;
+			}
 		}
 	}
 
