nDisplay: fixed p-node deadlock on ESC exit

#rnx

diff --git a/Engine/Plugins/Runtime/nDisplay/Source/DisplayCluster/Private/Cluster/DisplayClusterClusterManager.cpp b/Engine/Plugins/Runtime/nDisplay/Source/DisplayCluster/Private/Cluster/DisplayClusterClusterManager.cpp
--- a/Engine/Plugins/Runtime/nDisplay/Source/DisplayCluster/Private/Cluster/DisplayClusterClusterManager.cpp
+++ b/Engine/Plugins/Runtime/nDisplay/Source/DisplayCluster/Private/Cluster/DisplayClusterClusterManager.cpp
@@ -140,25 +140,13 @@ bool FDisplayClusterClusterManager::StartSession(UDisplayClusterConfigurationDat
 	return true;
 }
 
-void FDisplayClusterClusterManager::PreAppExit()
-{
-	// Trigger all data cache availability events to prevent
-	// client session threads to be deadlocked.
-	SetInternalSyncObjectsReleaseState(true);
-
-	// Also stop clients/servers
-	if (ClusterNodeCtrl)
-	{
-		ClusterNodeCtrl->Shutdown();
-	}
-}
-
 void FDisplayClusterClusterManager::EndSession()
 {
 	// Trigger all data cache availability events to prevent
 	// client session threads to be deadlocked.
 	SetInternalSyncObjectsReleaseState(true);
 
+	// Stop local clients/servers
 	if (ClusterNodeCtrl)
 	{
 		ClusterNodeCtrl->Shutdown();
diff --git a/Engine/Plugins/Runtime/nDisplay/Source/DisplayCluster/Private/Cluster/DisplayClusterClusterManager.h b/Engine/Plugins/Runtime/nDisplay/Source/DisplayCluster/Private/Cluster/DisplayClusterClusterManager.h
--- a/Engine/Plugins/Runtime/nDisplay/Source/DisplayCluster/Private/Cluster/DisplayClusterClusterManager.h
+++ b/Engine/Plugins/Runtime/nDisplay/Source/DisplayCluster/Private/Cluster/DisplayClusterClusterManager.h
@@ -27,7 +27,6 @@ public:
 	virtual bool Init(EDisplayClusterOperationMode OperationMode) override;
 	virtual void Release() override;
 	virtual bool StartSession(UDisplayClusterConfigurationData* InConfigData, const FString& NodeID) override;
-	virtual void PreAppExit() override;
 	virtual void EndSession() override;
 	virtual bool StartScene(UWorld* World) override;
 	virtual void EndScene() override;
diff --git a/Engine/Plugins/Runtime/nDisplay/Source/DisplayCluster/Private/DisplayClusterModule.cpp b/Engine/Plugins/Runtime/nDisplay/Source/DisplayCluster/Private/DisplayClusterModule.cpp
--- a/Engine/Plugins/Runtime/nDisplay/Source/DisplayCluster/Private/DisplayClusterModule.cpp
+++ b/Engine/Plugins/Runtime/nDisplay/Source/DisplayCluster/Private/DisplayClusterModule.cpp
@@ -121,16 +121,6 @@ bool FDisplayClusterModule::StartSession(UDisplayClusterConfigurationData* InCon
 	return bResult;
 }
 
-void FDisplayClusterModule::PreAppExit()
-{
-	UE_LOG(LogDisplayClusterModule, Log, TEXT("Preparing to application exit..."));
-
-	for (IPDisplayClusterManager* const Manager : Managers)
-	{
-		Manager->PreAppExit();
-	}
-}
-
 void FDisplayClusterModule::EndSession()
 {
 	UE_LOG(LogDisplayClusterModule, Log, TEXT("Stopping DisplayCluster session..."));
diff --git a/Engine/Plugins/Runtime/nDisplay/Source/DisplayCluster/Private/DisplayClusterModule.h b/Engine/Plugins/Runtime/nDisplay/Source/DisplayCluster/Private/DisplayClusterModule.h
--- a/Engine/Plugins/Runtime/nDisplay/Source/DisplayCluster/Private/DisplayClusterModule.h
+++ b/Engine/Plugins/Runtime/nDisplay/Source/DisplayCluster/Private/DisplayClusterModule.h
@@ -63,7 +63,6 @@ public:
 	virtual bool Init(EDisplayClusterOperationMode OperationMode) override;
 	virtual void Release() override;
 	virtual bool StartSession(UDisplayClusterConfigurationData* InConfigData, const FString& NodeId) override;
-	virtual void PreAppExit() override;
 	virtual void EndSession() override;
 	virtual bool StartScene(UWorld* InWorld) override;
 	virtual void EndScene() override;
diff --git a/Engine/Plugins/Runtime/nDisplay/Source/DisplayCluster/Private/Game/EngineClasses/Basics/DisplayClusterGameEngine.cpp b/Engine/Plugins/Runtime/nDisplay/Source/DisplayCluster/Private/Game/EngineClasses/Basics/DisplayClusterGameEngine.cpp
--- a/Engine/Plugins/Runtime/nDisplay/Source/DisplayCluster/Private/Game/EngineClasses/Basics/DisplayClusterGameEngine.cpp
+++ b/Engine/Plugins/Runtime/nDisplay/Source/DisplayCluster/Private/Game/EngineClasses/Basics/DisplayClusterGameEngine.cpp
@@ -61,14 +61,12 @@ void UDisplayClusterGameEngine::Init(class IEngineLoop* InEngineLoop)
 		FDisplayClusterAppExit::ExitApplication(FString("Couldn't initialize DisplayCluster module"));
 	}
 
+	// This is required to prevent GameThread dead-locking when Alt+F4 is used
+	// to terminate p-node of a 2+ nodes cluster. This gets called before
+	// virtual PreExit which allows to unlock GameThread in advance.
 	FCoreDelegates::ApplicationWillTerminateDelegate.AddLambda([this]()
 	{
-		UE_LOG(LogDisplayClusterEngine, Log, TEXT("ApplicationWillTerminateDelegate fired!"));
-
-		if (OperationMode == EDisplayClusterOperationMode::Cluster)
-		{
-			GDisplayCluster->PreAppExit();
-		}
+		PreExitImpl();
 	});
 
 	if (OperationMode == EDisplayClusterOperationMode::Cluster)
@@ -263,7 +261,15 @@ bool UDisplayClusterGameEngine::ValidateConfigFile(const FString& FilePath)
 
 void UDisplayClusterGameEngine::PreExit()
 {
-	UE_LOG(LogDisplayClusterEngine, Log, TEXT("UDisplayClusterGameEngine::PreExit"));
+	PreExitImpl();
+
+	// Release the engine
+	UGameEngine::PreExit();
+}
+
+void UDisplayClusterGameEngine::PreExitImpl()
+{
+	UE_LOG(LogDisplayClusterEngine, Log, TEXT("UDisplayClusterGameEngine::PreExitImpl"));
 
 	if (OperationMode == EDisplayClusterOperationMode::Cluster)
 	{
@@ -272,9 +278,6 @@ void UDisplayClusterGameEngine::PreExit()
 		// Close current DisplayCluster session
 		GDisplayCluster->EndSession();
 	}
-
-	// Release the engine
-	UGameEngine::PreExit();
 }
 
 bool UDisplayClusterGameEngine::LoadMap(FWorldContext& WorldContext, FURL URL, class UPendingNetGame* Pending, FString& Error)
diff --git a/Engine/Plugins/Runtime/nDisplay/Source/DisplayCluster/Private/IPDisplayClusterManager.h b/Engine/Plugins/Runtime/nDisplay/Source/DisplayCluster/Private/IPDisplayClusterManager.h
--- a/Engine/Plugins/Runtime/nDisplay/Source/DisplayCluster/Private/IPDisplayClusterManager.h
+++ b/Engine/Plugins/Runtime/nDisplay/Source/DisplayCluster/Private/IPDisplayClusterManager.h
@@ -30,10 +30,6 @@ public:
 	virtual bool StartSession(UDisplayClusterConfigurationData* InConfigData, const FString& InNodeId)
 	{ return true; }
 
-	// The callback comes from the engine during Tick(). Notifies about application exit once the frame is finished.
-	virtual void PreAppExit()
-	{ }
-
 	// Called on each session end at early step before exit (before UGameEngine::Preexit)
 	virtual void EndSession()
 	{ }
diff --git a/Engine/Plugins/Runtime/nDisplay/Source/DisplayCluster/Private/Misc/DisplayClusterAppExit.cpp b/Engine/Plugins/Runtime/nDisplay/Source/DisplayCluster/Private/Misc/DisplayClusterAppExit.cpp
--- a/Engine/Plugins/Runtime/nDisplay/Source/DisplayCluster/Private/Misc/DisplayClusterAppExit.cpp
+++ b/Engine/Plugins/Runtime/nDisplay/Source/DisplayCluster/Private/Misc/DisplayClusterAppExit.cpp
@@ -11,18 +11,13 @@
 #include "Editor/UnrealEdEngine.h"
 #endif
 
-FCriticalSection FDisplayClusterAppExit::InternalsSyncScope;
-
 
 void FDisplayClusterAppExit::ExitApplication(const FString& Msg)
 {
-	FScopeLock lock(&InternalsSyncScope);
-
-	UE_LOG(LogDisplayClusterModule, Log, TEXT("Exit requested - %s"), *Msg);
-
 	if (GEngine && GEngine->IsEditor())
 	{
 #if WITH_EDITOR
+		UE_LOG(LogDisplayClusterModule, Log, TEXT("PIE end requested - %s"), *Msg);
 		GUnrealEd->RequestEndPlayMap();
 #endif
 	}
@@ -30,12 +25,17 @@ void FDisplayClusterAppExit::ExitApplication(const FString& Msg)
 	{
 		if (!IsEngineExitRequested())
 		{
+			UE_LOG(LogDisplayClusterModule, Log, TEXT("Exit requested - %s"), *Msg);
+
 			if (IsInGameThread())
 			{
 				FPlatformMisc::RequestExit(false);
 			}
 			else
 			{
+				// For some reason UE4 generates crash info if FPlatformMisc::RequestExit gets called
+				// from a thread other than GameThread. Since it may be called from the networking
+				// session threads (failover pipeline), we don't want to generate unnecessary crash reports.
 				AsyncTask(ENamedThreads::GameThread, []()
 				{
 					FPlatformMisc::RequestExit(false);
diff --git a/Engine/Plugins/Runtime/nDisplay/Source/DisplayCluster/Private/Misc/DisplayClusterAppExit.h b/Engine/Plugins/Runtime/nDisplay/Source/DisplayCluster/Private/Misc/DisplayClusterAppExit.h
--- a/Engine/Plugins/Runtime/nDisplay/Source/DisplayCluster/Private/Misc/DisplayClusterAppExit.h
+++ b/Engine/Plugins/Runtime/nDisplay/Source/DisplayCluster/Private/Misc/DisplayClusterAppExit.h
@@ -12,7 +12,4 @@ class FDisplayClusterAppExit
 {
 public:
 	static void ExitApplication(const FString& Msg);
-
-private:
-	static FCriticalSection InternalsSyncScope;
 };
diff --git a/Engine/Plugins/Runtime/nDisplay/Source/DisplayCluster/Public/DisplayClusterGameEngine.h b/Engine/Plugins/Runtime/nDisplay/Source/DisplayCluster/Public/DisplayClusterGameEngine.h
--- a/Engine/Plugins/Runtime/nDisplay/Source/DisplayCluster/Public/DisplayClusterGameEngine.h
+++ b/Engine/Plugins/Runtime/nDisplay/Source/DisplayCluster/Public/DisplayClusterGameEngine.h
@@ -49,6 +49,7 @@ protected:
 	EDisplayClusterOperationMode DetectOperationMode() const;
 	bool GetResolvedNodeId(const UDisplayClusterConfigurationData* ConfigData, FString& NodeId) const;
 	bool ValidateConfigFile(const FString& FilePath);
+	void PreExitImpl();
 
 private:
 	bool OutOfSync() const;
