nDisplay: Added a 360 azimuthal projection mode to the light card editor 3D viewport, as well as a dropdown that switches from default projection to azimuthal projection

diff --git a/Engine/Plugins/Runtime/nDisplay/Shaders/Private/MeshProjectionShaders.usf b/Engine/Plugins/Runtime/nDisplay/Shaders/Private/MeshProjectionShaders.usf
new file mode 100644
--- /dev/null
+++ b/Engine/Plugins/Runtime/nDisplay/Shaders/Private/MeshProjectionShaders.usf
@@ -0,0 +1,67 @@
+// Copyright Epic Games, Inc. All Rights Reserved.
+
+#include "/Engine/Private/Common.ush"
+#include "/Engine/Generated/Material.ush"
+#include "/Engine/Generated/VertexFactory.ush"
+
+
+struct FMeshProjectionVSToPS
+{
+	FVertexFactoryInterpolantsVSToPS FactoryInterpolants;
+	float4 Position : SV_POSITION;
+	float4 SavedWorldPosition : POSITION1;
+};
+
+void AzimuthalVS(FVertexFactoryInput Input, out FMeshProjectionVSToPS Output)
+{
+	ResolvedView = ResolveView();
+
+	FVertexFactoryIntermediates VFIntermediates = GetVertexFactoryIntermediates(Input);
+	Output.SavedWorldPosition = VertexFactoryGetWorldPosition(Input, VFIntermediates);
+	float3x3 TangentToLocal = VertexFactoryGetTangentToLocal(Input, VFIntermediates);
+
+	FMaterialVertexParameters VertexParameters = GetMaterialVertexParameters(Input, VFIntermediates, Output.SavedWorldPosition.xyz, TangentToLocal);
+	Output.SavedWorldPosition.xyz += GetMaterialWorldPositionOffset(VertexParameters);
+
+	Output.FactoryInterpolants = VertexFactoryGetInterpolantsVSToPS(Input, VFIntermediates, VertexParameters);
+
+	// Project the vertex view position onto a unit sphere
+	float4 ViewPosition = mul(Output.SavedWorldPosition, ResolvedView.TranslatedWorldToView);
+	float rho = length(ViewPosition.xyz);
+	float3 UnitPosition = ViewPosition.xyz / rho;
+	
+	// Compute the position in view space of the point on the projection plane using azimthual equidistance projection
+	float2 PolarCoords = float2(acos(UnitPosition.z), atan2(UnitPosition.y, UnitPosition.x));
+	float3 PlanePosition = float3(PolarCoords.x * cos(PolarCoords.y), PolarCoords.x * sin(PolarCoords.y), 1);
+	
+	// Project the plane position into view space
+	float4 ProjectedPosition = float4(PlanePosition * rho / length(PlanePosition), 1);
+	
+	Output.Position = mul(ProjectedPosition, ResolvedView.ViewToClip);
+}
+
+void PerspectiveVS(FVertexFactoryInput Input, out FMeshProjectionVSToPS Output)
+{
+	ResolvedView = ResolveView();
+
+	FVertexFactoryIntermediates VFIntermediates = GetVertexFactoryIntermediates(Input);
+	Output.SavedWorldPosition = VertexFactoryGetWorldPosition(Input, VFIntermediates);
+	float3x3 TangentToLocal = VertexFactoryGetTangentToLocal(Input, VFIntermediates);
+
+	FMaterialVertexParameters VertexParameters = GetMaterialVertexParameters(Input, VFIntermediates, Output.SavedWorldPosition.xyz, TangentToLocal);
+	Output.SavedWorldPosition.xyz += GetMaterialWorldPositionOffset(VertexParameters);
+
+	Output.FactoryInterpolants = VertexFactoryGetInterpolantsVSToPS(Input, VFIntermediates, VertexParameters);
+	
+	Output.Position = mul(Output.SavedWorldPosition, ResolvedView.TranslatedWorldToClip);
+}
+
+void MainPS(FMeshProjectionVSToPS Input, in float4 SvPosition : SV_Position OPTIONAL_IsFrontFace, out float4 OutputColor : SV_Target0)
+{
+	FMaterialPixelParameters MaterialParameters = GetMaterialPixelParameters(Input.FactoryInterpolants, SvPosition);
+	FPixelMaterialInputs PixelMaterialInputs;
+	CalcMaterialParameters(MaterialParameters, PixelMaterialInputs, SvPosition, bIsFrontFace);
+
+	float3 Emissive = GetMaterialEmissive(PixelMaterialInputs);
+	OutputColor = float4(Emissive, 1);
+}
\ No newline at end of file
diff --git a/Engine/Plugins/Runtime/nDisplay/Source/DisplayCluster/Private/Game/EngineClasses/Scene/DisplayClusterPreviewComponent.cpp b/Engine/Plugins/Runtime/nDisplay/Source/DisplayCluster/Private/Game/EngineClasses/Scene/DisplayClusterPreviewComponent.cpp
--- a/Engine/Plugins/Runtime/nDisplay/Source/DisplayCluster/Private/Game/EngineClasses/Scene/DisplayClusterPreviewComponent.cpp
+++ b/Engine/Plugins/Runtime/nDisplay/Source/DisplayCluster/Private/Game/EngineClasses/Scene/DisplayClusterPreviewComponent.cpp
@@ -218,7 +218,14 @@ void UDisplayClusterPreviewComponent::UpdatePreviewMaterial()
 {
 	if (PreviewMaterialInstance != nullptr)
 	{
-		PreviewMaterialInstance->SetTextureParameterValue(TEXT("Preview"), RenderTarget);
+		if (OverrideTexture)
+		{
+			PreviewMaterialInstance->SetTextureParameterValue(TEXT("Preview"), OverrideTexture);
+		}
+		else
+		{
+			PreviewMaterialInstance->SetTextureParameterValue(TEXT("Preview"), RenderTarget);
+		}
 	}
 }
 
@@ -441,4 +448,13 @@ UTexture2D* UDisplayClusterPreviewComponent::GetOrCreateViewportPreviewTexture2D
 	return nullptr;
 }
 
+void UDisplayClusterPreviewComponent::SetOverrideTexture(UTexture* InOverrideTexture)
+{
+	if (OverrideTexture != InOverrideTexture)
+	{
+		OverrideTexture = InOverrideTexture;
+		UpdatePreviewMaterial();
+	}
+}
+
 #endif
diff --git a/Engine/Plugins/Runtime/nDisplay/Source/DisplayCluster/Public/Components/DisplayClusterPreviewComponent.h b/Engine/Plugins/Runtime/nDisplay/Source/DisplayCluster/Public/Components/DisplayClusterPreviewComponent.h
--- a/Engine/Plugins/Runtime/nDisplay/Source/DisplayCluster/Public/Components/DisplayClusterPreviewComponent.h
+++ b/Engine/Plugins/Runtime/nDisplay/Source/DisplayCluster/Public/Components/DisplayClusterPreviewComponent.h
@@ -75,6 +75,9 @@ public:
 	/** Create and retrieve a render texture 2d from the render target. */
 	UTexture2D* GetOrCreateViewportPreviewTexture2D();
 
+	/** Sets an override texture to display on the viewport instead of the render target */
+	void SetOverrideTexture(UTexture* InOverrideTexture);
+
 protected:
 	class IDisplayClusterViewport* GetCurrentViewport() const;
 	bool GetPreviewTextureSettings(FIntPoint& OutSize, float& OutGamma) const;
@@ -145,5 +148,8 @@ private:
 	UPROPERTY(Transient)
 	UTexture2D* PreviewTexture = nullptr;
 
+	UPROPERTY(Transient)
+	UTexture* OverrideTexture = nullptr;
+
 #endif /*WITH_EDITORONLY_DATA*/
 };
diff --git a/Engine/Plugins/Runtime/nDisplay/Source/DisplayClusterLightCardEditor/DisplayClusterLightCardEditor.Build.cs b/Engine/Plugins/Runtime/nDisplay/Source/DisplayClusterLightCardEditor/DisplayClusterLightCardEditor.Build.cs
--- a/Engine/Plugins/Runtime/nDisplay/Source/DisplayClusterLightCardEditor/DisplayClusterLightCardEditor.Build.cs
+++ b/Engine/Plugins/Runtime/nDisplay/Source/DisplayClusterLightCardEditor/DisplayClusterLightCardEditor.Build.cs
@@ -12,10 +12,14 @@ public class DisplayClusterLightCardEditor : ModuleRules
 				"DisplayCluster",
 				"DisplayClusterConfiguration",
 				"DisplayClusterOperator",
+				"DisplayClusterLightCardEditorShaders",
 				
 				"Core",
 				"CoreUObject",
 				"Engine",
+				"RenderCore",
+				"Renderer",
+				"RHI",
 				"InputCore",
 				"Slate",
 				"SlateCore",
diff --git a/Engine/Plugins/Runtime/nDisplay/Source/DisplayClusterLightCardEditor/Private/DisplayClusterLightCardEditorCommands.cpp b/Engine/Plugins/Runtime/nDisplay/Source/DisplayClusterLightCardEditor/Private/DisplayClusterLightCardEditorCommands.cpp
--- a/Engine/Plugins/Runtime/nDisplay/Source/DisplayClusterLightCardEditor/Private/DisplayClusterLightCardEditorCommands.cpp
+++ b/Engine/Plugins/Runtime/nDisplay/Source/DisplayClusterLightCardEditor/Private/DisplayClusterLightCardEditorCommands.cpp
@@ -8,6 +8,9 @@
 void FDisplayClusterLightCardEditorCommands::RegisterCommands()
 {
 	UI_COMMAND(ResetCamera, "Reset Camera", "Resets the camera to focus on the mesh", EUserInterfaceActionType::Button, FInputChord());
+
+	UI_COMMAND(PerspectiveProjection, "Perspective", "A perspective projection from the stage's view origin", EUserInterfaceActionType::Button, FInputChord());
+	UI_COMMAND(AzimuthalProjection, "Azimuthal", "An azimuthal hemisphere projection from the stage's center pointed upwards", EUserInterfaceActionType::Button, FInputChord());
 }
 
 #undef LOCTEXT_NAMESPACE
diff --git a/Engine/Plugins/Runtime/nDisplay/Source/DisplayClusterLightCardEditor/Private/DisplayClusterLightCardEditorCommands.h b/Engine/Plugins/Runtime/nDisplay/Source/DisplayClusterLightCardEditor/Private/DisplayClusterLightCardEditorCommands.h
--- a/Engine/Plugins/Runtime/nDisplay/Source/DisplayClusterLightCardEditor/Private/DisplayClusterLightCardEditorCommands.h
+++ b/Engine/Plugins/Runtime/nDisplay/Source/DisplayClusterLightCardEditor/Private/DisplayClusterLightCardEditorCommands.h
@@ -21,4 +21,7 @@ public:
 public:
 	// Viewport commands
 	TSharedPtr<FUICommandInfo> ResetCamera;
+
+	TSharedPtr<FUICommandInfo> PerspectiveProjection;
+	TSharedPtr<FUICommandInfo> AzimuthalProjection;
 };
diff --git a/Engine/Plugins/Runtime/nDisplay/Source/DisplayClusterLightCardEditor/Private/Viewport/DisplayClusterLightCardEditorViewportClient.cpp b/Engine/Plugins/Runtime/nDisplay/Source/DisplayClusterLightCardEditor/Private/Viewport/DisplayClusterLightCardEditorViewportClient.cpp
--- a/Engine/Plugins/Runtime/nDisplay/Source/DisplayClusterLightCardEditor/Private/Viewport/DisplayClusterLightCardEditorViewportClient.cpp
+++ b/Engine/Plugins/Runtime/nDisplay/Source/DisplayClusterLightCardEditor/Private/Viewport/DisplayClusterLightCardEditorViewportClient.cpp
@@ -2,36 +2,55 @@
 
 #include "DisplayClusterLightCardEditorViewportClient.h"
 
+#include "DisplayClusterConfigurationTypes.h"
 #include "DisplayClusterRootActor.h"
+#include "Components/DisplayClusterPreviewComponent.h"
+#include "Components/DisplayClusterCameraComponent.h"
+#include "Components/DisplayClusterScreenComponent.h"
 
 #include "AudioDevice.h"
 #include "EngineUtils.h"
+#include "EditorModes.h"
+#include "EditorModeManager.h"
 #include "PreviewScene.h"
 #include "UnrealEdGlobals.h"
 #include "ScopedTransaction.h"
 #include "SDisplayClusterLightCardEditor.h"
 #include "UnrealWidget.h"
+#include "RayTracingDebugVisualizationMenuCommands.h"
 #include "Components/DirectionalLightComponent.h"
 #include "Components/LineBatchComponent.h"
 #include "Components/PostProcessComponent.h"
 #include "Components/SkyLightComponent.h"
+#include "Components/PrimitiveComponent.h"
 #include "Editor/UnrealEdEngine.h"
 #include "Kismet/GameplayStatics.h"
 #include "Slate/SceneViewport.h"
 #include "Widgets/Docking/SDockTab.h"
+#include "CustomEditorStaticScreenPercentage.h"
+#include "LegacyScreenPercentageDriver.h"
+#include "EngineModule.h"
+#include "Debug/DebugDrawService.h"
+#include "Engine/Canvas.h"
+#include "Engine/TextureRenderTarget2D.h"
+#include "CameraController.h"
+
+
+#include "Renderer/Private/SceneRendering.h"
 
 #define LOCTEXT_NAMESPACE "DisplayClusterLightCardEditorViewportClient"
 
 FDisplayClusterLightCardEditorViewportClient::FDisplayClusterLightCardEditorViewportClient(FAdvancedPreviewScene& InPreviewScene,
                                                            const TWeakPtr<SEditorViewport>& InEditorViewportWidget,
                                                            TWeakPtr<SDisplayClusterLightCardEditor> InLightCardEditor) :
-	FEditorViewportClient(nullptr, &InPreviewScene, InEditorViewportWidget),
-	PreviewActorBounds(ForceInitToZero)
+	FEditorViewportClient(nullptr, &InPreviewScene, InEditorViewportWidget)
 {
 	check (InLightCardEditor.IsValid());
 	
 	LightCardEditorPtr = InLightCardEditor;
 	
+	MeshProjectionRenderer = MakeShared<FDisplayClusterMeshProjectionRenderer>();
+
 	SelectedActor = nullptr;
 	bDraggingActor = false;
 	ScopedTransaction = nullptr;
@@ -48,17 +67,20 @@ FDisplayClusterLightCardEditorViewportClient::FDisplayClusterLightCardEditorView
 	
 	ShowWidget(true);
 
-	SetViewMode(VMI_Lit);
+	SetViewMode(VMI_Unlit);
 	
 	ViewportType = LVT_Perspective;
 	bSetListenerPosition = false;
+	bUseNumpadCameraControl = false;
 	SetRealtime(true);
 	SetShowStats(true);
 
+	ProjectionFOVs.AddZeroed(2);
+	ProjectionFOVs[(int32)EDisplayClusterMeshProjectionType::Perspective] = 90.0f;
+	ProjectionFOVs[(int32)EDisplayClusterMeshProjectionType::Azimuthal] = 130.0f;
+
 	//This seems to be needed to get the correct world time in the preview.
 	SetIsSimulateInEditorViewport(true);
-
-	ResetCamera();
 	
 	UpdatePreviewActor(LightCardEditorPtr.Pin()->GetActiveRootActor().Get());
 }
@@ -77,6 +99,21 @@ void FDisplayClusterLightCardEditorViewportClient::Tick(float DeltaSeconds)
 {
 	FEditorViewportClient::Tick(DeltaSeconds);
 
+	// Camera position is locked to a specific location
+	FVector Location = FVector::ZeroVector;
+	if (ProjectionOriginComponent.IsValid())
+	{
+		Location = ProjectionOriginComponent->GetComponentLocation();
+	}
+
+	SetViewLocation(Location);
+
+	// View rotation is also locked for the azimuthal projection
+	if (ProjectionMode != EDisplayClusterMeshProjectionType::Perspective)
+	{
+		SetViewRotation(FVector::UpVector.Rotation());
+	}
+
 	// Tick the preview scene world.
 	if (!GIntraFrameDebuggingGameThread)
 	{
@@ -90,11 +127,185 @@ void FDisplayClusterLightCardEditorViewportClient::Tick(float DeltaSeconds)
 			PreviewScene->GetWorld()->Tick(IsRealtime() ? LEVELTICK_ViewportsOnly : LEVELTICK_TimeOnly, DeltaSeconds);
 		}
 	}
+
+	if (SpawnedRootActor.IsValid() && RootActorLevelInstance.IsValid())
+	{
+		// Pass the preview render targets from the level instance root actor to the preview root actor
+		UDisplayClusterConfigurationData* Config = RootActorLevelInstance->GetConfigData();
+
+		for (const TPair<FString, UDisplayClusterConfigurationClusterNode*>& NodePair : Config->Cluster->Nodes)
+		{
+			const UDisplayClusterConfigurationClusterNode* Node = NodePair.Value;
+			for (const TPair<FString, UDisplayClusterConfigurationViewport*>& ViewportPair : Node->Viewports)
+			{
+				UDisplayClusterPreviewComponent* LevelInstancePreviewComp = RootActorLevelInstance->GetPreviewComponent(NodePair.Key, ViewportPair.Key);
+				UDisplayClusterPreviewComponent* PreviewComp = SpawnedRootActor->GetPreviewComponent(NodePair.Key, ViewportPair.Key);
+
+				if (PreviewComp && LevelInstancePreviewComp)
+				{
+					PreviewComp->SetOverrideTexture(LevelInstancePreviewComp->GetRenderTargetTexture());
+				}
+			}
+		}
+	}
 }
 
 void FDisplayClusterLightCardEditorViewportClient::Draw(FViewport* InViewport, FCanvas* Canvas)
 {
-	FEditorViewportClient::Draw(InViewport, Canvas);
+	FViewport* ViewportBackup = Viewport;
+	Viewport = InViewport ? InViewport : Viewport;
+
+	UWorld* World = GetWorld();
+	FGameTime Time;
+	if (!World || (GetScene() != World->Scene) || UseAppTime()) 
+	{
+		Time = FGameTime::GetTimeSinceAppStart();
+	}
+	else
+	{
+		Time = World->GetTime();
+	}
+
+	FEngineShowFlags UseEngineShowFlags = EngineShowFlags;
+	if (OverrideShowFlagsFunc)
+	{
+		OverrideShowFlagsFunc(UseEngineShowFlags);
+	}
+
+	// Setup a FSceneViewFamily/FSceneView for the viewport.
+	FSceneViewFamilyContext ViewFamily(FSceneViewFamily::ConstructionValues(
+		Canvas->GetRenderTarget(),
+		GetScene(),
+		UseEngineShowFlags)
+		.SetTime(Time)
+		.SetRealtimeUpdate(IsRealtime() && FSlateThrottleManager::Get().IsAllowingExpensiveTasks())
+	);
+
+	ViewFamily.DebugDPIScale = GetDPIScale();
+	ViewFamily.bIsHDR = Viewport->IsHDRViewport();
+
+	ViewFamily.EngineShowFlags = UseEngineShowFlags;
+	ViewFamily.EngineShowFlags.CameraInterpolation = 0;
+	ViewFamily.EngineShowFlags.SetScreenPercentage(false);
+
+	ViewFamily.ViewExtensions = GEngine->ViewExtensions->GatherActiveExtensions(FSceneViewExtensionContext(InViewport));
+
+	for (auto ViewExt : ViewFamily.ViewExtensions)
+	{
+		ViewExt->SetupViewFamily(ViewFamily);
+	}
+
+	ViewFamily.ViewMode = VMI_Unlit;
+
+	EngineShowFlagOverride(ESFIM_Editor, ViewFamily.ViewMode, ViewFamily.EngineShowFlags, false);
+	EngineShowFlagOrthographicOverride(IsPerspective(), ViewFamily.EngineShowFlags);
+
+	ViewFamily.ExposureSettings = ExposureSettings;
+
+	// Setup the screen percentage and upscaling method for the view family.
+	{
+		checkf(ViewFamily.GetScreenPercentageInterface() == nullptr,
+			TEXT("Some code has tried to set up an alien screen percentage driver, that could be wrong if not supported very well by the RHI."));
+
+		if (SupportsLowDPIPreview() && IsLowDPIPreview() && ViewFamily.SupportsScreenPercentage())
+		{
+			ViewFamily.SecondaryViewFraction = GetDPIDerivedResolutionFraction();
+		}
+	}
+
+	FSceneView* View = nullptr;
+
+	View = CalcSceneView(&ViewFamily, INDEX_NONE);
+	SetupViewForRendering(ViewFamily,*View);
+
+	FSlateRect SafeFrame;
+	View->CameraConstrainedViewRect = View->UnscaledViewRect;
+	if (CalculateEditorConstrainedViewRect(SafeFrame, Viewport, Canvas->GetDPIScale()))
+	{
+		View->CameraConstrainedViewRect = FIntRect(SafeFrame.Left, SafeFrame.Top, SafeFrame.Right, SafeFrame.Bottom);
+	}
+
+	{
+		// If a screen percentage interface was not set by one of the view extension, then set the legacy one.
+		if (ViewFamily.GetScreenPercentageInterface() == nullptr)
+		{
+			float GlobalResolutionFraction = 1.0f;
+
+			if (SupportsPreviewResolutionFraction() && ViewFamily.SupportsScreenPercentage())
+			{
+				GlobalResolutionFraction = GetDefaultPrimaryResolutionFractionTarget();
+
+				// Force screen percentage's engine show flag to be turned on for preview screen percentage.
+				ViewFamily.EngineShowFlags.ScreenPercentage = (GlobalResolutionFraction != 1.0);
+			}
+
+			// In editor viewport, we ignore r.ScreenPercentage and FPostProcessSettings::ScreenPercentage by design.
+			ViewFamily.SetScreenPercentageInterface(new FLegacyScreenPercentageDriver(ViewFamily, GlobalResolutionFraction));
+		}
+
+		check(ViewFamily.GetScreenPercentageInterface() != nullptr);
+	}
+
+	Canvas->Clear(FLinearColor::Black);
+
+	FSceneViewInitOptions ViewInitOptions;
+	GetSceneViewInitOptions(ViewInitOptions);
+
+	MeshProjectionRenderer->Render(Canvas, GetScene(), ViewInitOptions, ProjectionMode);
+
+	if (View)
+	{
+		DrawCanvas(*Viewport, *View, *Canvas);
+	}
+
+	// Remove temporary debug lines.
+	// Possibly a hack. Lines may get added without the scene being rendered etc.
+	if (World->LineBatcher != NULL && (World->LineBatcher->BatchedLines.Num() || World->LineBatcher->BatchedPoints.Num() || World->LineBatcher->BatchedMeshes.Num() ) )
+	{
+		World->LineBatcher->Flush();
+	}
+
+	if (World->ForegroundLineBatcher != NULL && (World->ForegroundLineBatcher->BatchedLines.Num() || World->ForegroundLineBatcher->BatchedPoints.Num() || World->ForegroundLineBatcher->BatchedMeshes.Num() ) )
+	{
+		World->ForegroundLineBatcher->Flush();
+	}
+
+	// Draw the widget.
+	/*if (Widget && bShowWidget)
+	{
+		Widget->DrawHUD( Canvas );
+	}*/
+
+	// Axes indicators
+	if (bDrawAxes && !ViewFamily.EngineShowFlags.Game && !GLevelEditorModeTools().IsViewportUIHidden() && !IsVisualizeCalibrationMaterialEnabled())
+	{
+		// TODO: Figure out how we want the axes widget to be drawn
+		DrawAxes(Viewport, Canvas);
+	}
+
+	// NOTE: DebugCanvasObject will be created by UDebugDrawService::Draw() if it doesn't already exist.
+	FCanvas* DebugCanvas = Viewport->GetDebugCanvas();
+	UDebugDrawService::Draw(ViewFamily.EngineShowFlags, Viewport, View, DebugCanvas);
+	UCanvas* DebugCanvasObject = FindObjectChecked<UCanvas>(GetTransientPackage(),TEXT("DebugCanvasObject"));
+	DebugCanvasObject->Canvas = DebugCanvas;
+	DebugCanvasObject->Init( Viewport->GetSizeXY().X, Viewport->GetSizeXY().Y, View , DebugCanvas);
+
+	// Stats display
+	if( IsRealtime() && ShouldShowStats() && DebugCanvas)
+	{
+		const int32 XPos = 4;
+		TArray< FDebugDisplayProperty > EmptyPropertyArray;
+		DrawStatsHUD( World, Viewport, DebugCanvas, NULL, EmptyPropertyArray, GetViewLocation(), GetViewRotation() );
+	}
+
+	if(!IsRealtime())
+	{
+		// Wait for the rendering thread to finish drawing the view before returning.
+		// This reduces the apparent latency of dragging the viewport around.
+		FlushRenderingCommands();
+	}
+
+	Viewport = ViewportBackup;
 }
 
 UE::Widget::EWidgetMode FDisplayClusterLightCardEditorViewportClient::GetWidgetMode() const
@@ -117,20 +328,56 @@ FVector FDisplayClusterLightCardEditorViewportClient::GetWidgetLocation() const
 	return FEditorViewportClient::GetWidgetLocation();
 }
 
+FSceneView* FDisplayClusterLightCardEditorViewportClient::CalcSceneView(FSceneViewFamily* ViewFamily, const int32 StereoViewIndex)
+{
+	FSceneViewInitOptions ViewInitOptions;
+	GetSceneViewInitOptions(ViewInitOptions);
+
+	ViewInitOptions.ViewFamily = ViewFamily;
+
+	TimeForForceRedraw = 0.0;
+
+	FSceneView* View = new FSceneView(ViewInitOptions);
+
+	View->SubduedSelectionOutlineColor = GEngine->GetSubduedSelectionOutlineColor();
+
+	int32 FamilyIndex = ViewFamily->Views.Add(View);
+	check(FamilyIndex == View->StereoViewIndex || View->StereoViewIndex == INDEX_NONE);
+
+	View->StartFinalPostprocessSettings(View->ViewLocation);
+
+	OverridePostProcessSettings(*View);
+
+	View->EndFinalPostprocessSettings(ViewInitOptions);
+
+	for (int ViewExt = 0; ViewExt < ViewFamily->ViewExtensions.Num(); ViewExt++)
+	{
+		ViewFamily->ViewExtensions[ViewExt]->SetupView(*ViewFamily, *View);
+	}
+
+	return View;
+}
+
 bool FDisplayClusterLightCardEditorViewportClient::InputKey(FViewport* InViewport, int32 ControllerId, FKey Key, EInputEvent Event,
                                             float AmountDepressed, bool bGamepad)
 {
+	if ((Key == EKeys::MouseScrollUp || Key == EKeys::MouseScrollDown) && Event == IE_Pressed)
+	{
+		const int32 Sign = Key == EKeys::MouseScrollUp ? -1 : 1;
+		const float CurrentFOV = GetProjectionModeFOV(ProjectionMode);
+		const float NewFOV = FMath::Clamp(CurrentFOV + Sign * FOVScrollIncrement, CameraController->GetConfig().MinimumAllowedFOV, CameraController->GetConfig().MaximumAllowedFOV);
+
+		SetProjectionModeFOV(ProjectionMode, NewFOV);
+		return true;
+	}
+
 	return FEditorViewportClient::InputKey(InViewport, ControllerId, Key, Event, AmountDepressed, bGamepad);
 }
 
 bool FDisplayClusterLightCardEditorViewportClient::InputWidgetDelta(FViewport* InViewport, EAxisList::Type CurrentAxis, FVector& Drag,
                                                     FRotator& Rot, FVector& Scale)
 {
-	if (IsActorSelected() && bDraggingActor)
-	{
-		GEditor->ApplyDeltaToActor(SelectedActor.Get(), true, &Drag, &Rot, &Scale);
-		return true;
-	}
+	// TODO: Add logic for moving a light card in "2D" space, modifying its long and lat, and making sure it is flush with the nearest screen
 
 	return FEditorViewportClient::InputWidgetDelta(InViewport, CurrentAxis, Drag, Rot, Scale);
 }
@@ -170,23 +417,11 @@ void FDisplayClusterLightCardEditorViewportClient::TrackingStopped()
 void FDisplayClusterLightCardEditorViewportClient::ProcessClick(FSceneView& View, HHitProxy* HitProxy, FKey Key, EInputEvent Event,
                                                 uint32 HitX, uint32 HitY)
 {
-	if (HitProxy && HitProxy->IsA(HActor::StaticGetType()))
-	{
-		HActor* HitActor = static_cast<HActor*>(HitProxy);
-		SelectActor(HitActor->Actor);
-		return;
-	}
-
-	SelectActor(nullptr);
+	// TODO: Add logic to check if the hit proxy was a light card, and select that light card
 	
 	FEditorViewportClient::ProcessClick(View, HitProxy, Key, Event, HitX, HitY);
 }
 
-void FDisplayClusterLightCardEditorViewportClient::SetSceneViewport(TSharedPtr<FSceneViewport> InViewport)
-{
-	SceneViewportPtr = InViewport;
-}
-
 void FDisplayClusterLightCardEditorViewportClient::SelectActor(AActor* NewActor)
 {
 }
@@ -197,20 +432,7 @@ void FDisplayClusterLightCardEditorViewportClient::ResetSelection()
 	SetWidgetMode(UE::Widget::WM_None);
 }
 
-void FDisplayClusterLightCardEditorViewportClient::ResetCamera()
-{
-	ToggleOrbitCamera(true);
-	SetViewLocationForOrbiting(PreviewActorBounds.Origin);
-
-	// TODO: Better define defaults.
-	SetViewLocation(EditorViewportDefs::DefaultPerspectiveViewLocation + FVector(0.f, -400.f, 0.f));
-	SetViewRotation(EditorViewportDefs::DefaultPerspectiveViewRotation + FRotator(0.f, 180.f, 0.f));
-
-	Invalidate();
-}
-
-void FDisplayClusterLightCardEditorViewportClient::UpdatePreviewActor(
-	ADisplayClusterRootActor* RootActor)
+void FDisplayClusterLightCardEditorViewportClient::UpdatePreviewActor(ADisplayClusterRootActor* RootActor)
 {
 	UWorld* PreviewWorld = PreviewScene->GetWorld();
 	check (PreviewWorld);
@@ -219,10 +441,15 @@ void FDisplayClusterLightCardEditorViewportClient::UpdatePreviewActor(
 	{
 		PreviewWorld->EditorDestroyActor(SpawnedRootActor.Get(), false);
 		SpawnedRootActor.Reset();
+		RootActorLevelInstance.Reset();
+
+		MeshProjectionRenderer->ClearScene();
 	}
 	
 	if (RootActor)
 	{
+		RootActorLevelInstance = RootActor;
+
 		const FVector SpawnLocation = FVector::ZeroVector;
 		const FRotator SpawnRotation = FRotator::ZeroRotator;
 		
@@ -235,32 +462,63 @@ void FDisplayClusterLightCardEditorViewportClient::UpdatePreviewActor(
 		SpawnedRootActor = CastChecked<ADisplayClusterRootActor>(PreviewWorld->SpawnActor(RootActor->GetClass(),
 			&SpawnLocation, &SpawnRotation, MoveTemp(SpawnInfo)));
 
-		// Compute actor bounds as the sum of its visible parts
-		PreviewActorBounds = FBoxSphereBounds(ForceInitToZero);
-		for (UActorComponent* Component : RootActor->GetComponents())
+		// Spawned actor will take the transform values from the template, so manually reset them to zero here
+		SpawnedRootActor->SetActorLocation(FVector::ZeroVector);
+		SpawnedRootActor->SetActorRotation(FRotator::ZeroRotator);
+
+		FindProjectionOriginComponent();
+
+		// Filter out any primitives hidden in game except screen components
+		MeshProjectionRenderer->AddActor(SpawnedRootActor.Get(), [](const UPrimitiveComponent* PrimitiveComponent)
 		{
-			// Aggregate primitive components that either have collision enabled or are otherwise visible components in-game
-			if (UPrimitiveComponent* PrimComp = Cast<UPrimitiveComponent>(Component))
-			{
-				if (PrimComp->IsRegistered() && (!PrimComp->bHiddenInGame || PrimComp->IsCollisionEnabled()) && PrimComp->Bounds.SphereRadius < HALF_WORLD_MAX)
-				{
-					PreviewActorBounds = PreviewActorBounds + PrimComp->Bounds;
-				}
-			}
-		}
+			return !PrimitiveComponent->bHiddenInGame || PrimitiveComponent->IsA<UDisplayClusterScreenComponent>();
+		});
 
 		SpawnedRootActor->UpdatePreviewComponents();
 	}
 }
 
-bool FDisplayClusterLightCardEditorViewportClient::GetShowGrid() const
+void FDisplayClusterLightCardEditorViewportClient::SetProjectionMode(EDisplayClusterMeshProjectionType InProjectionMode)
 {
-	return IsSetShowGridChecked();
+	ProjectionMode = InProjectionMode;
+
+	if (ProjectionMode == EDisplayClusterMeshProjectionType::Perspective)
+	{
+		// TODO: Do we want to cache the perspective rotation and restore it when the user switches back?
+		SetViewRotation(FVector::ForwardVector.Rotation());
+	}
+	else if (ProjectionMode == EDisplayClusterMeshProjectionType::Azimuthal)
+	{
+		SetViewRotation(FVector::UpVector.Rotation());
+	}
+
+	FindProjectionOriginComponent();
+}
+
+float FDisplayClusterLightCardEditorViewportClient::GetProjectionModeFOV(EDisplayClusterMeshProjectionType InProjectionMode) const
+{
+	int32 ProjectionModeIndex = (int32)InProjectionMode;
+	if (ProjectionFOVs.Num() > ProjectionModeIndex)
+	{
+		return ProjectionFOVs[ProjectionModeIndex];
+	}
+	else
+	{
+		return ViewFOV;
+	}
 }
 
-void FDisplayClusterLightCardEditorViewportClient::ToggleShowGrid()
+void FDisplayClusterLightCardEditorViewportClient::SetProjectionModeFOV(EDisplayClusterMeshProjectionType InProjectionMode, float NewFOV)
 {
-	SetShowGrid();
+	int32 ProjectionModeIndex = (int32)InProjectionMode;
+	if (ProjectionFOVs.Num() > ProjectionModeIndex)
+	{
+		ProjectionFOVs[ProjectionModeIndex] = NewFOV;
+	}
+	else
+	{
+		ViewFOV = NewFOV;
+	}
 }
 
 void FDisplayClusterLightCardEditorViewportClient::BeginTransaction(const FText& Description)
@@ -280,4 +538,148 @@ void FDisplayClusterLightCardEditorViewportClient::EndTransaction()
 	}
 }
 
+void FDisplayClusterLightCardEditorViewportClient::GetScenePrimitiveComponents(TArray<UPrimitiveComponent*>& OutPrimitiveComponents)
+{
+	SpawnedRootActor->ForEachComponent<UPrimitiveComponent>(true, [&](UPrimitiveComponent* PrimitiveComponent)
+	{
+		OutPrimitiveComponents.Add(PrimitiveComponent);
+	});
+}
+
+void FDisplayClusterLightCardEditorViewportClient::GetSceneViewInitOptions(FSceneViewInitOptions& OutViewInitOptions)
+{
+	FSceneViewInitOptions ViewInitOptions;
+
+	FViewportCameraTransform& ViewTransform = GetViewTransform();
+
+	ViewInitOptions.ViewLocation = ViewTransform.GetLocation();
+	ViewInitOptions.ViewRotation = ViewTransform.GetRotation();
+	ViewInitOptions.ViewOrigin = ViewInitOptions.ViewLocation;
+
+	FIntPoint ViewportSize = Viewport->GetSizeXY();
+	ViewportSize.X = FMath::Max(ViewportSize.X, 1);
+	ViewportSize.Y = FMath::Max(ViewportSize.Y, 1);
+	FIntPoint ViewportOffset(0, 0);
+
+	ViewInitOptions.SetViewRectangle(FIntRect(ViewportOffset, ViewportOffset + ViewportSize));
+
+	AWorldSettings* WorldSettings = nullptr;
+	if (GetScene() != nullptr && GetScene()->GetWorld() != nullptr)
+	{
+		WorldSettings = GetScene()->GetWorld()->GetWorldSettings();
+	}
+
+	if (WorldSettings != nullptr)
+	{
+		ViewInitOptions.WorldToMetersScale = WorldSettings->WorldToMeters;
+	}
+
+	// Rotate view 90 degrees
+	ViewInitOptions.ViewRotationMatrix = CalcViewRotationMatrix(ViewInitOptions.ViewRotation) * FMatrix(
+		FPlane(0, 0, 1, 0),
+		FPlane(1, 0, 0, 0),
+		FPlane(0, 1, 0, 0),
+		FPlane(0, 0, 0, 1));
+
+	const float MinZ = GetNearClipPlane();
+	const float MaxZ = MinZ;
+	const float FieldOfView = GetProjectionModeFOV(ProjectionMode);
+
+	// Avoid zero ViewFOV's which cause divide by zero's in projection matrix
+	const float MatrixFOV = FMath::Max(0.001f, FieldOfView) * (float)PI / 360.0f;
+
+	float XAxisMultiplier;
+	float YAxisMultiplier;
+
+	EAspectRatioAxisConstraint AspectRatioAxisConstraint = GetDefault<ULevelEditorViewportSettings>()->AspectRatioAxisConstraint;
+
+	if (((ViewportSize.X > ViewportSize.Y) && (AspectRatioAxisConstraint == AspectRatio_MajorAxisFOV)) || (AspectRatioAxisConstraint == AspectRatio_MaintainXFOV))
+	{
+		//if the viewport is wider than it is tall
+		XAxisMultiplier = 1.0f;
+		YAxisMultiplier = ViewportSize.X / (float)ViewportSize.Y;
+	}
+	else
+	{
+		//if the viewport is taller than it is wide
+		XAxisMultiplier = ViewportSize.Y / (float)ViewportSize.X;
+		YAxisMultiplier = 1.0f;
+	}
+
+	if ((bool)ERHIZBuffer::IsInverted)
+	{
+		ViewInitOptions.ProjectionMatrix = FReversedZPerspectiveMatrix(
+			MatrixFOV,
+			MatrixFOV,
+			XAxisMultiplier,
+			YAxisMultiplier,
+			MinZ,
+			MaxZ
+			);
+	}
+	else
+	{
+		ViewInitOptions.ProjectionMatrix = FPerspectiveMatrix(
+			MatrixFOV,
+			MatrixFOV,
+			XAxisMultiplier,
+			YAxisMultiplier,
+			MinZ,
+			MaxZ
+			);
+	}
+
+	if (!ViewInitOptions.IsValidViewRectangle())
+	{
+		// Zero sized rects are invalid, so fake to 1x1 to avoid asserts later on
+		ViewInitOptions.SetViewRectangle(FIntRect(0, 0, 1, 1));
+	}
+
+	ViewInitOptions.SceneViewStateInterface = ViewState.GetReference();
+	ViewInitOptions.ViewElementDrawer = this;
+
+	ViewInitOptions.BackgroundColor = GetBackgroundColor();
+
+	ViewInitOptions.EditorViewBitflag = (uint64)1 << ViewIndex, // send the bit for this view - each actor will check it's visibility bits against this
+
+	// for ortho views to steal perspective view origin
+	ViewInitOptions.OverrideLODViewOrigin = FVector::ZeroVector;
+	ViewInitOptions.bUseFauxOrthoViewPos = true;
+
+	ViewInitOptions.FOV = FieldOfView;
+	ViewInitOptions.OverrideFarClippingPlaneDistance = GetFarClipPlaneOverride();
+	ViewInitOptions.CursorPos = CurrentMousePos;
+
+	OutViewInitOptions = ViewInitOptions;
+}
+
+void FDisplayClusterLightCardEditorViewportClient::FindProjectionOriginComponent()
+{
+	if (SpawnedRootActor.IsValid())
+	{
+		if (ProjectionMode == EDisplayClusterMeshProjectionType::Perspective)
+		{
+			TArray<UDisplayClusterCameraComponent*> ViewOriginComponents;
+			SpawnedRootActor->GetComponents<UDisplayClusterCameraComponent>(ViewOriginComponents);
+
+			if (ViewOriginComponents.Num())
+			{
+				ProjectionOriginComponent = ViewOriginComponents[0];
+			}
+			else
+			{
+				ProjectionOriginComponent = SpawnedRootActor->GetRootComponent();
+			}
+		}
+		else if (ProjectionMode == EDisplayClusterMeshProjectionType::Azimuthal)
+		{
+			ProjectionOriginComponent = SpawnedRootActor->GetRootComponent();
+		}
+	}
+	else
+	{
+		ProjectionOriginComponent.Reset();
+	}
+}
+
 #undef LOCTEXT_NAMESPACE
diff --git a/Engine/Plugins/Runtime/nDisplay/Source/DisplayClusterLightCardEditor/Private/Viewport/DisplayClusterLightCardEditorViewportClient.h b/Engine/Plugins/Runtime/nDisplay/Source/DisplayClusterLightCardEditor/Private/Viewport/DisplayClusterLightCardEditorViewportClient.h
--- a/Engine/Plugins/Runtime/nDisplay/Source/DisplayClusterLightCardEditor/Private/Viewport/DisplayClusterLightCardEditorViewportClient.h
+++ b/Engine/Plugins/Runtime/nDisplay/Source/DisplayClusterLightCardEditor/Private/Viewport/DisplayClusterLightCardEditorViewportClient.h
@@ -7,6 +7,8 @@
 #include "SEditorViewport.h"
 #include "AdvancedPreviewScene.h"
 
+#include "DisplayClusterMeshProjectionRenderer.h"
+
 class ADisplayClusterRootActor;
 class SDisplayClusterLightCardEditor;
 class FScopedTransaction;
@@ -27,35 +29,34 @@ public:
 	virtual bool CanCycleWidgetMode() const override { return true; }
 	virtual UE::Widget::EWidgetMode GetWidgetMode() const override;
 	virtual FVector GetWidgetLocation() const override;
+	virtual FSceneView* CalcSceneView(FSceneViewFamily* ViewFamily, const int32 StereoViewIndex = INDEX_NONE) override;
 	virtual bool IsLevelEditorClient() const override { return false; }
 	virtual bool InputKey(FViewport* InViewport, int32 ControllerId, FKey Key, EInputEvent Event, float AmountDepressed, bool bGamepad) override;
 	virtual bool InputWidgetDelta(FViewport* InViewport, EAxisList::Type CurrentAxis, FVector& Drag, FRotator& Rot, FVector& Scale) override;
 	virtual void TrackingStarted(const FInputEventState& InInputState, bool bIsDraggingWidget, bool bNudge) override;
 	virtual void TrackingStopped() override;
 	virtual void ProcessClick(FSceneView& View, HHitProxy* HitProxy, FKey Key, EInputEvent Event, uint32 HitX, uint32 HitY) override;
+	virtual ELevelViewportType GetViewportType() const override { return LVT_Perspective; }
 	// ~FEditorViewportClient
 
-	void SetSceneViewport(TSharedPtr<FSceneViewport> InViewport);
 	void SelectActor(AActor* NewActor);
 	void ResetSelection();
-	void ResetCamera();
 
 	void UpdatePreviewActor(ADisplayClusterRootActor* RootActor);
 	
-	/**
-	 * Returns true if the grid is currently visible in the viewport
-	 */
-	bool GetShowGrid() const;
-
-	/**
-	 * Will toggle the grid's visibility in the viewport
-	 */
-	void ToggleShowGrid();
-	
 	AActor* GetSelectedActor() const { return SelectedActor.Get(); }
 
 	bool IsActorSelected() const { return SelectedActor.IsValid(); }
 	
+	EDisplayClusterMeshProjectionType GetProjectionMode() const { return ProjectionMode; }
+	void SetProjectionMode(EDisplayClusterMeshProjectionType InProjectionMode);
+
+	/** Gets the field of view of the specified projection mode */
+	float GetProjectionModeFOV(EDisplayClusterMeshProjectionType InProjectionMode) const;
+
+	/** Sets the field of view of the specified projection mode */
+	void SetProjectionModeFOV(EDisplayClusterMeshProjectionType InProjectionMode, float NewFOV);
+
 protected:
 	/** Initiates a transaction. */
 	void BeginTransaction(const FText& Description);
@@ -63,17 +64,39 @@ protected:
 	/** Ends the current transaction, if one exists. */
 	void EndTransaction();
 
+	/** Gets a list of all primitive components to be rendered in the scene */
+	void GetScenePrimitiveComponents(TArray<UPrimitiveComponent*>& OutPrimitiveComponents);
+
+	/** Gets the scene view init options to use to create scene views for the preview scene */
+	void GetSceneViewInitOptions(FSceneViewInitOptions& OutViewInitOptions);
+
+	/** Finds a suitable primitive component on the stage actor to use as a projection origin */
+	void FindProjectionOriginComponent();
+
 private:
 	TWeakPtr<FSceneViewport> SceneViewportPtr;
 	TWeakPtr<SDisplayClusterLightCardEditor> LightCardEditorPtr;
 	TWeakObjectPtr<ADisplayClusterRootActor> SpawnedRootActor;
+	TWeakObjectPtr<ADisplayClusterRootActor> RootActorLevelInstance;
 	TWeakObjectPtr<AActor> SelectedActor;
 	
-	/** The full bounds of the preview scene (encompasses all visible components) */
-	FBoxSphereBounds PreviewActorBounds;
-	
+	/** The renderer for the viewport, which can render the meshes with a variety of projection types */
+	TSharedPtr<FDisplayClusterMeshProjectionRenderer> MeshProjectionRenderer;
+
 	/** The current transaction for undo/redo */
-	FScopedTransaction* ScopedTransaction;
+	FScopedTransaction* ScopedTransaction = nullptr;
 	
-	bool bDraggingActor;
+	bool bDraggingActor = false;
+
+	/** The current projection mode the 3D viewport is being displayed with */
+	EDisplayClusterMeshProjectionType ProjectionMode = EDisplayClusterMeshProjectionType::Perspective;
+
+	/** The component of the root actor that is acting as the projection origin. Can be either the root component (stage origin) or a view origin component */
+	TWeakObjectPtr<USceneComponent> ProjectionOriginComponent;
+
+	/** Stores each projection mode's field of view separately */
+	TArray<float> ProjectionFOVs;
+
+	/** The increment to change the FOV by when using the scroll wheel */
+	float FOVScrollIncrement = 5.0f;
 };
\ No newline at end of file
diff --git a/Engine/Plugins/Runtime/nDisplay/Source/DisplayClusterLightCardEditor/Private/Viewport/DisplayClusterLightcardEditorViewport.cpp b/Engine/Plugins/Runtime/nDisplay/Source/DisplayClusterLightCardEditor/Private/Viewport/DisplayClusterLightcardEditorViewport.cpp
--- a/Engine/Plugins/Runtime/nDisplay/Source/DisplayClusterLightCardEditor/Private/Viewport/DisplayClusterLightcardEditorViewport.cpp
+++ b/Engine/Plugins/Runtime/nDisplay/Source/DisplayClusterLightCardEditor/Private/Viewport/DisplayClusterLightcardEditorViewport.cpp
@@ -10,6 +10,8 @@
 #include "STransformViewportToolbar.h"
 #include "Slate/SceneViewport.h"
 
+#include "Kismet2/DebuggerCommands.h"
+
 #define LOCTEXT_NAMESPACE "DisplayClusterLightcardEditorViewport"
 
 class SDisplayClusterLightCardEditorViewportToolBar : public SViewportToolBar
@@ -39,30 +41,18 @@ public:
 					SNew(SEditorViewportToolbarMenu)
 					.ParentToolBar(SharedThis(this))
 					.Cursor(EMouseCursor::Default)
-					.Image("EditorViewportToolBar.MenuDropdown")
+					.Image("EditorViewportToolBar.OptionsDropdown")
 					.OnGetMenuContent(this, &SDisplayClusterLightCardEditorViewportToolBar::GeneratePreviewMenu)
 				]
 				+ SHorizontalBox::Slot()
 				.AutoWidth()
 				.Padding(2.0f, 2.0f)
 				[
-					SNew( SEditorViewportToolbarMenu )
-					.ParentToolBar( SharedThis( this ) )
-					.Cursor( EMouseCursor::Default )
-					.Label(this, &SDisplayClusterLightCardEditorViewportToolBar::GetCameraMenuLabel)
-					.LabelIcon(this, &SDisplayClusterLightCardEditorViewportToolBar::GetCameraMenuLabelIcon)
-					.OnGetMenuContent(this, &SDisplayClusterLightCardEditorViewportToolBar::GenerateCameraMenu)
-				]
-				+ SHorizontalBox::Slot()
-				.AutoWidth()
-				.Padding(2.0f, 2.0f)
-				[
-					SNew( SEditorViewportToolbarMenu )
-					.ParentToolBar( SharedThis( this ) )
-					.Cursor( EMouseCursor::Default )
-					.Label(this, &SDisplayClusterLightCardEditorViewportToolBar::GetViewMenuLabel)
-					.LabelIcon(this, &SDisplayClusterLightCardEditorViewportToolBar::GetViewMenuLabelIcon)
-					.OnGetMenuContent(this, &SDisplayClusterLightCardEditorViewportToolBar::GenerateViewMenu)
+					SNew(SEditorViewportToolbarMenu)
+					.ParentToolBar(SharedThis(this))
+					.Cursor(EMouseCursor::Default)
+					.Label(this, &SDisplayClusterLightCardEditorViewportToolBar::GetProjectionMenuLabel)
+					.OnGetMenuContent(this, &SDisplayClusterLightCardEditorViewportToolBar::GenerateProjectionMenu)
 				]
 				+ SHorizontalBox::Slot()
 				.Padding( 3.0f, 1.0f )
@@ -87,10 +77,9 @@ public:
 
 		FMenuBuilder PreviewOptionsMenuBuilder(bInShouldCloseWindowAfterMenuSelection, CommandList);
 		{
-			PreviewOptionsMenuBuilder.BeginSection("BlueprintEditorPreviewOptions", NSLOCTEXT("BlueprintEditor", "PreviewOptionsMenuHeader", "Preview Viewport Options"));
+			PreviewOptionsMenuBuilder.BeginSection("LightCardEditorViewportOptions", LOCTEXT("ViewportOptionsMenuHeader", "Viewport Options"));
 			{
 				PreviewOptionsMenuBuilder.AddMenuEntry(FDisplayClusterLightCardEditorCommands::Get().ResetCamera);
-				PreviewOptionsMenuBuilder.AddMenuEntry(FEditorViewportCommands::Get().ToggleRealTime);
 			}
 			PreviewOptionsMenuBuilder.EndSection();
 		}
@@ -98,45 +87,38 @@ public:
 		return PreviewOptionsMenuBuilder.MakeWidget();
 	}
 
-	FText GetCameraMenuLabel() const
+	FText GetProjectionMenuLabel() const
 	{
-		if(EditorViewport.IsValid())
-		{
-			return GetCameraMenuLabelFromViewportType(EditorViewport.Pin()->GetViewportClient()->GetViewportType());
-		}
-
-		return NSLOCTEXT("BlueprintEditor", "CameraMenuTitle_Default", "Camera");
-	}
+		FText Label = LOCTEXT("ProjectionMenuTitle_Default", "Projection");
 
-	const FSlateBrush* GetCameraMenuLabelIcon() const
-	{
-		if(EditorViewport.IsValid())
+		if (EditorViewport.IsValid())
 		{
-			return GetCameraMenuLabelIconFromViewportType( EditorViewport.Pin()->GetViewportClient()->GetViewportType() );
+			switch (EditorViewport.Pin()->GetLightCardEditorViewportClient()->GetProjectionMode())
+			{
+			case EDisplayClusterMeshProjectionType::Perspective:
+				Label = LOCTEXT("ProjectionwMenuTitle_Perspective", "Perspective");
+				break;
+
+			case EDisplayClusterMeshProjectionType::Azimuthal:
+				Label = LOCTEXT("ProjectionwMenuTitle_Azimuthal", "Azimuthal");
+				break;
+			}
 		}
 
-		return FEditorStyle::GetBrush(NAME_None);
+		return Label;
 	}
 
-	TSharedRef<SWidget> GenerateCameraMenu() const
+	TSharedRef<SWidget> GenerateProjectionMenu() const
 	{
-		TSharedPtr<const FUICommandList> CommandList = EditorViewport.IsValid()? EditorViewport.Pin()->GetCommandList(): nullptr;
+		TSharedPtr<const FUICommandList> CommandList = EditorViewport.IsValid() ? EditorViewport.Pin()->GetCommandList() : nullptr;
 
 		const bool bInShouldCloseWindowAfterMenuSelection = true;
-		FMenuBuilder CameraMenuBuilder(bInShouldCloseWindowAfterMenuSelection, CommandList);
-
-		CameraMenuBuilder.AddMenuEntry(FEditorViewportCommands::Get().Perspective);
+		FMenuBuilder ViewMenuBuilder(bInShouldCloseWindowAfterMenuSelection, CommandList);
 
-		CameraMenuBuilder.BeginSection("LevelViewportCameraType_Ortho", NSLOCTEXT("BlueprintEditor", "CameraTypeHeader_Ortho", "Orthographic"));
-			CameraMenuBuilder.AddMenuEntry(FEditorViewportCommands::Get().Top);
-			CameraMenuBuilder.AddMenuEntry(FEditorViewportCommands::Get().Bottom);
-			CameraMenuBuilder.AddMenuEntry(FEditorViewportCommands::Get().Left);
-			CameraMenuBuilder.AddMenuEntry(FEditorViewportCommands::Get().Right);
-			CameraMenuBuilder.AddMenuEntry(FEditorViewportCommands::Get().Front);
-			CameraMenuBuilder.AddMenuEntry(FEditorViewportCommands::Get().Back);
-		CameraMenuBuilder.EndSection();
+		ViewMenuBuilder.AddMenuEntry(FDisplayClusterLightCardEditorCommands::Get().PerspectiveProjection);
+		ViewMenuBuilder.AddMenuEntry(FDisplayClusterLightCardEditorCommands::Get().AzimuthalProjection);
 
-		return CameraMenuBuilder.MakeWidget();
+		return ViewMenuBuilder.MakeWidget();
 	}
 
 	FText GetViewMenuLabel() const
@@ -164,49 +146,6 @@ public:
 		return Label;
 	}
 
-	const FSlateBrush* GetViewMenuLabelIcon() const
-	{
-		static FName LitModeIconName("EditorViewport.LitMode");
-		static FName UnlitModeIconName("EditorViewport.UnlitMode");
-		static FName WireframeModeIconName("EditorViewport.WireframeMode");
-
-		FName Icon = NAME_None;
-
-		if (EditorViewport.IsValid())
-		{
-			switch (EditorViewport.Pin()->GetViewportClient()->GetViewMode())
-			{
-			case VMI_Lit:
-				Icon = LitModeIconName;
-				break;
-
-			case VMI_Unlit:
-				Icon = UnlitModeIconName;
-				break;
-
-			case VMI_BrushWireframe:
-				Icon = WireframeModeIconName;
-				break;
-			}
-		}
-
-		return FEditorStyle::GetBrush(Icon);
-	}
-
-	TSharedRef<SWidget> GenerateViewMenu() const
-	{
-		TSharedPtr<const FUICommandList> CommandList = EditorViewport.IsValid() ? EditorViewport.Pin()->GetCommandList() : nullptr;
-
-		const bool bInShouldCloseWindowAfterMenuSelection = true;
-		FMenuBuilder ViewMenuBuilder(bInShouldCloseWindowAfterMenuSelection, CommandList);
-
-		ViewMenuBuilder.AddMenuEntry(FEditorViewportCommands::Get().LitMode, NAME_None, NSLOCTEXT("BlueprintEditor", "LitModeMenuOption", "Lit"));
-		ViewMenuBuilder.AddMenuEntry(FEditorViewportCommands::Get().UnlitMode, NAME_None, NSLOCTEXT("BlueprintEditor", "UnlitModeMenuOption", "Unlit"));
-		ViewMenuBuilder.AddMenuEntry(FEditorViewportCommands::Get().WireframeMode, NAME_None, NSLOCTEXT("BlueprintEditor", "WireframeModeMenuOption", "Wireframe"));
-
-		return ViewMenuBuilder.MakeWidget();
-	}
-
 private:
 	/** Reference to the parent viewport */
 	TWeakPtr<SDisplayClusterLightCardEditorViewport> EditorViewport;
@@ -220,8 +159,6 @@ void SDisplayClusterLightCardEditorViewport::Construct(const FArguments& InArgs,
 	AdvancedPreviewScene->SetFloorVisibility(true);
 		
 	SEditorViewport::Construct(SEditorViewport::FArguments());
-
-	ViewportClient->SetSceneViewport(SceneViewport);
 }
 
 SDisplayClusterLightCardEditorViewport::~SDisplayClusterLightCardEditorViewport()
@@ -281,12 +218,38 @@ void SDisplayClusterLightCardEditorViewport::PopulateViewportOverlays(TSharedRef
 void SDisplayClusterLightCardEditorViewport::BindCommands()
 {
 	const FDisplayClusterLightCardEditorCommands& Commands = FDisplayClusterLightCardEditorCommands::Get();
-	
+
+	CommandList->MapAction(
+		FDisplayClusterLightCardEditorCommands::Get().PerspectiveProjection,
+		FExecuteAction::CreateSP(this, &SDisplayClusterLightCardEditorViewport::SetProjectionMode, EDisplayClusterMeshProjectionType::Perspective),
+		FCanExecuteAction(),
+		FIsActionChecked::CreateSP(this, &SDisplayClusterLightCardEditorViewport::IsProjectionModeSelected, EDisplayClusterMeshProjectionType::Perspective));
+
 	CommandList->MapAction(
-	Commands.ResetCamera,
-	FExecuteAction::CreateSP(ViewportClient.Get(), &FDisplayClusterLightCardEditorViewportClient::ResetCamera));
+		FDisplayClusterLightCardEditorCommands::Get().AzimuthalProjection,
+		FExecuteAction::CreateSP(this, &SDisplayClusterLightCardEditorViewport::SetProjectionMode, EDisplayClusterMeshProjectionType::Azimuthal),
+		FCanExecuteAction(),
+		FIsActionChecked::CreateSP(this, &SDisplayClusterLightCardEditorViewport::IsProjectionModeSelected, EDisplayClusterMeshProjectionType::Azimuthal));
 
 	SEditorViewport::BindCommands();
 }
 
+void SDisplayClusterLightCardEditorViewport::SetProjectionMode(EDisplayClusterMeshProjectionType InProjectionMode)
+{
+	if (ViewportClient.IsValid())
+	{
+		ViewportClient->SetProjectionMode(InProjectionMode);
+	}
+}
+
+bool SDisplayClusterLightCardEditorViewport::IsProjectionModeSelected(EDisplayClusterMeshProjectionType InProjectionMode) const
+{
+	if (ViewportClient.IsValid())
+	{
+		return ViewportClient->GetProjectionMode() == InProjectionMode;
+	}
+
+	return false;
+}
+
 #undef LOCTEXT_NAMESPACE
diff --git a/Engine/Plugins/Runtime/nDisplay/Source/DisplayClusterLightCardEditor/Private/Viewport/DisplayClusterLightcardEditorViewport.h b/Engine/Plugins/Runtime/nDisplay/Source/DisplayClusterLightCardEditor/Private/Viewport/DisplayClusterLightcardEditorViewport.h
--- a/Engine/Plugins/Runtime/nDisplay/Source/DisplayClusterLightCardEditor/Private/Viewport/DisplayClusterLightcardEditorViewport.h
+++ b/Engine/Plugins/Runtime/nDisplay/Source/DisplayClusterLightCardEditor/Private/Viewport/DisplayClusterLightcardEditorViewport.h
@@ -7,8 +7,11 @@
 #include "SCommonEditorViewportToolbarBase.h"
 #include "SEditorViewport.h"
 
+#include "DisplayClusterMeshProjectionRenderer.h"
+
 class SDisplayClusterLightCardEditor;
 class ADisplayClusterRootActor;
+class FDisplayClusterLightCardEditorViewportClient;
 
 /**
  * Slate widget which renders our view client.
@@ -32,6 +35,8 @@ public:
 
 	void SetRootActor(ADisplayClusterRootActor* NewRootActor);
 	
+	TSharedRef<FDisplayClusterLightCardEditorViewportClient> GetLightCardEditorViewportClient() const { return ViewportClient.ToSharedRef(); }
+
 private:
 	// SEditorViewport
 	virtual TSharedRef<FEditorViewportClient> MakeEditorViewportClient() override;
@@ -40,12 +45,15 @@ private:
 	virtual void BindCommands() override;
 	// ~SEditorViewport
 
+	void SetProjectionMode(EDisplayClusterMeshProjectionType InProjectionMode);
+	bool IsProjectionModeSelected(EDisplayClusterMeshProjectionType InProjectionMode) const;
+
 private:
 	/** Preview Scene - uses advanced preview settings */
 	TSharedPtr<FAdvancedPreviewScene> AdvancedPreviewScene;
 	
 	/** Level viewport client */
-	TSharedPtr<class FDisplayClusterLightCardEditorViewportClient> ViewportClient;
+	TSharedPtr<FDisplayClusterLightCardEditorViewportClient> ViewportClient;
 	TWeakPtr<SDisplayClusterLightCardEditor> LightCardEditorPtr;
 
 };
\ No newline at end of file
diff --git a/Engine/Plugins/Runtime/nDisplay/Source/DisplayClusterLightCardEditorShaders/DisplayClusterLightCardEditorShaders.Build.cs b/Engine/Plugins/Runtime/nDisplay/Source/DisplayClusterLightCardEditorShaders/DisplayClusterLightCardEditorShaders.Build.cs
new file mode 100644
--- /dev/null
+++ b/Engine/Plugins/Runtime/nDisplay/Source/DisplayClusterLightCardEditorShaders/DisplayClusterLightCardEditorShaders.Build.cs
@@ -0,0 +1,21 @@
+// Copyright Epic Games, Inc. All Rights Reserved.
+
+using UnrealBuildTool;
+using System.IO;
+
+public class DisplayClusterLightCardEditorShaders : ModuleRules
+{
+	public DisplayClusterLightCardEditorShaders(ReadOnlyTargetRules ROTargetRules) : base(ROTargetRules)
+	{
+		PrivateDependencyModuleNames.AddRange(
+			new string[] {
+				"Core",
+				"CoreUObject",
+				"Engine",
+				"RenderCore",
+				"Renderer",
+				"RHI",
+				"Projects"
+			});
+	}
+}
diff --git a/Engine/Plugins/Runtime/nDisplay/Source/DisplayClusterLightCardEditorShaders/Private/DisplayClusterLightCardEditorShadersModule.cpp b/Engine/Plugins/Runtime/nDisplay/Source/DisplayClusterLightCardEditorShaders/Private/DisplayClusterLightCardEditorShadersModule.cpp
new file mode 100644
--- /dev/null
+++ b/Engine/Plugins/Runtime/nDisplay/Source/DisplayClusterLightCardEditorShaders/Private/DisplayClusterLightCardEditorShadersModule.cpp
@@ -0,0 +1,24 @@
+// Copyright Epic Games, Inc. All Rights Reserved.
+
+#include "DisplayClusterLightCardEditorShadersModule.h"
+
+#include "Interfaces/IPluginManager.h"
+#include "Misc/Paths.h"
+#include "ShaderCore.h"
+
+#define NDISPLAY_SHADERS_MAP TEXT("/Plugin/nDisplay")
+
+void FDisplayClusterLightCardEditorShadersModule::StartupModule()
+{
+	if (!AllShaderSourceDirectoryMappings().Contains(NDISPLAY_SHADERS_MAP))
+	{
+		FString PluginShaderDir = FPaths::Combine(IPluginManager::Get().FindPlugin(TEXT("nDisplay"))->GetBaseDir(), TEXT("Shaders"));
+		AddShaderSourceDirectoryMapping(NDISPLAY_SHADERS_MAP, PluginShaderDir);
+	}
+}
+
+void FDisplayClusterLightCardEditorShadersModule::ShutdownModule()
+{
+}
+
+IMPLEMENT_MODULE(FDisplayClusterLightCardEditorShadersModule, DisplayClusterLightCardEditorShaders);
diff --git a/Engine/Plugins/Runtime/nDisplay/Source/DisplayClusterLightCardEditorShaders/Private/DisplayClusterLightCardEditorShadersModule.h b/Engine/Plugins/Runtime/nDisplay/Source/DisplayClusterLightCardEditorShaders/Private/DisplayClusterLightCardEditorShadersModule.h
new file mode 100644
--- /dev/null
+++ b/Engine/Plugins/Runtime/nDisplay/Source/DisplayClusterLightCardEditorShaders/Private/DisplayClusterLightCardEditorShadersModule.h
@@ -0,0 +1,18 @@
+// Copyright Epic Games, Inc. All Rights Reserved.
+
+#pragma once
+
+#include "IDisplayClusterLightCardEditorShaders.h"
+#include "UObject/WeakObjectPtr.h"
+
+/**
+ * Display Cluster editor module
+ */
+class FDisplayClusterLightCardEditorShadersModule :
+	public IDisplayClusterLightCardEditorShaders
+{
+public:
+	//~ IModuleInterface interface
+	virtual void StartupModule() override;
+	virtual void ShutdownModule() override;
+};
diff --git a/Engine/Plugins/Runtime/nDisplay/Source/DisplayClusterLightCardEditorShaders/Private/DisplayClusterMeshProjectionRenderer.cpp b/Engine/Plugins/Runtime/nDisplay/Source/DisplayClusterLightCardEditorShaders/Private/DisplayClusterMeshProjectionRenderer.cpp
new file mode 100644
--- /dev/null
+++ b/Engine/Plugins/Runtime/nDisplay/Source/DisplayClusterLightCardEditorShaders/Private/DisplayClusterMeshProjectionRenderer.cpp
@@ -0,0 +1,261 @@
+// Copyright Epic Games, Inc. All Rights Reserved.
+
+#include "DisplayClusterMeshProjectionRenderer.h"
+
+#include "Components/PrimitiveComponent.h"
+#include "InstanceCulling/InstanceCullingContext.h"
+#include "PrimitiveSceneInfo.h"
+#include "PrimitiveSceneProxy.h"
+#include "MeshPassProcessor.h"
+#include "MeshMaterialShader.h"
+#include "MeshPassProcessor.inl"
+#include "Shader.h"
+#include "CanvasTypes.h"
+#include "EngineModule.h"
+
+template<EDisplayClusterMeshProjectionType ProjectionType>
+class FMeshProjectionVS : public FMeshMaterialShader
+{
+public:
+	DECLARE_SHADER_TYPE(FMeshProjectionVS, MeshMaterial);
+
+	FMeshProjectionVS() { }
+	FMeshProjectionVS(const ShaderMetaType::CompiledShaderInitializerType& Initializer)
+		: FMeshMaterialShader(Initializer)
+	{
+		PassUniformBuffer.Bind(Initializer.ParameterMap, FSceneTextureUniformParameters::StaticStructMetadata.GetShaderVariableName());
+	}
+
+	static bool ShouldCompilePermutation(const FMeshMaterialShaderPermutationParameters& Parameters)
+	{
+		return EnumHasAllFlags(Parameters.Flags, EShaderPermutationFlags::HasEditorOnlyData)
+			&& IsFeatureLevelSupported(Parameters.Platform, ERHIFeatureLevel::SM5)
+			&& Parameters.VertexFactoryType == FindVertexFactoryType(TEXT("FLocalVertexFactory"));
+	}
+};
+
+using FMeshPerspectiveProjectionVS = FMeshProjectionVS<EDisplayClusterMeshProjectionType::Perspective>;
+using FMeshAzimuthalProjectionVS = FMeshProjectionVS<EDisplayClusterMeshProjectionType::Azimuthal>;
+
+IMPLEMENT_MATERIAL_SHADER_TYPE(template<>, FMeshPerspectiveProjectionVS, TEXT("/Plugin/nDisplay/Private/MeshProjectionShaders.usf"), TEXT("PerspectiveVS"), SF_Vertex);
+IMPLEMENT_MATERIAL_SHADER_TYPE(template<>, FMeshAzimuthalProjectionVS, TEXT("/Plugin/nDisplay/Private/MeshProjectionShaders.usf"), TEXT("AzimuthalVS"), SF_Vertex);
+
+class FMeshProjectionPS : public FMeshMaterialShader
+{
+public:
+	DECLARE_SHADER_TYPE(FMeshProjectionPS, MeshMaterial);
+
+	FMeshProjectionPS() { }
+	FMeshProjectionPS(const ShaderMetaType::CompiledShaderInitializerType& Initializer)
+		: FMeshMaterialShader(Initializer)
+	{
+		PassUniformBuffer.Bind(Initializer.ParameterMap, FSceneTextureUniformParameters::StaticStructMetadata.GetShaderVariableName());
+	}
+
+	static bool ShouldCompilePermutation(const FMeshMaterialShaderPermutationParameters& Parameters)
+	{
+		return EnumHasAllFlags(Parameters.Flags, EShaderPermutationFlags::HasEditorOnlyData)
+			&& IsFeatureLevelSupported(Parameters.Platform, ERHIFeatureLevel::SM5)
+			&& Parameters.VertexFactoryType == FindVertexFactoryType(TEXT("FLocalVertexFactory"));
+	}
+};
+
+IMPLEMENT_MATERIAL_SHADER_TYPE(, FMeshProjectionPS, TEXT("/Plugin/nDisplay/Private/MeshProjectionShaders.usf"), TEXT("MainPS"), SF_Pixel);
+
+BEGIN_SHADER_PARAMETER_STRUCT(FMeshProjectionPassParameters, )
+	SHADER_PARAMETER_STRUCT_REF(FViewUniformShaderParameters, View)
+	SHADER_PARAMETER_RDG_UNIFORM_BUFFER(FInstanceCullingGlobalUniforms, InstanceCulling)
+	SHADER_PARAMETER_STRUCT_INCLUDE(FInstanceCullingDrawParams, InstanceCullingDrawParams)
+	RENDER_TARGET_BINDING_SLOTS()
+END_SHADER_PARAMETER_STRUCT()
+
+template<EDisplayClusterMeshProjectionType ProjectionType> 
+class FLightCardEditorMeshPassProcessor : public FMeshPassProcessor
+{
+public:
+	FLightCardEditorMeshPassProcessor(const FScene* InScene, const FSceneView* InView, FMeshPassDrawListContext* InDrawListContext)
+		: FMeshPassProcessor(InScene, GMaxRHIFeatureLevel, InView, InDrawListContext)
+		, DrawRenderState(*InView)
+	{
+		DrawRenderState.SetDepthStencilState(TStaticDepthStencilState<false, CF_Always>::GetRHI());
+		DrawRenderState.SetBlendState(TStaticBlendState<>::GetRHI());
+	}
+
+	virtual void AddMeshBatch(const FMeshBatch& RESTRICT MeshBatch, uint64 BatchElementMask, const FPrimitiveSceneProxy* RESTRICT PrimitiveSceneProxy, int32 StaticMeshId = -1) override final
+	{
+		const FMaterialRenderProxy* FallbackMaterialRenderProxy = nullptr;
+		const FMaterial& Material = MeshBatch.MaterialRenderProxy->GetMaterialWithFallback(FeatureLevel, FallbackMaterialRenderProxy);
+		const FMaterialRenderProxy& MaterialRenderProxy = FallbackMaterialRenderProxy ? *FallbackMaterialRenderProxy : *MeshBatch.MaterialRenderProxy;
+
+		const FVertexFactory* VertexFactory = MeshBatch.VertexFactory;
+
+		TMeshProcessorShaders<FMeshProjectionVS<ProjectionType>, FMeshProjectionPS> PassShaders;
+
+		FMaterialShaderTypes ShaderTypes;
+		ShaderTypes.AddShaderType<FMeshProjectionVS<ProjectionType>>();
+		ShaderTypes.AddShaderType<FMeshProjectionPS>();
+
+		FMaterialShaders Shaders;
+		if (!Material.TryGetShaders(ShaderTypes, VertexFactory->GetType(), Shaders))
+		{
+			return;
+		}
+
+		Shaders.TryGetVertexShader(PassShaders.VertexShader);
+		Shaders.TryGetPixelShader(PassShaders.PixelShader);
+
+		const FMeshDrawingPolicyOverrideSettings OverrideSettings = ComputeMeshOverrideSettings(MeshBatch);
+		ERasterizerFillMode MeshFillMode = ComputeMeshFillMode(MeshBatch, Material, OverrideSettings);
+		ERasterizerCullMode MeshCullMode = CM_None;
+
+		FMeshMaterialShaderElementData ShaderElementData;
+		ShaderElementData.InitializeMeshMaterialData(ViewIfDynamicMeshCommand, PrimitiveSceneProxy, MeshBatch, StaticMeshId, false);
+
+		FMeshDrawCommandSortKey SortKey{};
+
+		BuildMeshDrawCommands(
+			MeshBatch,
+			BatchElementMask,
+			PrimitiveSceneProxy,
+			MaterialRenderProxy,
+			Material,
+			DrawRenderState,
+			PassShaders,
+			MeshFillMode,
+			MeshCullMode,
+			SortKey,
+			EMeshPassFeatures::Default,
+			ShaderElementData);
+	}
+
+private:
+	FMeshPassProcessorRenderState DrawRenderState;
+};
+
+void FDisplayClusterMeshProjectionRenderer::AddActor(AActor* Actor)
+{
+	AddActor(Actor, [](const UPrimitiveComponent* PrimitiveComponent)
+	{
+		return !PrimitiveComponent->bHiddenInGame;
+	});
+}
+
+void FDisplayClusterMeshProjectionRenderer::AddActor(AActor* Actor, const TFunctionRef<bool(const UPrimitiveComponent*)>& PrimitiveFilter)
+{
+	Actor->ForEachComponent<UPrimitiveComponent>(true, [&](UPrimitiveComponent* PrimitiveComponent)
+	{
+		if (PrimitiveFilter(PrimitiveComponent))
+		{
+			PrimitiveComponents.Add(PrimitiveComponent);
+		}
+	});
+}
+
+void FDisplayClusterMeshProjectionRenderer::RemoveActor(AActor* Actor)
+{
+	PrimitiveComponents.RemoveAll([Actor](TWeakObjectPtr<UPrimitiveComponent>& PrimitiveComponent)
+	{
+		return !PrimitiveComponent.IsValid() || PrimitiveComponent->GetOwner() == Actor;
+	});
+}
+
+void FDisplayClusterMeshProjectionRenderer::ClearScene()
+{
+	PrimitiveComponents.Empty();
+}
+
+void FDisplayClusterMeshProjectionRenderer::Render(FCanvas* Canvas, FSceneInterface* Scene, const FSceneViewInitOptions& ViewInitOptions, EDisplayClusterMeshProjectionType ProjectionType)
+{
+	Canvas->Flush_GameThread();
+	FRenderTarget* RenderTarget = Canvas->GetRenderTarget();
+
+	ENQUEUE_RENDER_COMMAND(FDrawProjectedMeshes)(
+		[RenderTarget, Scene, ViewInitOptions, ProjectionType, this](FRHICommandListImmediate& RHICmdList)
+		{
+			FMemMark Mark(FMemStack::Get());
+			FRDGBuilder GraphBuilder(RHICmdList);
+
+			FSceneViewFamily ViewFamily(FSceneViewFamily::ConstructionValues(
+				RenderTarget,
+				Scene,
+				FEngineShowFlags(ESFIM_Editor))
+				.SetTime(FGameTime::GetTimeSinceAppStart())
+				.SetGammaCorrection(1.0f));
+
+			if (Scene)
+			{
+				Scene->IncrementFrameNumber();
+				ViewFamily.FrameNumber = Scene->GetFrameNumber();
+			}
+			else
+			{
+				ViewFamily.FrameNumber = GFrameNumber;
+			}
+
+			FScenePrimitiveRenderingContextScopeHelper ScenePrimitiveRenderingContextScopeHelper(GetRendererModule().BeginScenePrimitiveRendering(GraphBuilder, &ViewFamily));
+
+			FSceneViewInitOptions NewInitOptions(ViewInitOptions);
+			NewInitOptions.ViewFamily = &ViewFamily;
+
+			GetRendererModule().CreateAndInitSingleView(RHICmdList, &ViewFamily, &NewInitOptions);
+			const FSceneView* View = ViewFamily.Views[0];
+
+			FRDGTextureRef OutputTexture = GraphBuilder.RegisterExternalTexture(CreateRenderTarget(RenderTarget->GetRenderTargetTexture(), TEXT("ViewRenderTarget")));
+
+			FMeshProjectionPassParameters* PassParameters = GraphBuilder.AllocParameters<FMeshProjectionPassParameters>();
+			PassParameters->View = View->ViewUniformBuffer;
+			PassParameters->InstanceCulling = FInstanceCullingContext::CreateDummyInstanceCullingUniformBuffer(GraphBuilder);
+			PassParameters->RenderTargets[0] = FRenderTargetBinding(OutputTexture, ERenderTargetLoadAction::ELoad);
+
+			GraphBuilder.AddPass(RDG_EVENT_NAME("MeshProjectionRenderer"), PassParameters, ERDGPassFlags::Raster | ERDGPassFlags::NeverCull, [View, ProjectionType, this](FRHICommandList& RHICmdList)
+			{
+				FIntRect ViewRect = View->UnscaledViewRect;
+				RHICmdList.SetViewport(ViewRect.Min.X, ViewRect.Min.Y, 0.0f, ViewRect.Max.X, ViewRect.Max.Y, 1.0f);
+
+				switch (ProjectionType)
+				{
+					case EDisplayClusterMeshProjectionType::Azimuthal:
+						RenderPrimitives_RenderThread<EDisplayClusterMeshProjectionType::Azimuthal>(View, RHICmdList);
+						break;
+
+					case EDisplayClusterMeshProjectionType::Perspective:
+					default:
+						RenderPrimitives_RenderThread<EDisplayClusterMeshProjectionType::Perspective>(View, RHICmdList);
+						break;
+				}
+			});
+
+			GraphBuilder.Execute();
+		});
+}
+
+template<EDisplayClusterMeshProjectionType ProjectionType>
+void FDisplayClusterMeshProjectionRenderer::RenderPrimitives_RenderThread(const FSceneView* View, FRHICommandList& RHICmdList)
+{
+	DrawDynamicMeshPass(*View, RHICmdList, [View, this](FDynamicPassMeshDrawListContext* DynamicMeshPassContext)
+	{
+		TArray<FPrimitiveSceneProxy*> PrimitiveSceneProxies;
+		for (const TWeakObjectPtr<UPrimitiveComponent>& PrimitiveComponent : PrimitiveComponents)
+		{
+			if (PrimitiveComponent.IsValid() && PrimitiveComponent->SceneProxy)
+			{
+				PrimitiveSceneProxies.Add(PrimitiveComponent->SceneProxy);
+			}
+		}
+
+		FLightCardEditorMeshPassProcessor<ProjectionType> MeshProcessor(nullptr, View, DynamicMeshPassContext);
+
+		for (FPrimitiveSceneProxy* PrimitiveProxy : PrimitiveSceneProxies)
+		{
+			if (const FMeshBatch* MeshBatch = PrimitiveProxy->GetPrimitiveSceneInfo()->GetMeshBatch(PrimitiveProxy->GetPrimitiveSceneInfo()->StaticMeshes.Num() - 1))
+			{
+				const uint64 BatchElementMask = ~0ull;
+				MeshProcessor.AddMeshBatch(*MeshBatch, BatchElementMask, PrimitiveProxy);
+			}
+		}
+	});
+}
+
+// Explicit template specialisations for RenderPrimitives_RenderThread
+template void FDisplayClusterMeshProjectionRenderer::RenderPrimitives_RenderThread<EDisplayClusterMeshProjectionType::Perspective>(const FSceneView* View, FRHICommandList& RHICmdList);
+template void FDisplayClusterMeshProjectionRenderer::RenderPrimitives_RenderThread<EDisplayClusterMeshProjectionType::Azimuthal>(const FSceneView* View, FRHICommandList& RHICmdList);
\ No newline at end of file
diff --git a/Engine/Plugins/Runtime/nDisplay/Source/DisplayClusterLightCardEditorShaders/Public/DisplayClusterMeshProjectionRenderer.h b/Engine/Plugins/Runtime/nDisplay/Source/DisplayClusterLightCardEditorShaders/Public/DisplayClusterMeshProjectionRenderer.h
new file mode 100644
--- /dev/null
+++ b/Engine/Plugins/Runtime/nDisplay/Source/DisplayClusterLightCardEditorShaders/Public/DisplayClusterMeshProjectionRenderer.h
@@ -0,0 +1,50 @@
+// Copyright Epic Games, Inc. All Rights Reserved.
+
+#pragma once
+
+#include "CoreMinimal.h"
+#include "RenderGraphBuilder.h"
+
+class AActor;
+class UPrimitiveComponent;
+class FSceneView;
+class UPrimitiveComponent;
+
+/** Indicates which kind of projection is used by the renderer */
+enum DISPLAYCLUSTERLIGHTCARDEDITORSHADERS_API EDisplayClusterMeshProjectionType
+{
+	/** Default linear projection */
+	Perspective,
+
+	/** Non-linear spherical projection based on the azimuthal equidistant map projection */
+	Azimuthal
+};
+
+/** A renderer that projects meshes to screen space using non-linear projection methods */
+class DISPLAYCLUSTERLIGHTCARDEDITORSHADERS_API FDisplayClusterMeshProjectionRenderer
+{
+public:
+	/** Adds an actor's primitive components to the list of primitives to render */
+	void AddActor(AActor* Actor);
+
+	/** Adds an actor's primitive components to the list of primitives to render, filtering which primitive components get rendered using the specified callback */
+	void AddActor(AActor* Actor, const TFunctionRef<bool(const UPrimitiveComponent*)>& PrimitiveFilter);
+
+	/** Removes an actor's primitive components from the list of primitives to render */
+	void RemoveActor(AActor* Actor);
+
+	/** Clears the list of primitives to render */
+	void ClearScene();
+
+	/** Renders its list of primitive components to the specified canvas using the desired projection type. Can be called from the game thread */
+	void Render(FCanvas* Canvas, FSceneInterface* Scene, const FSceneViewInitOptions& ViewInitOptions, EDisplayClusterMeshProjectionType  ProjectionType);
+
+private:
+	/** Renders the list of primitive components using the appropriate mesh pass processor given by the template parameter */
+	template<EDisplayClusterMeshProjectionType ProjectionType>
+	void RenderPrimitives_RenderThread(const FSceneView* View, FRHICommandList& RHICmdList);
+
+private:
+	/** The list of primitive components to render */
+	TArray<TWeakObjectPtr<UPrimitiveComponent>> PrimitiveComponents;
+};
\ No newline at end of file
diff --git a/Engine/Plugins/Runtime/nDisplay/Source/DisplayClusterLightCardEditorShaders/Public/IDisplayClusterLightCardEditorShaders.h b/Engine/Plugins/Runtime/nDisplay/Source/DisplayClusterLightCardEditorShaders/Public/IDisplayClusterLightCardEditorShaders.h
new file mode 100644
--- /dev/null
+++ b/Engine/Plugins/Runtime/nDisplay/Source/DisplayClusterLightCardEditorShaders/Public/IDisplayClusterLightCardEditorShaders.h
@@ -0,0 +1,40 @@
+// Copyright Epic Games, Inc. All Rights Reserved.
+
+#pragma once
+
+#include "CoreMinimal.h"
+#include "Modules/ModuleInterface.h"
+#include "Modules/ModuleManager.h"
+
+/**
+ * Display Cluster Light Card Editor Shaders module interface
+ */
+class IDisplayClusterLightCardEditorShaders : public IModuleInterface
+{
+public:
+	static constexpr const TCHAR* ModuleName = TEXT("DisplayClusterLightCardEditorShaders");
+
+public:
+	virtual ~IDisplayClusterLightCardEditorShaders() = default;
+
+	/**
+	* Singleton-like access to this module's interface.  This is just for convenience!
+	* Beware of calling this during the shutdown phase, though.  Your module might have been unloaded already.
+	*
+	* @return Returns singleton instance, loading the module on demand if needed
+	*/
+	static inline IDisplayClusterLightCardEditorShaders& Get()
+	{
+		return FModuleManager::GetModuleChecked<IDisplayClusterLightCardEditorShaders>(ModuleName);
+	}
+
+	/**
+	* Checks to see if this module is loaded and ready.  It is only valid to call Get() if IsAvailable() returns true.
+	*
+	* @return True if the module is loaded and ready to use
+	*/
+	static inline bool IsAvailable()
+	{
+		return FModuleManager::Get().IsModuleLoaded(ModuleName);
+	}
+};
\ No newline at end of file
diff --git a/Engine/Plugins/Runtime/nDisplay/nDisplay.uplugin b/Engine/Plugins/Runtime/nDisplay/nDisplay.uplugin
--- a/Engine/Plugins/Runtime/nDisplay/nDisplay.uplugin
+++ b/Engine/Plugins/Runtime/nDisplay/nDisplay.uplugin
@@ -108,6 +108,15 @@
 				"Linux"
 			]
 		},
+		{
+			"Name": "DisplayClusterLightCardEditorShaders",
+			"Type": "Runtime",
+			"LoadingPhase": "PostConfigInit",
+			"PlatformAllowList": [
+				"Win64",
+				"Linux"
+			]
+		},
 		{
 			"Name": "DisplayClusterStageMonitoring",
 			"Type": "Runtime",
