nDisplay: Disable viewport texture preview while changing the viewport region via interactive change or when interacting with the viewport node. This improves performance during a resize and fixes a rendering crash from a region being larger than the texture size.

diff --git a/Engine/Plugins/Runtime/nDisplay/Source/DisplayCluster/Private/Game/EngineClasses/Scene/DisplayClusterPreviewComponent.cpp b/Engine/Plugins/Runtime/nDisplay/Source/DisplayCluster/Private/Game/EngineClasses/Scene/DisplayClusterPreviewComponent.cpp
--- a/Engine/Plugins/Runtime/nDisplay/Source/DisplayCluster/Private/Game/EngineClasses/Scene/DisplayClusterPreviewComponent.cpp
+++ b/Engine/Plugins/Runtime/nDisplay/Source/DisplayCluster/Private/Game/EngineClasses/Scene/DisplayClusterPreviewComponent.cpp
@@ -448,6 +448,16 @@ UTexture2D* UDisplayClusterPreviewComponent::GetOrCreateViewportPreviewTexture2D
 	return nullptr;
 }
 
+UTexture2D* UDisplayClusterPreviewComponent::GetViewportPreviewTexture2D() const
+{
+	if (IsPreviewAvailable())
+	{
+		return PreviewTexture;
+	}
+
+	return nullptr;
+}
+
 void UDisplayClusterPreviewComponent::SetOverrideTexture(UTexture* InOverrideTexture)
 {
 	if (OverrideTexture != InOverrideTexture)
diff --git a/Engine/Plugins/Runtime/nDisplay/Source/DisplayCluster/Public/Components/DisplayClusterPreviewComponent.h b/Engine/Plugins/Runtime/nDisplay/Source/DisplayCluster/Public/Components/DisplayClusterPreviewComponent.h
--- a/Engine/Plugins/Runtime/nDisplay/Source/DisplayCluster/Public/Components/DisplayClusterPreviewComponent.h
+++ b/Engine/Plugins/Runtime/nDisplay/Source/DisplayCluster/Public/Components/DisplayClusterPreviewComponent.h
@@ -75,6 +75,9 @@ public:
 	/** Create and retrieve a render texture 2d from the render target. */
 	UTexture2D* GetOrCreateViewportPreviewTexture2D();
 
+	/** Return the render texture 2d but do not recreate it. */
+	UTexture2D* GetViewportPreviewTexture2D() const;
+
 	/** Sets an override texture to display on the viewport instead of the render target */
 	void SetOverrideTexture(UTexture* InOverrideTexture);
 
diff --git a/Engine/Plugins/Runtime/nDisplay/Source/DisplayClusterConfiguration/Private/DisplayClusterConfigurationTypes.cpp b/Engine/Plugins/Runtime/nDisplay/Source/DisplayClusterConfiguration/Private/DisplayClusterConfigurationTypes.cpp
--- a/Engine/Plugins/Runtime/nDisplay/Source/DisplayClusterConfiguration/Private/DisplayClusterConfigurationTypes.cpp
+++ b/Engine/Plugins/Runtime/nDisplay/Source/DisplayClusterConfiguration/Private/DisplayClusterConfigurationTypes.cpp
@@ -13,6 +13,10 @@
 
 #include "Engine/StaticMesh.h"
 
+#if WITH_EDITOR
+#include "Kismet2/CompilerResultsLog.h"
+#endif
+
 #define SAVE_MAP_TO_ARRAY(Map, DestArray) \
 	for (const auto& KeyVal : Map) \
 	{ \
@@ -174,14 +178,80 @@ UDisplayClusterConfigurationViewport::UDisplayClusterConfigurationViewport()
 }
 
 #if WITH_EDITOR
+void UDisplayClusterConfigurationViewport::PreEditChange(FEditPropertyChain& PropertyAboutToChange)
+{
+	if (TDoubleLinkedList<FProperty*>::TDoubleLinkedListNode* ActiveMemberNode = PropertyAboutToChange.
+		GetActiveMemberNode())
+	{
+		if (const FStructProperty* StructProperty = CastField<FStructProperty>(ActiveMemberNode->GetValue()))
+		{
+			if (StructProperty->GetFName() == GET_MEMBER_NAME_CHECKED(UDisplayClusterConfigurationViewport, Region))
+			{
+				if (DisablePreviewTexture())
+				{
+					bIsManagingPreviewTexture = true;
+				}
+			}
+		}
+	}
+	
+	Super::PreEditChange(PropertyAboutToChange);
+}
 
 void UDisplayClusterConfigurationViewport::PostEditChangeChainProperty(FPropertyChangedChainEvent& PropertyChangedEvent)
 {
+	if (PropertyChangedEvent.ChangeType != EPropertyChangeType::Interactive)
+	{
+		if (TDoubleLinkedList<FProperty*>::TDoubleLinkedListNode* ActiveMemberNode = PropertyChangedEvent.
+			PropertyChain.GetActiveMemberNode())
+		{
+			if (const FStructProperty* StructProperty = CastField<FStructProperty>(ActiveMemberNode->GetValue()))
+			{
+				if (StructProperty->GetFName() == GET_MEMBER_NAME_CHECKED(UDisplayClusterConfigurationViewport, Region))
+				{
+					if (bIsManagingPreviewTexture)
+					{
+						EnablePreviewTexture();
+					}
+				}
+			}
+		}
+	}
+	
 	Super::PostEditChangeChainProperty(PropertyChangedEvent);
 
 	OnPostEditChangeChainProperty.Broadcast(PropertyChangedEvent);
 }
 
+void UDisplayClusterConfigurationViewport::EnablePreviewTexture()
+{
+	bAllowPreviewTexture = true;
+	bIsManagingPreviewTexture = false;
+}
+
+bool UDisplayClusterConfigurationViewport::DisablePreviewTexture()
+{
+	if (bAllowPreviewTexture)
+	{
+		bAllowPreviewTexture = false;
+		return true;
+	}
+	
+	return false;
+}
+
+void UDisplayClusterConfigurationViewport::OnPreCompile(FCompilerResultsLog& MessageLog)
+{
+	Super::OnPreCompile(MessageLog);
+
+	if (!ensure(bAllowPreviewTexture))
+	{
+		// Verify correct rendering value is applied. This branch shouldn't be hit as long as the Region
+		// struct always has a final PostEditChangeProperty called without a change type of Interactive.
+		EnablePreviewTexture();
+	}
+}
+
 void UDisplayClusterConfigurationClusterNode::PostEditChangeChainProperty(FPropertyChangedChainEvent& PropertyChangedEvent)
 {
 	Super::PostEditChangeChainProperty(PropertyChangedEvent);
diff --git a/Engine/Plugins/Runtime/nDisplay/Source/DisplayClusterConfiguration/Public/DisplayClusterConfigurationTypes_Base.h b/Engine/Plugins/Runtime/nDisplay/Source/DisplayClusterConfiguration/Public/DisplayClusterConfigurationTypes_Base.h
--- a/Engine/Plugins/Runtime/nDisplay/Source/DisplayClusterConfiguration/Public/DisplayClusterConfigurationTypes_Base.h
+++ b/Engine/Plugins/Runtime/nDisplay/Source/DisplayClusterConfiguration/Public/DisplayClusterConfigurationTypes_Base.h
@@ -121,6 +121,14 @@ public:
 	virtual void Serialize(FArchive& Ar) override;
 	// ~UObject
 
+#if WITH_EDITOR
+protected:
+	friend class FDisplayClusterConfiguratorKismetCompilerContext;
+	
+	/** Called by nDisplay compiler prior to compilation. */
+	virtual void OnPreCompile(class FCompilerResultsLog& MessageLog) {}
+#endif
+	
 protected:
 	/** Called before saving to collect objects which should be exported as a sub object block. */
 	virtual void GetObjectsToExport(TArray<UObject*>& OutObjects) {}
diff --git a/Engine/Plugins/Runtime/nDisplay/Source/DisplayClusterConfiguration/Public/DisplayClusterConfigurationTypes_Viewport.h b/Engine/Plugins/Runtime/nDisplay/Source/DisplayClusterConfiguration/Public/DisplayClusterConfigurationTypes_Viewport.h
--- a/Engine/Plugins/Runtime/nDisplay/Source/DisplayClusterConfiguration/Public/DisplayClusterConfigurationTypes_Viewport.h
+++ b/Engine/Plugins/Runtime/nDisplay/Source/DisplayClusterConfiguration/Public/DisplayClusterConfigurationTypes_Viewport.h
@@ -141,11 +141,33 @@ public:
 public:
 	void GetReferencedMeshNames(TArray<FString>& OutMeshNames) const;
 
-private:
 #if WITH_EDITOR
 	// UObject interface
+	virtual void PreEditChange(FEditPropertyChain& PropertyAboutToChange) override;
 	virtual void PostEditChangeChainProperty(FPropertyChangedChainEvent& PropertyChangedEvent) override;
 	// End of UObject interface
+
+	/** Enable the preview texture. Only should be called by the object managing the preview texture state. */
+	void EnablePreviewTexture();
+	
+	/**
+	 * Signal that the preview texture should be disabled.
+	 * @return True if the preview texture was disabled. False if it was already disabled.
+	 */
+	bool DisablePreviewTexture();
+
+	/** If this viewport is allowed to render a preview texture. Used with resizing viewports. */
+	bool IsPreviewTextureAllowed() const { return bAllowPreviewTexture; }
+	
+protected:
+	virtual void OnPreCompile(class FCompilerResultsLog& MessageLog) override;
+	
+private:
+	/** If this viewport is allowed to render a preview texture. */
+	bool bAllowPreviewTexture = true;
+	
+	/** If this object is managing the preview texture state. */
+	bool bIsManagingPreviewTexture = false;
 #endif
 
 public:
diff --git a/Engine/Plugins/Runtime/nDisplay/Source/DisplayClusterConfigurator/Private/DisplayClusterConfiguratorCompiler.cpp b/Engine/Plugins/Runtime/nDisplay/Source/DisplayClusterConfigurator/Private/DisplayClusterConfiguratorCompiler.cpp
--- a/Engine/Plugins/Runtime/nDisplay/Source/DisplayClusterConfigurator/Private/DisplayClusterConfiguratorCompiler.cpp
+++ b/Engine/Plugins/Runtime/nDisplay/Source/DisplayClusterConfigurator/Private/DisplayClusterConfiguratorCompiler.cpp
@@ -75,6 +75,19 @@ void FDisplayClusterConfiguratorKismetCompilerContext::OnNewClassSet(UBlueprintG
 void FDisplayClusterConfiguratorKismetCompilerContext::PreCompile()
 {
 	Super::PreCompile();
+	
+	const UDisplayClusterBlueprint* DCBlueprint = CastChecked<UDisplayClusterBlueprint>(Blueprint);
+	if (const UDisplayClusterConfigurationData* ConfigData = DCBlueprint->GetConfig())
+	{
+		ForEachObjectWithOuter(ConfigData, [this](UObject* Child)
+		{
+			if (UDisplayClusterConfigurationData_Base* BaseNode = Cast<UDisplayClusterConfigurationData_Base>(Child))
+			{
+				BaseNode->OnPreCompile(MessageLog);
+			}
+		});
+	}
+
 	ValidateConfiguration();
 }
 
diff --git a/Engine/Plugins/Runtime/nDisplay/Source/DisplayClusterConfigurator/Private/Views/OutputMapping/EdNodes/DisplayClusterConfiguratorBaseNode.h b/Engine/Plugins/Runtime/nDisplay/Source/DisplayClusterConfigurator/Private/Views/OutputMapping/EdNodes/DisplayClusterConfiguratorBaseNode.h
--- a/Engine/Plugins/Runtime/nDisplay/Source/DisplayClusterConfigurator/Private/Views/OutputMapping/EdNodes/DisplayClusterConfiguratorBaseNode.h
+++ b/Engine/Plugins/Runtime/nDisplay/Source/DisplayClusterConfigurator/Private/Views/OutputMapping/EdNodes/DisplayClusterConfiguratorBaseNode.h
@@ -154,6 +154,7 @@ protected:
 	/** Remove invalid nodes. */
 	void CleanupChildrenNodes();
 	
+public:
 	template<class TObjectType>
 	TObjectType* GetObjectChecked() const
 	{
diff --git a/Engine/Plugins/Runtime/nDisplay/Source/DisplayClusterConfigurator/Private/Views/OutputMapping/EdNodes/DisplayClusterConfiguratorViewportNode.cpp b/Engine/Plugins/Runtime/nDisplay/Source/DisplayClusterConfigurator/Private/Views/OutputMapping/EdNodes/DisplayClusterConfiguratorViewportNode.cpp
--- a/Engine/Plugins/Runtime/nDisplay/Source/DisplayClusterConfigurator/Private/Views/OutputMapping/EdNodes/DisplayClusterConfiguratorViewportNode.cpp
+++ b/Engine/Plugins/Runtime/nDisplay/Source/DisplayClusterConfigurator/Private/Views/OutputMapping/EdNodes/DisplayClusterConfiguratorViewportNode.cpp
@@ -228,6 +228,15 @@ UTexture* UDisplayClusterConfiguratorViewportNode::GetPreviewTexture() const
 		UDisplayClusterConfiguratorWindowNode* ParentWindow = GetParentChecked<UDisplayClusterConfiguratorWindowNode>();
 		if (UDisplayClusterPreviewComponent* PreviewComp = RootActor->GetPreviewComponent(ParentWindow->GetNodeName(), GetNodeName()))
 		{
+			const UDisplayClusterConfigurationViewport* CfgViewport = GetObjectChecked<UDisplayClusterConfigurationViewport>();
+			if (!CfgViewport->IsPreviewTextureAllowed())
+			{
+				// The preview texture may be disabled during certain operations, like a resize. Use only an old preview
+				// if possible. This can improve performance and also prevent possible DX crashes if the render target
+				// size is smaller than the viewport region size.
+				return PreviewComp->GetViewportPreviewTexture2D();
+			}
+			
 			if (UTexture2D* Texture = PreviewComp->GetOrCreateViewportPreviewTexture2D())
 			{
 				return Texture;
diff --git a/Engine/Plugins/Runtime/nDisplay/Source/DisplayClusterConfigurator/Private/Views/OutputMapping/GraphNodes/SDisplayClusterConfiguratorViewportNode.cpp b/Engine/Plugins/Runtime/nDisplay/Source/DisplayClusterConfigurator/Private/Views/OutputMapping/GraphNodes/SDisplayClusterConfiguratorViewportNode.cpp
--- a/Engine/Plugins/Runtime/nDisplay/Source/DisplayClusterConfigurator/Private/Views/OutputMapping/GraphNodes/SDisplayClusterConfiguratorViewportNode.cpp
+++ b/Engine/Plugins/Runtime/nDisplay/Source/DisplayClusterConfigurator/Private/Views/OutputMapping/GraphNodes/SDisplayClusterConfiguratorViewportNode.cpp
@@ -196,6 +196,29 @@ bool SDisplayClusterConfiguratorViewportNode::IsAspectRatioFixed() const
 	return ViewportEdNode->IsFixedAspectRatio();
 }
 
+void SDisplayClusterConfiguratorViewportNode::BeginUserInteraction() const
+{
+	SDisplayClusterConfiguratorBaseNode::BeginUserInteraction();
+
+	const UDisplayClusterConfiguratorViewportNode* ViewportEdNode = GetGraphNodeChecked<UDisplayClusterConfiguratorViewportNode>();
+	UDisplayClusterConfigurationViewport* Viewport = ViewportEdNode->GetObjectChecked<UDisplayClusterConfigurationViewport>();
+	bIsManagingPreviewTexture = Viewport->DisablePreviewTexture();
+}
+
+void SDisplayClusterConfiguratorViewportNode::EndUserInteraction() const
+{
+	SDisplayClusterConfiguratorBaseNode::EndUserInteraction();
+	
+	if (bIsManagingPreviewTexture)
+	{
+		const UDisplayClusterConfiguratorViewportNode* ViewportEdNode = GetGraphNodeChecked<UDisplayClusterConfiguratorViewportNode>();
+		UDisplayClusterConfigurationViewport* Viewport = ViewportEdNode->GetObjectChecked<UDisplayClusterConfigurationViewport>();
+		
+		Viewport->EnablePreviewTexture();
+		bIsManagingPreviewTexture = false;
+	}
+}
+
 FSlateColor SDisplayClusterConfiguratorViewportNode::GetBackgroundColor() const
 {
 	TSharedPtr<FDisplayClusterConfiguratorBlueprintEditor> Toolkit = ToolkitPtr.Pin();
diff --git a/Engine/Plugins/Runtime/nDisplay/Source/DisplayClusterConfigurator/Private/Views/OutputMapping/GraphNodes/SDisplayClusterConfiguratorViewportNode.h b/Engine/Plugins/Runtime/nDisplay/Source/DisplayClusterConfigurator/Private/Views/OutputMapping/GraphNodes/SDisplayClusterConfiguratorViewportNode.h
--- a/Engine/Plugins/Runtime/nDisplay/Source/DisplayClusterConfigurator/Private/Views/OutputMapping/GraphNodes/SDisplayClusterConfiguratorViewportNode.h
+++ b/Engine/Plugins/Runtime/nDisplay/Source/DisplayClusterConfigurator/Private/Views/OutputMapping/GraphNodes/SDisplayClusterConfiguratorViewportNode.h
@@ -39,6 +39,8 @@ public:
 	virtual float GetNodeMinimumSize() const override;
 	virtual float GetNodeMaximumSize() const override;
 	virtual bool IsAspectRatioFixed() const override;
+	virtual void BeginUserInteraction() const override;
+	virtual void EndUserInteraction() const override;
 	//~ End of SDisplayClusterConfiguratorBaseNode interface
 
 private:
@@ -64,4 +66,7 @@ private:
 	TSharedPtr<SImage> BackgroundImage;
 
 	UTexture* CachedTexture;
+	
+	/** If this object is managing the preview texture, such as through resizing. */
+	mutable bool bIsManagingPreviewTexture = false;
 };
