nDisplay: fixed deadlock in the new barrier logic

diff --git a/Engine/Plugins/Runtime/nDisplay/Source/DisplayCluster/Private/Network/Barrier/DisplayClusterBarrierV2.cpp b/Engine/Plugins/Runtime/nDisplay/Source/DisplayCluster/Private/Network/Barrier/DisplayClusterBarrierV2.cpp
--- a/Engine/Plugins/Runtime/nDisplay/Source/DisplayCluster/Private/Network/Barrier/DisplayClusterBarrierV2.cpp
+++ b/Engine/Plugins/Runtime/nDisplay/Source/DisplayCluster/Private/Network/Barrier/DisplayClusterBarrierV2.cpp
@@ -28,7 +28,7 @@ FDisplayClusterBarrierV2::~FDisplayClusterBarrierV2()
 
 bool FDisplayClusterBarrierV2::Activate()
 {
-	FScopeLock Lock(&BarrierCS);
+	FScopeLock Lock(&DataCS);
 
 	UE_LOG(LogDisplayClusterNetwork, Log, TEXT("Barrier %s: activating..."), *Name);
 
@@ -49,7 +49,7 @@ bool FDisplayClusterBarrierV2::Activate()
 
 void FDisplayClusterBarrierV2::Deactivate()
 {
-	FScopeLock Lock(&BarrierCS);
+	FScopeLock Lock(&DataCS);
 
 	if (bActive)
 	{
@@ -74,77 +74,81 @@ EDisplayClusterBarrierWaitResult FDisplayClusterBarrierV2::Wait(const FString& N
 	double ThreadWaitTimeStart = 0;
 
 	{
-		FScopeLock Lock(&BarrierCS);
-
-		// Check if this thread has been previously dropped
-		if (NodesTimedout.Contains(NodeId))
-		{
-			UE_LOG(LogDisplayClusterNetwork, Verbose, TEXT("Barrier %s: node %s not allowed to join, it has been timed out previously"), *Name, *NodeId);
-			return EDisplayClusterBarrierWaitResult::TimeOut;
-		}
-
-		// Check if the barrier is active
-		if (!bActive)
-		{
-			UE_LOG(LogDisplayClusterNetwork, Verbose, TEXT("Barrier %s: not active"), *Name);
-			return EDisplayClusterBarrierWaitResult::NotActive;
-		}
-
-		// Check if this thread is allowed to sync at this barrier
-		if (!NodesAllowed.Contains(NodeId))
-		{
-			UE_LOG(LogDisplayClusterNetwork, Verbose, TEXT("Barrier %s: node %s not allowed to join, no permission"), *Name, *NodeId);
-			return EDisplayClusterBarrierWaitResult::NotAllowed;
-		}
-
-		// Register node
-		NodesAwaiting.Add(NodeId);
+		FScopeLock LockEntrance(&EntranceCS);
 
 		// Wait unless the barrier allows new threads to join. This will happen
 		// once all threads from the previous sync iteration leave the barrier,
 		// or the barrier gets deactivated.
 		EventInputGateOpen->Wait();
 
-		// Make sure the barrier was not deactivated while this thread was waiting at the input gate
-		if (bActive)
 		{
-			// Fixate awaiting start for this particular thread
-			ThreadWaitTimeStart = FPlatformTime::Seconds();
-
-			// One more thread awaiting
-			++ThreadCount;
+			FScopeLock LockData(&DataCS);
 
-			// In case this thread came first to the barrier, we need:
-			// - to fixate barrier awaiting start time
-			// - start watchdog timer
-			if (ThreadCount == 1)
+			// Check if this thread has been previously dropped
+			if (NodesTimedout.Contains(NodeId))
 			{
-				UE_LOG(LogDisplayClusterNetwork, Verbose, TEXT("Barrier %s: sync start"), *Name);
-
-				BarrierWaitTimeStart = ThreadWaitTimeStart;
-				WatchdogTimer.SetTimer(Timeout);
+				UE_LOG(LogDisplayClusterNetwork, Verbose, TEXT("Barrier %s: node %s not allowed to join, it has been timed out previously"), *Name, *NodeId);
+				return EDisplayClusterBarrierWaitResult::TimeOut;
 			}
 
-			UE_LOG(LogDisplayClusterNetwork, Verbose, TEXT("Barrier %s: awaiting threads amount - %d"), *Name, ThreadCount);
-
-			// In case this thread is the last one the barrier is awaiting for, we need:
-			// - to fixate barrier awaiting finish time
-			// - to open the output gate (release the barrier)
-			// - to close the input gate
-			// - reset watchdog timer
-			if (ThreadCount == ThreadLimit)
+			// Check if the barrier is active
+			if (!bActive)
 			{
-				BarrierWaitTimeFinish = FPlatformTime::Seconds();
-				BarrierWaitTimeOverall = BarrierWaitTimeFinish - BarrierWaitTimeStart;
+				UE_LOG(LogDisplayClusterNetwork, Verbose, TEXT("Barrier %s: not active"), *Name);
+				return EDisplayClusterBarrierWaitResult::NotActive;
+			}
 
-				UE_LOG(LogDisplayClusterNetwork, Verbose, TEXT("Barrier %s: sync end, barrier wait time %f"), *Name, BarrierWaitTimeOverall);
+			// Check if this thread is allowed to sync at this barrier
+			if (!NodesAllowed.Contains(NodeId))
+			{
+				UE_LOG(LogDisplayClusterNetwork, Verbose, TEXT("Barrier %s: node %s not allowed to join, no permission"), *Name, *NodeId);
+				return EDisplayClusterBarrierWaitResult::NotAllowed;
+			}
 
-				// All nodes here, pet the watchdog
-				WatchdogTimer.ResetTimer();
+			// Register node
+			NodesAwaiting.Add(NodeId);
 
-				// Close input gate, open output gate
-				EventInputGateOpen->Reset();
-				EventOutputGateOpen->Trigger();
+			// Make sure the barrier was not deactivated while this thread was waiting at the input gate
+			if (bActive)
+			{
+				// Fixate awaiting start for this particular thread
+				ThreadWaitTimeStart = FPlatformTime::Seconds();
+
+				// One more thread awaiting
+				++ThreadCount;
+
+				// In case this thread came first to the barrier, we need:
+				// - to fixate barrier awaiting start time
+				// - start watchdog timer
+				if (ThreadCount == 1)
+				{
+					UE_LOG(LogDisplayClusterNetwork, Verbose, TEXT("Barrier %s: sync start"), *Name);
+
+					BarrierWaitTimeStart = ThreadWaitTimeStart;
+					WatchdogTimer.SetTimer(Timeout);
+				}
+
+				UE_LOG(LogDisplayClusterNetwork, Verbose, TEXT("Barrier %s: awaiting threads amount - %d"), *Name, ThreadCount);
+
+				// In case this thread is the last one the barrier is awaiting for, we need:
+				// - to fixate barrier awaiting finish time
+				// - to open the output gate (release the barrier)
+				// - to close the input gate
+				// - reset watchdog timer
+				if (ThreadCount == ThreadLimit)
+				{
+					BarrierWaitTimeFinish = FPlatformTime::Seconds();
+					BarrierWaitTimeOverall = BarrierWaitTimeFinish - BarrierWaitTimeStart;
+
+					UE_LOG(LogDisplayClusterNetwork, Verbose, TEXT("Barrier %s: sync end, barrier wait time %f"), *Name, BarrierWaitTimeOverall);
+
+					// All nodes here, pet the watchdog
+					WatchdogTimer.ResetTimer();
+
+					// Close input gate, open output gate
+					EventInputGateOpen->Reset();
+					EventOutputGateOpen->Trigger();
+				}
 			}
 		}
 	}
@@ -156,21 +160,23 @@ EDisplayClusterBarrierWaitResult FDisplayClusterBarrierV2::Wait(const FString& N
 	const double ThreadWaitTimeFinish = FPlatformTime::Seconds();
 
 	{
-		FScopeLock Lock(&BarrierCS);
-
 		UE_LOG(LogDisplayClusterNetwork, Verbose, TEXT("Barrier %s: node %s is leaving the barrier"), *Name, *NodeId);
 
-		// Unregister node
-		NodesAwaiting.Remove(NodeId);
-
-		// Make sure the barrier was not deactivated while this thread was waiting at the output gate
-		if (bActive)
 		{
-			// In case this thread is leaving last, close output and open input
-			if (--ThreadCount == 0)
+			FScopeLock LockData(&DataCS);
+
+			// Unregister node
+			NodesAwaiting.Remove(NodeId);
+
+			// Make sure the barrier was not deactivated while this thread was waiting at the output gate
+			if (bActive)
 			{
-				EventOutputGateOpen->Reset();
-				EventInputGateOpen->Trigger();
+				// In case this thread is leaving last, close output and open input
+				if (--ThreadCount == 0)
+				{
+					EventOutputGateOpen->Reset();
+					EventInputGateOpen->Trigger();
+				}
 			}
 		}
 	}
@@ -194,7 +200,7 @@ void FDisplayClusterBarrierV2::UnregisterSyncNode(const FString& NodeId)
 {
 	UE_LOG(LogDisplayClusterNetwork, Log, TEXT("Barrier %s: unregistering node %s..."), *Name, *NodeId);
 
-	FScopeLock Lock(&BarrierCS);
+	FScopeLock Lock(&DataCS);
 
 	// Remove specified nodes from the management lists
 	if (NodesAllowed.Remove(NodeId) > 0 || NodesTimedout.Remove(NodeId) > 0)
@@ -207,7 +213,7 @@ void FDisplayClusterBarrierV2::UnregisterSyncNode(const FString& NodeId)
 		// In case it's a last missing node, we need to open the barrier
 		if (ThreadCount == ThreadLimit)
 		{
-			BarrierWaitTimeFinish = FPlatformTime::Seconds();
+			BarrierWaitTimeFinish  = FPlatformTime::Seconds();
 			BarrierWaitTimeOverall = BarrierWaitTimeFinish - BarrierWaitTimeStart;
 
 			UE_LOG(LogDisplayClusterNetwork, Verbose, TEXT("Barrier %s: sync end, barrier wait time %f"), *Name, BarrierWaitTimeOverall);
@@ -231,7 +237,7 @@ void FDisplayClusterBarrierV2::HandleBarrierTimeout()
 	// Being here means some nodes have not come to the barrier yet during specific time period. Those
 	// missing nodes will be considered as the lost ones. The barrier will continue working with the remaining nodes only.
 
-	FScopeLock Lock(&BarrierCS);
+	FScopeLock Lock(&DataCS);
 
 	UE_LOG(LogDisplayClusterNetwork, Log, TEXT("Barrier %s: Time out! %d nodes missing"), *Name, NodesAllowed.Num() - NodesAwaiting.Num());
 
diff --git a/Engine/Plugins/Runtime/nDisplay/Source/DisplayCluster/Private/Network/Barrier/DisplayClusterBarrierV2.h b/Engine/Plugins/Runtime/nDisplay/Source/DisplayCluster/Private/Network/Barrier/DisplayClusterBarrierV2.h
--- a/Engine/Plugins/Runtime/nDisplay/Source/DisplayCluster/Private/Network/Barrier/DisplayClusterBarrierV2.h
+++ b/Engine/Plugins/Runtime/nDisplay/Source/DisplayCluster/Private/Network/Barrier/DisplayClusterBarrierV2.h
@@ -89,5 +89,7 @@ private:
 	double BarrierWaitTimeOverall = 0;
 
 	// Barrier state CS
-	mutable FCriticalSection BarrierCS;
+	mutable FCriticalSection DataCS;
+	// Barrier entrance CS
+	mutable FCriticalSection EntranceCS;
 };
