nDisplay: fixed DCRA preview rendering, fixed crash on extreme multiplier values, added freeze outer rendering checkbox, added cvar to enable alpha channel passthrough, fixed output mapping preview, fixed crash on editing topology, fixed log spam

diff --git a/Engine/Plugins/Runtime/nDisplay/Shaders/Private/MPCDIOverlayShaders.usf b/Engine/Plugins/Runtime/nDisplay/Shaders/Private/MPCDIOverlayShaders.usf
--- a/Engine/Plugins/Runtime/nDisplay/Shaders/Private/MPCDIOverlayShaders.usf
+++ b/Engine/Plugins/Runtime/nDisplay/Shaders/Private/MPCDIOverlayShaders.usf
@@ -195,8 +195,14 @@ OutputPS IcvfxWarpPS(OutputVS IN)
 
 	// Compose all:
 #if VIEWPORT_INPUT
+
 	// First and single pass renders:
-	float4 OutColor = float4(InputTexture.Sample(InputSampler, ViewportUV).rgb, 1.0f);
+
+	#if VIEWPORT_INPUT_ALPHA
+		float4 OutColor = float4(InputTexture.Sample(InputSampler, ViewportUV));
+	#else
+		float4 OutColor = float4(InputTexture.Sample(InputSampler, ViewportUV).rgb, 1.0f);
+	#endif
 
 	#if OVERLAY_UNDER
 		OutColor.xyz = ComposeIcvfxOverlay(OutColor.xyz, OverlayUV);
diff --git a/Engine/Plugins/Runtime/nDisplay/Shaders/Private/MPCDIShaders.usf b/Engine/Plugins/Runtime/nDisplay/Shaders/Private/MPCDIShaders.usf
--- a/Engine/Plugins/Runtime/nDisplay/Shaders/Private/MPCDIShaders.usf
+++ b/Engine/Plugins/Runtime/nDisplay/Shaders/Private/MPCDIShaders.usf
@@ -45,7 +45,11 @@ OutputPS WarpPS(OutputVS IN)
 	float4 ViewportUVW = mul(WarpedUV, ViewportTextureProjectionMatrix);
 	float2 ViewportUV = ViewportUVW.xy / ViewportUVW.w;
 
+#if VIEWPORT_INPUT_ALPHA
+	float4 OutColor = float4(InputTexture.Sample(InputSampler, ViewportUV));
+#else
 	float4 OutColor = float4(InputTexture.Sample(InputSampler, ViewportUV).rgb, 1.0f);
+#endif
 
 	//@todo: add LUT here
 
diff --git a/Engine/Plugins/Runtime/nDisplay/Source/DisplayCluster/Private/Game/EngineClasses/Basics/DisplayClusterViewportClient.cpp b/Engine/Plugins/Runtime/nDisplay/Source/DisplayCluster/Private/Game/EngineClasses/Basics/DisplayClusterViewportClient.cpp
--- a/Engine/Plugins/Runtime/nDisplay/Source/DisplayCluster/Private/Game/EngineClasses/Basics/DisplayClusterViewportClient.cpp
+++ b/Engine/Plugins/Runtime/nDisplay/Source/DisplayCluster/Private/Game/EngineClasses/Basics/DisplayClusterViewportClient.cpp
@@ -495,7 +495,7 @@ void UDisplayClusterViewportClient::Draw(FViewport* InViewport, FCanvas* SceneCa
 				FRotator	ViewRotation;
 				FSceneView* View = LocalPlayer->CalcSceneView(&ViewFamily, ViewLocation, ViewRotation, InViewport, nullptr, ViewportContext.StereoViewIndex);
 
-				if (View && DCView.bDisableRender)
+				if (View && DCView.IsShouldRenderView() == false)
 				{
 					ViewFamily.Views.Remove(View);
 
@@ -889,8 +889,6 @@ bool UDisplayClusterViewportClient::Draw_PIE(FViewport* InViewport, FCanvas* Sce
 	default:
 		break;
 	}
-		
-
 
 	//Get world for render
 	UWorld* const MyWorld = GetWorld();
@@ -900,8 +898,13 @@ bool UDisplayClusterViewportClient::Draw_PIE(FViewport* InViewport, FCanvas* Sce
 	{
 		return false;
 	}
+
+	FDisplayClusterPreviewSettings PreviewSettings;
+	PreviewSettings.bPreviewEnablePostProcess = true;
+	PreviewSettings.bIsPIE = true;
+
 	// Update local node viewports (update\create\delete) and build new render frame
-	if (ViewportManager->UpdateConfiguration(RenderFrameMode, LocalNodeId, RootActor) == false)
+	if (ViewportManager->UpdateConfiguration(RenderFrameMode, LocalNodeId, RootActor, &PreviewSettings) == false)
 	{
 		return false;
 	}
@@ -913,7 +916,8 @@ bool UDisplayClusterViewportClient::Draw_PIE(FViewport* InViewport, FCanvas* Sce
 	}
 
 	bool bOutFrameRendered = false;
-	if (ViewportManager->RenderInEditor(RenderFrame, InViewport, 0, -1, bOutFrameRendered) == false)
+	int32 RenderedViewportsAmount = 0;
+	if (ViewportManager->RenderInEditor(RenderFrame, InViewport, 0, -1, RenderedViewportsAmount, bOutFrameRendered) == false)
 	{
 		return false;
 	}
diff --git a/Engine/Plugins/Runtime/nDisplay/Source/DisplayCluster/Private/Game/EngineClasses/Scene/DisplayClusterPreviewComponent.cpp b/Engine/Plugins/Runtime/nDisplay/Source/DisplayCluster/Private/Game/EngineClasses/Scene/DisplayClusterPreviewComponent.cpp
--- a/Engine/Plugins/Runtime/nDisplay/Source/DisplayCluster/Private/Game/EngineClasses/Scene/DisplayClusterPreviewComponent.cpp
+++ b/Engine/Plugins/Runtime/nDisplay/Source/DisplayCluster/Private/Game/EngineClasses/Scene/DisplayClusterPreviewComponent.cpp
@@ -13,6 +13,11 @@
 #include "Render/Projection/IDisplayClusterProjectionPolicy.h"
 #include "Render/Projection/IDisplayClusterProjectionPolicyFactory.h"
 
+#include "Render/Viewport/DisplayClusterViewport.h"
+#include "Render/Viewport/Containers/DisplayClusterViewportReadPixels.h"
+#include "Render/Viewport/RenderFrame/DisplayClusterRenderFrameSettings.h"
+#include "Render/Viewport/DisplayClusterViewportHelpers.h"
+
 #include "Materials/Material.h"
 #include "Materials/MaterialInstanceDynamic.h"
 
@@ -47,8 +52,6 @@ void UDisplayClusterPreviewComponent::OnComponentCreated()
 	DECLARE_SCOPE_CYCLE_COUNTER(TEXT("UDisplayClusterPreviewComponent::OnComponentCreated"), STAT_OnComponentCreated, STATGROUP_NDisplay);
 	
 	Super::OnComponentCreated();
-
-	InitializePreviewMaterial();
 }
 
 void UDisplayClusterPreviewComponent::DestroyComponent(bool bPromoteChildren)
@@ -58,12 +61,23 @@ void UDisplayClusterPreviewComponent::DestroyComponent(bool bPromoteChildren)
 	ReleasePreviewMesh();
 	ReleasePreviewMaterial();
 
-	ReleasePreviewTexture();
 	ReleasePreviewRenderTarget();
 
 	Super::DestroyComponent(bPromoteChildren);
 }
 
+void UDisplayClusterPreviewComponent::ResetPreviewComponent(bool bInRestoreSceneMaterial)
+{
+	if (bInRestoreSceneMaterial)
+	{
+		RestorePreviewMeshMaterial();
+	}
+	else
+	{
+		UpdatePreviewMesh();
+	}
+}
+
 IDisplayClusterViewport* UDisplayClusterPreviewComponent::GetCurrentViewport() const
 {
 	if (RootActor != nullptr)
@@ -93,12 +107,6 @@ void UDisplayClusterPreviewComponent::RestorePreviewMeshMaterial()
 {
 	UpdatePreviewMeshReference();
 
-	if (!bIsRootActorPreviewMesh)
-	{
-		// Forged created meshes, dont restore
-		ReleasePreviewMesh();
-	}
-
 	if (PreviewMesh && OriginalMaterial)
 	{
 		// Restore
@@ -123,6 +131,9 @@ void UDisplayClusterPreviewComponent::SetPreviewMeshMaterial()
 			}
 		}
 
+		InitializePreviewMaterial();
+		UpdatePreviewMaterial();
+
 		// Set preview material
 		if (PreviewMaterialInstance)
 		{
@@ -142,57 +153,66 @@ void UDisplayClusterPreviewComponent::UpdatePreviewMeshReference()
 	}
 }
 
-bool UDisplayClusterPreviewComponent::UpdatePreviewMesh(bool bRestoreOriginalMaterial)
+bool UDisplayClusterPreviewComponent::UpdatePreviewMesh()
 {
 	DECLARE_SCOPE_CYCLE_COUNTER(TEXT("UDisplayClusterPreviewComponent::UpdatePreviewMesh"), STAT_UpdatePreviewMesh, STATGROUP_NDisplay);
-	
+
 	check(IsInGameThread());
 
 	UpdatePreviewMeshReference();
 
-	// And search for new mesh reference
-	IDisplayClusterViewport* Viewport = GetCurrentViewport();
-	if (Viewport != nullptr && Viewport->GetProjectionPolicy().IsValid() && Viewport->GetProjectionPolicy()->HasPreviewMesh())
+	if (IsPreviewEnabled())
 	{
-		// create warp mesh or update changes
-		if (Viewport->GetProjectionPolicy()->IsConfigurationChanged(&WarpMeshSavedProjectionPolicy))
-		{
-			if (bRestoreOriginalMaterial)
-			{
-				RestorePreviewMeshMaterial();
-			}
-			ReleasePreviewMesh();
-		}
+		check(ViewportConfig);
 
-		if (PreviewMesh == nullptr)
+		// And search for new mesh reference
+		IDisplayClusterViewport* Viewport = GetCurrentViewport();
+		const bool bIsViewportEnabled = RenderTarget != nullptr && Viewport != nullptr && Viewport->GetRenderSettings().bEnable && Viewport->GetProjectionPolicy().IsValid();
+		if (bIsViewportEnabled)
 		{
-			// Get new mesh ptr
-			PreviewMesh = Viewport->GetProjectionPolicy()->GetOrCreatePreviewMeshComponent(Viewport, bIsRootActorPreviewMesh);
-
-			// Assign preview material to mesh
-			SetPreviewMeshMaterial();
-
-			if (!ensure(ViewportConfig))
+			// Handle preview mesh:
+			if (Viewport->GetProjectionPolicy()->HasPreviewMesh())
 			{
-				// Can be null during a reimport.
-				// @TODO reimport: See if we can avoid this during reimport.
-				return false;
+				// create warp mesh or update changes
+				if (PreviewMesh != nullptr)
+				{
+					if (Viewport->GetProjectionPolicy()->IsConfigurationChanged(&WarpMeshSavedProjectionPolicy))
+					{
+						RestorePreviewMeshMaterial();
+						ReleasePreviewMesh();
+					}
+				}
+
+				if (PreviewMesh == nullptr)
+				{
+					// Get new mesh ptr
+					PreviewMesh = Viewport->GetProjectionPolicy()->GetOrCreatePreviewMeshComponent(Viewport, bIsRootActorPreviewMesh);
+
+					// Update saved proj policy parameters
+					WarpMeshSavedProjectionPolicy = ViewportConfig->ProjectionPolicy;
+				}
+
+				if (OriginalMaterial == nullptr)
+				{
+					// Assign preview material to mesh
+					SetPreviewMeshMaterial();
+				}
+
+				return true;
 			}
 
-			// Update saved proj policy parameters
-			WarpMeshSavedProjectionPolicy = ViewportConfig->ProjectionPolicy;
-			return true;
-		}
-	}
-	else
-	{
-		if (bRestoreOriginalMaterial)
-		{
+			// Policy without preview mesh
 			RestorePreviewMeshMaterial();
+			ReleasePreviewMesh();
+
+			return true;
 		}
-		ReleasePreviewMesh();
 	}
 
+	// Viewport don't render
+	RestorePreviewMeshMaterial();
+	ReleasePreviewMesh();
+
 	return false;
 }
 
@@ -258,31 +278,46 @@ void UDisplayClusterPreviewComponent::UpdatePreviewRenderTarget()
 {
 	DECLARE_SCOPE_CYCLE_COUNTER(TEXT("UDisplayClusterPreviewComponent::UpdatePreviewRenderTarget"), STAT_UpdatePreviewRenderTarget, STATGROUP_NDisplay);
 	
-	FIntPoint TextureSize(1,1);
-	float     TextureGamma = 1.f;
+	EPixelFormat TextureFormat = EPixelFormat::PF_Unknown;
+	FIntPoint    TextureSize(1,1);
+	float        TextureGamma = 1.f;
+	bool         bTextureSRGB = false;
 
-	if (GetPreviewTextureSettings(TextureSize, TextureGamma))
+	if (GetPreviewTextureSettings(TextureSize, TextureFormat, TextureGamma, bTextureSRGB))
 	{
-		// Create new RTT for preview Material in scene
-		if (RenderTarget == nullptr)
+		if (RenderTarget != nullptr)
 		{
-			static const TConsoleVariableData<int32>* CVarDefaultBackBufferPixelFormat = IConsoleManager::Get().FindTConsoleVariableDataInt(TEXT("r.DefaultBackBufferPixelFormat"));
-			static const EPixelFormat SceneTargetFormat = EDefaultBackBufferPixelFormat::Convert2PixelFormat(EDefaultBackBufferPixelFormat::FromInt(CVarDefaultBackBufferPixelFormat->GetValueOnGameThread()));
+			// Re-create RTT when format changed
+			if (RenderTarget->GetFormat() != TextureFormat)
+			{
+				ReleasePreviewRenderTarget();
+			}
+		}
 
-			RenderTarget = NewObject<UTextureRenderTarget2D>(this);
-			RenderTarget->ClearColor = FLinearColor::Black;
-			// Always use linear color gamma in scene material
-			RenderTarget->TargetGamma = TextureGamma;
-			RenderTarget->InitCustomFormat(TextureSize.X, TextureSize.Y, SceneTargetFormat, false);
+		if (RenderTarget != nullptr)
+		{
+			// Update an existing RTT resource only when settings change
+			if (RenderTarget->TargetGamma != TextureGamma
+				|| RenderTarget->SRGB != bTextureSRGB
+				|| RenderTarget->GetSurfaceWidth() != TextureSize.X
+				|| RenderTarget->GetSurfaceHeight() != TextureSize.Y)
+			{
+				RenderTarget->TargetGamma = TextureGamma;
+				RenderTarget->SRGB = bTextureSRGB;
 
-			InitializePreviewMaterial();
-			UpdatePreviewMaterial();
+				RenderTarget->ResizeTarget(TextureSize.X, TextureSize.Y);
+			}
 		}
-		// Update exist RTT resource:
 		else
 		{
-			RenderTarget->ResizeTarget(TextureSize.X, TextureSize.Y);
+			// Create new RTT
+			RenderTarget = NewObject<UTextureRenderTarget2D>(this);
+			RenderTarget->ClearColor = FLinearColor::Black;
 			RenderTarget->TargetGamma = TextureGamma;
+			RenderTarget->SRGB = bTextureSRGB;
+
+			RenderTarget->InitCustomFormat(TextureSize.X, TextureSize.Y, TextureFormat, false);
+			UpdatePreviewMaterial();
 		}
 	}
 	else
@@ -301,161 +336,34 @@ void UDisplayClusterPreviewComponent::UpdatePreviewRenderTarget()
 	}
 }
 
-bool UDisplayClusterPreviewComponent::GetPreviewTextureSettings(FIntPoint& OutSize, float& OutGamma) const
-{
-	IDisplayClusterViewport* Viewport = GetCurrentViewport();
-	if (Viewport != nullptr)
-	{
-		// The viewport size is already capped for RenderSettings
-		const TArray<FDisplayClusterViewport_Context>& Contexts = Viewport->GetContexts();
-		if (Contexts.Num() > 0)
-		{
-			OutSize = Contexts[0].FrameTargetRect.Size();
-
-			check(OutSize.X > 0);
-			check(OutSize.Y > 0);
-
-			//! Get gamma from current FViewport
-			OutGamma = 2.2f;
-
-			return true;
-		}
-	}
-
-	return false;
-}
-
-bool UDisplayClusterPreviewComponent::IsPreviewAvailable() const
-{
-	IDisplayClusterViewport* Viewport = GetCurrentViewport();
-	return (Viewport != nullptr) && Viewport->GetProjectionPolicy().IsValid();
-}
-
-void UDisplayClusterPreviewComponent::ReleasePreviewTexture()
+bool UDisplayClusterPreviewComponent::GetPreviewTextureSettings(FIntPoint& OutSize, EPixelFormat& OutTextureFormat, float& OutGamma, bool& bOutSRGB) const
 {
-	if (PreviewTexture != nullptr)
-	{
-		PreviewTexture = nullptr;
-	}
-}
-
-bool UDisplayClusterPreviewComponent::UpdatePreviewTexture()
-{
-	DECLARE_SCOPE_CYCLE_COUNTER(TEXT("UDisplayClusterPreviewComponent::UpdatePreviewTexture"), STAT_UpdatePreviewTexture, STATGROUP_NDisplay);
-	
-	check(RenderTarget);
-
-	TArray<FColor> SurfData;
-	FRenderTarget* RenderTargetResource = RenderTarget->GameThread_GetRenderTargetResource();
-	RenderTargetResource->ReadPixels(SurfData);
+	if (IDisplayClusterViewport* PublicViewport = GetCurrentViewport())
 	{
-		// Check for invalid data.. could happen if a viewport was unbound from a screen/mesh/ect but still has a policy assigned.
-		const FColor EmptyColor(ForceInitToZero);
-		if (SurfData.Num() > 0 && SurfData[0] == EmptyColor)
+		if (FDisplayClusterViewport* Viewport = static_cast<FDisplayClusterViewport*>(PublicViewport))
 		{
-			// Quick test on first element shows we might have an invalid texture.
-			const TSet<FColor> TestEmpty(SurfData);
-			if (TestEmpty.Num() == 1 && *TestEmpty.CreateConstIterator() == EmptyColor)
+			// The viewport size is already capped for RenderSettings
+			const TArray<FDisplayClusterViewport_Context>& Contexts = PublicViewport->GetContexts();
+			if (Contexts.Num() > 0)
 			{
-				// Check rest of the texture -- Texture is blank
-				ReleasePreviewTexture();
-				return false;
-			}
-		}
-	}
-
-	FIntPoint DstSize = RenderTargetResource->GetSizeXY();
-	bool SRGB = RenderTarget->SRGB;
-
-	// If source rendertarget texture changed
-	if (PreviewTexture != nullptr)
-	{
-		if (PreviewTexture->GetSizeX() != DstSize.X || PreviewTexture->GetSizeY() != DstSize.Y || PreviewTexture->SRGB != SRGB)
-		{
-			// Size changed, re-create
-			ReleasePreviewTexture();
-		}
-	}
-
-	if (PreviewTexture == nullptr)
-	{
-		// Create new resource
-		PreviewTexture = UTexture2D::CreateTransient(DstSize.X, DstSize.Y, PF_B8G8R8A8);
-		if (PreviewTexture == nullptr)
-		{
-			return false;
-		}
-
-		PreviewTexture->MipGenSettings = TMGS_NoMipmaps;
-		PreviewTexture->SRGB = RenderTarget->SRGB;
-	}
+				DisplayClusterViewportHelpers::GetPreviewRenderTargetDesc_Editor(Viewport->GetRenderFrameSettings(), OutTextureFormat, OutGamma, bOutSRGB);
 
-	// Alpha channel must non-opaque for preview in OutputMapping
-	for (FColor& PixelIt : SurfData)
-	{
-		PixelIt.A = 255;
-	}
-
-	// Transfer data
-	{
-		void* TextureData = PreviewTexture->GetPlatformData()->Mips[0].BulkData.Lock(LOCK_READ_WRITE);
-		const int32 TextureDataSize = SurfData.Num() * 4;
-		FMemory::Memcpy(TextureData, SurfData.GetData(), TextureDataSize);
-		PreviewTexture->GetPlatformData()->Mips[0].BulkData.Unlock();
-		PreviewTexture->UpdateResource();
-	}
-
-	if (PreviewTexture->GetOuter() != this)
-	{
-		PreviewTexture->Rename(nullptr, this, REN_DoNotDirty | REN_ForceNoResetLoaders | REN_DontCreateRedirectors);
-	}
+				OutSize = Contexts[0].FrameTargetRect.Size();
 
-	return true;
-}
-
-void UDisplayClusterPreviewComponent::HandleRenderTargetTextureDeferredUpdate(const EDisplayClusterRenderFrameMode InRenderFrameMode)
-{
-	check(IsInGameThread());
-
-	switch (InRenderFrameMode)
-	{
-	case EDisplayClusterRenderFrameMode::PreviewInScene:
-		if (RenderTargetSurfaceChangedCnt < 1)
-		{
-			RenderTargetSurfaceChangedCnt = 2;
-		}
-		break;
-	default:
-		break;
-	}
-}
+				check(OutSize.X > 0);
+				check(OutSize.Y > 0);
 
-UTexture2D* UDisplayClusterPreviewComponent::GetOrCreateViewportPreviewTexture2D()
-{
-	if (IsPreviewAvailable())
-	{
-		if (RenderTarget && RenderTargetSurfaceChangedCnt)
-		{
-			if (--RenderTargetSurfaceChangedCnt == 0)
-			{
-				UpdatePreviewTexture();
+				return true;
 			}
 		}
-
-		return PreviewTexture;
 	}
 
-	return nullptr;
+	return false;
 }
 
-UTexture2D* UDisplayClusterPreviewComponent::GetViewportPreviewTexture2D() const
+UTexture* UDisplayClusterPreviewComponent::GetViewportPreviewTexture2D()
 {
-	if (IsPreviewAvailable())
-	{
-		return PreviewTexture;
-	}
-
-	return nullptr;
+	return RenderTarget;
 }
 
 void UDisplayClusterPreviewComponent::SetOverrideTexture(UTexture* InOverrideTexture)
diff --git a/Engine/Plugins/Runtime/nDisplay/Source/DisplayCluster/Private/Game/EngineClasses/Scene/DisplayClusterRootActor.cpp b/Engine/Plugins/Runtime/nDisplay/Source/DisplayCluster/Private/Game/EngineClasses/Scene/DisplayClusterRootActor.cpp
--- a/Engine/Plugins/Runtime/nDisplay/Source/DisplayCluster/Private/Game/EngineClasses/Scene/DisplayClusterRootActor.cpp
+++ b/Engine/Plugins/Runtime/nDisplay/Source/DisplayCluster/Private/Game/EngineClasses/Scene/DisplayClusterRootActor.cpp
@@ -488,6 +488,7 @@ bool ADisplayClusterRootActor::GetHiddenInGamePrimitives(TSet<FPrimitiveComponen
 		{
 			GetTypedPrimitives<UStaticMeshComponent>(OutPrimitives, &WarpMeshNames);
 			GetTypedPrimitives<UProceduralMeshComponent>(OutPrimitives, &WarpMeshNames);
+			GetTypedPrimitives<UDisplayClusterScreenComponent>(OutPrimitives, &WarpMeshNames);
 		}
 	}
 
diff --git a/Engine/Plugins/Runtime/nDisplay/Source/DisplayCluster/Private/Game/EngineClasses/Scene/DisplayClusterRootActorEditor.cpp b/Engine/Plugins/Runtime/nDisplay/Source/DisplayCluster/Private/Game/EngineClasses/Scene/DisplayClusterRootActorEditor.cpp
--- a/Engine/Plugins/Runtime/nDisplay/Source/DisplayCluster/Private/Game/EngineClasses/Scene/DisplayClusterRootActorEditor.cpp
+++ b/Engine/Plugins/Runtime/nDisplay/Source/DisplayCluster/Private/Game/EngineClasses/Scene/DisplayClusterRootActorEditor.cpp
@@ -33,6 +33,7 @@
 #include "Render/Viewport/IDisplayClusterViewportManager.h"
 #include "Render/Viewport/IDisplayClusterViewport.h"
 #include "Render/Viewport/IDisplayClusterViewportProxy.h"
+#include "Render/Viewport/RenderFrame/DisplayClusterRenderFrame.h"
 
 #include "Engine/TextureRenderTarget2D.h"
 
@@ -49,8 +50,39 @@
 #include "LevelEditor.h"
 #include "EditorSupportDelegates.h"
 
-void ADisplayClusterRootActor::ResetPreviewInternals()
+//////////////////////////////////////////////////////////////////////////////////////////////
+int32 GDisplayClusterPreviewAllowMultiGPURendering = 0;
+static FAutoConsoleVariableRef CVarDisplayClusterPreviewAllowMultiGPURendering(
+	TEXT("DC.Preview.AllowMultiGPURendering"),
+	GDisplayClusterPreviewAllowMultiGPURendering,
+	TEXT("Allow mGPU for preview rendering (0 == disabled)"),
+	ECVF_RenderThreadSafe
+);
+
+int32 GDisplayClusterPreviewMultiGPURenderingMinIndex = 0;
+static FAutoConsoleVariableRef CVarDisplayClusterPreviewMultiGPURenderingMinIndex(
+	TEXT("DC.Preview.MultiGPURenderingMinIndex"),
+	GDisplayClusterPreviewMultiGPURenderingMinIndex,
+	TEXT("Distribute mGPU render on GPU from #min to #max indices"),
+	ECVF_RenderThreadSafe
+);
+
+int32 GDisplayClusterPreviewMultiGPURenderingMaxIndex = 0;
+static FAutoConsoleVariableRef CVarDisplayClusterPreviewMultiGPURenderingMaxIndex(
+	TEXT("DC.Preview.MultiGPURenderingMaxIndex"),
+	GDisplayClusterPreviewMultiGPURenderingMaxIndex,
+	TEXT("Distribute mGPU render on GPU from #min to #max indices"),
+	ECVF_RenderThreadSafe
+);
+
+//////////////////////////////////////////////////////////////////////////////////////////////
+// ADisplayClusterRootActor
+//////////////////////////////////////////////////////////////////////////////////////////////
+void ADisplayClusterRootActor::ResetPreviewInternals_Editor()
 {
+	// Reset preview components before DCRA rebuild
+	ResetPreviewComponents_Editor(true);
+
 	PreviewRenderFrame.Reset();
 
 	TickPerFrameCounter = 0;
@@ -63,7 +95,7 @@ void ADisplayClusterRootActor::Constructor_Editor()
 	// Allow tick in editor for preview rendering
 	PrimaryActorTick.bStartWithTickEnabled = true;
 
-	ResetPreviewInternals();
+	ResetPreviewInternals_Editor();
 }
 
 void ADisplayClusterRootActor::Destructor_Editor()
@@ -76,7 +108,7 @@ void ADisplayClusterRootActor::Tick_Editor(float DeltaSeconds)
 {
 	if (!IsPreviewEnabled())
 	{
-		ResetPreviewInternals();
+		ResetPreviewInternals_Editor();
 	}
 	else if (!IsRunningGameOrPIE())
 	{
@@ -87,30 +119,24 @@ void ADisplayClusterRootActor::Tick_Editor(float DeltaSeconds)
 			bDeferPreviewGeneration = false;
 			UpdatePreviewComponents();
 		}
-		else
+		
+		// Update preview RTTs correspond to 'TickPerFrame' value
+		if (++TickPerFrameCounter >= TickPerFrame)
 		{
-			// Update preview RTTs correspond to 'TickPerFrame' value
-			if (++TickPerFrameCounter >= TickPerFrame)
-			{
-				TickPerFrameCounter = 0;
+			TickPerFrameCounter = 0;
 
-				// Render viewport for preview material RTTs
-				if (NodesPerFrame > 0)
-				{
-					// Render preview RTTs
-					ImplRenderPreviewForSceneMaterials_Editor(NodesPerFrame);
-				}
-			}
-
-			// preview frustums on each tick
-			RenderPreviewFrustums();
+			// Render viewport for preview material RTTs
+			ImplRenderPreview_Editor();
 		}
+
+		// preview frustums on each tick
+		ImplRenderPreviewFrustums_Editor();
 	}
 }
 
 void ADisplayClusterRootActor::PostActorCreated_Editor()
 {
-	ResetPreviewInternals();
+	ResetPreviewInternals_Editor();
 }
 
 void ADisplayClusterRootActor::PostLoad_Editor()
@@ -118,11 +144,11 @@ void ADisplayClusterRootActor::PostLoad_Editor()
 	// Generating the preview on load for instances in the world can't be done on PostLoad, components may not have loaded flags present.
 	bDeferPreviewGeneration = true;
 
-	ResetPreviewInternals();
+	ResetPreviewInternals_Editor();
 }
 void ADisplayClusterRootActor::BeginDestroy_Editor()
 {
-	ResetPreviewInternals();
+	ResetPreviewInternals_Editor();
 
 	ReleasePreviewComponents();
 
@@ -140,7 +166,8 @@ void ADisplayClusterRootActor::RerunConstructionScripts_Editor()
 
 	UpdateInnerFrustumPriority();
 
-	ResetPreviewInternals();
+	// Reset preview components before DCRA rebuild
+	ResetPreviewComponents_Editor(false);
 }
 
 bool ADisplayClusterRootActor::IsPreviewEnabled() const
@@ -156,41 +183,28 @@ bool ADisplayClusterRootActor::IsPreviewEnabled() const
 	}
 #endif
 
-	return true;
+	return bPreviewEnable;
 }
 
 // Return all RTT RHI resources for preview
-void ADisplayClusterRootActor::GetPreviewRenderTargetableTextures(const EDisplayClusterRenderFrameMode InRenderFrameMode, const TArray<FString>& InViewportNames, TArray<FTextureRHIRef>& OutTextures)
+void ADisplayClusterRootActor::GetPreviewRenderTargetableTextures(const TArray<FString>& InViewportNames, TArray<FTextureRHIRef>& OutTextures)
 {
 	check(IsInGameThread());
 
-	if(IsPreviewEnabled())
+	for(const TPair<FString, UDisplayClusterPreviewComponent*>& PreviewComponentIt : PreviewComponents)
 	{
-		for(const TPair<FString, UDisplayClusterPreviewComponent*>& PreviewComponentIt : PreviewComponents)
+		if (PreviewComponentIt.Value)
 		{
-			if (PreviewComponentIt.Value)
+			const int32 OutTextureIndex = InViewportNames.Find(PreviewComponentIt.Value->GetViewportId());
+			if (OutTextureIndex != INDEX_NONE)
 			{
-				const int32 OutTextureIndex = InViewportNames.Find(PreviewComponentIt.Value->GetViewportId());
-				if (OutTextureIndex != INDEX_NONE)
+				// Add scope for func GetRenderTargetTexture()
+				if (UTextureRenderTarget2D* RenderTarget2D = PreviewComponentIt.Value->GetRenderTargetTexture())
 				{
-					// Add scope for func GetRenderTargetTexture()
-					UTextureRenderTarget2D* RenderTarget2D = nullptr;
-					switch (InRenderFrameMode)
+					FTextureRenderTargetResource* DstRenderTarget = RenderTarget2D->GameThread_GetRenderTargetResource();
+					if (DstRenderTarget != nullptr)
 					{
-					case EDisplayClusterRenderFrameMode::PreviewInScene:
-						RenderTarget2D = PreviewComponentIt.Value->GetRenderTargetTexture();
-						break;
-					default:
-						break;
-					}
-					
-					if (RenderTarget2D != nullptr)
-					{
-						FTextureRenderTargetResource* DstRenderTarget = RenderTarget2D->GameThread_GetRenderTargetResource();
-						if (DstRenderTarget != nullptr)
-						{
-							OutTextures[OutTextureIndex] = DstRenderTarget->TextureRHI;
-						}
+						OutTextures[OutTextureIndex] = DstRenderTarget->TextureRHI;
 					}
 				}
 			}
@@ -281,7 +295,7 @@ IDisplayClusterViewport* ADisplayClusterRootActor::FindPreviewViewport(const FSt
 	return nullptr;
 }
 
-bool ADisplayClusterRootActor::UpdatePreviewConfiguration_Editor(const EDisplayClusterRenderFrameMode InRenderFrameMode, const FString& InClusterNodeId)
+bool ADisplayClusterRootActor::ImplUpdatePreviewConfiguration_Editor(const FString& InClusterNodeId)
 {
 	PreviewRenderFrame.Reset();
 
@@ -289,24 +303,28 @@ bool ADisplayClusterRootActor::UpdatePreviewConfiguration_Editor(const EDisplayC
 	{
 		FDisplayClusterPreviewSettings PreviewSettings;
 		PreviewSettings.PreviewRenderTargetRatioMult = PreviewRenderTargetRatioMult;
-		PreviewSettings.PreviewMaxTextureSize = PreviewMaxTextureSize;
-		PreviewSettings.bAllowMultiGPURenderingInEditor = bAllowMultiGPURendering;
 
-		PreviewSettings.MinGPUIndex = FMath::Max(0, MinGPUIndex);
-		PreviewSettings.MaxGPUIndex = FMath::Max(0, MaxGPUIndex);
+		PreviewSettings.bFreezePreviewRender = bFreezePreviewRender;
+		PreviewSettings.bPreviewEnablePostProcess = bPreviewEnablePostProcess;
+
+		PreviewSettings.PreviewMaxTextureDimension = PreviewMaxTextureDimension;
+		
+		PreviewSettings.bAllowMultiGPURenderingInEditor = GDisplayClusterPreviewAllowMultiGPURendering != 0;
+		PreviewSettings.MinGPUIndex = FMath::Max(0, GDisplayClusterPreviewMultiGPURenderingMinIndex);
+		PreviewSettings.MaxGPUIndex = FMath::Max(0, GDisplayClusterPreviewMultiGPURenderingMaxIndex);
 
-		return ViewportManager->UpdateConfiguration(InRenderFrameMode, InClusterNodeId, this, &PreviewSettings);
+		return ViewportManager->UpdateConfiguration(EDisplayClusterRenderFrameMode::PreviewInScene, InClusterNodeId, this, &PreviewSettings);
 	}
 
 	return false;
 }
 
-void ADisplayClusterRootActor::RenderPreviewClusterNode_Editor(const EDisplayClusterRenderFrameMode InRenderFrameMode, const FString& InClusterNodeId)
+bool ADisplayClusterRootActor::ImplUpdatePreviewRenderFrame_Editor(const FString& InClusterNodeId)
 {
 	if (!PreviewRenderFrame.IsValid())
 	{
 		// Begin render new frame for cluster node
-		if (UpdatePreviewConfiguration_Editor(InRenderFrameMode, InClusterNodeId))
+		if (ImplUpdatePreviewConfiguration_Editor(InClusterNodeId))
 		{
 			// Update all preview components resources before render
 			for (const TTuple<FString, UDisplayClusterPreviewComponent*>& PreviewComponentIt : PreviewComponents)
@@ -327,90 +345,133 @@ void ADisplayClusterRootActor::RenderPreviewClusterNode_Editor(const EDisplayClu
 				if (ViewportManager->BeginNewFrame(nullptr, PreviewWorld, *PreviewRenderFrame))
 				{
 					PreviewViewportIndex = 0;
+					return true;
 				}
+
+				return false;
 			}
 		}
 	}
 
-	if (PreviewRenderFrame.IsValid() && PreviewViewportIndex >= 0)
+	return true;
+}
+
+bool ADisplayClusterRootActor::ImplRenderPassPreviewClusterNode_Editor()
+{
+	if (!PreviewRenderFrame.IsValid() || PreviewViewportIndex < 0)
 	{
-		bool bFrameRendered = false;
-		ViewportManager->RenderInEditor(*PreviewRenderFrame, nullptr, PreviewViewportIndex, ViewportsPerFrame, bFrameRendered);
+		return false;
+	}
 
-		PreviewViewportIndex += ViewportsPerFrame;
+	int32 ViewportsAmount = ViewportsPerFrame - PreviewViewportsRenderedInThisFrameCnt;
+	if (ViewportsAmount <= 0)
+	{
+		// All viewports for this pass is rendered
+		return false;
+	}
 
-		if (bFrameRendered)
-		{
-			// Send event about cluster node RTT changed
-			for (const TPair<FString, UDisplayClusterPreviewComponent*>& PreviewComponentIt : PreviewComponents)
-			{
-				if (PreviewComponentIt.Value && PreviewComponentIt.Value->GetClusterNodeId() == InClusterNodeId)
-				{
-					PreviewComponentIt.Value->HandleRenderTargetTextureDeferredUpdate(InRenderFrameMode);
-				}
-			}
+	bool bFrameRendered = false;
+	int32 RenderedViewportsAmount = 0;
 
-			PreviewViewportIndex = -1;
-			PreviewRenderFrame.Reset();
+	ViewportManager->RenderInEditor(*PreviewRenderFrame, nullptr, PreviewViewportIndex, ViewportsAmount, RenderedViewportsAmount, bFrameRendered);
 
-			// Send event about RTT changed
-			OnPreviewGenerated.ExecuteIfBound();
-		}
+	// Increase viewport index
+	PreviewViewportIndex += ViewportsAmount;
+
+	// Count only rendered viewports
+	PreviewViewportsRenderedInThisFrameCnt += RenderedViewportsAmount;
+
+	if (bFrameRendered)
+	{
+		// current cluster node is composed
+		PreviewViewportIndex = -1;
+		PreviewRenderFrame.Reset();
+
+		// Send event about RTT changed
+		OnPreviewGenerated.ExecuteIfBound();
+
+		return true;
 	}
+
+	return false;
 }
 
-void ADisplayClusterRootActor::ImplRenderPreviewForSceneMaterials_Editor(const int32 InNumNodesForRender)
+void ADisplayClusterRootActor::ImplRenderPreview_Editor()
 {
-	if (bPreviewRenderEveryFrame && PreviewClusterNodeIndex < 0)
+	if (!bPreviewEnable || !ViewportManager.IsValid())
 	{
+		// no preview
+		return;
+	}
+
+	if (PreviewClusterNodeIndex < 0)
+	{
+		if (bFreezePreviewRender)
+		{
+			return;
+		}
+
 		// Allow preview render
 		PreviewClusterNodeIndex = 0;
 	}
 
-	if (bPreviewEnable && ViewportManager.IsValid() && PreviewClusterNodeIndex >= 0)
+	// per-node render
+	TArray<FString> ExistClusterNodesIDs;
+	CurrentConfigData->Cluster->GetNodeIds(ExistClusterNodesIDs);
+
+	// When viewports count changed, reset render cycle
+	if (PreviewClusterNodeIndex >= ExistClusterNodesIDs.Num())
 	{
-		// per-node render
-		TArray<FString> ExistClusterNodesIDs;
-		CurrentConfigData->Cluster->GetNodeIds(ExistClusterNodesIDs);
+		PreviewClusterNodeIndex = 0;
+	}
 
-		// When viewports count changed, reset render cycle
-		if (PreviewClusterNodeIndex >= ExistClusterNodesIDs.Num())
+	// Try render all nodes in one pass
+	int32 NumNodesForSceneMaterialPreview = ExistClusterNodesIDs.Num();
+	PreviewViewportsRenderedInThisFrameCnt = 0;
+
+	for (int32 NodeIt = 0; NodeIt < NumNodesForSceneMaterialPreview; NodeIt++)
+	{
+		int32 ViewportsAmount = ViewportsPerFrame - PreviewViewportsRenderedInThisFrameCnt;
+		if (ViewportsAmount <= 0)
 		{
-			PreviewClusterNodeIndex = 0;
+			// All viewports for this pass is rendered
+			return;
 		}
 
-		int32 NumNodesForSceneMaterialPreview = FMath::Min(InNumNodesForRender, ExistClusterNodesIDs.Num());
+		const FString& ClusterNodeId = ExistClusterNodesIDs[PreviewClusterNodeIndex];
 
-		for (int32 NodeIt = 0; NodeIt < NumNodesForSceneMaterialPreview; NodeIt++)
+		if (PreviewRenderFrameClusterNodeId != ClusterNodeId)
 		{
-			const FString& ClusterNodeId = ExistClusterNodesIDs[PreviewClusterNodeIndex];
+			PreviewRenderFrameClusterNodeId = ClusterNodeId;
+			PreviewRenderFrame.Reset();
+			PreviewViewportIndex = -1;
+		}
 
-			// Render this cluster node
-			RenderPreviewClusterNode_Editor(EDisplayClusterRenderFrameMode::PreviewInScene, ClusterNodeId);
+		ImplUpdatePreviewRenderFrame_Editor(PreviewRenderFrameClusterNodeId);
 
-			if (PreviewRenderFrame.IsValid())
-			{
-				// Cluster node rendere in progress..
-				break;
-			}
+		// Render this cluster node viewports
+		if (!ImplRenderPassPreviewClusterNode_Editor())
+		{
+			// Cluster node render still in progress..
+			return;
+		}
 
-			// Loop over cluster nodes
-			PreviewClusterNodeIndex++;
-			if (PreviewClusterNodeIndex >= ExistClusterNodesIDs.Num())
-			{
-				PreviewClusterNodeIndex = bPreviewRenderEveryFrame  ? 0 : -1;
-			}
+		// Loop over cluster nodes
+		PreviewClusterNodeIndex++;
+		if (PreviewClusterNodeIndex >= ExistClusterNodesIDs.Num())
+		{
+			PreviewClusterNodeIndex = bFreezePreviewRender ? -1 : 0;
+		}
 
-			if (PreviewClusterNodeIndex < 0)
-			{
-				// Frame captured. stop render
-				break;
-			}
+		if (PreviewClusterNodeIndex < 0)
+		{
+			// Frame captured. stop render
+			break;
 		}
 	}
 }
 
-void ADisplayClusterRootActor::RenderPreviewFrustums()
+void ADisplayClusterRootActor::ImplRenderPreviewFrustums_Editor()
 {
 	if (!ViewportManager.IsValid())
 	{
@@ -464,30 +525,59 @@ void ADisplayClusterRootActor::RenderPreviewFrustums()
 		}
 	}
 
-	for (IDisplayClusterViewport* Viewport : FrustumPreviewViewports)
+	for (IDisplayClusterViewport* ViewportIt : FrustumPreviewViewports)
 	{
-		const TArray<FDisplayClusterViewport_Context>& Contexts = Viewport->GetContexts();
-		for (const FDisplayClusterViewport_Context& Context : Contexts)
+		const TArray<FDisplayClusterViewport_Context>& Contexts = ViewportIt->GetContexts();
+
+		// Due to rendering optimizations for DCRA preview (node ​​rendered over multiple frames). 
+		// As a result, viewport math values ​​(stored in contexts) may not be ready at the moment.
+		FFrustumPreviewViewportContextCache FrustumPreviewViewportContext;
+		bool bIsValidViewportContext = false;
+
+		// Preview rendered only in mono
+		if (Contexts.Num() == 1 && Contexts[0].bIsValidProjectionMatrix && Contexts[0].bIsValidViewLocation && Contexts[0].bIsValidViewRotation)
 		{
-			FMatrix ProjectionMatrix = Context.ProjectionMatrix;
+			FrustumPreviewViewportContext.ProjectionMatrix = Contexts[0].ProjectionMatrix;
+			FrustumPreviewViewportContext.ViewLocation = Contexts[0].ViewLocation;
+			FrustumPreviewViewportContext.ViewRotation = Contexts[0].ViewRotation;
 
-			uint32 ViewportContextNum = 0;
-			Viewport->GetProjectionMatrix(ViewportContextNum, ProjectionMatrix);
+			if (ViewportIt->GetProjectionMatrix(0, FrustumPreviewViewportContext.ProjectionMatrix))
+			{
+				bIsValidViewportContext = true;
+			}
+		}
+
+		// Get cached value
+		if (!bIsValidViewportContext)
+		{
+			FFrustumPreviewViewportContextCache* const LastValidContext = FrustumPreviewViewportContextCache.Find(ViewportIt->GetId());
+			if (LastValidContext)
+			{
+				FrustumPreviewViewportContext = *LastValidContext;
+				bIsValidViewportContext = true;
+			}
+		}
+
+		if(bIsValidViewportContext)
+		{
+			// Update cache
+			FrustumPreviewViewportContextCache.Emplace(ViewportIt->GetId(), FrustumPreviewViewportContext);
+
+			// Render the frustum
 
-			FMatrix ViewRotationMatrix = FInverseRotationMatrix(Context.ViewRotation) * FMatrix(
+			FMatrix ViewRotationMatrix = FInverseRotationMatrix(FrustumPreviewViewportContext.ViewRotation) * FMatrix(
 				FPlane(0, 0, 1, 0),
 				FPlane(1, 0, 0, 0),
 				FPlane(0, 1, 0, 0),
 				FPlane(0, 0, 0, 1));
+			const FMatrix ViewMatrix = FTranslationMatrix(-FrustumPreviewViewportContext.ViewLocation) * ViewRotationMatrix;
 
-			const FMatrix ViewMatrix = FTranslationMatrix(-Context.ViewLocation) * ViewRotationMatrix;
-
-			RenderPreviewFrustum(ProjectionMatrix, ViewMatrix, Context.ViewLocation);
+			ImplRenderPreviewViewportFrustum_Editor(FrustumPreviewViewportContext.ProjectionMatrix, ViewMatrix, FrustumPreviewViewportContext.ViewLocation);
 		}
 	}
 }
 
-void ADisplayClusterRootActor::RenderPreviewFrustum(const FMatrix ProjectionMatrix, const FMatrix ViewMatrix, const FVector ViewOrigin)
+void ADisplayClusterRootActor::ImplRenderPreviewViewportFrustum_Editor(const FMatrix ProjectionMatrix, const FMatrix ViewMatrix, const FVector ViewOrigin)
 {
 	const float FarPlane = PreviewICVFXFrustumsFarDistance;
 	const float NearPlane = GNearClippingPlane;
@@ -569,6 +659,9 @@ static FName Name_RelativeScale3D = USceneComponent::GetRelativeScale3DPropertyN
 
 void ADisplayClusterRootActor::PostEditChangeProperty(FPropertyChangedEvent& PropertyChangedEvent)
 {
+	// any property update causes a DCRA rebuild. Restore preview material show
+	ResetPreviewComponents_Editor(true);
+
 	const FName PropertyName = (PropertyChangedEvent.Property != nullptr) ? PropertyChangedEvent.Property->GetFName() : NAME_None;
 	
 	// The AActor method, simplified and modified to skip construction scripts.
@@ -581,11 +674,12 @@ void ADisplayClusterRootActor::PostEditChangeProperty(FPropertyChangedEvent& Pro
 		}
 
 		const bool bTransformationChanged = (PropertyName == Name_RelativeLocation || PropertyName == Name_RelativeRotation || PropertyName == Name_RelativeScale3D);
+
 		if ((GEditor && GEditor->bIsSimulatingInEditor && GetWorld() != nullptr) || ReregisterComponentsWhenModified())
 		{
 			// If a transaction is occurring we rely on the true parent method instead.
 			check (!CurrentTransactionAnnotation.IsValid());
-			
+
 			UnregisterAllComponents();
 			ReregisterAllComponents();
 		}
@@ -619,13 +713,23 @@ void ADisplayClusterRootActor::PostEditChangeProperty(FPropertyChangedEvent& Pro
 	
 	// Config file has been changed, we should rebuild the nDisplay hierarchy
 	// Cluster node ID has been changed
-	if (PropertyName == GET_MEMBER_NAME_CHECKED(ADisplayClusterRootActor, PreviewNodeId))
+	if (PropertyName == GET_MEMBER_NAME_CHECKED(ADisplayClusterRootActor, bFreezePreviewRender)
+	|| PropertyName == GET_MEMBER_NAME_CHECKED(ADisplayClusterRootActor, PreviewRenderTargetRatioMult)
+	|| PropertyName == GET_MEMBER_NAME_CHECKED(ADisplayClusterRootActor, bPreviewEnablePostProcess)
+	|| PropertyName == GET_MEMBER_NAME_CHECKED(ADisplayClusterRootActor, bPreviewICVFXFrustums)
+	|| PropertyName == GET_MEMBER_NAME_CHECKED(ADisplayClusterRootActor, PreviewICVFXFrustumsFarDistance)
+	|| PropertyName == GET_MEMBER_NAME_CHECKED(ADisplayClusterRootActor, PreviewMaxTextureDimension)
+	|| PropertyName == GET_MEMBER_NAME_CHECKED(ADisplayClusterRootActor, PreviewNodeId)
+	|| PropertyName == GET_MEMBER_NAME_CHECKED(ADisplayClusterRootActor, RenderMode))
 	{
 		bReinitializeActor = false;
-		AsyncTask(ENamedThreads::GameThread, [this]()
-		{
-			UpdatePreviewComponents();
-		});
+	}
+	else if (PropertyName == GET_MEMBER_NAME_CHECKED(ADisplayClusterRootActor, TickPerFrame)
+		|| PropertyName == GET_MEMBER_NAME_CHECKED(ADisplayClusterRootActor, ViewportsPerFrame))
+	{
+		ResetPreviewInternals_Editor();
+
+		bReinitializeActor = false;
 	}
 	else if (PropertyName == GET_MEMBER_NAME_CHECKED(ADisplayClusterRootActor, InnerFrustumPriority))
 	{
@@ -636,12 +740,29 @@ void ADisplayClusterRootActor::PostEditChangeProperty(FPropertyChangedEvent& Pro
 	{
 		InitializeRootActor();
 	}
+
+	ResetPreviewComponents_Editor(false);
 }
 
 void ADisplayClusterRootActor::PostEditMove(bool bFinished)
 {
 	// Don't update the preview with the config data if we're just moving the actor.
 	Super::PostEditMove(bFinished);
+
+	FrustumPreviewViewportContextCache.Empty();
+
+	ResetPreviewComponents_Editor(false);
+}
+
+void ADisplayClusterRootActor::ResetPreviewComponents_Editor(bool bInRestoreSceneMaterial)
+{
+	TArray<UDisplayClusterPreviewComponent*> AllPreviewComponents;
+	GetComponents<UDisplayClusterPreviewComponent>(AllPreviewComponents);
+
+	for (UDisplayClusterPreviewComponent* ExistingComp : AllPreviewComponents)
+	{
+		ExistingComp->ResetPreviewComponent(bInRestoreSceneMaterial);
+	}
 }
 
 void ADisplayClusterRootActor::UpdatePreviewComponents()
@@ -654,14 +775,12 @@ void ADisplayClusterRootActor::UpdatePreviewComponents()
 	// Render new preview
 	PreviewClusterNodeIndex = 0;
 
-	UpdatePreviewConfiguration_Editor(EDisplayClusterRenderFrameMode::PreviewInScene, DisplayClusterConfigurationStrings::gui::preview::PreviewNodeAll);
+	ImplUpdatePreviewConfiguration_Editor(DisplayClusterConfigurationStrings::gui::preview::PreviewNodeAll);
 
 	TArray<UDisplayClusterPreviewComponent*> IteratedPreviewComponents;
 	
 	if (CurrentConfigData != nullptr && bPreviewEnable)
 	{
-		const bool bAllComponentsVisible = PreviewNodeId.Equals(DisplayClusterConfigurationStrings::gui::preview::PreviewNodeAll, ESearchCase::IgnoreCase);
-
 		for (const TPair<FString, UDisplayClusterConfigurationClusterNode*>& Node : CurrentConfigData->Cluster->Nodes)
 		{
 			if (Node.Value == nullptr)
@@ -669,33 +788,27 @@ void ADisplayClusterRootActor::UpdatePreviewComponents()
 				continue;
 			}
 
-			if (bAllComponentsVisible || Node.Key.Equals(PreviewNodeId, ESearchCase::IgnoreCase))
+			for (const TPair<FString, UDisplayClusterConfigurationViewport*>& Viewport : Node.Value->Viewports)
 			{
-				for (const TPair<FString, UDisplayClusterConfigurationViewport*>& Viewport : Node.Value->Viewports)
+				const FString PreviewCompId = GeneratePreviewComponentName_Editor(Node.Key, Viewport.Key);
+				UDisplayClusterPreviewComponent* PreviewComp = PreviewComponents.FindRef(PreviewCompId);
+				if (!PreviewComp)
 				{
-					if (bAllComponentsVisible || Node.Key.Equals(PreviewNodeId, ESearchCase::IgnoreCase))
-					{
-						const FString PreviewCompId = GeneratePreviewComponentName(Node.Key, Viewport.Key);
-						UDisplayClusterPreviewComponent* PreviewComp = PreviewComponents.FindRef(PreviewCompId);
-						if (!PreviewComp)
-						{
-							PreviewComp = NewObject<UDisplayClusterPreviewComponent>(this, FName(*PreviewCompId), RF_DuplicateTransient | RF_Transactional | RF_NonPIEDuplicateTransient);
-							check(PreviewComp);
-						
-							PreviewComponents.Emplace(PreviewCompId, PreviewComp);
-						}
-
-						if (GetWorld() && !PreviewComp->IsRegistered())
-						{
-							PreviewComp->RegisterComponent();
-						}
-
-						// Always reinitialize so changes impact the preview component.
-						PreviewComp->InitializePreviewComponent(this, Node.Key, Viewport.Key, Viewport.Value);
-
-						IteratedPreviewComponents.Add(PreviewComp);
-					}
+					PreviewComp = NewObject<UDisplayClusterPreviewComponent>(this, FName(*PreviewCompId), RF_DuplicateTransient | RF_Transactional | RF_NonPIEDuplicateTransient);
+					check(PreviewComp);
+
+					PreviewComponents.Emplace(PreviewCompId, PreviewComp);
 				}
+
+				if (GetWorld() && !PreviewComp->IsRegistered())
+				{
+					PreviewComp->RegisterComponent();
+				}
+
+				// Always reinitialize so changes impact the preview component.
+				PreviewComp->InitializePreviewComponent(this, Node.Key, Viewport.Key, Viewport.Value);
+
+				IteratedPreviewComponents.Add(PreviewComp);
 			}
 		}
 	}
@@ -710,6 +823,9 @@ void ADisplayClusterRootActor::UpdatePreviewComponents()
 		{
 			PreviewComponents.Remove(ExistingComp->GetName());
 			
+			// Reset preview components before unregister
+			ExistingComp->ResetPreviewComponent(true);
+
 			ExistingComp->UnregisterComponent();
 			ExistingComp->DestroyComponent();
 		}
@@ -722,6 +838,8 @@ void ADisplayClusterRootActor::ReleasePreviewComponents()
 	{
 		if (CompPair.Value)
 		{
+			CompPair.Value->ResetPreviewComponent(true);
+
 			CompPair.Value->UnregisterComponent();
 			CompPair.Value->DestroyComponent();
 		}
@@ -731,14 +849,14 @@ void ADisplayClusterRootActor::ReleasePreviewComponents()
 	OnPreviewDestroyed.ExecuteIfBound();
 }
 
-FString ADisplayClusterRootActor::GeneratePreviewComponentName(const FString& NodeId, const FString& ViewportId) const
+FString ADisplayClusterRootActor::GeneratePreviewComponentName_Editor(const FString& NodeId, const FString& ViewportId) const
 {
 	return FString::Printf(TEXT("%s_%s"), *NodeId, *ViewportId);
 }
 
 UDisplayClusterPreviewComponent* ADisplayClusterRootActor::GetPreviewComponent(const FString& NodeId, const FString& ViewportId)
 {
-	const FString PreviewCompId = GeneratePreviewComponentName(NodeId, ViewportId);
+	const FString PreviewCompId = GeneratePreviewComponentName_Editor(NodeId, ViewportId);
 	if (PreviewComponents.Contains(PreviewCompId))
 	{
 		return PreviewComponents[PreviewCompId];
diff --git a/Engine/Plugins/Runtime/nDisplay/Source/DisplayCluster/Private/Game/EngineClasses/Scene/DisplayClusterScreenComponent.cpp b/Engine/Plugins/Runtime/nDisplay/Source/DisplayCluster/Private/Game/EngineClasses/Scene/DisplayClusterScreenComponent.cpp
--- a/Engine/Plugins/Runtime/nDisplay/Source/DisplayCluster/Private/Game/EngineClasses/Scene/DisplayClusterScreenComponent.cpp
+++ b/Engine/Plugins/Runtime/nDisplay/Source/DisplayCluster/Private/Game/EngineClasses/Scene/DisplayClusterScreenComponent.cpp
@@ -27,6 +27,14 @@ UDisplayClusterScreenComponent::UDisplayClusterScreenComponent(const FObjectInit
 	{
 		SetVisibility(true);
 	}
+
+#if WITH_EDITORONLY_DATA
+	SetIsVisualizationComponent(true);
+#endif
+
+	bVisibleInReflectionCaptures = false;
+	bVisibleInRayTracing = false;
+	bVisibleInRealTimeSkyCaptures = false;
 #endif
 
 	// Default screen size
diff --git a/Engine/Plugins/Runtime/nDisplay/Source/DisplayCluster/Private/Render/Device/DisplayClusterDeviceBase.cpp b/Engine/Plugins/Runtime/nDisplay/Source/DisplayCluster/Private/Render/Device/DisplayClusterDeviceBase.cpp
--- a/Engine/Plugins/Runtime/nDisplay/Source/DisplayCluster/Private/Render/Device/DisplayClusterDeviceBase.cpp
+++ b/Engine/Plugins/Runtime/nDisplay/Source/DisplayCluster/Private/Render/Device/DisplayClusterDeviceBase.cpp
@@ -273,7 +273,8 @@ void FDisplayClusterDeviceBase::CalculateStereoViewOffset(const int32 ViewIndex,
 		return;
 	}
 
-	const FDisplayClusterViewport_Context& ViewportContext = ViewportPtr->GetContexts()[ViewportContextNum];
+	const TArray<FDisplayClusterViewport_Context>& ViewportContexts = ViewportPtr->GetContexts();
+	const FDisplayClusterViewport_Context& ViewportContext = ViewportContexts[ViewportContextNum];
 
 	UE_LOG(LogDisplayClusterRender, VeryVerbose, TEXT("OLD ViewLoc: %s, ViewRot: %s"), *ViewLocation.ToString(), *ViewRotation.ToString());
 	UE_LOG(LogDisplayClusterRender, VeryVerbose, TEXT("WorldToMeters: %f"), WorldToMeters);
@@ -308,23 +309,14 @@ void FDisplayClusterDeviceBase::CalculateStereoViewOffset(const int32 ViewIndex,
 	const float EyeOffset = CfgEyeDist / 2.f;
 	const float EyeOffsetValues[] = { -EyeOffset, 0.f, EyeOffset };
 
-	auto DecodeEyeType = [](const EStereoscopicPass StereoPass)
+	// Decode current eye type
+	EDisplayClusterEyeType EyeType = EDisplayClusterEyeType::Mono;
+	if(ViewportContexts.Num() > 1)
 	{
-		switch (StereoPass)
-		{
-		case EStereoscopicPass::eSSP_PRIMARY:
-			return EDisplayClusterEyeType::StereoLeft;
-		case EStereoscopicPass::eSSP_SECONDARY:
-			return EDisplayClusterEyeType::StereoRight;
-		default:
-			break;
-		}
-
-		return EDisplayClusterEyeType::Mono;
-	};
+		// Support stereo:
+		EyeType = (ViewportContextNum == 0)? EDisplayClusterEyeType::StereoLeft : EDisplayClusterEyeType::StereoRight;
+	}
 
-	// Decode current eye type	
-	const EDisplayClusterEyeType EyeType = DecodeEyeType(ViewportContext.StereoscopicPass);
 	const int32 EyeIndex = (int32)EyeType;
 
 	float PassOffset = 0.f;
@@ -514,8 +506,6 @@ void FDisplayClusterDeviceBase::RenderTexture_RenderThread(FRHICommandListImmedi
 //////////////////////////////////////////////////////////////////////////////////////////////
 void FDisplayClusterDeviceBase::UpdateViewport(bool bUseSeparateRenderTarget, const class FViewport& Viewport, class SViewport* ViewportWidget)
 {
-	check(IsInGameThread());
-
 	// Store viewport
 	if (!MainViewport)
 	{
@@ -526,8 +516,8 @@ void FDisplayClusterDeviceBase::UpdateViewport(bool bUseSeparateRenderTarget, co
 
 void FDisplayClusterDeviceBase::CalculateRenderTargetSize(const class FViewport& Viewport, uint32& InOutSizeX, uint32& InOutSizeY)
 {
-	check(IsInGameThread());
-	check(InOutSizeX > 0 && InOutSizeY > 0);
+	InOutSizeX = FMath::Max(1, (int32)InOutSizeX);
+	InOutSizeY = FMath::Max(1, (int32)InOutSizeY);
 }
 
 bool FDisplayClusterDeviceBase::NeedReAllocateViewportRenderTarget(const class FViewport& Viewport)
@@ -541,7 +531,7 @@ bool FDisplayClusterDeviceBase::NeedReAllocateViewportRenderTarget(const class F
 	uint32 newSizeX = rtSize.X;
 	uint32 newSizeY = rtSize.Y;
 
-	//CalculateRenderTargetSize(Viewport, newSizeX, newSizeY);
+	CalculateRenderTargetSize(Viewport, newSizeX, newSizeY);
 
 	// Here we conclude if need to re-allocate
 	const bool Result = (newSizeX != rtSize.X || newSizeY != rtSize.Y);
diff --git a/Engine/Plugins/Runtime/nDisplay/Source/DisplayCluster/Private/Render/Viewport/Configuration/DisplayClusterViewportConfiguration.cpp b/Engine/Plugins/Runtime/nDisplay/Source/DisplayCluster/Private/Render/Viewport/Configuration/DisplayClusterViewportConfiguration.cpp
--- a/Engine/Plugins/Runtime/nDisplay/Source/DisplayCluster/Private/Render/Viewport/Configuration/DisplayClusterViewportConfiguration.cpp
+++ b/Engine/Plugins/Runtime/nDisplay/Source/DisplayCluster/Private/Render/Viewport/Configuration/DisplayClusterViewportConfiguration.cpp
@@ -91,22 +91,39 @@ bool FDisplayClusterViewportConfiguration::ImplUpdateConfiguration(EDisplayClust
 				RenderFrameSettings.PreviewRenderTargetRatioMult = InPreviewSettings->PreviewRenderTargetRatioMult;
 
 				// Limit preview textures max size
-				RenderFrameSettings.PreviewMaxTextureSize = InPreviewSettings->PreviewMaxTextureSize;
+				RenderFrameSettings.PreviewMaxTextureDimension = InPreviewSettings->PreviewMaxTextureDimension;
+
+				// Hack preview gamma.
+				// In a scene, PostProcess always renders on top of the preview textures.
+				// But in it, PostProcess is also rendered with the flag turned off.
+				RenderFrameSettings.bPreviewEnablePostProcess = InPreviewSettings->bPreviewEnablePostProcess;
 
 				// Support mGPU for preview rendering
 				RenderFrameSettings.bAllowMultiGPURenderingInEditor = InPreviewSettings->bAllowMultiGPURenderingInEditor;
 				RenderFrameSettings.PreviewMinGPUIndex = InPreviewSettings->MinGPUIndex;
 				RenderFrameSettings.PreviewMaxGPUIndex = InPreviewSettings->MaxGPUIndex;
 
-
 				RenderFrameSettings.bIsRenderingInEditor = true;
 			}
+			else
+			{
+				RenderFrameSettings.bPreviewEnablePostProcess = false;
+				RenderFrameSettings.bAllowMultiGPURenderingInEditor = false;
+				RenderFrameSettings.bIsRenderingInEditor = false;
+			}
 
 			ConfigurationBase.Update(InClusterNodeId);
 			ConfigurationICVFX.Update();
 			ConfigurationProjectionPolicy.Update();
 			ConfigurationICVFX.PostUpdate();
 
+#if WITH_EDITOR
+			if (InPreviewSettings != nullptr && !InPreviewSettings->bIsPIE)
+			{
+				ConfigurationICVFX.PostUpdatePreview_Editor(*InPreviewSettings);
+			}
+#endif
+
 			ImplUpdateConfigurationVisibility(*RootActor, *ConfigurationData);
 
 			ConfigurationBase.UpdateClusterNodePostProcess(InClusterNodeId);
@@ -131,6 +148,8 @@ bool FDisplayClusterViewportConfiguration::UpdatePreviewConfiguration(EDisplayCl
 {
 	if (InClusterNodeId.Equals(DisplayClusterConfigurationStrings::gui::preview::PreviewNodeAll, ESearchCase::IgnoreCase))
 	{
+		check(!InPreviewSettings.bIsPIE);
+
 		// initialize all nodes
 		ADisplayClusterRootActor* RootActor = GetRootActor();
 		if (RootActor != nullptr)
diff --git a/Engine/Plugins/Runtime/nDisplay/Source/DisplayCluster/Private/Render/Viewport/Configuration/DisplayClusterViewportConfigurationHelpers_ICVFX.cpp b/Engine/Plugins/Runtime/nDisplay/Source/DisplayCluster/Private/Render/Viewport/Configuration/DisplayClusterViewportConfigurationHelpers_ICVFX.cpp
--- a/Engine/Plugins/Runtime/nDisplay/Source/DisplayCluster/Private/Render/Viewport/Configuration/DisplayClusterViewportConfigurationHelpers_ICVFX.cpp
+++ b/Engine/Plugins/Runtime/nDisplay/Source/DisplayCluster/Private/Render/Viewport/Configuration/DisplayClusterViewportConfigurationHelpers_ICVFX.cpp
@@ -38,26 +38,73 @@
 
 #include "Misc/DisplayClusterLog.h"
 
-// Initialize camera policy with camera component and settings
-static bool ImplUpdateCameraProjectionSettings(TSharedPtr<IDisplayClusterProjectionPolicy, ESPMode::ThreadSafe>& InOutCameraProjection, const FDisplayClusterConfigurationICVFX_CameraSettings& CameraSettings, UCameraComponent* const CameraComponent)
+////////////////////////////////////////////////////////////////////////////////
+// Experimental feature: to be approved after testing
+int32 GDisplayClusterPreviewEnableReuseViewportInCluster = 1;
+static FAutoConsoleVariableRef CVarDisplayClusterPreviewEnableReuseViewportInCluster(
+	TEXT("DC.Preview.EnableReuseViewportInCluster"),
+	GDisplayClusterPreviewEnableReuseViewportInCluster,
+	TEXT("Experimental feature (0 == disabled, 1 == enabled)"),
+	ECVF_RenderThreadSafe
+);
+
+////////////////////////////////////////////////////////////////////////////////
+namespace DisplayClusterViewportConfigurationHelpers_ICVFX_Impl
 {
-	FDisplayClusterProjectionCameraPolicySettings PolicyCameraSettings;
-	PolicyCameraSettings.FOVMultiplier = CameraSettings.CustomFrustum.bEnable ? CameraSettings.CustomFrustum.FieldOfViewMultiplier : 1.f;
-
-	// Lens correction
-	PolicyCameraSettings.FrustumRotation = CameraSettings.FrustumRotation;
-	PolicyCameraSettings.FrustumOffset   = CameraSettings.FrustumOffset;
+	static float ImplGetFieldOfViewMultiplier(const FDisplayClusterConfigurationICVFX_CameraSettings& CameraSettings)
+	{
+		return CameraSettings.CustomFrustum.bEnable ? CameraSettings.CustomFrustum.FieldOfViewMultiplier : 1.f;
+	}
 
 	// Initialize camera policy with camera component and settings
-	return IDisplayClusterProjection::Get().CameraPolicySetCamera(InOutCameraProjection, CameraComponent, PolicyCameraSettings);
-}
+	static bool ImplUpdateCameraProjectionSettings(TSharedPtr<IDisplayClusterProjectionPolicy, ESPMode::ThreadSafe>& InOutCameraProjection, const FDisplayClusterConfigurationICVFX_CameraSettings& CameraSettings, UCameraComponent* const CameraComponent)
+	{
+		FDisplayClusterProjectionCameraPolicySettings PolicyCameraSettings;
+		PolicyCameraSettings.FOVMultiplier = ImplGetFieldOfViewMultiplier(CameraSettings);
 
-// Return unique ICVFX name
-FString ImplGetNameICVFX(const FString& InClusterNodeId, const FString& InViewportId, const FString& InResourceId)
-{
-	return FString::Printf(TEXT("%s_%s_%s_%s"), *InClusterNodeId, DisplayClusterViewportStrings::icvfx::prefix, *InViewportId, *InResourceId);
-}
+		// Lens correction
+		PolicyCameraSettings.FrustumRotation = CameraSettings.FrustumRotation;
+		PolicyCameraSettings.FrustumOffset = CameraSettings.FrustumOffset;
+
+		// Initialize camera policy with camera component and settings
+		return IDisplayClusterProjection::Get().CameraPolicySetCamera(InOutCameraProjection, CameraComponent, PolicyCameraSettings);
+	}
+
+	// Return unique ICVFX name
+	static FString ImplGetNameICVFX(const FString& InClusterNodeId, const FString& InViewportId, const FString& InResourceId)
+	{
+		return FString::Printf(TEXT("%s_%s_%s_%s"), *InClusterNodeId, DisplayClusterViewportStrings::icvfx::prefix, *InViewportId, *InResourceId);
+	}
+
+#if WITH_EDITOR
+	static bool IsPreviewEnableReuseViewportInCluster_EditorImpl(ADisplayClusterRootActor& InRootActor)
+	{
+		if (GDisplayClusterPreviewEnableReuseViewportInCluster > 0)
+		{
+			FDisplayClusterViewportManager* ViewportManager = FDisplayClusterViewportConfigurationHelpers_ICVFX::GetViewportManager(InRootActor);
+			if (ViewportManager)
+			{
+				const FDisplayClusterRenderFrameSettings& RenderFrameSettings = ViewportManager->GetRenderFrameSettings();
+				switch (RenderFrameSettings.RenderMode)
+				{
+				case EDisplayClusterRenderFrameMode::PreviewInScene:
+					return true;
+				default:
+					break;
+				}
+			}
+		}
+
+		return false;
+	}
+#endif
+};
 
+using namespace DisplayClusterViewportConfigurationHelpers_ICVFX_Impl;
+
+////////////////////////////////////////////////////////////////////////////////
+// FDisplayClusterViewportConfigurationHelpers_ICVFX
+////////////////////////////////////////////////////////////////////////////////
 FDisplayClusterViewport* FDisplayClusterViewportConfigurationHelpers_ICVFX::ImplFindViewport(ADisplayClusterRootActor& RootActor, const FString& InViewportId, const FString& InResourceId)
 {
 	FDisplayClusterViewportManager* ViewportManager = FDisplayClusterViewportConfigurationHelpers_ICVFX::GetViewportManager(RootActor);
@@ -174,6 +221,100 @@ FDisplayClusterViewport* FDisplayClusterViewportConfigurationHelpers_ICVFX::GetO
 	return CameraViewport;
 }
 
+#if WITH_EDITOR
+void FDisplayClusterViewportConfigurationHelpers_ICVFX::PreviewReuseInnerFrustumViewportWithinClusterNodes_Editor(FDisplayClusterViewport& InCameraViewport, ADisplayClusterRootActor& InRootActor, UDisplayClusterICVFXCameraComponent& InCameraComponent)
+{
+	if (!IsPreviewEnableReuseViewportInCluster_EditorImpl(InRootActor))
+	{
+		return;
+	}
+
+	const FDisplayClusterRenderFrameSettings& RenderFrameSettingsConstRef = InCameraViewport.GetRenderFrameSettings();
+	switch (RenderFrameSettingsConstRef.RenderMode)
+	{
+	case EDisplayClusterRenderFrameMode::PreviewInScene:
+		// Only for preview mode
+		break;
+	default:
+		return;
+	}
+
+	const FString ICVFXCameraId = InCameraComponent.GetCameraUniqueId();
+
+	// Search for rendered camera viewport on other cluster nodes
+	FDisplayClusterViewportManager* ViewportManager = FDisplayClusterViewportConfigurationHelpers_ICVFX::GetViewportManager(InRootActor);
+	if (ViewportManager != nullptr)
+	{
+		for (FDisplayClusterViewport* ViewportIt : ViewportManager->ImplGetWholeClusterViewports_Editor())
+		{
+			if (ViewportIt != nullptr
+			&& (ViewportIt->RenderSettingsICVFX.RuntimeFlags & ViewportRuntime_ICVFXIncamera) != 0 
+			&& (ViewportIt->InputShaderResources.Num() > 0 && ViewportIt->InputShaderResources[0] != nullptr && ViewportIt->Contexts.Num() > 0)
+			&& (ViewportIt->RenderSettings.OverrideViewportId.IsEmpty())
+			&& (ViewportIt->GetClusterNodeId() != InCameraViewport.GetClusterNodeId()))
+			{
+				// this is incamera viewport. Check by name
+				const FString RequiredViewportId = ImplGetNameICVFX(ViewportIt->GetClusterNodeId(), ICVFXCameraId, DisplayClusterViewportStrings::icvfx::camera);
+				if (RequiredViewportId.Equals(ViewportIt->GetId()))
+				{
+					if (FDisplayClusterViewportConfigurationHelpers_OpenColorIO::IsInnerFrustumViewportSettingsEqual_Editor(*ViewportIt, InCameraViewport, InCameraComponent)
+					&&  FDisplayClusterViewportConfigurationHelpers_Postprocess::IsInnerFrustumViewportSettingsEqual_Editor(*ViewportIt, InCameraViewport, InCameraComponent))
+					{
+						// Reuse exist viewport:
+						InCameraViewport.RenderSettings.OverrideViewportId = ViewportIt->GetId();
+						return;
+					}
+				}
+			}
+		}
+	}
+}
+
+void FDisplayClusterViewportConfigurationHelpers_ICVFX::PreviewReuseChromakeyViewportWithinClusterNodes_Editor(FDisplayClusterViewport& InChromakeyViewport, ADisplayClusterRootActor& InRootActor, UDisplayClusterICVFXCameraComponent& InCameraComponent)
+{
+	if (!IsPreviewEnableReuseViewportInCluster_EditorImpl(InRootActor))
+	{
+		return;
+	}
+
+	const FDisplayClusterRenderFrameSettings& RenderFrameSettingsConstRef = InChromakeyViewport.GetRenderFrameSettings();
+	switch (RenderFrameSettingsConstRef.RenderMode)
+	{
+	case EDisplayClusterRenderFrameMode::PreviewInScene:
+		// Only for preview mode
+		break;
+	default:
+		return;
+	}
+
+	const FString ICVFXCameraId = InCameraComponent.GetCameraUniqueId();
+
+	// Search for rendered camera viewport on other cluster nodes
+	FDisplayClusterViewportManager* ViewportManager = FDisplayClusterViewportConfigurationHelpers_ICVFX::GetViewportManager(InRootActor);
+	if (ViewportManager != nullptr)
+	{
+		for (FDisplayClusterViewport* ViewportIt : ViewportManager->ImplGetWholeClusterViewports_Editor())
+		{
+			if (ViewportIt != nullptr
+				&& (ViewportIt->RenderSettingsICVFX.RuntimeFlags & ViewportRuntime_ICVFXChromakey) != 0
+				&& (ViewportIt->InputShaderResources.Num() > 0 && ViewportIt->InputShaderResources[0] != nullptr && ViewportIt->Contexts.Num() > 0)
+				&& (ViewportIt->RenderSettings.OverrideViewportId.IsEmpty())
+				&& (ViewportIt->GetClusterNodeId() != InChromakeyViewport.GetClusterNodeId()))
+			{
+				// this is chromakey viewport. Check by name
+				const FString RequiredViewportId = ImplGetNameICVFX(ViewportIt->GetClusterNodeId(), ICVFXCameraId, DisplayClusterViewportStrings::icvfx::chromakey);
+				if (RequiredViewportId.Equals(ViewportIt->GetId()))
+				{
+					// Reuse exist viewport:
+					InChromakeyViewport.RenderSettings.OverrideViewportId = ViewportIt->GetId();
+					return;
+				}
+			}
+		}
+	}
+}
+#endif
+
 FDisplayClusterViewport* FDisplayClusterViewportConfigurationHelpers_ICVFX::GetOrCreateChromakeyViewport(ADisplayClusterRootActor& RootActor, UDisplayClusterICVFXCameraComponent& InCameraComponent)
 {
 	const FString CameraId = InCameraComponent.GetCameraUniqueId();
@@ -278,7 +419,7 @@ FDisplayClusterShaderParameters_ICVFX::FCameraSettings FDisplayClusterViewportCo
 	const int32 CameraRenderOrder = RootActor.GetInnerFrustumPriority(InnerFrustumID);
 	Result.RenderOrder = (CameraRenderOrder<0) ? CameraSettings.RenderSettings.RenderOrder : CameraRenderOrder;
 
-	const float FieldOfViewMultiplier = CameraSettings.CustomFrustum.bEnable ? CameraSettings.CustomFrustum.FieldOfViewMultiplier : 1;
+	const float FieldOfViewMultiplier = ImplGetFieldOfViewMultiplier(CameraSettings);
 
 	// softedge adjustments	
 	const float Overscan = (FieldOfViewMultiplier > 0) ? FieldOfViewMultiplier : 1;
@@ -478,7 +619,6 @@ void FDisplayClusterViewportConfigurationHelpers_ICVFX::UpdateChromakeyViewportS
 	if (ChromakeySettings.ChromakeyRenderTexture.bReplaceCameraViewport)
 	{
 		InCameraViewport.RenderSettings.OverrideViewportId = DstViewport.GetId();
-		InCameraViewport.RenderSettings.bSkipRendering = true;
 	}
 }
 
@@ -636,7 +776,6 @@ void FDisplayClusterViewportConfigurationHelpers_ICVFX::UpdateLightcardViewportS
 	if (InRenderSettings.bReplaceViewport)
 	{
 		BaseViewport.RenderSettings.OverrideViewportId = DstViewport.GetId();
-		BaseViewport.RenderSettings.bSkipRendering = true;
 	}
 	else
 	{
@@ -691,13 +830,18 @@ void FDisplayClusterViewportConfigurationHelpers_ICVFX::UpdateCameraViewportBuff
 {
 	float BufferRatio = CameraSettings.BufferRatio;
 
-	const float FieldOfViewMultiplier = CameraSettings.CustomFrustum.bEnable ? CameraSettings.CustomFrustum.FieldOfViewMultiplier : 1;
-
-	// adapt resolution should work as a shortcut to improve rendering quality
 	if (CameraSettings.CustomFrustum.bAdaptResolution)
 	{
-		DstViewport.RenderSettings.RenderTargetRatio *= FieldOfViewMultiplier;
+		const float FieldOfViewMultiplier = ImplGetFieldOfViewMultiplier(CameraSettings);
+
+		// adapt resolution should work as a shortcut to improve rendering quality
+		DstViewport.RenderSettings.RenderTargetAdaptRatio = FieldOfViewMultiplier;
+	}
+	else
+	{
+		// Don't use an adaptive resolution multiplier
+		DstViewport.RenderSettings.RenderTargetAdaptRatio = 1.f;
 	}
 
 	DstViewport.Owner.SetViewportBufferRatio(DstViewport, BufferRatio);
-}
\ No newline at end of file
+}
diff --git a/Engine/Plugins/Runtime/nDisplay/Source/DisplayCluster/Private/Render/Viewport/Configuration/DisplayClusterViewportConfigurationHelpers_ICVFX.h b/Engine/Plugins/Runtime/nDisplay/Source/DisplayCluster/Private/Render/Viewport/Configuration/DisplayClusterViewportConfigurationHelpers_ICVFX.h
--- a/Engine/Plugins/Runtime/nDisplay/Source/DisplayCluster/Private/Render/Viewport/Configuration/DisplayClusterViewportConfigurationHelpers_ICVFX.h
+++ b/Engine/Plugins/Runtime/nDisplay/Source/DisplayCluster/Private/Render/Viewport/Configuration/DisplayClusterViewportConfigurationHelpers_ICVFX.h
@@ -39,6 +39,11 @@ public:
 	static FDisplayClusterViewport* GetOrCreateChromakeyViewport(ADisplayClusterRootActor& RootActor, UDisplayClusterICVFXCameraComponent& InCameraComponent);
 	static FDisplayClusterViewport* GetOrCreateLightcardViewport(FDisplayClusterViewport& BaseViewport, ADisplayClusterRootActor& RootActor, bool bIsOpenColorIOViewportExist);
 
+#if WITH_EDITOR
+	static void PreviewReuseInnerFrustumViewportWithinClusterNodes_Editor(FDisplayClusterViewport& InCameraViewport, ADisplayClusterRootActor& RootActor, UDisplayClusterICVFXCameraComponent& InCameraComponent);
+	static void PreviewReuseChromakeyViewportWithinClusterNodes_Editor(FDisplayClusterViewport& InChromakeyViewport, ADisplayClusterRootActor& RootActor, UDisplayClusterICVFXCameraComponent& InCameraComponent);
+#endif
+
 	static void UpdateCameraViewportSettings(FDisplayClusterViewport& DstViewport, ADisplayClusterRootActor& RootActor, UDisplayClusterICVFXCameraComponent& InCameraComponent);
 	static void UpdateChromakeyViewportSettings(FDisplayClusterViewport& DstViewport, FDisplayClusterViewport& InCameraViewport, ADisplayClusterRootActor& RootActor, UDisplayClusterICVFXCameraComponent& InCameraComponent);
 	static void UpdateLightcardViewportSetting(FDisplayClusterViewport& DstViewport, FDisplayClusterViewport& BaseViewport, ADisplayClusterRootActor& RootActor, bool bIsOpenColorIOViewportExist);
diff --git a/Engine/Plugins/Runtime/nDisplay/Source/DisplayCluster/Private/Render/Viewport/Configuration/DisplayClusterViewportConfigurationHelpers_OpenColorIO.cpp b/Engine/Plugins/Runtime/nDisplay/Source/DisplayCluster/Private/Render/Viewport/Configuration/DisplayClusterViewportConfigurationHelpers_OpenColorIO.cpp
--- a/Engine/Plugins/Runtime/nDisplay/Source/DisplayCluster/Private/Render/Viewport/Configuration/DisplayClusterViewportConfigurationHelpers_OpenColorIO.cpp
+++ b/Engine/Plugins/Runtime/nDisplay/Source/DisplayCluster/Private/Render/Viewport/Configuration/DisplayClusterViewportConfigurationHelpers_OpenColorIO.cpp
@@ -77,6 +77,45 @@ bool FDisplayClusterViewportConfigurationHelpers_OpenColorIO::UpdateLightcardVie
 	return false;
 }
 
+#if WITH_EDITOR
+bool FDisplayClusterViewportConfigurationHelpers_OpenColorIO::IsInnerFrustumViewportSettingsEqual_Editor(const FDisplayClusterViewport& InViewport1, const FDisplayClusterViewport& InViewport2, UDisplayClusterICVFXCameraComponent& InCameraComponent)
+{
+	const FDisplayClusterConfigurationICVFX_CameraSettings& CameraSettings = InCameraComponent.GetCameraSettingsICVFX();
+	if (CameraSettings.AllNodesOCIOConfiguration.bIsEnabled)
+	{
+		for (const FDisplayClusterConfigurationOCIOProfile& OCIOProfileIt : CameraSettings.PerNodeOCIOProfiles)
+		{
+			if (OCIOProfileIt.bIsEnabled)
+			{
+				const FString* CustomNode1 = OCIOProfileIt.ApplyOCIOToObjects.FindByPredicate([ClusterNodeId = InViewport1.GetClusterNodeId()](const FString& InClusterNodeId)
+				{
+					return ClusterNodeId.Equals(InClusterNodeId, ESearchCase::IgnoreCase);
+				});
+
+				const FString* CustomNode2 = OCIOProfileIt.ApplyOCIOToObjects.FindByPredicate([ClusterNodeId = InViewport2.GetClusterNodeId()](const FString& InClusterNodeId)
+				{
+					return ClusterNodeId.Equals(InClusterNodeId, ESearchCase::IgnoreCase);
+				});
+
+				if (CustomNode1 && CustomNode2)
+				{
+					// equal custom settings
+					return true;
+				}
+
+				if (CustomNode1 || CustomNode2)
+				{
+					// one of node has custom settings
+					return false;
+				}
+			}
+		}
+	}
+
+	return true;
+}
+#endif
+
 bool FDisplayClusterViewportConfigurationHelpers_OpenColorIO::UpdateICVFXCameraViewport(FDisplayClusterViewport& DstViewport, ADisplayClusterRootActor& RootActor, UDisplayClusterICVFXCameraComponent& InCameraComponent)
 {
 	const FDisplayClusterConfigurationICVFX_CameraSettings& CameraSettings = InCameraComponent.GetCameraSettingsICVFX();
diff --git a/Engine/Plugins/Runtime/nDisplay/Source/DisplayCluster/Private/Render/Viewport/Configuration/DisplayClusterViewportConfigurationHelpers_OpenColorIO.h b/Engine/Plugins/Runtime/nDisplay/Source/DisplayCluster/Private/Render/Viewport/Configuration/DisplayClusterViewportConfigurationHelpers_OpenColorIO.h
--- a/Engine/Plugins/Runtime/nDisplay/Source/DisplayCluster/Private/Render/Viewport/Configuration/DisplayClusterViewportConfigurationHelpers_OpenColorIO.h
+++ b/Engine/Plugins/Runtime/nDisplay/Source/DisplayCluster/Private/Render/Viewport/Configuration/DisplayClusterViewportConfigurationHelpers_OpenColorIO.h
@@ -20,6 +20,11 @@ public:
 	static bool UpdateICVFXCameraViewport(FDisplayClusterViewport& DstViewport, ADisplayClusterRootActor& RootActor, UDisplayClusterICVFXCameraComponent& InCameraComponent);
 	static bool UpdateLightcardViewport(FDisplayClusterViewport& DstViewport, FDisplayClusterViewport& BaseViewport, ADisplayClusterRootActor& RootActor);
 
+#if WITH_EDITOR
+	// return true when same settings used for both viewports
+	static bool IsInnerFrustumViewportSettingsEqual_Editor(const FDisplayClusterViewport& InViewport1, const FDisplayClusterViewport& InViewport2, UDisplayClusterICVFXCameraComponent& InCameraComponent);
+#endif
+
 private:
 	static bool ImplUpdate(FDisplayClusterViewport& DstViewport, const FOpenColorIODisplayConfiguration& InConfiguration);
 	static void ImplDisable(FDisplayClusterViewport& DstViewport);
diff --git a/Engine/Plugins/Runtime/nDisplay/Source/DisplayCluster/Private/Render/Viewport/Configuration/DisplayClusterViewportConfigurationHelpers_Postprocess.cpp b/Engine/Plugins/Runtime/nDisplay/Source/DisplayCluster/Private/Render/Viewport/Configuration/DisplayClusterViewportConfigurationHelpers_Postprocess.cpp
--- a/Engine/Plugins/Runtime/nDisplay/Source/DisplayCluster/Private/Render/Viewport/Configuration/DisplayClusterViewportConfigurationHelpers_Postprocess.cpp
+++ b/Engine/Plugins/Runtime/nDisplay/Source/DisplayCluster/Private/Render/Viewport/Configuration/DisplayClusterViewportConfigurationHelpers_Postprocess.cpp
@@ -166,10 +166,47 @@ static bool ImplUpdateIncameraAllNodesColorGrading(FDisplayClusterViewport& DstV
 	return true;
 }
 
+#if WITH_EDITOR
+// return true when same settings used for both viewports
+bool FDisplayClusterViewportConfigurationHelpers_Postprocess::IsInnerFrustumViewportSettingsEqual_Editor(const FDisplayClusterViewport& InViewport1, const FDisplayClusterViewport& InViewport2, UDisplayClusterICVFXCameraComponent& InCameraComponent)
+{
+	const FDisplayClusterConfigurationICVFX_CameraSettings& CameraSettings = InCameraComponent.GetCameraSettingsICVFX();
+	for (const FDisplayClusterConfigurationViewport_PerNodeColorGrading& ColorGradingProfileIt : CameraSettings.PerNodeColorGrading)
+	{
+		if (ColorGradingProfileIt.bIsEnabled)
+		{
+			const FString* CustomNode1 = ColorGradingProfileIt.ApplyPostProcessToObjects.FindByPredicate([ClusterNodeId = InViewport1.GetClusterNodeId()](const FString& InClusterNodeId)
+			{
+				return ClusterNodeId.Equals(InClusterNodeId, ESearchCase::IgnoreCase);
+			});
+
+			const FString* CustomNode2 = ColorGradingProfileIt.ApplyPostProcessToObjects.FindByPredicate([ClusterNodeId = InViewport2.GetClusterNodeId()](const FString& InClusterNodeId)
+			{
+				return ClusterNodeId.Equals(InClusterNodeId, ESearchCase::IgnoreCase);
+			});
+
+			if (CustomNode1 && CustomNode2)
+			{
+				// equal custom settings
+				return true;
+			}
+
+			if (CustomNode1 || CustomNode2)
+			{
+				// one of node has custom settings
+				return false;
+			}
+		}
+	}
+
+	return true;
+}
+#endif
+
 bool FDisplayClusterViewportConfigurationHelpers_Postprocess::ImplUpdateInnerFrustumColorGrading(FDisplayClusterViewport& DstViewport, ADisplayClusterRootActor& RootActor, UDisplayClusterICVFXCameraComponent& InCameraComponent)
 {
 	const FDisplayClusterConfigurationICVFX_CameraSettings& CameraSettings = InCameraComponent.GetCameraSettingsICVFX();
-	const FDisplayClusterRenderFrameSettings& RenderFrameSettings = DstViewport.Owner.GetRenderFrameSettings();	
+	const FDisplayClusterRenderFrameSettings& RenderFrameSettings = DstViewport.GetRenderFrameSettings();
 
 	// per node color grading first (it includes all nodes blending too)
 	const FString& ClusterNodeId = DstViewport.GetClusterNodeId();
diff --git a/Engine/Plugins/Runtime/nDisplay/Source/DisplayCluster/Private/Render/Viewport/Configuration/DisplayClusterViewportConfigurationHelpers_Postprocess.h b/Engine/Plugins/Runtime/nDisplay/Source/DisplayCluster/Private/Render/Viewport/Configuration/DisplayClusterViewportConfigurationHelpers_Postprocess.h
--- a/Engine/Plugins/Runtime/nDisplay/Source/DisplayCluster/Private/Render/Viewport/Configuration/DisplayClusterViewportConfigurationHelpers_Postprocess.h
+++ b/Engine/Plugins/Runtime/nDisplay/Source/DisplayCluster/Private/Render/Viewport/Configuration/DisplayClusterViewportConfigurationHelpers_Postprocess.h
@@ -23,6 +23,11 @@ public:
 
 	static bool UpdateLightcardPostProcessSettings(FDisplayClusterViewport& DstViewport, FDisplayClusterViewport& BaseViewport, ADisplayClusterRootActor& RootActor);
 
+#if WITH_EDITOR
+	// return true when same settings used for both viewports
+	static bool IsInnerFrustumViewportSettingsEqual_Editor(const FDisplayClusterViewport& InViewport1, const FDisplayClusterViewport& InViewport2, UDisplayClusterICVFXCameraComponent& InCameraComponent);
+#endif
+
 public:
 	static void BlendPostProcessSettings(FPostProcessSettings& OutputPP, const FDisplayClusterConfigurationViewport_ColorGradingRenderingSettings& ClusterPPSettings, const FDisplayClusterConfigurationViewport_ColorGradingRenderingSettings& ViewportPPSettings);
 	static void CopyPPSStructConditional(FDisplayClusterConfigurationViewport_ColorGradingRenderingSettings* OutViewportPPSettings, FPostProcessSettings* InPPS);
diff --git a/Engine/Plugins/Runtime/nDisplay/Source/DisplayCluster/Private/Render/Viewport/Configuration/DisplayClusterViewportConfigurationICVFX.cpp b/Engine/Plugins/Runtime/nDisplay/Source/DisplayCluster/Private/Render/Viewport/Configuration/DisplayClusterViewportConfigurationICVFX.cpp
--- a/Engine/Plugins/Runtime/nDisplay/Source/DisplayCluster/Private/Render/Viewport/Configuration/DisplayClusterViewportConfigurationICVFX.cpp
+++ b/Engine/Plugins/Runtime/nDisplay/Source/DisplayCluster/Private/Render/Viewport/Configuration/DisplayClusterViewportConfigurationICVFX.cpp
@@ -7,6 +7,7 @@
 #include "DisplayClusterViewportConfigurationHelpers_Visibility.h"
 
 #include "DisplayClusterRootActor.h"
+#include "DisplayClusterConfigurationTypes_Viewport.h"
 
 #include "IDisplayClusterProjection.h"
 #include "Render/Projection/IDisplayClusterProjectionPolicy.h"
@@ -16,11 +17,21 @@
 #include "Render/Viewport/DisplayClusterViewport.h"
 #include "Render/Viewport/DisplayClusterViewportManager.h"
 #include "Render/Viewport/DisplayClusterViewportStrings.h"
+#include "Render/Viewport/RenderFrame/DisplayClusterRenderFrameSettings.h"
 
 #include "Components/DisplayClusterICVFXCameraComponent.h"
 
 #include "Misc/DisplayClusterLog.h"
 
+////////////////////////////////////////////////////////////////////////////////
+int32 GDisplayClusterEnableAlphaChannelRendering = 0;
+static FAutoConsoleVariableRef CVarDisplayClusterEnableAlphaChannelRendering(
+	TEXT("DC.EnableAlphaChannelRendering"),
+	GDisplayClusterEnableAlphaChannelRendering,
+	TEXT("Enable alpha channel rendering to backbuffer (0 == disabled, 1 == enabled)"),
+	ECVF_RenderThreadSafe
+);
+
 ////////////////////////////////////////////////////////////////////////////////
 // FDisplayClusterViewportConfigurationCameraViewport
 ////////////////////////////////////////////////////////////////////////////////
@@ -48,6 +59,11 @@ public:
 			// Support projection policy update
 			FDisplayClusterViewportConfigurationHelpers::UpdateProjectionPolicy(*CameraViewport);
 
+#if WITH_EDITOR
+			// Reuse for EditorPreview
+			FDisplayClusterViewportConfigurationHelpers_ICVFX::PreviewReuseInnerFrustumViewportWithinClusterNodes_Editor(*CameraViewport, InRootActor, InCameraComponent);
+#endif
+
 			return true;
 		}
 
@@ -69,6 +85,11 @@ public:
 			// Support projection policy update
 			FDisplayClusterViewportConfigurationHelpers::UpdateProjectionPolicy(*ChromakeyViewport);
 
+#if WITH_EDITOR
+			// reuse for EditorPreview
+			FDisplayClusterViewportConfigurationHelpers_ICVFX::PreviewReuseChromakeyViewportWithinClusterNodes_Editor(*ChromakeyViewport, InRootActor, InCameraComponent);
+#endif
+
 			return true;
 		}
 
@@ -360,14 +381,37 @@ void FDisplayClusterViewportConfigurationICVFX::Update()
 				}
 			}
 
-			// Sort cameras by render order for each target
 			for (FDisplayClusterViewport* TargetIt : TargetViewports)
 			{
 				if (TargetIt)
 				{
+					// Sort cameras by render order for each target
 					TargetIt->RenderSettingsICVFX.ICVFX.SortCamerasRenderOrder();
 				}
 			}
+
+			if (StageSettings.bFreezeRenderOuterViewports)
+			{
+				// Freeze rendering for outer viewports
+				for (FDisplayClusterViewport* TargetIt : TargetViewports)
+				{
+					TargetIt->RenderSettings.bFreezeRendering = true;
+				}
+
+				// Freeze render for lightcards when outer viewports freezed
+				if (StageSettings.Lightcard.bIgnoreOuterViewportsFreezingForLightcards == false)
+				{
+					const EDisplayClusterViewportRuntimeICVFXFlags LightcardViewportMask = ViewportRuntime_ICVFXLightcardColor | ViewportRuntime_ICVFXLightcard | ViewportRuntime_ICVFXLightcardColor;
+
+					for (FDisplayClusterViewport* ViewportIt : ViewportManager->ImplGetViewports())
+					{
+						if (ViewportIt && (ViewportIt->RenderSettingsICVFX.RuntimeFlags & LightcardViewportMask) != 0)
+						{
+							ViewportIt->RenderSettings.bFreezeRendering = true;
+						}
+					}
+				}
+			}
 		}
 
 		ImplFinishReallocateViewports(*ViewportManager);
@@ -386,6 +430,15 @@ void FDisplayClusterViewportConfigurationICVFX::PostUpdate()
 
 		// Support additional hide list for icvfx cameras:
 		UpdateCameraHideList(*ViewportManager);
+
+		// Support alpha channel rendering for WarpBlend
+		if (GDisplayClusterEnableAlphaChannelRendering != 0)
+		{
+			for (FDisplayClusterViewport* ViewportIt : ViewportManager->ImplGetViewports())
+			{
+				ViewportIt->RenderSettings.bWarpBlendRenderAlphaChannel = true;
+			}
+		}
 	}
 }
 
@@ -513,3 +566,9 @@ EDisplayClusterViewportICVFXFlags FDisplayClusterViewportConfigurationICVFX::Imp
 	// Collect all targets disable flags
 	return ~(InvFlags);
 }
+
+#if WITH_EDITOR
+void FDisplayClusterViewportConfigurationICVFX::PostUpdatePreview_Editor(const FDisplayClusterPreviewSettings& InPreviewSettings)
+{
+}
+#endif
diff --git a/Engine/Plugins/Runtime/nDisplay/Source/DisplayCluster/Private/Render/Viewport/Configuration/DisplayClusterViewportConfigurationICVFX.h b/Engine/Plugins/Runtime/nDisplay/Source/DisplayCluster/Private/Render/Viewport/Configuration/DisplayClusterViewportConfigurationICVFX.h
--- a/Engine/Plugins/Runtime/nDisplay/Source/DisplayCluster/Private/Render/Viewport/Configuration/DisplayClusterViewportConfigurationICVFX.h
+++ b/Engine/Plugins/Runtime/nDisplay/Source/DisplayCluster/Private/Render/Viewport/Configuration/DisplayClusterViewportConfigurationICVFX.h
@@ -10,6 +10,8 @@ class FDisplayClusterViewportManager;
 class FDisplayClusterViewport;
 class ADisplayClusterRootActor;
 
+struct FDisplayClusterPreviewSettings;
+
 class FDisplayClusterViewportConfigurationICVFX
 {
 public:
@@ -20,6 +22,10 @@ public:
 	void Update();
 	void PostUpdate();
 
+#if WITH_EDITOR
+	void PostUpdatePreview_Editor(const FDisplayClusterPreviewSettings& InPreviewSettings);
+#endif
+
 private:
 	void ImplBeginReallocateViewports(FDisplayClusterViewportManager& ViewportManager);
 	void ImplFinishReallocateViewports(FDisplayClusterViewportManager& ViewportManager);
diff --git a/Engine/Plugins/Runtime/nDisplay/Source/DisplayCluster/Private/Render/Viewport/Configuration/DisplayClusterViewportConfigurationProjectionPolicy.cpp b/Engine/Plugins/Runtime/nDisplay/Source/DisplayCluster/Private/Render/Viewport/Configuration/DisplayClusterViewportConfigurationProjectionPolicy.cpp
--- a/Engine/Plugins/Runtime/nDisplay/Source/DisplayCluster/Private/Render/Viewport/Configuration/DisplayClusterViewportConfigurationProjectionPolicy.cpp
+++ b/Engine/Plugins/Runtime/nDisplay/Source/DisplayCluster/Private/Render/Viewport/Configuration/DisplayClusterViewportConfigurationProjectionPolicy.cpp
@@ -116,7 +116,6 @@ bool FDisplayClusterViewportConfigurationProjectionPolicy::UpdateCameraPolicy_IC
 	if (ExistCameraViewport)
 	{
 		DstViewport.RenderSettings.OverrideViewportId = ExistCameraViewport->GetId();
-		DstViewport.RenderSettings.bSkipRendering = true;
 	}
 	else
 	{
diff --git a/Engine/Plugins/Runtime/nDisplay/Source/DisplayCluster/Private/Render/Viewport/Containers/DisplayClusterViewportReadPixels.h b/Engine/Plugins/Runtime/nDisplay/Source/DisplayCluster/Private/Render/Viewport/Containers/DisplayClusterViewportReadPixels.h
new file mode 100644
--- /dev/null
+++ b/Engine/Plugins/Runtime/nDisplay/Source/DisplayCluster/Private/Render/Viewport/Containers/DisplayClusterViewportReadPixels.h
@@ -0,0 +1,12 @@
+// Copyright Epic Games, Inc. All Rights Reserved.
+
+#pragma once
+
+#include "CoreMinimal.h"
+
+class FDisplayClusterViewportReadPixelsData
+{
+public:
+	FIntPoint Size;
+	TArray<FColor> Pixels;
+};
diff --git a/Engine/Plugins/Runtime/nDisplay/Source/DisplayCluster/Private/Render/Viewport/DisplayClusterViewport.cpp b/Engine/Plugins/Runtime/nDisplay/Source/DisplayCluster/Private/Render/Viewport/DisplayClusterViewport.cpp
--- a/Engine/Plugins/Runtime/nDisplay/Source/DisplayCluster/Private/Render/Viewport/DisplayClusterViewport.cpp
+++ b/Engine/Plugins/Runtime/nDisplay/Source/DisplayCluster/Private/Render/Viewport/DisplayClusterViewport.cpp
@@ -1,6 +1,7 @@
 // Copyright Epic Games, Inc. All Rights Reserved.
 
 #include "Render/Viewport/DisplayClusterViewport.h"
+#include "Render/Viewport/DisplayClusterViewportHelpers.h"
 #include "Render/Viewport/DisplayClusterViewportManager.h"
 #include "Render/Viewport/DisplayClusterViewportProxy.h"
 
@@ -49,6 +50,11 @@ IDisplayClusterViewportManager& FDisplayClusterViewport::GetOwner() const
 	return Owner;
 }
 
+const FDisplayClusterRenderFrameSettings& FDisplayClusterViewport::GetRenderFrameSettings() const
+{
+	return Owner.GetRenderFrameSettings();
+}
+
 const TArray<FSceneViewExtensionRef> FDisplayClusterViewport::GatherActiveExtensions(FViewport* InViewport) const
 {
 	if (InViewport)
@@ -206,17 +212,8 @@ void FDisplayClusterViewport::SetupSceneView(uint32 ContextNum, class UWorld* Wo
 		InOutView.bAllowCrossGPUTransfer = (Contexts[ContextNum].bAllowGPUTransferOptimization == false);
 	}
 
-	if (RenderSettings.bSkipRendering)
-	{
-		//@todo: research better way to disable viewport rendering
-		// now just set showOnly=[]
-		InOutView.ShowOnlyPrimitives.Emplace();
-	}
-	else
-	{
-		// Apply visibility settigns to view
-		VisibilitySettings.SetupSceneView(World, InOutView);
-	}
+	// Apply visibility settigns to view
+	VisibilitySettings.SetupSceneView(World, InOutView);
 
 	// Handle Motion blur parameters
 	CameraMotionBlur.SetupSceneView(Contexts[ContextNum], InOutView);
@@ -233,8 +230,8 @@ inline void AdjustRect(FIntRect& InOutRect, const float multX, const float multY
 FIntRect FDisplayClusterViewport::GetValidRect(const FIntRect& InRect, const TCHAR* DbgSourceName)
 {
 	// The target always needs be within GMaxTextureDimensions, larger dimensions are not supported by the engine
-	static const int32 MaxTextureSize = 1 << (GMaxTextureMipCount - 1);
-	static const int32 MinTextureSize = 16;
+	const int32 MaxTextureSize = DisplayClusterViewportHelpers::GetMaxTextureDimension();
+	const int32 MinTextureSize = DisplayClusterViewportHelpers::GetMinTextureDimension();
 
 	int32 Width  = FMath::Max(MinTextureSize, InRect.Width());
 	int32 Height = FMath::Max(MinTextureSize, InRect.Height());
@@ -254,57 +251,82 @@ FIntRect FDisplayClusterViewport::GetValidRect(const FIntRect& InRect, const TCH
 	OutRect.Min.X = FMath::Min(OutRect.Min.X, MaxTextureSize);
 	OutRect.Min.Y = FMath::Min(OutRect.Min.Y, MaxTextureSize);
 
-	const int32 MaxX = FMath::RoundHalfToEven(OutRect.Max.X * RectScale);
-	const int32 MaxY = FMath::RoundHalfToEven(OutRect.Max.Y * RectScale);
+	const FIntPoint ScaledRectMax = DisplayClusterViewportHelpers::ScaleTextureSize(OutRect.Max, RectScale);
 
-	OutRect.Max.X = FMath::Clamp(MaxX, OutRect.Min.X, MaxTextureSize);
-	OutRect.Max.Y = FMath::Clamp(MaxY, OutRect.Min.Y, MaxTextureSize);
+	OutRect.Max.X = FMath::Clamp(ScaledRectMax.X, OutRect.Min.X, MaxTextureSize);
+	OutRect.Max.Y = FMath::Clamp(ScaledRectMax.Y, OutRect.Min.Y, MaxTextureSize);
 
 	return OutRect;
 }
 
-bool FDisplayClusterViewport::UpdateFrameContexts(const uint32 InStereoViewIndex, const FDisplayClusterRenderFrameSettings& InFrameSettings)
+float FDisplayClusterViewport::GetClusterRenderTargetRatioMult(const FDisplayClusterRenderFrameSettings& InFrameSettings) const
 {
-	check(IsInGameThread());
+	float ClusterRenderTargetRatioMult = InFrameSettings.ClusterRenderTargetRatioMult;
 
-	// Release old contexts
-	Contexts.Empty();
+	// Support Outer viewport cluster rtt multiplier
+	if ((RenderSettingsICVFX.RuntimeFlags & ViewportRuntime_ICVFXTarget) != 0)
+	{
+		ClusterRenderTargetRatioMult *= InFrameSettings.ClusterICVFXOuterViewportRenderTargetRatioMult;
+	}
+	else if ((RenderSettingsICVFX.RuntimeFlags & ViewportRuntime_ICVFXIncamera) != 0)
+	{
+		ClusterRenderTargetRatioMult *= InFrameSettings.ClusterICVFXInnerViewportRenderTargetRatioMult;
+	}
 
-	RenderTargets.Empty();
-	OutputFrameTargetableResources.Empty();
-	AdditionalFrameTargetableResources.Empty();
+	// Cluster mult downscale in range 0..1
+	return FMath::Clamp(ClusterRenderTargetRatioMult, 0.f, 1.f);
+}
 
-#if WITH_EDITOR
-	OutputPreviewTargetableResource.SafeRelease();
-#endif
+FIntPoint FDisplayClusterViewport::GetDesiredContextSize(const FIntPoint& InSize, const FDisplayClusterRenderFrameSettings& InFrameSettings) const
+{
+	const float ClusterRenderTargetRatioMult = GetClusterRenderTargetRatioMult(InFrameSettings);
 
-	InputShaderResources.Empty();
-	AdditionalTargetableResources.Empty();
-	MipsShaderResources.Empty();
+	// Check size multipliers in order bellow:
+	const float RenderTargetAdaptRatio = DisplayClusterViewportHelpers::GetValidSizeMultiplier(InSize, RenderSettings.RenderTargetAdaptRatio, ClusterRenderTargetRatioMult * RenderSettings.RenderTargetRatio);
+	const float RenderTargetRatio = DisplayClusterViewportHelpers::GetValidSizeMultiplier(InSize, RenderSettings.RenderTargetRatio, ClusterRenderTargetRatioMult * RenderTargetAdaptRatio);
+	const float ClusterMult = DisplayClusterViewportHelpers::GetValidSizeMultiplier(InSize, ClusterRenderTargetRatioMult, RenderTargetRatio * RenderTargetAdaptRatio);
 
-	if (RenderSettings.bEnable == false)
+	FIntPoint DesiredContextSize = DisplayClusterViewportHelpers::ScaleTextureSize(InSize, FMath::Max(RenderTargetAdaptRatio * RenderTargetRatio * ClusterMult, 0.f));
+
+	const int32 MaxTextureSize = DisplayClusterViewportHelpers::GetMaxTextureDimension();
+	DesiredContextSize.X = FMath::Min(DesiredContextSize.X, MaxTextureSize);
+	DesiredContextSize.Y = FMath::Min(DesiredContextSize.Y, MaxTextureSize);
+
+	return DesiredContextSize;
+}
+
+float FDisplayClusterViewport::GetCustomBufferRatio(const FDisplayClusterRenderFrameSettings& InFrameSettings) const
+{
+	float CustomBufferRatio = RenderSettings.BufferRatio;
+
+	// Global multiplier
+	CustomBufferRatio *= InFrameSettings.ClusterBufferRatioMult;
+
+	if ((RenderSettingsICVFX.RuntimeFlags & ViewportRuntime_ICVFXTarget) != 0)
 	{
-		// Exclude this viewport from render and logic, but object still exist
-		return false;
+		// Outer viewport
+		CustomBufferRatio *= InFrameSettings.ClusterICVFXOuterViewportBufferRatioMult;
 	}
-
-	if (PostRenderSettings.GenerateMips.IsEnabled())
+	else if ((RenderSettingsICVFX.RuntimeFlags & ViewportRuntime_ICVFXIncamera) != 0)
 	{
-		//Check if current projection policy supports this feature
-		if (!ProjectionPolicy.IsValid() || !ProjectionPolicy->ShouldUseSourceTextureWithMips())
-		{
-			// Don't create unused mips texture
-			PostRenderSettings.GenerateMips.Reset();
-		}
+		// Inner Frustum
+		CustomBufferRatio *= InFrameSettings.ClusterICVFXInnerFrustumBufferRatioMult;
 	}
 
-	uint32 FrameTargetsAmmount = 2;
+	return CustomBufferRatio;
+}
+
+bool FDisplayClusterViewport::UpdateFrameContexts(const uint32 InStereoViewIndex, const FDisplayClusterRenderFrameSettings& InFrameSettings)
+{
+	check(IsInGameThread());
+
+	uint32 FrameTargetsAmount = 2;
 	FIntRect     DesiredFrameTargetRect = RenderSettings.Rect;
 	{
 		switch (InFrameSettings.RenderMode)
 		{
 		case EDisplayClusterRenderFrameMode::Mono:
-			FrameTargetsAmmount = 1;
+			FrameTargetsAmount = 1;
 			break;
 		case EDisplayClusterRenderFrameMode::SideBySide:
 			AdjustRect(DesiredFrameTargetRect, 0.5f, 1.f);
@@ -322,7 +344,7 @@ bool FDisplayClusterViewport::UpdateFrameContexts(const uint32 InStereoViewIndex
 			DesiredFrameTargetRect = FIntRect(FIntPoint(0, 0), DesiredFrameTargetRect.Size());
 
 			// Mono
-			FrameTargetsAmmount = 1;
+			FrameTargetsAmount = 1;
 			break;
 		}
 		default:
@@ -331,58 +353,82 @@ bool FDisplayClusterViewport::UpdateFrameContexts(const uint32 InStereoViewIndex
 	}
 
 	// Special case mono->stereo
-	uint32 ViewportContextAmmount = RenderSettings.bForceMono ? 1 : FrameTargetsAmmount;
+	uint32 ViewportContextAmount = RenderSettings.bForceMono ? 1 : FrameTargetsAmount;
 
-	// Make sure the frame target rect doesn't exceed the maximum resolution, and preserve its aspect ratio if it needs to be clamped
-	FIntRect FrameTargetRect = GetValidRect(DesiredFrameTargetRect, TEXT("Context Frame"));
+#if WITH_EDITOR
+	OutputPreviewTargetableResource.SafeRelease();
+#endif
 
-	// Exclude zero-size viewports from render
-	if (FrameTargetRect.Size().GetMin() <= 0)
+	// Discard resources that are not used in frame composition
+	RenderTargets.Empty();
+	OutputFrameTargetableResources.Empty();
+	AdditionalFrameTargetableResources.Empty();
+
+	// Freeze the image in the viewport only after the frame has been rendered
+	if (RenderSettings.bFreezeRendering && RenderSettings.bEnable)
 	{
-		UE_LOG(LogDisplayClusterViewport, Error, TEXT("The viewport '%s' FrameTarget rect has zero size %dx%d: Disabled"), *GetId(), FrameTargetRect.Size().X, FrameTargetRect.Size().Y);
-		return false;
-	}
+		// Freeze only when all resources valid
+		if (Contexts.Num() > 0 && Contexts.Num() == InputShaderResources.Num())
+		{
+			DisplayClusterViewportHelpers::FreezeRenderingOfViewportTextureResources(InputShaderResources);
+			DisplayClusterViewportHelpers::FreezeRenderingOfViewportTextureResources(AdditionalTargetableResources);
+			DisplayClusterViewportHelpers::FreezeRenderingOfViewportTextureResources(MipsShaderResources);
 
+			// Update context links for freezed viewport
+			for (FDisplayClusterViewport_Context& ContextIt : Contexts)
+			{
+				ContextIt.StereoscopicPass = FDisplayClusterViewportStereoscopicPass::EncodeStereoscopicPass(ContextIt.ContextNum, ViewportContextAmount, InFrameSettings);
+				ContextIt.StereoViewIndex = (int32)(InStereoViewIndex + ContextIt.ContextNum);
+				ContextIt.bDisableRender = true;
+			}
 
-	FIntPoint DesiredContextSize = FrameTargetRect.Size();
+			return true;
+		}
+	}
 
-	float ClusterRenderTargetRatioMult = InFrameSettings.ClusterRenderTargetRatioMult;
+	// Release old contexts
+	Contexts.Empty();
 
-	// Support Outer viewport cluster rtt multiplier
-	if ((RenderSettingsICVFX.RuntimeFlags & ViewportRuntime_ICVFXTarget) != 0)
+	// Free internal resources
+	InputShaderResources.Empty();
+	AdditionalTargetableResources.Empty();
+	MipsShaderResources.Empty();
+
+	if (RenderSettings.bEnable == false)
 	{
-		ClusterRenderTargetRatioMult *= InFrameSettings.ClusterICVFXOuterViewportRenderTargetRatioMult;
+		// Exclude this viewport from render and logic, but object still exist
+		return false;
 	}
-	else
-	if ((RenderSettingsICVFX.RuntimeFlags & ViewportRuntime_ICVFXIncamera) != 0)
+
+	if (PostRenderSettings.GenerateMips.IsEnabled())
 	{
-		ClusterRenderTargetRatioMult *= InFrameSettings.ClusterICVFXInnerViewportRenderTargetRatioMult;
+		//Check if current projection policy supports this feature
+		if (!ProjectionPolicy.IsValid() || !ProjectionPolicy->ShouldUseSourceTextureWithMips())
+		{
+			// Don't create unused mips texture
+			PostRenderSettings.GenerateMips.Reset();
+		}
 	}
 
-	// Cluster mult downscale in range 0..1
-	float ClusterRTTMult = FMath::Clamp(ClusterRenderTargetRatioMult, 0.f, 1.f);
+	// Make sure the frame target rect doesn't exceed the maximum resolution, and preserve its aspect ratio if it needs to be clamped
+	FIntRect FrameTargetRect = GetValidRect(DesiredFrameTargetRect, TEXT("Context Frame"));
+
+	// Exclude zero-size viewports from render
+	if (FrameTargetRect.Size().GetMin() <= 0)
+	{
+		UE_LOG(LogDisplayClusterViewport, Error, TEXT("The viewport '%s' FrameTarget rect has zero size %dx%d: Disabled"), *GetId(), FrameTargetRect.Size().X, FrameTargetRect.Size().Y);
+		return false;
+	}
 
 	// Scale context for rendering
-	float ViewportContextSizeMult = FMath::Max(RenderSettings.RenderTargetRatio * ClusterRTTMult, 0.f);
-	DesiredContextSize.X *= ViewportContextSizeMult;
-	DesiredContextSize.Y *= ViewportContextSizeMult;
+	FIntPoint DesiredContextSize = GetDesiredContextSize(FrameTargetRect.Size(), InFrameSettings);
 
 #if WITH_EDITOR
 	switch (InFrameSettings.RenderMode)
 	{
 	case EDisplayClusterRenderFrameMode::PreviewInScene:
 		{
-			int32 MaxSize = FMath::Max(DesiredContextSize.X, DesiredContextSize.Y);
-			if (MaxSize > InFrameSettings.PreviewMaxTextureSize)
-			{
-				float RectScale = float(InFrameSettings.PreviewMaxTextureSize) / float(MaxSize);
-
-				const int32 MaxX = FMath::RoundHalfToEven(DesiredContextSize.X * RectScale);
-				const int32 MaxY = FMath::RoundHalfToEven(DesiredContextSize.Y * RectScale);
-
-				// Use limited texture size for preview
-				DesiredContextSize = FIntPoint(MaxX, MaxY);
-			}
+			DesiredContextSize = DisplayClusterViewportHelpers::GetTextureSizeLessThanMax(DesiredContextSize, InFrameSettings.PreviewMaxTextureDimension);
 		}
 		break;
 	default:
@@ -390,15 +436,16 @@ bool FDisplayClusterViewport::UpdateFrameContexts(const uint32 InStereoViewIndex
 	}
 #endif
 
-	FIntRect RenderTargetRect = GetValidRect(FIntRect(FIntPoint(0, 0), DesiredContextSize), TEXT("Context RenderTarget"));
-
 	// Exclude zero-size viewports from render
-	if (RenderTargetRect.Size().GetMin()<=0)
+	if (DesiredContextSize.GetMin() <= 0)
 	{
-		UE_LOG(LogDisplayClusterViewport, Error, TEXT("The viewport '%s' RenderTarget rect has zero size %dx%d: Disabled"), *GetId(), RenderTargetRect.Size().X, RenderTargetRect.Size().Y);
+		UE_LOG(LogDisplayClusterViewport, Error, TEXT("The viewport '%s' RenderTarget rect has zero size %dx%d: Disabled"), *GetId(), DesiredContextSize.X, DesiredContextSize.Y);
 		return false;
 	}
 
+	// Build RTT rect
+	FIntRect RenderTargetRect = FIntRect(FIntPoint(0, 0), DesiredContextSize);
+
 	// Support custom frustum rendering feature
 	CustomFrustumRendering.Update(*this, RenderTargetRect);
 
@@ -407,14 +454,29 @@ bool FDisplayClusterViewport::UpdateFrameContexts(const uint32 InStereoViewIndex
 	// Support overscan rendering feature
 	OverscanRendering.Update(*this, RenderTargetRect);
 
-	// Support override feature
-	bool bDisableRender = PostRenderSettings.Replace.IsEnabled();
-	
-	
+	const float BaseCustomBufferRatio = GetCustomBufferRatio(InFrameSettings);
+
+	// Fix buffer ratio value vs MaxTextureSize:
+	const float CustomBufferRatio = DisplayClusterViewportHelpers::GetValidSizeMultiplier(RenderTargetRect.Size(), BaseCustomBufferRatio, 1.f);
+
+	// Setup resource usage logic:
+	bool bDisableRender = false;
+	if (PostRenderSettings.Replace.IsEnabled())
+	{
+		bDisableRender = true;
+	}
+
+	bool bDisableInternalResources = false;
+	if (RenderSettings.bSkipRendering || RenderSettings.OverrideViewportId.IsEmpty() == false)
+	{
+		bDisableInternalResources = true;
+		bDisableRender = true;
+	}
+
 	//Add new contexts
-	for (uint32 ContextIt = 0; ContextIt < ViewportContextAmmount; ++ContextIt)
+	for (uint32 ContextIt = 0; ContextIt < ViewportContextAmount; ++ContextIt)
 	{
-		const EStereoscopicPass StereoscopicPass = (InFrameSettings.bIsRenderingInEditor) ? EStereoscopicPass::eSSP_FULL : FDisplayClusterViewportStereoscopicPass::EncodeStereoscopicPass(ContextIt, ViewportContextAmmount);
+		const EStereoscopicPass StereoscopicPass = FDisplayClusterViewportStereoscopicPass::EncodeStereoscopicPass(ContextIt, ViewportContextAmount, InFrameSettings);
 		const int32 StereoViewIndex = (int32)(InStereoViewIndex + ContextIt);
 
 		FDisplayClusterViewport_Context Context(ContextIt, StereoscopicPass, StereoViewIndex);
@@ -481,30 +543,38 @@ bool FDisplayClusterViewport::UpdateFrameContexts(const uint32 InStereoViewIndex
 		Context.RenderTargetRect = RenderTargetRect;
 		Context.ContextSize = ContextSize;
 
+		Context.CustomBufferRatio = CustomBufferRatio;
+
 		Context.bDisableRender = bDisableRender;
 
 		Contexts.Add(Context);
 	}
 
 	// Reserve for resources
-	RenderTargets.AddZeroed(FrameTargetsAmmount);
-	InputShaderResources.AddZeroed(FrameTargetsAmmount);
-
-	if (ShouldUseAdditionalTargetableResource())
+	if (!bDisableRender)
 	{
-		AdditionalTargetableResources.AddZeroed(FrameTargetsAmmount);
+		RenderTargets.AddZeroed(FrameTargetsAmount);
 	}
 
-	// Setup Mips resources:
-	for (FDisplayClusterViewport_Context& ContextIt: Contexts)
+	if (!bDisableInternalResources)
 	{
-		ContextIt.NumMips = PostRenderSettings.GenerateMips.GetRequiredNumMips(ContextIt.ContextSize);
-		if (ContextIt.NumMips > 1)
+		InputShaderResources.AddZeroed(FrameTargetsAmount);
+
+		if (ShouldUseAdditionalTargetableResource())
+		{
+			AdditionalTargetableResources.AddZeroed(FrameTargetsAmount);
+		}
+
+		// Setup Mips resources:
+		for (FDisplayClusterViewport_Context& ContextIt : Contexts)
 		{
-			MipsShaderResources.AddZeroed(1);
+			ContextIt.NumMips = DisplayClusterViewportHelpers::GetMaxTextureNumMips(InFrameSettings, PostRenderSettings.GenerateMips.GetRequiredNumMips(ContextIt.ContextSize));
+			if (ContextIt.NumMips > 1)
+			{
+				MipsShaderResources.AddZeroed(1);
+			}
 		}
 	}
 
 	return true;
 }
-
diff --git a/Engine/Plugins/Runtime/nDisplay/Source/DisplayCluster/Private/Render/Viewport/DisplayClusterViewport.h b/Engine/Plugins/Runtime/nDisplay/Source/DisplayCluster/Private/Render/Viewport/DisplayClusterViewport.h
--- a/Engine/Plugins/Runtime/nDisplay/Source/DisplayCluster/Private/Render/Viewport/DisplayClusterViewport.h
+++ b/Engine/Plugins/Runtime/nDisplay/Source/DisplayCluster/Private/Render/Viewport/DisplayClusterViewport.h
@@ -8,6 +8,7 @@
 #include "Render/Viewport/Containers/ImplDisplayClusterViewport_CustomFrustum.h"
 #include "Render/Viewport/Containers/ImplDisplayClusterViewport_Overscan.h"
 #include "Render/Viewport/Containers/DisplayClusterViewportRemap.h"
+#include "Render/Viewport/Containers/DisplayClusterViewport_Enums.h"
 
 #include "Render/Viewport/DisplayClusterViewport_CustomPostProcessSettings.h"
 #include "Render/Viewport/DisplayClusterViewport_TextureShare.h"
@@ -62,6 +63,7 @@ public:
 		check(IsInGameThread());
 		return ClusterNodeId;
 	}
+
 	virtual const FDisplayClusterViewport_RenderSettings& GetRenderSettings() const override
 	{
 		check(IsInGameThread());
@@ -108,6 +110,8 @@ public:
 
 	virtual IDisplayClusterViewportManager& GetOwner() const override;
 
+	const FDisplayClusterRenderFrameSettings& GetRenderFrameSettings() const;
+
 	//////////////////////////////////////////////////////
 	/// ~IDisplayClusterViewport
 	//////////////////////////////////////////////////////
@@ -118,11 +122,13 @@ public:
 	FSceneView* ImplCalcScenePreview(class FSceneViewFamilyContext& InOutViewFamily, uint32 ContextNum);
 	bool    ImplPreview_CalculateStereoViewOffset(const uint32 InContextNum, FRotator& ViewRotation, const float WorldToMeters, FVector& ViewLocation);
 	FMatrix ImplPreview_GetStereoProjectionMatrix(const uint32 InContextNum);
+
+	bool GetPreviewPixels(TSharedPtr<class FDisplayClusterViewportReadPixelsData, ESPMode::ThreadSafe>& OutPixelsData) const;
+
 #endif //WITH_EDITOR
 
 	// Get from logic request for additional targetable resource
 	bool ShouldUseAdditionalTargetableResource() const;
-
 	bool ShouldUseAdditionalFrameTargetableResource() const;
 	bool ShouldUseFullSizeFrameTargetableResource() const;
 
@@ -171,6 +177,11 @@ public:
 public:
 	FIntRect GetValidRect(const FIntRect& InRect, const TCHAR* DbgSourceName);
 
+private:
+	float GetClusterRenderTargetRatioMult(const FDisplayClusterRenderFrameSettings& InFrameSettings) const;
+	FIntPoint GetDesiredContextSize(const FIntPoint& InSize, const FDisplayClusterRenderFrameSettings& InFrameSettings) const;
+	float GetCustomBufferRatio(const FDisplayClusterRenderFrameSettings& InFrameSettings) const;
+
 #if WITH_EDITOR
 	// Support view states for preview
 private:
@@ -252,6 +263,8 @@ protected:
 	TArray<FDisplayClusterViewportTextureResource*> AdditionalFrameTargetableResources;
 
 #if WITH_EDITOR
+	friend class UDisplayClusterPreviewComponent;
+
 	FTextureRHIRef OutputPreviewTargetableResource;
 #endif
 
diff --git a/Engine/Plugins/Runtime/nDisplay/Source/DisplayCluster/Private/Render/Viewport/DisplayClusterViewportHelpers.cpp b/Engine/Plugins/Runtime/nDisplay/Source/DisplayCluster/Private/Render/Viewport/DisplayClusterViewportHelpers.cpp
new file mode 100644
--- /dev/null
+++ b/Engine/Plugins/Runtime/nDisplay/Source/DisplayCluster/Private/Render/Viewport/DisplayClusterViewportHelpers.cpp
@@ -0,0 +1,232 @@
+// Copyright Epic Games, Inc. All Rights Reserved.
+
+#include "Render/Viewport/DisplayClusterViewportHelpers.h"
+#include "Render/Viewport/RenderTarget/DisplayClusterRenderTargetResource.h"
+#include "Render/Viewport/RenderFrame/DisplayClusterRenderFrameSettings.h"
+
+#include "Engine/TextureRenderTarget2D.h"
+
+////////////////////////////////////////////////////////////////////////////////
+int32 GDisplayClusterOverrideMaxTextureDimension = 8192;
+static FAutoConsoleVariableRef CVarDisplayClusterOverrideMaxTextureDimension(
+	TEXT("DC.OverrideMaxTextureDimension"),
+	GDisplayClusterOverrideMaxTextureDimension,
+	TEXT("Override max texture dimension for nDisplay rendering (-1 == disabled, positive values == override)"),
+	ECVF_RenderThreadSafe
+);
+
+int32 GDisplayClusterOverrideMinTextureDimension = -1;
+static FAutoConsoleVariableRef CVarDisplayClusterOverrideMinTextureDimension(
+	TEXT("DC.OverrideMinTextureDimension"),
+	GDisplayClusterOverrideMinTextureDimension,
+	TEXT("Override min texture dimension for nDisplay rendering (-1 == disabled, positive values == override)"),
+	ECVF_RenderThreadSafe
+);
+
+int32 GDisplayClusterDefaultPreviewPixelFormat = 1;
+static FAutoConsoleVariableRef CVarDisplayClusterDefaultPreviewPixelFormat(
+	TEXT("DC.Preview.DefaultPixelFormat"),
+	GDisplayClusterDefaultPreviewPixelFormat,
+	TEXT("Defines the default preview RTT pixel format.\n")
+	TEXT(" 0: 8bit fixed point RGBA\n")
+	TEXT(" 1: 16bit Float RGBA\n")
+	TEXT(" 2: 10bit fixed point RGB and 2bit Alpha\n"),
+	ECVF_RenderThreadSafe
+);
+
+int32 GDisplayClusterMaxNumMips = -1;
+static FAutoConsoleVariableRef CVarDisplayClusterMaxNumMips(
+	TEXT("DC.TextureMaxNumMips"),
+	GDisplayClusterMaxNumMips,
+	TEXT("Maximum number of mips for viewport texture.\n")
+	TEXT(" '0' - disable mips generation.\n")
+	TEXT("'-1' - disable this limit.\n")
+	TEXT("positive value - set the limit.\n"),
+	ECVF_RenderThreadSafe
+);
+
+namespace PreviewPixelFormat
+{
+	static ETextureRenderTargetFormat ImplGetRenderTargetFormatFromInt(const int32 InDefaultPreviewPixelFormat)
+	{
+		const int32 ValidIndex = FMath::Clamp(InDefaultPreviewPixelFormat, 0, 2);
+		static const ETextureRenderTargetFormat SPixelFormat[] = { RTF_RGBA8, RTF_RGBA16f, RTF_RGB10A2 };
+
+		return SPixelFormat[ValidIndex];
+	}
+
+	static EPixelFormat GetPreviewPixelFormat()
+	{
+		const ETextureRenderTargetFormat RenderTargetFormat = ImplGetRenderTargetFormatFromInt(GDisplayClusterDefaultPreviewPixelFormat);
+		
+		return GetPixelFormatFromRenderTargetFormat(RenderTargetFormat);
+	}
+};
+
+////////////////////////////////////////////////////////////////////////////////
+int32 DisplayClusterViewportHelpers::GetMaxTextureNumMips(const FDisplayClusterRenderFrameSettings& InRenderFrameSettings, const int32 InNumMips)
+{
+	int32 NumMips = InNumMips;
+
+#if WITH_EDITOR
+	switch (InRenderFrameSettings.RenderMode)
+	{
+	case EDisplayClusterRenderFrameMode::PreviewInScene:
+		if (GDisplayClusterDefaultPreviewPixelFormat != 0)
+		{
+			//@todo: now UE support mips generation only for fixed point textures (8bit RGBA)
+			// Remove this hack latter
+			// Disable preview mips generation in case of unsupported RTT texture format.
+			NumMips = 0;
+		}
+		break;
+	default:
+		break;
+	}
+#endif
+
+	if (GDisplayClusterMaxNumMips >= 0)
+	{
+		return FMath::Min(GDisplayClusterMaxNumMips, NumMips);
+	}
+
+	return NumMips;
+}
+
+int32 DisplayClusterViewportHelpers::GetMaxTextureDimension()
+{
+	// The target always needs be within GMaxTextureDimensions, larger dimensions are not supported by the engine
+	static const int32 MaxTextureDimension = 1 << (GMaxTextureMipCount - 1);
+
+	// Use overridden value
+	if (GDisplayClusterOverrideMaxTextureDimension > 0)
+	{
+		// Protect vs wrong CVar values
+		if (GDisplayClusterOverrideMinTextureDimension < 0 || GDisplayClusterOverrideMaxTextureDimension > GDisplayClusterOverrideMinTextureDimension)
+		{
+			return GDisplayClusterOverrideMaxTextureDimension;
+		}
+	}
+
+	return MaxTextureDimension;
+}
+
+int32 DisplayClusterViewportHelpers::GetMinTextureDimension()
+{
+	static const int32 MinTextureDimension = 16;
+
+	// Use overridden value
+	if (GDisplayClusterOverrideMinTextureDimension > 0)
+	{
+		// Protect vs wrong CVar values
+		if (GDisplayClusterOverrideMaxTextureDimension < 0 || GDisplayClusterOverrideMaxTextureDimension > GDisplayClusterOverrideMinTextureDimension)
+		{
+			return GDisplayClusterOverrideMinTextureDimension;
+		}
+	}
+
+	return MinTextureDimension;
+}
+
+void DisplayClusterViewportHelpers::FreezeRenderingOfViewportTextureResources(TArray<FDisplayClusterViewportTextureResource*>& TextureResources)
+{
+	for (FDisplayClusterViewportTextureResource* TextureResource : TextureResources)
+	{
+		if (TextureResource != nullptr)
+		{
+			TextureResource->RaiseViewportResourceState(EDisplayClusterViewportResourceState::DisableReallocate);
+		}
+	}
+}
+
+bool DisplayClusterViewportHelpers::IsValidTextureSize(const FIntPoint& InSize)
+{
+	return InSize.GetMin() >= GetMinTextureDimension() && InSize.GetMax() <= GetMaxTextureDimension();
+}
+
+// Return size less than max texture dimension. Aspect ratio not changed
+FIntPoint DisplayClusterViewportHelpers::GetTextureSizeLessThanMax(const FIntPoint& InSize, const int32 InMaxTextureDimension)
+{
+	if (InSize.GetMax() > InMaxTextureDimension)
+	{
+		const float DownscaleMult =  double(InMaxTextureDimension) / double(InSize.GetMax());
+
+		return ScaleTextureSize(InSize, DownscaleMult);
+	}
+
+	return InSize;
+}
+
+FIntPoint DisplayClusterViewportHelpers::ScaleTextureSize(const FIntPoint& InSize, float InMult)
+{
+	const int32 ScaledX = FMath::CeilToInt(InSize.X * InMult);
+	const int32 ScaledY = FMath::CeilToInt(InSize.Y * InMult);
+
+	return FIntPoint(ScaledX, ScaledY);
+}
+
+float DisplayClusterViewportHelpers::GetValidSizeMultiplier(const FIntPoint& InSize, const float InSizeMult, const float InBaseSizeMult)
+{
+	// find best possible size mult in range 1..InSizeMult
+	if (InSizeMult > 1.f)
+	{
+		const FIntPoint ScaledSize = ScaleTextureSize(InSize, FMath::Max(InSizeMult * InBaseSizeMult, 0.f));
+		if (DisplayClusterViewportHelpers::IsValidTextureSize(ScaledSize) == false)
+		{
+			// Try change 'RenderTargetAdaptRatio' to min possible value
+			const float BaseMult = FMath::Max(InBaseSizeMult, 0.f);
+			const FIntPoint MinScaledSize = ScaleTextureSize(InSize, BaseMult);
+
+			if (DisplayClusterViewportHelpers::IsValidTextureSize(MinScaledSize) == false)
+			{
+				// BaseSizeMult to big. Disable size mult
+				return 1.f;
+			}
+			else
+			{
+				const int32 MinDimension = MinScaledSize.GetMin();
+				const int32 MaxDimension = GetMaxTextureDimension();
+
+				// Get the maximum allowed multiplier value
+				const float OutMult = float(double(MaxDimension) / double(MinDimension));
+
+#if UE_BUILD_DEBUG
+				// debug purpose
+				const FIntPoint FinalSize = ScaleTextureSize(InSize, FMath::Max(OutMult * InBaseSizeMult, 0.f));
+				check(DisplayClusterViewportHelpers::IsValidTextureSize(FinalSize));
+#endif
+
+				return OutMult;
+			}
+		}
+	}
+
+	return InSizeMult;
+}
+
+#if WITH_EDITOR
+void DisplayClusterViewportHelpers::GetPreviewRenderTargetDesc_Editor(const FDisplayClusterRenderFrameSettings& InRenderFrameSettings, EPixelFormat& OutPixelFormat, float& OutDisplayGamma, bool& bOutSRGB)
+{
+	OutPixelFormat = PreviewPixelFormat::GetPreviewPixelFormat();
+
+	// Default gamma value
+	OutDisplayGamma = 2.2f;
+
+	//!
+	bOutSRGB = true;
+
+	switch (InRenderFrameSettings.RenderMode)
+	{
+	case EDisplayClusterRenderFrameMode::PreviewInScene:
+		if (InRenderFrameSettings.bPreviewEnablePostProcess == false)
+		{
+			// Disable postprocess for preview. Use Gamma 1.f
+			OutDisplayGamma = 1.f;
+		}
+		break;
+	default:
+		break;
+	}
+}
+#endif
+
diff --git a/Engine/Plugins/Runtime/nDisplay/Source/DisplayCluster/Private/Render/Viewport/DisplayClusterViewportHelpers.h b/Engine/Plugins/Runtime/nDisplay/Source/DisplayCluster/Private/Render/Viewport/DisplayClusterViewportHelpers.h
new file mode 100644
--- /dev/null
+++ b/Engine/Plugins/Runtime/nDisplay/Source/DisplayCluster/Private/Render/Viewport/DisplayClusterViewportHelpers.h
@@ -0,0 +1,35 @@
+// Copyright Epic Games, Inc. All Rights Reserved.
+
+#pragma once
+#include "CoreMinimal.h"
+#include "PixelFormat.h"
+
+class FDisplayClusterViewportTextureResource;
+struct FDisplayClusterRenderFrameSettings;
+
+class DisplayClusterViewportHelpers
+{
+public:
+	// Control texture resources dimensions inside nDisplay
+	static int32 GetMinTextureDimension();
+	static int32 GetMaxTextureDimension();
+
+	// Texture size helpers
+	static bool IsValidTextureSize(const FIntPoint& InSize);
+
+	static FIntPoint ScaleTextureSize(const FIntPoint& InSize, float InMult);
+
+	// Return size less than max texture dimension. Aspect ratio not changed
+	static FIntPoint GetTextureSizeLessThanMax(const FIntPoint& InSize, const int32 InMaxTextureDimension);
+
+	static float GetValidSizeMultiplier(const FIntPoint& InSize, const float InSizeMult, const float InBaseSizeMult);
+
+	// Support render freeze feature
+	static void FreezeRenderingOfViewportTextureResources(TArray<FDisplayClusterViewportTextureResource*>& TextureResources);
+
+	static int32 GetMaxTextureNumMips(const FDisplayClusterRenderFrameSettings& InRenderFrameSettings, const int32 InNumMips);
+
+#if WITH_EDITOR
+	static void GetPreviewRenderTargetDesc_Editor(const FDisplayClusterRenderFrameSettings& InRenderFrameSettings, EPixelFormat& OutPixelFormat, float& OutDisplayGamma, bool& bOutSRGB);
+#endif
+};
diff --git a/Engine/Plugins/Runtime/nDisplay/Source/DisplayCluster/Private/Render/Viewport/DisplayClusterViewportManager.cpp b/Engine/Plugins/Runtime/nDisplay/Source/DisplayCluster/Private/Render/Viewport/DisplayClusterViewportManager.cpp
--- a/Engine/Plugins/Runtime/nDisplay/Source/DisplayCluster/Private/Render/Viewport/DisplayClusterViewportManager.cpp
+++ b/Engine/Plugins/Runtime/nDisplay/Source/DisplayCluster/Private/Render/Viewport/DisplayClusterViewportManager.cpp
@@ -273,7 +273,7 @@ bool FDisplayClusterViewportManager::ShouldUseAdditionalFrameTargetableResource(
 
 bool FDisplayClusterViewportManager::UpdateConfiguration(EDisplayClusterRenderFrameMode InRenderMode, const FString& InClusterNodeId, ADisplayClusterRootActor* InRootActorPtr, const FDisplayClusterPreviewSettings* InPreviewSettings)
 {
-	ClusterNodeViewports.Empty();
+	ImplUpdateClusterNodeViewports(InRenderMode, InClusterNodeId);
 
 	if (InRootActorPtr)
 	{
@@ -289,8 +289,6 @@ bool FDisplayClusterViewportManager::UpdateConfiguration(EDisplayClusterRenderFr
 		{
 			if (Configuration->UpdateConfiguration(InRenderMode, InClusterNodeId))
 			{
-				ImplUpdateClusterNodeViewports();
-
 				return true;
 			}
 		}
@@ -299,8 +297,6 @@ bool FDisplayClusterViewportManager::UpdateConfiguration(EDisplayClusterRenderFr
 		{
 			if (Configuration->UpdatePreviewConfiguration(InRenderMode, InClusterNodeId, *InPreviewSettings))
 			{
-				ImplUpdateClusterNodeViewports();
-
 				return true;
 			}
 		}
@@ -317,28 +313,28 @@ const FDisplayClusterRenderFrameSettings& FDisplayClusterViewportManager::GetRen
 	return Configuration->GetRenderFrameSettings();
 }
 
-void FDisplayClusterViewportManager::ImplUpdateClusterNodeViewports()
+void FDisplayClusterViewportManager::ImplUpdateClusterNodeViewports(const EDisplayClusterRenderFrameMode InRenderMode, const FString& InClusterNodeId)
 {
-	const FDisplayClusterRenderFrameSettings& RenderFrameSettings = GetRenderFrameSettings();
+	ClusterNodeViewports.Empty();
 
-	switch (RenderFrameSettings.RenderMode)
+	switch (InRenderMode)
 	{
 	case EDisplayClusterRenderFrameMode::PreviewInScene:
 		{
 			// Per-viewport preview render required special render order for viewports
-			// Sort viewports: move linked and overrided viewports to last
+			// Sort viewports: move linked and overridden viewports to last
 			TArray<FDisplayClusterViewport*> LinkedViewports;
-			TArray<FDisplayClusterViewport*> OverridedViewports;
+			TArray<FDisplayClusterViewport*> OverriddenViewports;
 			for (FDisplayClusterViewport* Viewport : Viewports)
 			{
-				if (Viewport && Viewport->GetClusterNodeId() == RenderFrameSettings.ClusterNodeId)
+				if (Viewport && Viewport->GetClusterNodeId() == InClusterNodeId)
 				{
 					bool bHasParentViewport = Viewport->RenderSettings.GetParentViewportId().IsEmpty() == false;
 					bool bHasViewportOverride = Viewport->RenderSettings.OverrideViewportId.IsEmpty() == false;
 
 					if (bHasViewportOverride)
 					{
-						OverridedViewports.Add(Viewport);
+						OverriddenViewports.Add(Viewport);
 					}
 					else if (bHasParentViewport)
 					{
@@ -352,7 +348,7 @@ void FDisplayClusterViewportManager::ImplUpdateClusterNodeViewports()
 			}
 
 			ClusterNodeViewports.Append(LinkedViewports);
-			ClusterNodeViewports.Append(OverridedViewports);
+			ClusterNodeViewports.Append(OverriddenViewports);
 		}
 		break;
 	default:
@@ -453,7 +449,7 @@ bool FDisplayClusterViewportManager::BeginNewFrame(FViewport* InViewport, UWorld
 void FDisplayClusterViewportManager::UpdateDesiredNumberOfViews(FDisplayClusterRenderFrame& InOutRenderFrame)
 {
 	InOutRenderFrame.DesiredNumberOfViews = 0;
-	InOutRenderFrame.ViewportsAmmount = 0;
+	InOutRenderFrame.ViewportsAmount = 0;
 
 	for (FDisplayClusterRenderFrame::FFrameRenderTarget& RenderTargetIt : InOutRenderFrame.RenderTargets)
 	{
@@ -463,18 +459,15 @@ void FDisplayClusterViewportManager::UpdateDesiredNumberOfViews(FDisplayClusterR
 
 			for (FDisplayClusterRenderFrame::FFrameView& ViewIt : ViewFamilyIt.Views)
 			{
-				InOutRenderFrame.ViewportsAmmount++;
+				InOutRenderFrame.ViewportsAmount++;
 
-				if (ViewIt.bDisableRender == false)
+				if (ViewIt.bDisableRender == false && ViewIt.Viewport != nullptr)
 				{
 					ViewFamilyIt.NumViewsForRender++;
 
-					if (ViewIt.Viewport)
-					{
-						// Get StereoViewIndex for this viewport
-						const int32 StereoViewIndex = ViewIt.Viewport->GetContexts()[ViewIt.ContextNum].StereoViewIndex;
-						InOutRenderFrame.DesiredNumberOfViews = FMath::Max(StereoViewIndex + 1, InOutRenderFrame.DesiredNumberOfViews);
-					}
+					// Get StereoViewIndex for this viewport
+					const int32 StereoViewIndex = ViewIt.Viewport->GetContexts()[ViewIt.ContextNum].StereoViewIndex;
+					InOutRenderFrame.DesiredNumberOfViews++;
 				}
 			}
 		}
@@ -506,6 +499,15 @@ void FDisplayClusterViewportManager::FinalizeNewFrame()
 	// Send updated viewports data to render thread proxy
 	ViewportManagerProxy->ImplUpdateViewports(ClusterNodeViewports);
 
+	// Finish update settings
+	for (FDisplayClusterViewport* Viewport : ClusterNodeViewports)
+	{
+		if (Viewport)
+		{
+			Viewport->RenderSettings.FinishUpdateSettings();
+		}
+	}
+
 	// Update postprocess data from game thread
 	PostProcessManager->Tick();
 
@@ -550,13 +552,30 @@ FSceneViewFamily::ConstructionValues FDisplayClusterViewportManager::CreateViewF
 			break;
 	}
 
+	const FDisplayClusterRenderFrameSettings& RenderFrameSettings = GetRenderFrameSettings();
+	switch (RenderFrameSettings.RenderMode)
+	{
+	case EDisplayClusterRenderFrameMode::PreviewInScene:
+
+		// Disable eye adaptation for preview
+		InEngineShowFlags.EyeAdaptation = 0;
+
+		if (RenderFrameSettings.bPreviewEnablePostProcess == false)
+		{
+			// Disable postprocess for preview
+			InEngineShowFlags.PostProcessing = 0;
+		}
+		break;
+	default:
+		break;
+	}
+
 	return FSceneViewFamily::ConstructionValues(InFrameTarget.RenderTargetPtr, InScene, InEngineShowFlags)
 		.SetResolveScene(bResolveScene)
 		.SetRealtimeUpdate(true)
 		.SetAdditionalViewFamily(bInAdditionalViewFamily);
 }
 
-
 void FDisplayClusterViewportManager::ConfigureViewFamily(const FDisplayClusterRenderFrame::FFrameRenderTarget& InFrameTarget, const FDisplayClusterRenderFrame::FFrameViewFamily& InFrameViewFamily, FSceneViewFamilyContext& ViewFamily)
 {
 	// Note: EngineShowFlags should have already been configured in CreateViewFamilyConstructionValues.
@@ -570,7 +589,7 @@ void FDisplayClusterViewportManager::ConfigureViewFamily(const FDisplayClusterRe
 			ViewExt->SetupViewFamily(ViewFamily);
 		}
 	}
-
+	
 	// Setup capture mode:
 	{
 		ViewFamily.SceneCaptureCompositeMode = ESceneCaptureCompositeMode::SCCM_Overwrite;
@@ -586,6 +605,24 @@ void FDisplayClusterViewportManager::ConfigureViewFamily(const FDisplayClusterRe
 				break;
 		}
 	}
+
+#if WITH_EDITOR
+	const FDisplayClusterRenderFrameSettings& RenderFrameSettings = GetRenderFrameSettings();
+	switch (RenderFrameSettings.RenderMode)
+	{
+	case EDisplayClusterRenderFrameMode::PreviewInScene:
+		if (RenderFrameSettings.bPreviewEnablePostProcess == false)
+		{
+			// Disable postprocess for preview
+			ViewFamily.SceneCaptureSource = ESceneCaptureSource::SCS_SceneColorHDR;
+			ViewFamily.SceneCaptureCompositeMode = ESceneCaptureCompositeMode::SCCM_Overwrite;
+		}
+		break;
+	default:
+		break;
+	}
+#endif
+
 }
 
 void FDisplayClusterViewportManager::RenderFrame(FViewport* InViewport)
@@ -688,6 +725,7 @@ FDisplayClusterViewport* FDisplayClusterViewportManager::ImplCreateViewport(cons
 
 	// Add viewport on gamethread
 	Viewports.Add(NewViewport);
+	ClusterNodeViewports.Add(NewViewport);
 
 	// Add viewport proxy on renderthread
 	ViewportManagerProxy->ImplCreateViewport(NewViewport->ViewportProxy);
@@ -707,12 +745,24 @@ void FDisplayClusterViewportManager::ImplDeleteViewport(FDisplayClusterViewport*
 	// Delete viewport proxy on render thread
 	ViewportManagerProxy->ImplDeleteViewport(ExistViewport->ViewportProxy);
 
-	// Remove viewport obj from manager
-	int32 ViewportIndex = Viewports.Find(ExistViewport);
-	if (ViewportIndex != INDEX_NONE)
 	{
-		Viewports[ViewportIndex] = nullptr;
-		Viewports.RemoveAt(ViewportIndex);
+		// Remove viewport obj from manager
+		int32 ViewportIndex = Viewports.Find(ExistViewport);
+		if (ViewportIndex != INDEX_NONE)
+		{
+			Viewports[ViewportIndex] = nullptr;
+			Viewports.RemoveAt(ViewportIndex);
+		}
+	}
+
+	{
+		// Remove cluster node viewport
+		int32 ViewportIndex = ClusterNodeViewports.Find(ExistViewport);
+		if (ViewportIndex != INDEX_NONE)
+		{
+			ClusterNodeViewports[ViewportIndex] = nullptr;
+			ClusterNodeViewports.RemoveAt(ViewportIndex);
+		}
 	}
 
 	delete ExistViewport;
diff --git a/Engine/Plugins/Runtime/nDisplay/Source/DisplayCluster/Private/Render/Viewport/DisplayClusterViewportManager.h b/Engine/Plugins/Runtime/nDisplay/Source/DisplayCluster/Private/Render/Viewport/DisplayClusterViewportManager.h
--- a/Engine/Plugins/Runtime/nDisplay/Source/DisplayCluster/Private/Render/Viewport/DisplayClusterViewportManager.h
+++ b/Engine/Plugins/Runtime/nDisplay/Source/DisplayCluster/Private/Render/Viewport/DisplayClusterViewportManager.h
@@ -66,9 +66,16 @@ public:
 	virtual void RenderFrame(FViewport* InViewport) override;
 
 #if WITH_EDITOR
-	virtual bool RenderInEditor(class FDisplayClusterRenderFrame& InRenderFrame, FViewport* InViewport, const uint32 InFirstViewportNum, const int32 InViewportsAmmount, bool& bOutFrameRendered) override;
+	virtual bool RenderInEditor(class FDisplayClusterRenderFrame& InRenderFrame, FViewport* InViewport, const uint32 InFirstViewportNum, const int32 InViewportsAmount, int32& OutViewportsAmount, bool& bOutFrameRendered) override;
 	
 	void ImplUpdatePreviewRTTResources();
+
+	const TArray<FDisplayClusterViewport*>& ImplGetWholeClusterViewports_Editor() const
+	{
+		check(IsInGameThread());
+
+		return Viewports;
+	}
 #endif
 
 	virtual IDisplayClusterViewport* FindViewport(const FString& InViewportId) const override
@@ -80,7 +87,7 @@ public:
 
 	virtual const TArrayView<IDisplayClusterViewport*> GetViewports() const override
 	{
-		return TArrayView<IDisplayClusterViewport*>((IDisplayClusterViewport**)(Viewports.GetData()), Viewports.Num());
+		return TArrayView<IDisplayClusterViewport*>((IDisplayClusterViewport**)(ClusterNodeViewports.GetData()), ClusterNodeViewports.Num());
 	}
 
 	virtual void MarkComponentGeometryDirty(const FName InComponentName = NAME_None) override;
@@ -98,7 +105,7 @@ public:
 	const TArray<FDisplayClusterViewport*>& ImplGetViewports() const
 	{
 		check(IsInGameThread());
-		return Viewports;
+		return ClusterNodeViewports;
 	}
 
 	FDisplayClusterViewport* ImplFindViewport(const FString& InViewportId) const;
@@ -128,7 +135,7 @@ private:
 	void ResetSceneRenderTargetSize();
 	void UpdateSceneRenderTargetSize();
 	void HandleViewportRTTChanges(const TArray<FDisplayClusterViewport_Context>& PrevContexts, const TArray<FDisplayClusterViewport_Context>& Contexts);
-	void ImplUpdateClusterNodeViewports();
+	void ImplUpdateClusterNodeViewports(const EDisplayClusterRenderFrameMode InRenderMode, const FString& InClusterNodeId);
 
 protected:
 	friend FDisplayClusterViewportManagerProxy;
diff --git a/Engine/Plugins/Runtime/nDisplay/Source/DisplayCluster/Private/Render/Viewport/DisplayClusterViewportManagerProxy.cpp b/Engine/Plugins/Runtime/nDisplay/Source/DisplayCluster/Private/Render/Viewport/DisplayClusterViewportManagerProxy.cpp
--- a/Engine/Plugins/Runtime/nDisplay/Source/DisplayCluster/Private/Render/Viewport/DisplayClusterViewportManagerProxy.cpp
+++ b/Engine/Plugins/Runtime/nDisplay/Source/DisplayCluster/Private/Render/Viewport/DisplayClusterViewportManagerProxy.cpp
@@ -230,8 +230,8 @@ void FDisplayClusterViewportManagerProxy::UpdateDeferredResources_RenderThread(F
 {
 	check(IsInRenderingThread());
 
-	TArray<FDisplayClusterViewportProxy*> OverridedViewports;
-	OverridedViewports.Reserve(ClusterNodeViewportProxies.Num());
+	TArray<FDisplayClusterViewportProxy*> OverriddenViewports;
+	OverriddenViewports.Reserve(ClusterNodeViewportProxies.Num());
 
 	for (FDisplayClusterViewportProxy* ViewportProxy : ClusterNodeViewportProxies)
 	{
@@ -242,12 +242,12 @@ void FDisplayClusterViewportManagerProxy::UpdateDeferredResources_RenderThread(F
 		else
 		{
 			// Update after all
-			OverridedViewports.Add(ViewportProxy);
+			OverriddenViewports.Add(ViewportProxy);
 		}
 	}
 
 	// Update deferred viewports after all
-	for (FDisplayClusterViewportProxy* ViewportProxy : OverridedViewports)
+	for (FDisplayClusterViewportProxy* ViewportProxy : OverriddenViewports)
 	{
 		ViewportProxy->UpdateDeferredResources(RHICmdList);
 	}
@@ -352,8 +352,7 @@ void FDisplayClusterViewportManagerProxy::UpdateFrameResources_RenderThread(FRHI
 						case EWarpPass::Render:
 							PrjPolicy->ApplyWarpBlend_RenderThread(RHICmdList, ViewportProxy);
 
-							// Implement ViewportRemap feature after warp
-							ViewportProxy->ImplViewportRemap_RenderThread(RHICmdList);
+							ViewportProxy->PostResolveViewport_RenderThread(RHICmdList);
 							break;
 
 						case EWarpPass::End:
@@ -372,8 +371,7 @@ void FDisplayClusterViewportManagerProxy::UpdateFrameResources_RenderThread(FRHI
 							// just resolve not warped viewports to frame target texture
 							ViewportProxy->ResolveResources_RenderThread(RHICmdList, EDisplayClusterViewportResourceType::InputShaderResource, ViewportProxy->GetOutputResourceType_RenderThread());
 							
-							// Implement ViewportRemap feature after resolve
-							ViewportProxy->ImplViewportRemap_RenderThread(RHICmdList);
+							ViewportProxy->PostResolveViewport_RenderThread(RHICmdList);
 							break;
 
 						default:
@@ -404,39 +402,51 @@ void FDisplayClusterViewportManagerProxy::DoCrossGPUTransfers_RenderThread(FRHIC
 
 	for (FDisplayClusterViewportProxy* ViewportProxy : ClusterNodeViewportProxies)
 	{
-		for (FDisplayClusterViewport_Context& ViewportContext : ViewportProxy->Contexts)
+		bool bShouldCrossGPUTransfersViewportRenderResources = true;
+
+		// Skip a frozen viewport that has already been transferred between GPUs
+		// The first time freezing should do the transfer (RenderTargets must be assigned on the first pass)
+		if (ViewportProxy->RenderSettings.bFreezeRendering && ViewportProxy->RenderTargets.Num() == 0)
+		{
+			bShouldCrossGPUTransfersViewportRenderResources = false;
+		}
+
+		if (bShouldCrossGPUTransfersViewportRenderResources)
 		{
-			if (ViewportContext.bAllowGPUTransferOptimization && ViewportContext.GPUIndex >= 0)
+			for (FDisplayClusterViewport_Context& ViewportContext : ViewportProxy->Contexts)
 			{
-				// Use optimized cross GPU transfer for this context
+				if (ViewportContext.bAllowGPUTransferOptimization && ViewportContext.GPUIndex >= 0)
+				{
+					// Use optimized cross GPU transfer for this context
 
-				FRenderTarget* RenderTarget = ViewportProxy->RenderTargets[ViewportContext.ContextNum];
-				FRHITexture2D* TextureRHI = ViewportProxy->RenderTargets[ViewportContext.ContextNum]->GetViewportRenderTargetResourceRHI();
+					FRenderTarget* RenderTarget = ViewportProxy->RenderTargets[ViewportContext.ContextNum];
+					FRHITexture2D* TextureRHI = ViewportProxy->RenderTargets[ViewportContext.ContextNum]->GetViewportRenderTargetResourceRHI();
 
-				FRHIGPUMask RenderTargetGPUMask = (GNumExplicitGPUsForRendering > 1 && RenderTarget) ? RenderTarget->GetGPUMask(RHICmdList) : FRHIGPUMask::GPU0();
-				{
-					static auto CVar = IConsoleManager::Get().FindConsoleVariable(TEXT("r.PathTracing.GPUCount"));
-					if (CVar && CVar->GetInt() > 1)
+					FRHIGPUMask RenderTargetGPUMask = (GNumExplicitGPUsForRendering > 1 && RenderTarget) ? RenderTarget->GetGPUMask(RHICmdList) : FRHIGPUMask::GPU0();
 					{
-						RenderTargetGPUMask = FRHIGPUMask::All(); // Broadcast to all GPUs 
+						static auto CVar = IConsoleManager::Get().FindConsoleVariable(TEXT("r.PathTracing.GPUCount"));
+						if (CVar && CVar->GetInt() > 1)
+						{
+							RenderTargetGPUMask = FRHIGPUMask::All(); // Broadcast to all GPUs 
+						}
 					}
-				}
 
-				FRHIGPUMask ContextGPUMask = FRHIGPUMask::FromIndex(ViewportContext.GPUIndex);
+					FRHIGPUMask ContextGPUMask = FRHIGPUMask::FromIndex(ViewportContext.GPUIndex);
 
-				if (ContextGPUMask != RenderTargetGPUMask)
-				{
-					// Clamp the view rect by the rendertarget rect to prevent issues when resizing the viewport.
-					const FIntRect TransferRect = ViewportContext.RenderTargetRect;
-
-					if (TransferRect.Width() > 0 && TransferRect.Height() > 0)
+					if (ContextGPUMask != RenderTargetGPUMask)
 					{
-						for (uint32 RenderTargetGPUIndex : RenderTargetGPUMask)
+						// Clamp the view rect by the rendertarget rect to prevent issues when resizing the viewport.
+						const FIntRect TransferRect = ViewportContext.RenderTargetRect;
+
+						if (TransferRect.Width() > 0 && TransferRect.Height() > 0)
 						{
-							if (!ContextGPUMask.Contains(RenderTargetGPUIndex))
+							for (uint32 RenderTargetGPUIndex : RenderTargetGPUMask)
 							{
-								FTransferResourceParams ResourceParams(TextureRHI, TransferRect, ContextGPUMask.GetFirstIndex(), RenderTargetGPUIndex, true, ViewportContext.bEnabledGPUTransferLockSteps);
-								TransferResources.Add(ResourceParams);
+								if (!ContextGPUMask.Contains(RenderTargetGPUIndex))
+								{
+									FTransferResourceParams ResourceParams(TextureRHI, TransferRect, ContextGPUMask.GetFirstIndex(), RenderTargetGPUIndex, true, ViewportContext.bEnabledGPUTransferLockSteps);
+									TransferResources.Add(ResourceParams);
+								}
 							}
 						}
 					}
diff --git a/Engine/Plugins/Runtime/nDisplay/Source/DisplayCluster/Private/Render/Viewport/DisplayClusterViewportManager_EditorPreview.cpp b/Engine/Plugins/Runtime/nDisplay/Source/DisplayCluster/Private/Render/Viewport/DisplayClusterViewportManager_EditorPreview.cpp
--- a/Engine/Plugins/Runtime/nDisplay/Source/DisplayCluster/Private/Render/Viewport/DisplayClusterViewportManager_EditorPreview.cpp
+++ b/Engine/Plugins/Runtime/nDisplay/Source/DisplayCluster/Private/Render/Viewport/DisplayClusterViewportManager_EditorPreview.cpp
@@ -37,12 +37,10 @@ void FDisplayClusterViewportManager::ImplUpdatePreviewRTTResources()
 {
 	check(IsInGameThread());
 
-	const EDisplayClusterRenderFrameMode RenderMode = GetRenderFrameSettings().RenderMode;
-
-	// Only for preview modes:
-	switch (RenderMode)
+	switch (GetRenderFrameSettings().RenderMode)
 	{
 	case EDisplayClusterRenderFrameMode::PreviewInScene:
+		// Only for preview modes:
 		break;
 	default:
 		return;
@@ -78,7 +76,7 @@ void FDisplayClusterViewportManager::ImplUpdatePreviewRTTResources()
 	}
 
 	// Get all supported preview rtt resources from root actor
-	RootActor->GetPreviewRenderTargetableTextures(RenderMode, PreviewViewportNames, PreviewRenderTargetableTextures);
+	RootActor->GetPreviewRenderTargetableTextures(PreviewViewportNames, PreviewRenderTargetableTextures);
 
 	// Configure preview RTT to viwports:
 	for (int32 ViewportIndex = 0; ViewportIndex < PreviewViewportNames.Num(); ViewportIndex++)
@@ -101,9 +99,10 @@ void FDisplayClusterViewportManager::ImplUpdatePreviewRTTResources()
 	}
 }
 
-bool FDisplayClusterViewportManager::RenderInEditor(class FDisplayClusterRenderFrame& InRenderFrame, FViewport* InViewport, const uint32 InFirstViewportNum, const int32 InViewportsAmmount, bool& bOutFrameRendered)
+bool FDisplayClusterViewportManager::RenderInEditor(class FDisplayClusterRenderFrame& InRenderFrame, FViewport* InViewport, const uint32 InFirstViewportNum, const int32 InViewportsAmount, int32& OutViewportsAmount, bool& bOutFrameRendered)
 {
 	bOutFrameRendered = false;
+	OutViewportsAmount = 0;
 
 	UWorld* CurrentWorld = GetCurrentWorld();
 	if (CurrentWorld == nullptr)
@@ -118,7 +117,6 @@ bool FDisplayClusterViewportManager::RenderInEditor(class FDisplayClusterRenderF
 	const bool bIsRenderedImmediatelyAfterAnotherViewFamily = false;
 
 	int32 ViewportIndex = 0;
-	int32 RenderedViewportsAmmount = 0;
 	bool bViewportsRenderPassDone = false;
 
 	for (FDisplayClusterRenderFrame::FFrameRenderTarget& RenderTargetIt : InRenderFrame.RenderTargets)
@@ -138,71 +136,78 @@ bool FDisplayClusterViewportManager::RenderInEditor(class FDisplayClusterRenderF
 				break;
 			}
 
-			// Create the view family for rendering the world scene to the viewport's render target
-			FSceneViewFamilyContext ViewFamily(CreateViewFamilyConstructionValues(
-				RenderTargetIt,
-				PreviewScene,
-				EngineShowFlags,
-				bAdditionalViewFamily
-			));
-
-			ConfigureViewFamily(RenderTargetIt, ViewFamiliesIt, ViewFamily);
-
-			for (FDisplayClusterRenderFrame::FFrameView& ViewIt : ViewFamiliesIt.Views)
+			if ((ViewportIndex + ViewFamiliesIt.Views.Num()) < (int32)(InFirstViewportNum + 1))
 			{
-				bool bViewportAlreadyRendered = ViewportIndex < (int32)InFirstViewportNum;
-				ViewportIndex++;
-
-				if (!bViewportAlreadyRendered)
-				{
-					FDisplayClusterViewport* ViewportPtr = static_cast<FDisplayClusterViewport*>(ViewIt.Viewport);
+				// Skip already rendered viewports
+				ViewportIndex += ViewFamiliesIt.Views.Num();
+			}
+			else
+			{
+				// Create the view family for rendering the world scene to the viewport's render target
+				FSceneViewFamilyContext ViewFamily(CreateViewFamilyConstructionValues(
+					RenderTargetIt,
+					PreviewScene,
+					EngineShowFlags,
+					bAdditionalViewFamily
+				));
 
-					check(ViewportPtr != nullptr);
-					check(ViewIt.ContextNum < (uint32)ViewportPtr->Contexts.Num());
+				ConfigureViewFamily(RenderTargetIt, ViewFamiliesIt, ViewFamily);
 
-					// Calculate the player's view information.
-					FVector  ViewLocation;
-					FRotator ViewRotation;
-					FSceneView* View = ViewportPtr->ImplCalcScenePreview(ViewFamily, ViewIt.ContextNum);
+				for (FDisplayClusterRenderFrame::FFrameView& ViewIt : ViewFamiliesIt.Views)
+				{
+					bool bViewportAlreadyRendered = ViewportIndex < (int32)InFirstViewportNum;
+					ViewportIndex++;
 
-					if (View && ViewIt.bDisableRender)
+					if (!bViewportAlreadyRendered)
 					{
-						ViewFamily.Views.Remove(View);
+						FDisplayClusterViewport* ViewportPtr = static_cast<FDisplayClusterViewport*>(ViewIt.Viewport);
+						check(ViewportPtr != nullptr);
+						check(ViewIt.ContextNum < (uint32)ViewportPtr->Contexts.Num());
 
-						delete View;
-						View = nullptr;
-					}
+						// Calculate the player's view information.
+						FVector  ViewLocation;
+						FRotator ViewRotation;
+						FSceneView* View = ViewportPtr->ImplCalcScenePreview(ViewFamily, ViewIt.ContextNum);
 
-					if (View)
-					{
-						// Apply viewport context settings to view (crossGPU, visibility, etc)
-						ViewIt.Viewport->SetupSceneView(ViewIt.ContextNum, PreviewScene->GetWorld(), ViewFamily, *View);
+						if (View != nullptr && ViewIt.IsShouldRenderView() == false)
+						{
+							ViewFamily.Views.Remove(View);
 
-						RenderedViewportsAmmount++;
-						if (InViewportsAmmount > 0 && RenderedViewportsAmmount >= InViewportsAmmount)
+							delete View;
+							View = nullptr;
+						}
+
+						if (View != nullptr)
 						{
-							bViewportsRenderPassDone = true;
-							break;
+							// Count all rendered viewports on this render pass
+							OutViewportsAmount++;
+
+							// Apply viewport context settings to view (crossGPU, visibility, etc)
+							ViewIt.Viewport->SetupSceneView(ViewIt.ContextNum, PreviewScene->GetWorld(), ViewFamily, *View);
+
+							// The allowed number of viewports is limited by the "InViewportAmount" value.
+							if (InViewportsAmount > 0 && OutViewportsAmount >= InViewportsAmount)
+							{
+								bViewportsRenderPassDone = true;
+								break;
+							}
 						}
 					}
 				}
-			}
 
-			if (ViewFamily.Views.Num() > 0)
-			{
-				// Screen percentage is still not supported in scene capture.
-				ViewFamily.EngineShowFlags.ScreenPercentage = false;
-				ViewFamily.SetScreenPercentageInterface(new FLegacyScreenPercentageDriver(ViewFamily, 1.0f));
+				if (ViewFamily.Views.Num() > 0)
+				{
+					// Screen percentage is still not supported in scene capture.
+					ViewFamily.EngineShowFlags.ScreenPercentage = false;
+					ViewFamily.SetScreenPercentageInterface(new FLegacyScreenPercentageDriver(ViewFamily, 1.0f));
 
-				ViewFamily.bIsRenderedImmediatelyAfterAnotherViewFamily = bIsRenderedImmediatelyAfterAnotherViewFamily;
+					ViewFamily.bIsRenderedImmediatelyAfterAnotherViewFamily = bIsRenderedImmediatelyAfterAnotherViewFamily;
 
-				FCanvas Canvas(RenderTargetIt.RenderTargetPtr, nullptr, PreviewScene->GetWorld(), ERHIFeatureLevel::SM5, FCanvas::CDM_DeferDrawing /*FCanvas::CDM_ImmediateDrawing*/, 1.0f);
-				Canvas.Clear(FLinearColor::Black);
+					FCanvas Canvas(RenderTargetIt.RenderTargetPtr, nullptr, PreviewScene->GetWorld(), ERHIFeatureLevel::SM5, FCanvas::CDM_DeferDrawing /*FCanvas::CDM_ImmediateDrawing*/, 1.0f);
+					Canvas.Clear(FLinearColor::Black);
 
-				GetRendererModule().BeginRenderingViewFamily(&Canvas, &ViewFamily);
+					GetRendererModule().BeginRenderingViewFamily(&Canvas, &ViewFamily);
 
-				if (GetRenderFrameSettings().bAllowMultiGPURenderingInEditor)
-				{
 					if (GNumExplicitGPUsForRendering > 1)
 					{
 						const FRHIGPUMask SubmitGPUMask = ViewFamily.Views.Num() == 1 ? ViewFamily.Views[0]->GPUMask : FRHIGPUMask::All();
@@ -218,12 +223,18 @@ bool FDisplayClusterViewportManager::RenderInEditor(class FDisplayClusterRenderF
 		}
 	}
 
+	// Render cluster node compose on this pass if possible
+	bool bHasExtraRenderPass = (InViewportsAmount < 0) || (OutViewportsAmount < InViewportsAmount);
+
 	// bViewportsRenderPassDone means all viewports for this cluster node have been rendered
 	// and we can start frame composing(creating mips, warps, composing icvfx, OutputRemap, etc.).
 	// The main idea is to move the frame composition to a separate render pass.
 	// This will reduce the load on the CPU + GPU and increase the FPS for the preview.
-	if (ViewportIndex == InRenderFrame.ViewportsAmmount && !bViewportsRenderPassDone)
+	if (!bViewportsRenderPassDone && bHasExtraRenderPass)
 	{
+		// Count cluster node compose render pass as viewport
+		OutViewportsAmount++;
+
 		// Handle special viewports game-thread logic at frame end
 		// custom postprocess single frame flag must be removed at frame end on game thread
 		FinalizeNewFrame();
diff --git a/Engine/Plugins/Runtime/nDisplay/Source/DisplayCluster/Private/Render/Viewport/DisplayClusterViewportProxy.cpp b/Engine/Plugins/Runtime/nDisplay/Source/DisplayCluster/Private/Render/Viewport/DisplayClusterViewportProxy.cpp
--- a/Engine/Plugins/Runtime/nDisplay/Source/DisplayCluster/Private/Render/Viewport/DisplayClusterViewportProxy.cpp
+++ b/Engine/Plugins/Runtime/nDisplay/Source/DisplayCluster/Private/Render/Viewport/DisplayClusterViewportProxy.cpp
@@ -17,6 +17,7 @@
 
 #if WITH_EDITOR
 #include "DisplayClusterRootActor.h"
+#include "Render/Viewport/Containers/DisplayClusterViewportReadPixels.h"
 #endif
 
 #include "RHIStaticStates.h"
@@ -43,6 +44,9 @@ static TAutoConsoleVariable<int32> CVarDisplayClusterRenderOverscanResolve(
 	ECVF_RenderThreadSafe
 );
 
+// The viewport override has the maximum depth. This protects against a link cycle
+static const int32 DisplayClusterViewportProxyResourcesOverrideRecursionDepthMax = 4;
+
 FDisplayClusterViewportProxy::FDisplayClusterViewportProxy(const FDisplayClusterViewportManagerProxy& InOwner, const FDisplayClusterViewport& RenderViewport)
 	: ViewportId(RenderViewport.ViewportId)
 	, ClusterNodeId(RenderViewport.ClusterNodeId)
@@ -71,14 +75,22 @@ bool ImplGetTextureResources_RenderThread(const TArray<FDisplayClusterViewportTe
 
 	if (InResources.Num() > 0)
 	{
-		OutResources.AddDefaulted(InResources.Num());
+		for (int32 ResourceIndex = 0; ResourceIndex < InResources.Num(); ResourceIndex++)
+		{
+			if(InResources[ResourceIndex] != nullptr)
+			{
+				OutResources.Add(InResources[ResourceIndex]->GetViewportResourceRHI());
+			}
+			
+		}
 
-		for (int32 ResourceIndex = 0; ResourceIndex < OutResources.Num(); ResourceIndex++)
+		if (OutResources.Num() == InResources.Num())
 		{
-			OutResources[ResourceIndex] = InResources[ResourceIndex]->GetViewportResourceRHI();
+			return true;
 		}
 
-		return true;
+		// Some resources lost
+		OutResources.Empty();
 	}
 
 	return false;
@@ -91,14 +103,43 @@ bool ImplGetTextureResources_RenderThread(const TArray<FTextureRenderTargetResou
 
 	if (InResources.Num() > 0)
 	{
-		OutResources.AddDefaulted(InResources.Num());
+		for (int32 ResourceIndex = 0; ResourceIndex < InResources.Num(); ResourceIndex++)
+		{
+			if (InResources[ResourceIndex] != nullptr)
+			{
+				OutResources.Add(InResources[ResourceIndex]->GetTexture2DRHI());
+			}
+		}
 
-		for (int32 ResourceIndex = 0; ResourceIndex < OutResources.Num(); ResourceIndex++)
+		if (OutResources.Num() == InResources.Num())
 		{
-			OutResources[ResourceIndex] = InResources[ResourceIndex]->GetTexture2DRHI();
+			return true;
 		}
 
-		return true;
+		// Some resources lost
+		OutResources.Empty();
+	}
+
+	return false;
+}
+
+bool FDisplayClusterViewportProxy::IsShouldOverrideViewportResource(const EDisplayClusterViewportResourceType InResourceType) const
+{
+	// Override resources from other viewport
+	if (RenderSettings.OverrideViewportId.IsEmpty() == false)
+	{
+		switch (InResourceType)
+		{
+		case EDisplayClusterViewportResourceType::InternalRenderTargetResource:
+		case EDisplayClusterViewportResourceType::InputShaderResource:
+		case EDisplayClusterViewportResourceType::MipsShaderResource:
+		case EDisplayClusterViewportResourceType::AdditionalTargetableResource:
+			return true;
+
+		default:
+			// By default use all output resources from this viewport
+			break;
+		}
 	}
 
 	return false;
@@ -106,9 +147,28 @@ bool ImplGetTextureResources_RenderThread(const TArray<FTextureRenderTargetResou
 
 //  Return viewport scene proxy resources by type
 bool FDisplayClusterViewportProxy::GetResources_RenderThread(const EDisplayClusterViewportResourceType InResourceType, TArray<FRHITexture2D*>& OutResources) const
+{
+	return ImplGetResources_RenderThread(InResourceType, OutResources, false);
+}
+
+bool FDisplayClusterViewportProxy::ImplGetResources_RenderThread(const EDisplayClusterViewportResourceType InResourceType, TArray<FRHITexture2D*>& OutResources, const int32 InRecursionDepth) const
 {
 	check(IsInRenderingThread());
 
+	// Override resources from other viewport
+	if (IsShouldOverrideViewportResource(InResourceType))
+	{
+		if (InRecursionDepth < DisplayClusterViewportProxyResourcesOverrideRecursionDepthMax)
+		{
+			if (FDisplayClusterViewportProxy const* OverrideViewportProxy = Owner.ImplFindViewport_RenderThread(RenderSettings.OverrideViewportId))
+			{
+				return OverrideViewportProxy->ImplGetResources_RenderThread(InResourceType, OutResources, InRecursionDepth + 1);
+			}
+		}
+
+		return false;
+	}
+
 	OutResources.Empty();
 
 	switch (InResourceType)
@@ -117,34 +177,31 @@ bool FDisplayClusterViewportProxy::GetResources_RenderThread(const EDisplayClust
 	{
 		if (Contexts.Num() > 0)
 		{
-			OutResources.AddDefaulted(Contexts.Num());
 			for (int32 ContextIt = 0; ContextIt < Contexts.Num(); ContextIt++)
 			{
-				//Support Override:
-				if (Contexts[ContextIt].bDisableRender)
+				// Support texture replace:
+				if (PostRenderSettings.Replace.IsEnabled())
 				{
-					if (!PostRenderSettings.Replace.TextureRHI.IsValid())
-					{
-						OutResources.Empty();
-						return false;
-					}
-
-					OutResources[ContextIt] = PostRenderSettings.Replace.TextureRHI->GetTexture2D();
+					OutResources.Add(PostRenderSettings.Replace.TextureRHI->GetTexture2D());
 				}
 				else
 				{
-					FDisplayClusterViewportRenderTargetResource* Input = RenderTargets[ContextIt];
-					if (Input == nullptr || Input->IsInitialized() == false)
+					if (ContextIt < RenderTargets.Num())
 					{
-						OutResources.Empty();
-						return false;
+						if (FDisplayClusterViewportRenderTargetResource* Input = RenderTargets[ContextIt])
+						{
+							OutResources.Add(Input->GetViewportRenderTargetResourceRHI());
+						}
 					}
-
-					OutResources[ContextIt] = Input->GetViewportRenderTargetResourceRHI();
 				}
 			}
 		}
 
+		if (Contexts.Num() != OutResources.Num())
+		{
+			OutResources.Empty();
+		}
+
 		return OutResources.Num() > 0;
 	}
 
@@ -189,10 +246,19 @@ bool FDisplayClusterViewportProxy::GetResources_RenderThread(const EDisplayClust
 	return false;
 }
 
-void FDisplayClusterViewportProxy::ImplViewportRemap_RenderThread(FRHICommandListImmediate& RHICmdList) const
+void FDisplayClusterViewportProxy::PostResolveViewport_RenderThread(FRHICommandListImmediate& RHICmdList) const
 {
 	check(IsInRenderingThread());
 
+	// Implement ViewportRemap feature
+	ImplViewportRemap_RenderThread(RHICmdList);
+
+	// Implement read pixels for preview DCRA
+	ImplPreviewReadPixels_RenderThread(RHICmdList);
+}
+
+void FDisplayClusterViewportProxy::ImplViewportRemap_RenderThread(FRHICommandListImmediate& RHICmdList) const
+{
 #if WITH_EDITOR
 	switch(Owner.GetRenderFrameSettings_RenderThread().RenderMode)
 	{
@@ -232,6 +298,70 @@ void FDisplayClusterViewportProxy::ImplViewportRemap_RenderThread(FRHICommandLis
 	}
 }
 
+void FDisplayClusterViewportProxy::ImplPreviewReadPixels_RenderThread(FRHICommandListImmediate& RHICmdList) const
+{
+#if WITH_EDITOR
+	if (RenderSettings.bPreviewReadPixels)
+	{
+		bPreviewReadPixels = true;
+	}
+
+	// Now try to read until success:
+	if (bPreviewReadPixels && Contexts.Num() > 0 && OutputPreviewTargetableResource.IsValid())
+	{
+		check(Owner.GetRenderFrameSettings_RenderThread().RenderMode == EDisplayClusterRenderFrameMode::PreviewInScene);
+
+		// We should synchronize thread for preview read
+		FScopeLock Lock(&PreviewPixelsCSGuard);
+
+		if (PreviewPixels.IsValid() == false)
+		{
+			// Read pixels from this texture
+			FRHITexture* Texture = OutputPreviewTargetableResource;
+			if (FRHITexture2D* Texture2D = static_cast<FRHITexture2D*>(Texture))
+			{
+				
+				// Clear deferred read flag
+				bPreviewReadPixels = false;
+
+				TSharedPtr<FDisplayClusterViewportReadPixelsData, ESPMode::ThreadSafe> ReadData = MakeShared<FDisplayClusterViewportReadPixelsData, ESPMode::ThreadSafe>();
+
+				RHICmdList.ReadSurfaceData(
+					Texture,
+					FIntRect(FIntPoint(0, 0), Texture2D->GetSizeXY()),
+					ReadData->Pixels,
+					FReadSurfaceDataFlags(RCM_UNorm, CubeFace_MAX)
+				);
+
+				ReadData->Size = Texture2D->GetSizeXY();
+
+				// Expose result:
+				PreviewPixels = ReadData;
+			}
+		}
+	}
+#endif
+}
+
+#if WITH_EDITOR
+bool FDisplayClusterViewportProxy::GetPreviewPixels_GameThread(TSharedPtr<FDisplayClusterViewportReadPixelsData, ESPMode::ThreadSafe>& OutPixelsData) const
+{
+	check(IsInGameThread());
+
+	// We should synchronize thread for preview read
+	FScopeLock Lock(&PreviewPixelsCSGuard);
+
+	if (PreviewPixels.IsValid())
+	{
+		OutPixelsData = PreviewPixels;
+		PreviewPixels.Reset();
+		return true;
+	}
+
+	return false;
+}
+#endif
+
 EDisplayClusterViewportResourceType FDisplayClusterViewportProxy::GetOutputResourceType_RenderThread() const
 {
 	check(IsInRenderingThread());
@@ -261,9 +391,28 @@ EDisplayClusterViewportResourceType FDisplayClusterViewportProxy::GetOutputResou
 }
 
 bool FDisplayClusterViewportProxy::GetResourcesWithRects_RenderThread(const EDisplayClusterViewportResourceType InResourceType, TArray<FRHITexture2D*>& OutResources, TArray<FIntRect>& OutResourceRects) const
+{
+	return ImplGetResourcesWithRects_RenderThread(InResourceType, OutResources, OutResourceRects, false);
+}
+
+bool FDisplayClusterViewportProxy::ImplGetResourcesWithRects_RenderThread(const EDisplayClusterViewportResourceType InResourceType, TArray<FRHITexture2D*>& OutResources, TArray<FIntRect>& OutResourceRects, const int32 InRecursionDepth) const
 {
 	check(IsInRenderingThread());
 
+	// Override resources from other viewport
+	if(IsShouldOverrideViewportResource(InResourceType))
+	{
+		if (InRecursionDepth < DisplayClusterViewportProxyResourcesOverrideRecursionDepthMax)
+		{
+			if (FDisplayClusterViewportProxy const* OverrideViewportProxy = Owner.ImplFindViewport_RenderThread(RenderSettings.OverrideViewportId))
+			{
+				return OverrideViewportProxy->ImplGetResourcesWithRects_RenderThread(InResourceType, OutResources, OutResourceRects, InRecursionDepth + 1);
+			}
+		}
+
+		return false;
+	}
+
 	if (!GetResources_RenderThread(InResourceType, OutResources))
 	{
 		return false;
@@ -277,7 +426,7 @@ bool FDisplayClusterViewportProxy::GetResourcesWithRects_RenderThread(const EDis
 	case EDisplayClusterViewportResourceType::InternalRenderTargetResource:
 		for (int32 ContextIt = 0; ContextIt < OutResourceRects.Num(); ContextIt++)
 		{
-			if (Contexts[ContextIt].bDisableRender)
+			if (PostRenderSettings.Replace.IsEnabled())
 			{
 				// Get image from Override
 				OutResourceRects[ContextIt] = PostRenderSettings.Replace.Rect;
@@ -310,34 +459,20 @@ void FDisplayClusterViewportProxy::UpdateDeferredResources(FRHICommandListImmedi
 {
 	check(IsInRenderingThread());
 
-	FDisplayClusterViewportProxy const * SourceProxy = this;
-	EDisplayClusterViewportResourceType SourceType = EDisplayClusterViewportResourceType::InternalRenderTargetResource;
-
-	if (RenderSettings.OverrideViewportId.IsEmpty() == false)
+	if (RenderTargets.Num() == 0 && PostRenderSettings.Replace.IsEnabled() == false)
 	{
-		// link resources from other viewport
-		//@todo: implement correct sorting order for rendering in frame manager
-		//@todo: Implement stereo supporting, now mono only
-
-		FDisplayClusterViewportProxy const* OverrideViewportProxy = Owner.ImplFindViewport_RenderThread(RenderSettings.OverrideViewportId);
-		if (OverrideViewportProxy)
-		{
-			SourceProxy = OverrideViewportProxy;
-			// Get after postprocessing
-			SourceType = EDisplayClusterViewportResourceType::InputShaderResource;
-		}
+		// Internal RTT required for deferred update. Except when use 'Replace' as source
+		return;
 	}
-	else
+
+	if (RenderSettings.bFreezeRendering || RenderSettings.bSkipRendering || RenderSettings.OverrideViewportId.IsEmpty() == false)
 	{
-		if (RenderSettings.bSkipRendering)
-		{
-			//@todo: support skip rendering
-			return;
-		}
+		// Disable deferred update
+		return;
 	}
 
 	// Pass 0: Resolve from RTT region to separated viewport context resource:
-	ImplResolveResources_RenderThread(RHICmdList, SourceProxy, SourceType, EDisplayClusterViewportResourceType::InputShaderResource);
+	ResolveResources_RenderThread(RHICmdList, EDisplayClusterViewportResourceType::InternalRenderTargetResource, EDisplayClusterViewportResourceType::InputShaderResource);
 
 	// Pass 1: Generate blur postprocess effect for render target texture rect for all contexts
 	if (PostRenderSettings.PostprocessBlur.IsEnabled())
@@ -396,7 +531,8 @@ static void DirectCopyTextureImpl_RenderThread(FRHICommandListImmediate& RHICmdL
 	RHICmdList.CopyToResolveTarget(SrcTexture, DstTexture, Params);
 }
 
-static void ResampleCopyTextureImpl_RenderThread(FRHICommandListImmediate& RHICmdList, FRHITexture* SrcTexture, FRHITexture* DstTexture, const FIntRect& SrcRect, const FIntRect& DstRect)
+template<class TScreenPixelShader>
+void ResampleCopyTextureImpl_RenderThread(FRHICommandListImmediate& RHICmdList, FRHITexture* SrcTexture, FRHITexture* DstTexture, const FIntRect& SrcRect, const FIntRect& DstRect)
 {
 	// Texture format mismatch, use a shader to do the copy.
 	// #todo-renderpasses there's no explicit resolve here? Do we need one?
@@ -413,13 +549,15 @@ static void ResampleCopyTextureImpl_RenderThread(FRHICommandListImmediate& RHICm
 
 		FGraphicsPipelineStateInitializer GraphicsPSOInit;
 		RHICmdList.ApplyCachedRenderTargets(GraphicsPSOInit);
+
 		GraphicsPSOInit.BlendState = TStaticBlendState<>::GetRHI();
 		GraphicsPSOInit.RasterizerState = TStaticRasterizerState<>::GetRHI();
 		GraphicsPSOInit.DepthStencilState = TStaticDepthStencilState<false, CF_Always>::GetRHI();
 
+
 		FGlobalShaderMap* ShaderMap = GetGlobalShaderMap(GMaxRHIFeatureLevel);
 		TShaderMapRef<FScreenVS> VertexShader(ShaderMap);
-		TShaderMapRef<FScreenPS> PixelShader(ShaderMap);
+		TShaderMapRef<TScreenPixelShader> PixelShader(ShaderMap);
 
 		GraphicsPSOInit.BoundShaderState.VertexDeclarationRHI = GFilterVertexDeclaration.VertexDeclarationRHI;
 		GraphicsPSOInit.BoundShaderState.VertexShaderRHI = VertexShader.GetVertexShader();
@@ -453,7 +591,7 @@ static void ResampleCopyTextureImpl_RenderThread(FRHICommandListImmediate& RHICm
 	RHICmdList.EndRenderPass();
 }
 
-void ImplResolveResource(FRHICommandListImmediate& RHICmdList, FRHITexture2D* InputResource, const FIntRect& InputRect, FRHITexture2D* OutputResource, const FIntRect& OutputRect, bool bOutputIsMipsResource)
+void ImplResolveResource(FRHICommandListImmediate& RHICmdList, FRHITexture2D* InputResource, const FIntRect& InputRect, FRHITexture2D* OutputResource, const FIntRect& OutputRect, const bool bOutputIsMipsResource, const bool bOutputIsPreviewResource)
 {
 	check(InputResource);
 	check(OutputResource);
@@ -476,10 +614,18 @@ void ImplResolveResource(FRHICommandListImmediate& RHICmdList, FRHITexture2D* In
 	}
 	else
 	{
-		if (InputRect.Size() != OutputRect.Size() || InputResource->GetFormat() != OutputResource->GetFormat())
+		if (InputRect.Size() != OutputRect.Size() || InputResource->GetFormat() != OutputResource->GetFormat() || bOutputIsPreviewResource)
 		{
 			// Resample size with shader
-			ResampleCopyTextureImpl_RenderThread(RHICmdList, InputResource, OutputResource, InputRect, OutputRect);
+			if(bOutputIsPreviewResource)
+			{
+				// Preview require a normal alpha (re-invert)
+				ResampleCopyTextureImpl_RenderThread<FScreenPSInvertAlpha>(RHICmdList, InputResource, OutputResource, InputRect, OutputRect);
+			}
+			else
+			{
+				ResampleCopyTextureImpl_RenderThread<FScreenPS>(RHICmdList, InputResource, OutputResource, InputRect, OutputRect);
+			}
 		}
 		else
 		{
@@ -506,12 +652,19 @@ bool FDisplayClusterViewportProxy::ImplResolveResources_RenderThread(FRHICommand
 
 	bool bOutputIsMipsResource = OutputResourceType == EDisplayClusterViewportResourceType::MipsShaderResource;
 
+#if WITH_EDITOR
+	// This resolve pattern always called once for preview. This flag force to invert alpha (at this point alpha from engine is inverted)
+	const bool bOutputIsPreviewResource = InputResourceType == EDisplayClusterViewportResourceType::InputShaderResource && OutputResourceType == EDisplayClusterViewportResourceType::OutputPreviewTargetableResource;
+#else
+	const bool bOutputIsPreviewResource = false;
+#endif
+
 	TArray<FRHITexture2D*> InputResources, OutputResources;
 	TArray<FIntRect> InputResourcesRect, OutputResourcesRect;
 	if (SourceProxy->GetResourcesWithRects_RenderThread(InputResourceType, InputResources, InputResourcesRect) && GetResourcesWithRects_RenderThread(OutputResourceType, OutputResources, OutputResourcesRect))
 	{
-		const int32 InputAmmount = FMath::Min(InputResources.Num(), OutputResources.Num());
-		for (int32 InputIt = 0; InputIt < InputAmmount; InputIt++)
+		const int32 InputAmount = FMath::Min(InputResources.Num(), OutputResources.Num());
+		for (int32 InputIt = 0; InputIt < InputAmount; InputIt++)
 		{
 			FIntRect SourceRect = InputResourcesRect[InputIt];
 
@@ -529,18 +682,18 @@ bool FDisplayClusterViewportProxy::ImplResolveResources_RenderThread(FRHICommand
 				break;
 			}
 
-			if ((InputIt + 1) == InputAmmount)
+			if ((InputIt + 1) == InputAmount)
 			{
 				// last input mono -> stereo outputs
 				for (int32 OutputIt = InputIt; OutputIt < OutputResources.Num(); OutputIt++)
 				{
-					ImplResolveResource(RHICmdList, InputResources[InputIt], SourceRect, OutputResources[OutputIt], OutputResourcesRect[OutputIt], bOutputIsMipsResource);
+					ImplResolveResource(RHICmdList, InputResources[InputIt], SourceRect, OutputResources[OutputIt], OutputResourcesRect[OutputIt], bOutputIsMipsResource, bOutputIsPreviewResource);
 				}
 				break;
 			}
 			else
 			{
-				ImplResolveResource(RHICmdList, InputResources[InputIt], SourceRect, OutputResources[InputIt], OutputResourcesRect[InputIt], bOutputIsMipsResource);
+				ImplResolveResource(RHICmdList, InputResources[InputIt], SourceRect, OutputResources[InputIt], OutputResourcesRect[InputIt], bOutputIsMipsResource, bOutputIsPreviewResource);
 			}
 		}
 
diff --git a/Engine/Plugins/Runtime/nDisplay/Source/DisplayCluster/Private/Render/Viewport/DisplayClusterViewportProxy.h b/Engine/Plugins/Runtime/nDisplay/Source/DisplayCluster/Private/Render/Viewport/DisplayClusterViewportProxy.h
--- a/Engine/Plugins/Runtime/nDisplay/Source/DisplayCluster/Private/Render/Viewport/DisplayClusterViewportProxy.h
+++ b/Engine/Plugins/Runtime/nDisplay/Source/DisplayCluster/Private/Render/Viewport/DisplayClusterViewportProxy.h
@@ -14,6 +14,7 @@ class FDisplayClusterViewport;
 class IDisplayClusterShaders;
 class IDisplayClusterProjectionPolicy;
 class IDisplayClusterRender_MeshComponent;
+class FDisplayClusterViewportReadPixelsData;
 
 class FDisplayClusterViewportProxy
 	: public IDisplayClusterViewportProxy
@@ -87,9 +88,11 @@ public:
 	// ~IDisplayClusterViewportProxy
 	///////////////////////////////
 
-	bool ImplResolveResources_RenderThread(FRHICommandListImmediate& RHICmdList, FDisplayClusterViewportProxy const* SourceProxy, const EDisplayClusterViewportResourceType InputResourceType, const EDisplayClusterViewportResourceType OutputResourceType) const;
+	void PostResolveViewport_RenderThread(FRHICommandListImmediate& RHICmdList) const;
 
-	void ImplViewportRemap_RenderThread(FRHICommandListImmediate& RHICmdList) const;
+#if WITH_EDITOR
+	bool GetPreviewPixels_GameThread(TSharedPtr<FDisplayClusterViewportReadPixelsData, ESPMode::ThreadSafe>& OutPixelsData) const;
+#endif
 
 	inline bool FindContext_RenderThread(const int32 ViewIndex, uint32* OutContextNum)
 	{
@@ -111,6 +114,17 @@ public:
 		return false;
 	}
 
+private:
+	bool ImplGetResourcesWithRects_RenderThread(const EDisplayClusterViewportResourceType InResourceType, TArray<FRHITexture2D*>& OutResources, TArray<FIntRect>& OutResourceRects, const int32 InRecursionDepth) const;
+	bool ImplGetResources_RenderThread(const EDisplayClusterViewportResourceType InResourceType, TArray<FRHITexture2D*>& OutResources, const int32 InRecursionDepth) const;
+
+	void ImplViewportRemap_RenderThread(FRHICommandListImmediate& RHICmdList) const;
+	void ImplPreviewReadPixels_RenderThread(FRHICommandListImmediate& RHICmdList) const;
+
+	bool ImplResolveResources_RenderThread(FRHICommandListImmediate& RHICmdList, FDisplayClusterViewportProxy const* SourceProxy, const EDisplayClusterViewportResourceType InputResourceType, const EDisplayClusterViewportResourceType OutputResourceType) const;
+
+	bool IsShouldOverrideViewportResource(const EDisplayClusterViewportResourceType InResourceType) const;
+
 protected:
 	friend FDisplayClusterViewportProxyData;
 	friend FDisplayClusterViewportManagerProxy;
@@ -145,6 +159,10 @@ protected:
 
 #if WITH_EDITOR
 	FTextureRHIRef OutputPreviewTargetableResource;
+	
+	mutable bool bPreviewReadPixels = false;
+	mutable FCriticalSection PreviewPixelsCSGuard;
+	mutable TSharedPtr<FDisplayClusterViewportReadPixelsData, ESPMode::ThreadSafe> PreviewPixels;
 #endif
 
 	// unique viewport resources
diff --git a/Engine/Plugins/Runtime/nDisplay/Source/DisplayCluster/Private/Render/Viewport/DisplayClusterViewportStereoscopicPass.h b/Engine/Plugins/Runtime/nDisplay/Source/DisplayCluster/Private/Render/Viewport/DisplayClusterViewportStereoscopicPass.h
--- a/Engine/Plugins/Runtime/nDisplay/Source/DisplayCluster/Private/Render/Viewport/DisplayClusterViewportStereoscopicPass.h
+++ b/Engine/Plugins/Runtime/nDisplay/Source/DisplayCluster/Private/Render/Viewport/DisplayClusterViewportStereoscopicPass.h
@@ -4,13 +4,28 @@
 
 #include "CoreMinimal.h"
 #include "StereoRendering.h"
+#include "Render/Viewport/RenderFrame/DisplayClusterRenderFrameSettings.h"
 
 struct FDisplayClusterViewportStereoscopicPass
 {
-	inline static EStereoscopicPass EncodeStereoscopicPass(const uint32 ContextNum, const uint32 ContextAmmount)
+	// 'EStereoscopicPass': Stereoscopic rendering passes.
+	// FULL implies stereoscopic rendering isn't enabled for this pass
+	// PRIMARY implies the view needs its own pass, while SECONDARY implies the view can be instanced
+
+	inline static EStereoscopicPass EncodeStereoscopicPass(const uint32 ContextNum, const uint32 ContextAmount, const FDisplayClusterRenderFrameSettings& InFrameSettings)
 	{
-		// Stereoscopic rendering
-		if (ContextAmmount > 1)
+#if WITH_EDITOR
+		if (InFrameSettings.bIsRenderingInEditor)
+		{
+			return EStereoscopicPass::eSSP_FULL;
+		}
+#endif
+
+		// 'bEnableStereoscopicRenderingOptimization'
+		// When stereoscopic rendering optimization is enabled, one image family and one RTT for both eyes is used. 
+		// The result both eyes are rendered in the same family. It makes sense to use PRIMARY/SECONDARY.
+		// Otherwise, we render each view into a unique family.As a result, only PRIMARY is used.
+		if (InFrameSettings.bEnableStereoscopicRenderingOptimization && ContextAmount > 1)
 		{
 			switch (ContextNum)
 			{
@@ -29,7 +44,7 @@ struct FDisplayClusterViewportStereoscopicPass
 			}
 		}
 
-		// Monoscopic rendering
+		// Render as primary
 		return EStereoscopicPass::eSSP_PRIMARY;
 	}
 };
diff --git a/Engine/Plugins/Runtime/nDisplay/Source/DisplayCluster/Private/Render/Viewport/DisplayClusterViewport_EditorPreview.cpp b/Engine/Plugins/Runtime/nDisplay/Source/DisplayCluster/Private/Render/Viewport/DisplayClusterViewport_EditorPreview.cpp
--- a/Engine/Plugins/Runtime/nDisplay/Source/DisplayCluster/Private/Render/Viewport/DisplayClusterViewport_EditorPreview.cpp
+++ b/Engine/Plugins/Runtime/nDisplay/Source/DisplayCluster/Private/Render/Viewport/DisplayClusterViewport_EditorPreview.cpp
@@ -2,7 +2,7 @@
 
 #include "Components/DisplayClusterPreviewComponent.h"
 #include "Render/Viewport/DisplayClusterViewport.h"
-
+#include "Render/Viewport/DisplayClusterViewportProxy.h"
 #include "Render/Projection/IDisplayClusterProjectionPolicy.h"
 
 #include "Misc/DisplayClusterLog.h"
@@ -23,7 +23,7 @@
 #include "Components/DisplayClusterScreenComponent.h"
 
 #include "Render/Viewport/DisplayClusterViewportManager.h"
-#include "Render/Viewport/IDisplayClusterViewport.h"
+#include "Render/Viewport/RenderFrame/DisplayClusterRenderFrameSettings.h"
 
 #include "Render/Viewport/Configuration/DisplayClusterViewportConfigurationHelpers_Postprocess.h"
 
@@ -44,10 +44,10 @@ void FDisplayClusterViewport::CleanupViewState()
 
 FSceneViewStateInterface* FDisplayClusterViewport::GetViewState(uint32 ViewIndex)
 {
-	int32 RequiredAmmount = (int32)ViewIndex - ViewStates.Num() + 1;
-	if(RequiredAmmount > 0)
+	int32 RequiredAmount = (int32)ViewIndex - ViewStates.Num() + 1;
+	if(RequiredAmount > 0)
 	{
-		ViewStates.AddDefaulted(RequiredAmmount);
+		ViewStates.AddDefaulted(RequiredAmount);
 	}
 
 	if (ViewStates[ViewIndex].GetReference() == NULL)
@@ -116,6 +116,11 @@ FSceneView* FDisplayClusterViewport::ImplCalcScenePreview(FSceneViewFamilyContex
 
 		ViewInitOptions.BackgroundColor = FLinearColor::Black;
 
+		if (Owner.GetRenderFrameSettings().bPreviewEnablePostProcess == false)
+		{
+			ViewInitOptions.OverlayColor = FLinearColor::Black;
+		}
+
 		FSceneView* View = new FSceneView(ViewInitOptions);
 
 		View->bIsSceneCapture = true;
@@ -240,23 +245,15 @@ bool FDisplayClusterViewport::ImplPreview_CalculateStereoViewOffset(const uint32
 	const float EyeOffset = CfgEyeDist / 2.f;
 	const float EyeOffsetValues[] = { -EyeOffset, 0.f, EyeOffset };
 
-	auto DecodeEyeType = [](const EStereoscopicPass StereoPass)
+	// Decode current eye type	
+	// Decode current eye type
+	EDisplayClusterEyeType EyeType = EDisplayClusterEyeType::Mono;
+	if (Contexts.Num() > 1)
 	{
-		switch (StereoPass)
-		{
-		case EStereoscopicPass::eSSP_PRIMARY:
-			return EDisplayClusterEyeType::StereoLeft;
-		case EStereoscopicPass::eSSP_SECONDARY:
-			return EDisplayClusterEyeType::StereoRight;
-		default:
-			break;
-		}
-
-		return EDisplayClusterEyeType::Mono;
-	};
+		// Support stereo:
+		EyeType = (InContextNum == 0) ? EDisplayClusterEyeType::StereoLeft : EDisplayClusterEyeType::StereoRight;
+	}
 
-	// Decode current eye type	
-	const EDisplayClusterEyeType EyeType = DecodeEyeType(ViewportContext.StereoscopicPass);
 	const int32 EyeIndex = (int32)EyeType;
 
 	float PassOffset = 0.f;
@@ -325,4 +322,10 @@ FMatrix FDisplayClusterViewport::ImplPreview_GetStereoProjectionMatrix(const uin
 	return PrjMatrix;
 }
 
+bool FDisplayClusterViewport::GetPreviewPixels(TSharedPtr<FDisplayClusterViewportReadPixelsData, ESPMode::ThreadSafe>& OutPixelsData) const
+{
+	check(IsInGameThread());
+
+	return (ViewportProxy != nullptr) && ViewportProxy->GetPreviewPixels_GameThread(OutPixelsData);
+}
 #endif
diff --git a/Engine/Plugins/Runtime/nDisplay/Source/DisplayCluster/Private/Render/Viewport/DisplayClusterViewport_Math.cpp b/Engine/Plugins/Runtime/nDisplay/Source/DisplayCluster/Private/Render/Viewport/DisplayClusterViewport_Math.cpp
--- a/Engine/Plugins/Runtime/nDisplay/Source/DisplayCluster/Private/Render/Viewport/DisplayClusterViewport_Math.cpp
+++ b/Engine/Plugins/Runtime/nDisplay/Source/DisplayCluster/Private/Render/Viewport/DisplayClusterViewport_Math.cpp
@@ -31,7 +31,11 @@ bool FDisplayClusterViewport::CalculateView(const uint32 InContextNum, FVector&
 	{
 		// Store the view location/rotation
 		Contexts[InContextNum].ViewLocation = InOutViewLocation;
+		Contexts[InContextNum].bIsValidViewLocation = true;
+
 		Contexts[InContextNum].ViewRotation = InOutViewRotation;
+		Contexts[InContextNum].bIsValidViewRotation = true;
+
 		Contexts[InContextNum].WorldToMeters = WorldToMeters;
 
 		return true;
@@ -45,6 +49,7 @@ bool FDisplayClusterViewport::GetProjectionMatrix(const uint32 InContextNum, FMa
 	if (ProjectionPolicy.IsValid() && ProjectionPolicy->GetProjectionMatrix(this, InContextNum, OutPrjMatrix))
 	{
 		Contexts[InContextNum].ProjectionMatrix = OutPrjMatrix;
+		Contexts[InContextNum].bIsValidProjectionMatrix = true;
 
 		if (OverscanRendering.IsEnabled())
 		{
diff --git a/Engine/Plugins/Runtime/nDisplay/Source/DisplayCluster/Private/Render/Viewport/RenderFrame/DisplayClusterRenderFrameManager.cpp b/Engine/Plugins/Runtime/nDisplay/Source/DisplayCluster/Private/Render/Viewport/RenderFrame/DisplayClusterRenderFrameManager.cpp
--- a/Engine/Plugins/Runtime/nDisplay/Source/DisplayCluster/Private/Render/Viewport/RenderFrame/DisplayClusterRenderFrameManager.cpp
+++ b/Engine/Plugins/Runtime/nDisplay/Source/DisplayCluster/Private/Render/Viewport/RenderFrame/DisplayClusterRenderFrameManager.cpp
@@ -95,49 +95,26 @@ bool FDisplayClusterRenderFrameManager::BuildSimpleFrame(FViewport* InViewport,
 		{
 			for (const FDisplayClusterViewport_Context& ContextIt : ViewportIt->GetContexts())
 			{
-				bool bShouldUseRenderTarget = true;
-
 				FDisplayClusterRenderFrame::FFrameView FrameView;
 				{
 					FrameView.ContextNum = ContextIt.ContextNum;
 					FrameView.Viewport = ViewportIt;
 					FrameView.bDisableRender = ContextIt.bDisableRender;
-
-					// Simple frame use unique RTT  for each viewport, so disable RTT when viewport rendering disabled
-					if (ContextIt.bDisableRender)
-					{
-						bShouldUseRenderTarget = false;
-					}
-				}
-
-				float CustomBufferRatio = ViewportIt->GetRenderSettings().BufferRatio;
-				{
-					// Global multiplier
-					CustomBufferRatio *= InRenderFrameSettings.ClusterBufferRatioMult;
-
-					if ((ViewportIt->RenderSettingsICVFX.RuntimeFlags & ViewportRuntime_ICVFXTarget) != 0)
-					{
-						// Outer viewport
-						CustomBufferRatio *= InRenderFrameSettings.ClusterICVFXOuterViewportBufferRatioMult;
-					}
-					else
-					if ((ViewportIt->RenderSettingsICVFX.RuntimeFlags & ViewportRuntime_ICVFXIncamera) != 0)
-					{
-						// Inner Frustum
-						CustomBufferRatio *= InRenderFrameSettings.ClusterICVFXInnerFrustumBufferRatioMult;
-					}
+					FrameView.bFreezeRendering = ViewportIt->RenderSettings.bFreezeRendering;
 				}
 
 				FDisplayClusterRenderFrame::FFrameViewFamily FrameViewFamily;
 				{
 					FrameViewFamily.Views.Add(FrameView);
-					FrameViewFamily.CustomBufferRatio = CustomBufferRatio;
+					FrameViewFamily.CustomBufferRatio = ContextIt.CustomBufferRatio;
 					FrameViewFamily.ViewExtensions = ViewportIt->GatherActiveExtensions(InViewport);
 				}
 
 				FDisplayClusterRenderFrame::FFrameRenderTarget FrameRenderTarget;
 				{
-					FrameRenderTarget.bShouldUseRenderTarget = bShouldUseRenderTarget;
+					// Simple frame use unique RTT  for each viewport, so disable RTT when viewport rendering disabled
+					FrameRenderTarget.bShouldUseRenderTarget = ViewportIt->RenderTargets.Num() > 0;
+
 					FrameRenderTarget.ViewFamilies.Add(FrameViewFamily);
 
 					FrameRenderTarget.RenderTargetSize = ContextIt.RenderTargetRect.Max;
diff --git a/Engine/Plugins/Runtime/nDisplay/Source/DisplayCluster/Private/Render/Viewport/RenderFrame/DisplayClusterRenderFrameSettings.h b/Engine/Plugins/Runtime/nDisplay/Source/DisplayCluster/Private/Render/Viewport/RenderFrame/DisplayClusterRenderFrameSettings.h
--- a/Engine/Plugins/Runtime/nDisplay/Source/DisplayCluster/Private/Render/Viewport/RenderFrame/DisplayClusterRenderFrameSettings.h
+++ b/Engine/Plugins/Runtime/nDisplay/Source/DisplayCluster/Private/Render/Viewport/RenderFrame/DisplayClusterRenderFrameSettings.h
@@ -21,6 +21,9 @@ struct FDisplayClusterRenderFrameSettings
 	// Postprocess can use full size backbuffer. This disables Frame RTT size optimization
 	bool bShouldUseFullSizeFrameTargetableResource = false;
 
+	// Enable Post Process for preview rendering
+	bool bPreviewEnablePostProcess = false;
+
 	// Preview RTT size multiplier
 	float PreviewRenderTargetRatioMult = 1.f;
 
@@ -57,7 +60,10 @@ struct FDisplayClusterRenderFrameSettings
 	int32 PreviewMaxGPUIndex = 0;
 
 	// The maximum dimension of any texture for preview
-	int32 PreviewMaxTextureSize = 2048;
+	int32 PreviewMaxTextureDimension = 2048;
+
+	// Performance: Stereoscopic rendering uses a single viewfamily(RTT) for both eyes
+	bool bEnableStereoscopicRenderingOptimization = false;
 
 	// Performance: Allow merge multiple viewports on single RTT with atlasing (required for bAllowViewFamilyMergeOptimization)
 	bool bAllowRenderTargetAtlasing = false;
diff --git a/Engine/Plugins/Runtime/nDisplay/Source/DisplayCluster/Private/Render/Viewport/RenderTarget/DisplayClusterRenderTargetManager.cpp b/Engine/Plugins/Runtime/nDisplay/Source/DisplayCluster/Private/Render/Viewport/RenderTarget/DisplayClusterRenderTargetManager.cpp
--- a/Engine/Plugins/Runtime/nDisplay/Source/DisplayCluster/Private/Render/Viewport/RenderTarget/DisplayClusterRenderTargetManager.cpp
+++ b/Engine/Plugins/Runtime/nDisplay/Source/DisplayCluster/Private/Render/Viewport/RenderTarget/DisplayClusterRenderTargetManager.cpp
@@ -85,19 +85,16 @@ bool FDisplayClusterRenderTargetManager::AllocateRenderFrameResources(FViewport*
 								FDisplayClusterViewport* ViewportPtr = static_cast<FDisplayClusterViewport*>(ViewIt.Viewport);
 								if (ViewportPtr)
 								{
-									// Array already resized in function FDisplayClusterViewport::UpdateFrameContexts() with RenderTargets.AddZeroed(ViewportContextAmmount);
+									// Array already resized in function FDisplayClusterViewport::UpdateFrameContexts() with RenderTargets.AddZeroed(ViewportContextAmount);
 									check(ViewIt.ContextNum < (uint32)ViewportPtr->RenderTargets.Num());
+									check(ViewportPtr->Contexts[ViewIt.ContextNum].bDisableRender == false);
 
-									// Support override
-									if (ViewportPtr->Contexts[ViewIt.ContextNum].bDisableRender == false)
-									{
-										ViewportPtr->RenderTargets[ViewIt.ContextNum] = NewResource;
+									ViewportPtr->RenderTargets[ViewIt.ContextNum] = NewResource;
 
-										if (NewResource->GetViewportResourceState(EDisplayClusterViewportResourceState::Initialized) == false)
-										{
-											// Log: New resource created
-											UE_LOG(LogDisplayClusterViewport, Verbose, TEXT("Created new ViewportRenderTarget resource %08X (%dx%d) for viewport '%s'"), NewResource, NewResource->GetSizeX(), NewResource->GetSizeY(), *ViewportPtr->GetId());
-										}
+									if (NewResource->GetViewportResourceState(EDisplayClusterViewportResourceState::Initialized) == false)
+									{
+										// Log: New resource created
+										UE_LOG(LogDisplayClusterViewport, Verbose, TEXT("Created new ViewportRenderTarget resource %08X (%dx%d) for viewport '%s'"), NewResource, NewResource->GetSizeX(), NewResource->GetSizeY(), *ViewportPtr->GetId());
 									}
 								}
 							}
@@ -110,32 +107,34 @@ bool FDisplayClusterRenderTargetManager::AllocateRenderFrameResources(FViewport*
 		// Allocate viewport internal resources
 		for (FDisplayClusterViewport* ViewportIt : InViewports)
 		{
-			const bool bShouldUseAdditionalTargetableResource = ViewportIt->AdditionalTargetableResources.Num() > 0;
-
-			EPixelFormat CustomFormat = ImplGetCustomFormat(ViewportIt->RenderSettings.CaptureMode);
-
-			// Allocate all context resources:
-			for (const FDisplayClusterViewport_Context& ContextIt : ViewportIt->GetContexts())
+			if (ViewportIt && ViewportIt->RenderSettings.bFreezeRendering == false)
 			{
+				EPixelFormat CustomFormat = ImplGetCustomFormat(ViewportIt->RenderSettings.CaptureMode);
 
-				const FIntPoint& ContextSize = ContextIt.ContextSize;
-				const uint32&     ContextNum = ContextIt.ContextNum;
+				// Allocate all context resources:
+				for (const FDisplayClusterViewport_Context& ContextIt : ViewportIt->GetContexts())
+				{
+					const FIntPoint& ContextSize = ContextIt.ContextSize;
+					const uint32& ContextNum = ContextIt.ContextNum;
 
-				check(ViewportIt->InputShaderResources.Num() > 0);
-				ViewportIt->InputShaderResources[ContextNum] = ResourcesPool->AllocateTextureResource(ContextSize, false, CustomFormat);
-				ImplViewportTextureResourceLogging(ViewportIt, ContextNum, TEXT("InputShader"), ViewportIt->InputShaderResources[ContextNum]);
+					if (ViewportIt->InputShaderResources.Num() > (int32)ContextNum)
+					{
+						ViewportIt->InputShaderResources[ContextNum] = ResourcesPool->AllocateTextureResource(ContextSize, false, CustomFormat);
+						ImplViewportTextureResourceLogging(ViewportIt, ContextNum, TEXT("InputShader"), ViewportIt->InputShaderResources[ContextNum]);
+					}
 
-				// Allocate custom resources:
-				if (bShouldUseAdditionalTargetableResource)
-				{
-					ViewportIt->AdditionalTargetableResources[ContextNum] = ResourcesPool->AllocateTextureResource(ContextSize, true, CustomFormat);
-					ImplViewportTextureResourceLogging(ViewportIt, ContextNum, TEXT("AdditionalRTT"), ViewportIt->AdditionalTargetableResources[ContextNum]);
-				}
+					// Allocate custom resources:
+					if (ViewportIt->AdditionalTargetableResources.Num() > (int32)ContextNum)
+					{
+						ViewportIt->AdditionalTargetableResources[ContextNum] = ResourcesPool->AllocateTextureResource(ContextSize, true, CustomFormat);
+						ImplViewportTextureResourceLogging(ViewportIt, ContextNum, TEXT("AdditionalRTT"), ViewportIt->AdditionalTargetableResources[ContextNum]);
+					}
 
-				if (ContextIt.NumMips > 1)
-				{
-					ViewportIt->MipsShaderResources[ContextNum] = ResourcesPool->AllocateTextureResource(ContextSize, false, CustomFormat, ContextIt.NumMips);
-					ImplViewportTextureResourceLogging(ViewportIt, ContextNum, TEXT("Mips"), ViewportIt->MipsShaderResources[ContextNum]);
+					if (ViewportIt->MipsShaderResources.Num() > (int32)ContextNum)
+					{
+						ViewportIt->MipsShaderResources[ContextNum] = ResourcesPool->AllocateTextureResource(ContextSize, false, CustomFormat, ContextIt.NumMips);
+						ImplViewportTextureResourceLogging(ViewportIt, ContextNum, TEXT("Mips"), ViewportIt->MipsShaderResources[ContextNum]);
+					}
 				}
 			}
 		}
diff --git a/Engine/Plugins/Runtime/nDisplay/Source/DisplayCluster/Private/Render/Viewport/RenderTarget/DisplayClusterRenderTargetResource.h b/Engine/Plugins/Runtime/nDisplay/Source/DisplayCluster/Private/Render/Viewport/RenderTarget/DisplayClusterRenderTargetResource.h
--- a/Engine/Plugins/Runtime/nDisplay/Source/DisplayCluster/Private/Render/Viewport/RenderTarget/DisplayClusterRenderTargetResource.h
+++ b/Engine/Plugins/Runtime/nDisplay/Source/DisplayCluster/Private/Render/Viewport/RenderTarget/DisplayClusterRenderTargetResource.h
@@ -78,6 +78,9 @@ enum class EDisplayClusterViewportResourceState : uint8
 	Unused,
 	Deleted,
 
+	// Disable reallocation of this resource for the current frame
+	DisableReallocate,
+
 	COUNT
 };
 ///////////////////////////////////////////////////////////////////
diff --git a/Engine/Plugins/Runtime/nDisplay/Source/DisplayCluster/Private/Render/Viewport/RenderTarget/DisplayClusterRenderTargetResourcesPool.cpp b/Engine/Plugins/Runtime/nDisplay/Source/DisplayCluster/Private/Render/Viewport/RenderTarget/DisplayClusterRenderTargetResourcesPool.cpp
--- a/Engine/Plugins/Runtime/nDisplay/Source/DisplayCluster/Private/Render/Viewport/RenderTarget/DisplayClusterRenderTargetResourcesPool.cpp
+++ b/Engine/Plugins/Runtime/nDisplay/Source/DisplayCluster/Private/Render/Viewport/RenderTarget/DisplayClusterRenderTargetResourcesPool.cpp
@@ -2,6 +2,8 @@
 
 #include "DisplayClusterRenderTargetResourcesPool.h"
 #include "Render/Viewport/RenderFrame/DisplayClusterRenderFrameSettings.h"
+#include "Render/Viewport/DisplayClusterViewportHelpers.h"
+
 #include "RHICommandList.h"
 #include "Engine/RendererSettings.h"
 #include "RenderingThread.h"
@@ -55,41 +57,16 @@ namespace DisplayClusterRenderTargetResourcesPool
 			}
 		}
 
-		// Get default settings for preview rendering:
-		static TConsoleVariableData<int32>* CVarDefaultBackBufferPixelFormat = IConsoleManager::Get().FindTConsoleVariableDataInt(TEXT("r.DefaultBackBufferPixelFormat"));
-		EPixelFormat Format = EDefaultBackBufferPixelFormat::Convert2PixelFormat(EDefaultBackBufferPixelFormat::FromInt(CVarDefaultBackBufferPixelFormat->GetValueOnGameThread()));
-
-		// default gamma for preview in editor
-		float Gamma = 2.2f;
-		switch (InRenderFrameSettings.RenderMode)
-		{
-		case EDisplayClusterRenderFrameMode::PreviewInScene:
-			Gamma = 1.f;
-			break;
-		default:
-			break;
-		}
-
-		// By default sRGB disabled (Preview rendering)
-		bool bShouldUseSRGB = false;
+#if WITH_EDITOR
+		EPixelFormat PreviewPixelFormat;
+		float PreviewDisplayGamma = 1.f;
+		bool bPreviewSRGB = true;
+		DisplayClusterViewportHelpers::GetPreviewRenderTargetDesc_Editor(InRenderFrameSettings, PreviewPixelFormat, PreviewDisplayGamma, bPreviewSRGB);
 
-		return new FDisplayClusterViewportResourceSettings(InRenderFrameSettings.ClusterNodeId, Format, bShouldUseSRGB, Gamma);
-	}
-
-	static bool IsTextureSizeValidImpl(const FIntPoint& InSize)
-	{
-		static const int32 MaxTextureSize = 1 << (GMaxTextureMipCount - 1);
-
-		// just check the texture size is valid
-		if (InSize.X <= 0 || InSize.Y <= 0 || InSize.X > MaxTextureSize || InSize.Y > MaxTextureSize)
-		{
-			return false;
-		}
-
-
-		//@todo - maybe check free size of video memory before texture allocation
-
-		return true;
+		return new FDisplayClusterViewportResourceSettings(InRenderFrameSettings.ClusterNodeId, PreviewPixelFormat, bPreviewSRGB, PreviewDisplayGamma);
+#else
+		return nullptr;
+#endif
 	}
 };
 
@@ -117,9 +94,16 @@ void FDisplayClusterRenderTargetResourcesPool::ImplBeginReallocateResources(TArr
 	// Mark cluster node resources as unused
 	for (TViewportResourceType* ResourceIt : InOutViewportResources)
 	{
-		if (ResourceIt != nullptr && ResourceIt->GetResourceSettingsConstRef().IsClusterNodeNameEqual(*ResourceSettings))
+		if (ResourceIt != nullptr && ResourceIt->GetViewportResourceState(EDisplayClusterViewportResourceState::DisableReallocate) == false)
 		{
-			ResourceIt->RaiseViewportResourceState(EDisplayClusterViewportResourceState::Unused);
+			if (const bool bIsCurrentClusterNodeResource = ResourceIt->GetResourceSettingsConstRef().IsClusterNodeNameEqual(*ResourceSettings))
+			{
+				ResourceIt->RaiseViewportResourceState(EDisplayClusterViewportResourceState::Unused);
+			}
+			else
+			{
+				ResourceIt->RaiseViewportResourceState(EDisplayClusterViewportResourceState::DisableReallocate);
+			}
 		}
 	}
 }
@@ -149,6 +133,9 @@ void FDisplayClusterRenderTargetResourcesPool::ImplFinishReallocateResources(TAr
 				UnusedResources.Add(ResourceIt);
 				ResourceIt->RaiseViewportResourceState(EDisplayClusterViewportResourceState::Deleted);
 			}
+
+			// Clear 'DisableReallocate' at frame end
+			ResourceIt->ClearViewportResourceState(EDisplayClusterViewportResourceState::DisableReallocate);
 		}
 	}
 
@@ -190,7 +177,7 @@ void FDisplayClusterRenderTargetResourcesPool::ImplReleaseResources(TArray<TView
 template <typename TViewportResourceType>
 TViewportResourceType* FDisplayClusterRenderTargetResourcesPool::ImplAllocateResource(TArray<TViewportResourceType*>& InOutViewportResources, const FDisplayClusterViewportResourceSettings& InSettings)
 {
-	if (!IsTextureSizeValidImpl(InSettings.Size))
+	if (!DisplayClusterViewportHelpers::IsValidTextureSize(InSettings.Size))
 	{
 		return nullptr;
 	}
diff --git a/Engine/Plugins/Runtime/nDisplay/Source/DisplayCluster/Public/Components/DisplayClusterPreviewComponent.h b/Engine/Plugins/Runtime/nDisplay/Source/DisplayCluster/Public/Components/DisplayClusterPreviewComponent.h
--- a/Engine/Plugins/Runtime/nDisplay/Source/DisplayCluster/Public/Components/DisplayClusterPreviewComponent.h
+++ b/Engine/Plugins/Runtime/nDisplay/Source/DisplayCluster/Public/Components/DisplayClusterPreviewComponent.h
@@ -12,7 +12,7 @@
 
 class IDisplayClusterProjectionPolicy;
 class UTextureRenderTarget2D;
-class UTexture2D;
+class UTexture;
 class UDisplayClusterConfigurationViewport;
 class UMaterial;
 class UMaterialInstanceDynamic;
@@ -42,8 +42,6 @@ public:
 public:
 	bool InitializePreviewComponent(ADisplayClusterRootActor* RootActor, const FString& ClusterNodeId, const FString& ViewportId, UDisplayClusterConfigurationViewport* ViewportConfig);
 
-	bool IsPreviewAvailable() const;
-
 	const FString& GetViewportId() const
 	{ 
 		return ViewportId; 
@@ -65,49 +63,41 @@ public:
 	}
 
 	void UpdatePreviewResources();
-	void HandleRenderTargetTextureDeferredUpdate(const EDisplayClusterRenderFrameMode InRenderFrameMode);
 
-	UMeshComponent* GetPreviewMesh() const
+	UMeshComponent* GetPreviewMesh()
 	{
+		UpdatePreviewMeshReference();
+
 		return PreviewMesh;
 	}
 
 	/** Create and retrieve a render texture 2d from the render target. */
-	UTexture2D* GetOrCreateViewportPreviewTexture2D();
+	UTexture* GetViewportPreviewTexture2D();
 
-	/** Return the render texture 2d but do not recreate it. */
-	UTexture2D* GetViewportPreviewTexture2D() const;
+	void ResetPreviewComponent(bool bInRestoreSceneMaterial);
 
 	/** Sets an override texture to display on the viewport instead of the render target */
 	void SetOverrideTexture(UTexture* InOverrideTexture);
 
 protected:
+	bool IsPreviewEnabled() const;
+
 	class IDisplayClusterViewport* GetCurrentViewport() const;
-	bool GetPreviewTextureSettings(FIntPoint& OutSize, float& OutGamma) const;
+	bool GetPreviewTextureSettings(FIntPoint& OutSize, EPixelFormat& OutTextureFormat, float& OutGamma, bool& bOutSRGB) const;
 
 	void UpdatePreviewRenderTarget();
 	void ReleasePreviewRenderTarget();
 
-	bool UpdatePreviewMesh(bool bRestoreOriginalMaterial = false);
+	bool UpdatePreviewMesh();
 	void ReleasePreviewMesh();
 	void UpdatePreviewMeshReference();
 
-	bool UpdatePreviewTexture();
-	void ReleasePreviewTexture();
-
 	void InitializePreviewMaterial();
 	void ReleasePreviewMaterial();
 	void UpdatePreviewMaterial();
 
 	void RestorePreviewMeshMaterial();
 	void SetPreviewMeshMaterial();
-
-	bool IsPreviewEnabled() const;
-
-protected:
-	// Set to true, when RTT surface updated
-	int32 RenderTargetSurfaceChangedCnt = 0;
-	bool bIsEditingProperty = false;
 #endif /* WITH_EDITOR */
 
 #if WITH_EDITORONLY_DATA
@@ -148,9 +138,6 @@ private:
 	UPROPERTY()
 	UMaterialInstanceDynamic* PreviewMaterialInstance = nullptr;
 
-	UPROPERTY(Transient)
-	UTexture2D* PreviewTexture = nullptr;
-
 	UPROPERTY(Transient)
 	UTexture* OverrideTexture = nullptr;
 
diff --git a/Engine/Plugins/Runtime/nDisplay/Source/DisplayCluster/Public/DisplayClusterRootActor.h b/Engine/Plugins/Runtime/nDisplay/Source/DisplayCluster/Public/DisplayClusterRootActor.h
--- a/Engine/Plugins/Runtime/nDisplay/Source/DisplayCluster/Public/DisplayClusterRootActor.h
+++ b/Engine/Plugins/Runtime/nDisplay/Source/DisplayCluster/Public/DisplayClusterRootActor.h
@@ -170,7 +170,7 @@ protected:
 // Placed here to ensure layout builders process referencers first
 //////////////////////////////////////////////////////////////////////////////////////////////
 #if WITH_EDITORONLY_DATA
-private:	
+private:
 	UPROPERTY(EditAnywhere, Transient, Category = Viewports, meta = (PropertyPath = "CurrentConfigData.RenderFrameSettings.ClusterICVFXOuterViewportBufferRatioMult"))
 	FDisplayClusterEditorPropertyReference ViewportScreenPercentageMultiplierRef;
 
@@ -180,6 +180,9 @@ private:
 	UPROPERTY(EditInstanceOnly, Transient, Category = Viewports, meta = (DisplayName = "Viewport Overscan", PropertyPath = "CurrentConfigData.Cluster.Nodes.Viewports.RenderSettings.Overscan", ToolTip = "Render a larger frame than specified in the configuration to achieve continuity across displays when using post-processing effects."))
 	FDisplayClusterEditorPropertyReference ViewportOverscanRef;
 
+	UPROPERTY(EditAnywhere, Transient, Category = Viewports, meta = (PropertyPath = "CurrentConfigData.StageSettings.bFreezeRenderOuterViewports"))
+	FDisplayClusterEditorPropertyReference FreezeRenderOuterViewportsRef;
+
 	UPROPERTY(EditAnywhere, Transient, Category = Viewports, meta = (PropertyPath = "CurrentConfigData.StageSettings.HideList"))
 	FDisplayClusterEditorPropertyReference ClusterHideListRef;
 
@@ -291,55 +294,45 @@ public:
 	UPROPERTY(EditAnywhere, BlueprintReadWrite, Category = "Editor Preview", meta = (DisplayName = "Enable Editor Preview"))
 	bool bPreviewEnable = true;
 	
-	/** render preview every frame */
-	UPROPERTY(EditAnywhere, BlueprintReadWrite, Category = "Editor Preview", meta = (DisplayName = "Capture Every Frame", EditCondition = "bPreviewEnable"))
-	bool bPreviewRenderEveryFrame = true;
+	/** Adjust resolution scaling for the editor preview. */
+	UPROPERTY(EditAnywhere, BlueprintReadWrite, Category = "Editor Preview", meta = (DisplayName = "Preview Screen Percentage", ClampMin = "0.05", UIMin = "0.05", ClampMax = "1", UIMax = "1", EditCondition = "bPreviewEnable"))
+	float PreviewRenderTargetRatioMult = 0.25;
+
+	/** Enable PostProcess for preview. */
+	UPROPERTY(EditAnywhere, BlueprintReadWrite, Category = "Editor Preview", meta = (DisplayName = "Enable Post Process", EditCondition = "bPreviewEnable"))
+	bool bPreviewEnablePostProcess = false;
+
+	/** Freeze preview render.  This will impact editor performance. */
+	UPROPERTY(EditAnywhere, BlueprintReadWrite, Category = "Editor Preview", meta = (DisplayName = "Freeze Editor Preview", EditCondition = "bPreviewEnable"))
+	bool bFreezePreviewRender = false;
+
+	/** Render ICVFX Frustums */
+	UPROPERTY(EditAnywhere, BlueprintReadWrite, Category = "Editor Preview", meta = (DisplayName = "Enable Camera Frustums", EditCondition = "bPreviewEnable"))
+	bool bPreviewICVFXFrustums = false;
+
+	/** Render ICVFX Frustums */
+	UPROPERTY(EditAnywhere, BlueprintReadWrite, Category = "Editor Preview", meta = (DisplayName = "Camera Frustum Distance", EditCondition = "bPreviewEnable"))
+	float PreviewICVFXFrustumsFarDistance = 1000.0f;
 
 	/** Selectively preview a specific viewport or show all/none. */
 	UPROPERTY(EditAnywhere, BlueprintReadWrite, Category = "Editor Preview", meta = (DisplayName = "Preview Node", EditCondition = "bPreviewEnable"))
-	FString PreviewNodeId = DisplayClusterConfigurationStrings::gui::preview::PreviewNodeAll;
+	FString PreviewNodeId = DisplayClusterConfigurationStrings::gui::preview::PreviewNodeNone;
 
 	/** Render Mode */
-	UPROPERTY(EditAnywhere, BlueprintReadWrite, Category = "Editor Preview", meta = (EditCondition = "bPreviewEnable"))
+	UPROPERTY(EditAnywhere, BlueprintReadWrite, Category = "Editor Preview", meta = (DisplayName = "Render Mode", EditCondition = "bPreviewEnable"))
 	EDisplayClusterConfigurationRenderMode RenderMode = EDisplayClusterConfigurationRenderMode::Mono;
 
-	/** Render with mGPU */
-	UPROPERTY(EditAnywhere, BlueprintReadWrite, Category = "Editor Preview", meta = (EditCondition = "bPreviewEnable"))
-	bool bAllowMultiGPURendering = true;
-
-	UPROPERTY(EditAnywhere, BlueprintReadWrite, Category = "Editor Preview", meta = (EditCondition = "bPreviewEnable"))
-	int MinGPUIndex = 1;
-
-	UPROPERTY(EditAnywhere, BlueprintReadWrite, Category = "Editor Preview", meta = (EditCondition = "bPreviewEnable"))
-	int MaxGPUIndex = 1;
-
 	/** Tick Per Frame */
-	UPROPERTY(EditAnywhere, BlueprintReadWrite, Category = "Editor Preview", meta = (ClampMin = "1", UIMin = "1", ClampMax = "200", UIMax = "200", EditCondition = "bPreviewEnable"))
+	UPROPERTY(EditAnywhere, BlueprintReadWrite, Category = "Editor Preview", AdvancedDisplay, meta = (DisplayName = "Tick Per Frame", ClampMin = "1", UIMin = "1", ClampMax = "200", UIMax = "200", EditCondition = "bPreviewEnable"))
 	int TickPerFrame = 1;
 
-	/** Nodes Per Frame */
-	UPROPERTY(EditAnywhere, BlueprintReadWrite, Category = "Editor Preview", meta = (ClampMin = "1", UIMin = "1", ClampMax = "200", UIMax = "200", EditCondition = "bPreviewEnable"))
-	int NodesPerFrame = 1;
-
-	/** Viewports Per Frame */
-	UPROPERTY(EditAnywhere, BlueprintReadWrite, Category = "Editor Preview", meta = (ClampMin = "1", UIMin = "1", ClampMax = "200", UIMax = "200", EditCondition = "bPreviewEnable"))
+	/** Max amount of Viewports Per Frame */
+	UPROPERTY(EditAnywhere, BlueprintReadWrite, Category = "Editor Preview", AdvancedDisplay, meta = (DisplayName = "Viewports Per Frame", ClampMin = "1", UIMin = "1", ClampMax = "200", UIMax = "200", EditCondition = "bPreviewEnable"))
 	int ViewportsPerFrame = 1;
 
-	/** Adjust resolution scaling for the editor preview. */
-	UPROPERTY(EditAnywhere, BlueprintReadWrite, Category = "Editor Preview", meta = (DisplayName = "Preview Screen Percentage", ClampMin = "0.05", UIMin = "0.05", ClampMax = "1", UIMax = "1", EditCondition = "bPreviewEnable"))
-	float PreviewRenderTargetRatioMult = 0.25;
-
-	/** Render ICVFX Frustums */
-	UPROPERTY(EditAnywhere, BlueprintReadWrite, Category = "Editor Preview", meta = (DisplayName = "ICVFX Camera Frustums", EditCondition = "bPreviewEnable"))
-	bool bPreviewICVFXFrustums = false;
-
-	/** Render ICVFX Frustums */
-	UPROPERTY(EditAnywhere, BlueprintReadWrite, Category = "Editor Preview", meta = (DisplayName = "ICVFX Camera Frustums Distance", EditCondition = "bPreviewEnable"))
-	float PreviewICVFXFrustumsFarDistance = 1000.0f;
-
 	/** The maximum dimension of any internal texture for preview. Use less memory for large preview viewports */
-	UPROPERTY(EditAnywhere, BlueprintReadWrite, Category = "Editor Preview", meta = (DisplayName = "Preview Texture Max Size", ClampMin = "64", UIMin = "64", ClampMax = "4096", UIMax = "4096", EditCondition = "bPreviewEnable"))
-	int PreviewMaxTextureSize = 2048;
+	UPROPERTY(EditAnywhere, BlueprintReadWrite, Category = "Editor Preview", AdvancedDisplay, meta = (DisplayName = "Preview Texture Max Size", ClampMin = "64", UIMin = "64", ClampMax = "4096", UIMax = "4096", EditCondition = "bPreviewEnable"))
+	int PreviewMaxTextureDimension = 2048;
 
 private:
 	UPROPERTY(Transient)
@@ -372,11 +365,15 @@ public:
 	void BeginDestroy_Editor();
 	void RerunConstructionScripts_Editor();
 
+	// Preview components free referenced meshes and materials
+	void ResetPreviewComponents_Editor(bool bInRestoreSceneMaterial);
+
 	UDisplayClusterPreviewComponent* GetPreviewComponent(const FString& NodeId, const FString& ViewportId);
 
 	void UpdatePreviewComponents();
 	void ReleasePreviewComponents();
 
+
 	float GetPreviewRenderTargetRatioMult() const
 	{
 		return PreviewRenderTargetRatioMult;
@@ -384,7 +381,7 @@ public:
 
 	IDisplayClusterViewport* FindPreviewViewport(const FString& InViewportId) const;
 
-	void GetPreviewRenderTargetableTextures(const EDisplayClusterRenderFrameMode InRenderFrameMode, const TArray<FString>& InViewportNames, TArray<FTextureRHIRef>& OutTextures);
+	void GetPreviewRenderTargetableTextures(const TArray<FString>& InViewportNames, TArray<FTextureRHIRef>& OutTextures);
 
 	void UpdateInnerFrustumPriority();
 	void ResetInnerFrustumPriority();
@@ -399,17 +396,18 @@ public:
 	}
 
 protected:
-	FString GeneratePreviewComponentName(const FString& NodeId, const FString& ViewportId) const;
+	FString GeneratePreviewComponentName_Editor(const FString& NodeId, const FString& ViewportId) const;
+	void ResetPreviewInternals_Editor();
 
-	void ImplRenderPreviewForSceneMaterials_Editor(const int32 InNumNodesForRender);
+	bool ImplUpdatePreviewConfiguration_Editor(const FString& InClusterNodeId);
 
-	void RenderPreviewClusterNode_Editor(const EDisplayClusterRenderFrameMode InRenderFrameMode, const FString& InClusterNodeId);
-	bool UpdatePreviewConfiguration_Editor(const EDisplayClusterRenderFrameMode InRenderFrameMode, const FString& InClusterNodeId);
+	void ImplRenderPreview_Editor();
+	bool ImplRenderPassPreviewClusterNode_Editor();
 
-	void ResetPreviewInternals();
-	
-	void RenderPreviewFrustums();
-	void RenderPreviewFrustum(const FMatrix ProjectionMatrix, const FMatrix ViewMatrix, const FVector ViewOrigin);
+	bool ImplUpdatePreviewRenderFrame_Editor(const FString& InClusterNodeId);
+
+	void ImplRenderPreviewFrustums_Editor();
+	void ImplRenderPreviewViewportFrustum_Editor(const FMatrix ProjectionMatrix, const FMatrix ViewMatrix, const FVector ViewOrigin);
 
 	virtual void PostEditChangeProperty(FPropertyChangedEvent& PropertyChangedEvent) override;
 	virtual void PostEditMove(bool bFinished) override;
@@ -425,9 +423,20 @@ private:
 	int32 PreviewClusterNodeIndex = 0;
 	int32 PreviewViewportIndex = 0;
 	TUniquePtr<FDisplayClusterRenderFrame> PreviewRenderFrame;
+	FString PreviewRenderFrameClusterNodeId;
+	int32 PreviewViewportsRenderedInThisFrameCnt = 0;
 
 	FOnPreviewUpdated OnPreviewGenerated;
 	FOnPreviewUpdated OnPreviewDestroyed;
 
+	struct FFrustumPreviewViewportContextCache
+	{
+		FVector  ViewLocation;
+		FRotator ViewRotation;
+		FMatrix  ProjectionMatrix;
+	};
+	// Cache the last valid viewport context
+	TMap<FString, FFrustumPreviewViewportContextCache> FrustumPreviewViewportContextCache;
+
 #endif
 };
diff --git a/Engine/Plugins/Runtime/nDisplay/Source/DisplayCluster/Public/Render/Viewport/Containers/DisplayClusterPreviewSettings.h b/Engine/Plugins/Runtime/nDisplay/Source/DisplayCluster/Public/Render/Viewport/Containers/DisplayClusterPreviewSettings.h
--- a/Engine/Plugins/Runtime/nDisplay/Source/DisplayCluster/Public/Render/Viewport/Containers/DisplayClusterPreviewSettings.h
+++ b/Engine/Plugins/Runtime/nDisplay/Source/DisplayCluster/Public/Render/Viewport/Containers/DisplayClusterPreviewSettings.h
@@ -9,13 +9,23 @@ struct FDisplayClusterPreviewSettings
 	// Preview RTT size multiplier
 	float PreviewRenderTargetRatioMult = 1.f;
 
+	// Enable/Disable preview rendering. When disabled preview image freeze
+	bool bFreezePreviewRender = false;
+
+	// Hack preview gamma.
+	// In a scene, PostProcess always renders on top of the preview textures.
+	// But in it, PostProcess is also rendered with the flag turned off.
+	bool bPreviewEnablePostProcess = false;
+
 	// The maximum dimension of any texture for preview
-	int32 PreviewMaxTextureSize = 2048;
+	int32 PreviewMaxTextureDimension = 2048;
 
 	// Allow mGPU in editor mode
 	bool bAllowMultiGPURenderingInEditor = false;
 
+	// Special flag for PIE preview rendering
+	bool bIsPIE = false;
+
 	int32 MinGPUIndex = 1;
 	int32 MaxGPUIndex = 1;
 };
-
diff --git a/Engine/Plugins/Runtime/nDisplay/Source/DisplayCluster/Public/Render/Viewport/Containers/DisplayClusterViewport_Context.h b/Engine/Plugins/Runtime/nDisplay/Source/DisplayCluster/Public/Render/Viewport/Containers/DisplayClusterViewport_Context.h
--- a/Engine/Plugins/Runtime/nDisplay/Source/DisplayCluster/Public/Render/Viewport/Containers/DisplayClusterViewport_Context.h
+++ b/Engine/Plugins/Runtime/nDisplay/Source/DisplayCluster/Public/Render/Viewport/Containers/DisplayClusterViewport_Context.h
@@ -17,15 +17,19 @@ public:
 public:
 	const uint32            ContextNum;
 
-	const EStereoscopicPass StereoscopicPass;
-	const int32             StereoViewIndex;
+	EStereoscopicPass StereoscopicPass;
+	int32             StereoViewIndex;
 
 	// Camera location and orientation
 	FVector  ViewLocation = FVector::ZeroVector;
+	bool bIsValidViewLocation = false;
+
 	FRotator ViewRotation = FRotator::ZeroRotator;
+	bool bIsValidViewRotation = false;
 
 	// Projection Matrix
 	FMatrix ProjectionMatrix = FMatrix::Identity;
+	bool bIsValidProjectionMatrix = false;
 
 	// Overscan Projection Matrix (internal use)
 	FMatrix OverscanProjectionMatrix = FMatrix::Identity;
@@ -51,6 +55,9 @@ public:
 	// Location and size on a frame target texture
 	FIntRect FrameTargetRect;
 
+	// Buffer ratio
+	float CustomBufferRatio = 1;
+
 	// Mips number for additional MipsShader resources
 	int32 NumMips = 1;
 
diff --git a/Engine/Plugins/Runtime/nDisplay/Source/DisplayCluster/Public/Render/Viewport/Containers/DisplayClusterViewport_RenderSettings.h b/Engine/Plugins/Runtime/nDisplay/Source/DisplayCluster/Public/Render/Viewport/Containers/DisplayClusterViewport_RenderSettings.h
--- a/Engine/Plugins/Runtime/nDisplay/Source/DisplayCluster/Public/Render/Viewport/Containers/DisplayClusterViewport_RenderSettings.h
+++ b/Engine/Plugins/Runtime/nDisplay/Source/DisplayCluster/Public/Render/Viewport/Containers/DisplayClusterViewport_RenderSettings.h
@@ -36,6 +36,15 @@ public:
 	// Skip rendering for this viewport
 	bool bSkipRendering = false;
 
+	// Freeze viewport resources, skip rendering internal viewport resources. But still use it for final compositing
+	bool bFreezeRendering = false;
+
+	// Render alpha channel from input texture to warp output
+	bool bWarpBlendRenderAlphaChannel = false;
+
+	// Read viewport pixels for preview (this flag is cleared at the end of the frame)
+	bool bPreviewReadPixels = false;
+
 	// Useful to render some viewports in mono, then copied to stereo backbuffers identical image
 	bool bForceMono = false;
 
@@ -48,9 +57,12 @@ public:
 	// Allow ScreenPercentage 
 	float BufferRatio = 1;
 
-	// Performance: Render to scale RTT, resolved with shader to viewport (Custom value)
+	// Performance: Render target base resolution multiplier
 	float RenderTargetRatio = 1;
 
+	// Performance: Render target adaptive resolution multiplier
+	float RenderTargetAdaptRatio = 1;
+
 	// Viewport can overlap each other on backbuffer. This value uses to sorting order
 	int32 OverlapOrder = 0;
 
@@ -71,12 +83,19 @@ public:
 		bVisible = true;
 		bEnable = true;
 		bSkipRendering = false;
+		bFreezeRendering = false;
+		bWarpBlendRenderAlphaChannel = false;
 
 		CaptureMode = EDisplayClusterViewportCaptureMode::Default;
 
 		OverrideViewportId.Empty();
 	}
 
+	inline void FinishUpdateSettings()
+	{
+		bPreviewReadPixels = false;
+	}
+
 	inline const FString& GetParentViewportId() const
 	{
 		return ParentViewportId;
diff --git a/Engine/Plugins/Runtime/nDisplay/Source/DisplayCluster/Public/Render/Viewport/IDisplayClusterViewportManager.h b/Engine/Plugins/Runtime/nDisplay/Source/DisplayCluster/Public/Render/Viewport/IDisplayClusterViewportManager.h
--- a/Engine/Plugins/Runtime/nDisplay/Source/DisplayCluster/Public/Render/Viewport/IDisplayClusterViewportManager.h
+++ b/Engine/Plugins/Runtime/nDisplay/Source/DisplayCluster/Public/Render/Viewport/IDisplayClusterViewportManager.h
@@ -103,11 +103,23 @@ public:
 	virtual void AddReferencedObjects(FReferenceCollector& Collector) = 0;
 
 #if WITH_EDITOR
-	virtual bool RenderInEditor(class FDisplayClusterRenderFrame& InRenderFrame, FViewport* InViewport, const uint32 InFirstViewportNum, const int32 InViewportsAmmount, bool& bOutFrameRendered) = 0;
+	/**
+	* Render in editor (preview)
+	* [Game thread func]
+	*
+	* @param InRenderFrame - render frame setup
+	* @param InViewport
+	* @param InFirstViewportNum - begin render from this viewport in frame
+	* @param InViewportsAmount - max viewports for render
+	* @param OutViewportsAmount - total viewport rendered
+	* @param bOutFrameRendered - true, if cluster node composition pass done (additional pass after last viewport is rendered)
+	*
+	* @return - true, if render success
+	*/
+	virtual bool RenderInEditor(class FDisplayClusterRenderFrame& InRenderFrame, FViewport* InViewport, const uint32 InFirstViewportNum, const int32 InViewportsAmount, int32& OutViewportsAmount, bool& bOutFrameRendered) = 0;
 #endif
 
 	///////////////////////////////////////////////////////////////////////////////////////////////////////////////////////////////////////////////////////////////////////////
-
 	/**
 	* Find viewport object by name
 	* [Game thread func]
diff --git a/Engine/Plugins/Runtime/nDisplay/Source/DisplayCluster/Public/Render/Viewport/RenderFrame/DisplayClusterRenderFrame.h b/Engine/Plugins/Runtime/nDisplay/Source/DisplayCluster/Public/Render/Viewport/RenderFrame/DisplayClusterRenderFrame.h
--- a/Engine/Plugins/Runtime/nDisplay/Source/DisplayCluster/Public/Render/Viewport/RenderFrame/DisplayClusterRenderFrame.h
+++ b/Engine/Plugins/Runtime/nDisplay/Source/DisplayCluster/Public/Render/Viewport/RenderFrame/DisplayClusterRenderFrame.h
@@ -19,6 +19,12 @@ public:
 public:
 	class FFrameView
 	{
+	public:
+		inline bool IsShouldRenderView() const
+		{
+			return !bDisableRender && !bFreezeRendering;
+		}
+
 	public:
 		// Viewport game-thread data
 		class IDisplayClusterViewport* Viewport = nullptr;
@@ -27,6 +33,7 @@ public:
 		uint32 ContextNum = 0;
 
 		bool bDisableRender = false;
+		bool bFreezeRendering = false;
 	};
 
 	class FFrameViewFamily
@@ -84,7 +91,7 @@ public:
 	FIntRect FrameRect;
 
 	int32 DesiredNumberOfViews = 0;
-	int32 ViewportsAmmount = 0;
+	int32 ViewportsAmount = 0;
 
 	IDisplayClusterViewportManager* ViewportManager = nullptr;
 };
diff --git a/Engine/Plugins/Runtime/nDisplay/Source/DisplayClusterConfiguration/Public/DisplayClusterConfigurationTypes_ICVFX.h b/Engine/Plugins/Runtime/nDisplay/Source/DisplayClusterConfiguration/Public/DisplayClusterConfigurationTypes_ICVFX.h
--- a/Engine/Plugins/Runtime/nDisplay/Source/DisplayClusterConfiguration/Public/DisplayClusterConfigurationTypes_ICVFX.h
+++ b/Engine/Plugins/Runtime/nDisplay/Source/DisplayClusterConfiguration/Public/DisplayClusterConfigurationTypes_ICVFX.h
@@ -239,6 +239,10 @@ public:
 	UPROPERTY(EditAnywhere, BlueprintReadWrite, Category = "Light Cards", meta = (DisplayName = "Enable Light Cards"))
 	bool bEnable = true;
 
+	/** Enable\Disable freeze rendering for lightcards when outer viewports rendering also freezed. This will impact performance. */
+	UPROPERTY(EditAnywhere, BlueprintReadWrite, Category = "Light Cards", meta = (DisplayName = "Ignore Freezing of Outer Viewports for Lightcards"))
+	bool bIgnoreOuterViewportsFreezingForLightcards = true;
+
 	/** Specify how to render Light Cards in relation to the inner frustum. */
 	UPROPERTY(EditAnywhere, BlueprintReadWrite, Category = "Light Cards", meta = (DisplayName = "Blending Mode", EditCondition = "bEnable"))
 	EDisplayClusterConfigurationICVFX_LightcardRenderMode Blendingmode = EDisplayClusterConfigurationICVFX_LightcardRenderMode::Under;
@@ -544,6 +548,10 @@ public:
 	UPROPERTY(EditAnywhere, BlueprintReadWrite, Category = "Light Cards", meta = (ExpandProperties))
 	FDisplayClusterConfigurationICVFX_LightcardSettings Lightcard;
 
+	/** Freeze rendering for viewports. This improves performance. */
+	UPROPERTY(EditAnywhere, BlueprintReadWrite, Category = Viewports, meta = (DisplayName = "Freeze Viewports"))
+	bool bFreezeRenderOuterViewports = false;
+
 	// Hide list for all icvfx viewports (outer, inner, cameras, etc)
 	// (This allow to hide all actors from layers for icvfx render logic)
 	UPROPERTY(EditAnywhere, BlueprintReadWrite, Category = Viewports, meta = (DisplayName = "Content Hidden from Entire Cluster", ToolTip = "Content specified here will not appear anywhere in the nDisplay cluster", Substitutions = "LayersTooltip = Layers hidden from the entire nDisplay cluster, ActorsTooltip = Actors hidden from the entire nDisplay cluster"))
diff --git a/Engine/Plugins/Runtime/nDisplay/Source/DisplayClusterConfigurator/Private/DisplayClusterConfiguratorCompiler.cpp b/Engine/Plugins/Runtime/nDisplay/Source/DisplayClusterConfigurator/Private/DisplayClusterConfiguratorCompiler.cpp
--- a/Engine/Plugins/Runtime/nDisplay/Source/DisplayClusterConfigurator/Private/DisplayClusterConfiguratorCompiler.cpp
+++ b/Engine/Plugins/Runtime/nDisplay/Source/DisplayClusterConfigurator/Private/DisplayClusterConfiguratorCompiler.cpp
@@ -126,7 +126,7 @@ void FDisplayClusterConfiguratorKismetCompilerContext::CopyTermDefaultsToDefault
 	ADisplayClusterRootActor* RootActor = CastChecked<ADisplayClusterRootActor>(DefaultObject);
 	if (Blueprint->bIsNewlyCreated)
 	{
-		RootActor->PreviewNodeId = DisplayClusterConfigurationStrings::gui::preview::PreviewNodeAll;
+		RootActor->PreviewNodeId = DisplayClusterConfigurationStrings::gui::preview::PreviewNodeNone;
 	}
 	
 	// Restore the old config data back to the CDO to preserve transaction history.
diff --git a/Engine/Plugins/Runtime/nDisplay/Source/DisplayClusterConfigurator/Private/DisplayClusterConfiguratorEditorSubsystem.cpp b/Engine/Plugins/Runtime/nDisplay/Source/DisplayClusterConfigurator/Private/DisplayClusterConfiguratorEditorSubsystem.cpp
--- a/Engine/Plugins/Runtime/nDisplay/Source/DisplayClusterConfigurator/Private/DisplayClusterConfiguratorEditorSubsystem.cpp
+++ b/Engine/Plugins/Runtime/nDisplay/Source/DisplayClusterConfigurator/Private/DisplayClusterConfiguratorEditorSubsystem.cpp
@@ -25,7 +25,7 @@ UDisplayClusterBlueprint* UDisplayClusterConfiguratorEditorSubsystem::ImportAsse
 	{
 		return nullptr;
 	}
-	RootActor->PreviewNodeId = DisplayClusterConfigurationStrings::gui::preview::PreviewNodeAll;
+	RootActor->PreviewNodeId = DisplayClusterConfigurationStrings::gui::preview::PreviewNodeNone;
 
 	UDisplayClusterBlueprint* NewBlueprint = FDisplayClusterConfiguratorUtils::CreateBlueprintFromRootActor(RootActor, InName, InParent);
 	if (NewBlueprint == nullptr)
diff --git a/Engine/Plugins/Runtime/nDisplay/Source/DisplayClusterConfigurator/Private/Views/Details/DisplayClusterRootActorDetailsCustomization.cpp b/Engine/Plugins/Runtime/nDisplay/Source/DisplayClusterConfigurator/Private/Views/Details/DisplayClusterRootActorDetailsCustomization.cpp
--- a/Engine/Plugins/Runtime/nDisplay/Source/DisplayClusterConfigurator/Private/Views/Details/DisplayClusterRootActorDetailsCustomization.cpp
+++ b/Engine/Plugins/Runtime/nDisplay/Source/DisplayClusterConfigurator/Private/Views/Details/DisplayClusterRootActorDetailsCustomization.cpp
@@ -187,13 +187,11 @@ bool FDisplayClusterRootActorDetailsCustomization::RebuildNodeIdOptionsList()
 	}
 
 	// Initialize special options
-	NodeIdOptionAll  = MakeShared<FString>(DisplayClusterConfigurationStrings::gui::preview::PreviewNodeAll);
 	NodeIdOptionNone = MakeShared<FString>(DisplayClusterConfigurationStrings::gui::preview::PreviewNodeNone);
 
 	// Fill combobox with the options
 	NodeIdOptions.Reset();
 	NodeIdOptions.Emplace(NodeIdOptionNone);
-	NodeIdOptions.Emplace(NodeIdOptionAll);
 	for (const TPair<FString, UDisplayClusterConfigurationClusterNode*>& Node : ConfigurationData->Cluster->Nodes)
 	{
 		if (!Node.Key.IsEmpty())
@@ -230,7 +228,7 @@ void FDisplayClusterRootActorDetailsCustomization::UpdateNodeIdSelection()
 		else
 		{
 			// Set combobox selected item (options list is not empty here)
-			NodeIdComboBox->SetSelectedItem(NodeIdOptionAll);
+			NodeIdComboBox->SetSelectedItem(NodeIdOptionNone);
 		}
 	}
 }
diff --git a/Engine/Plugins/Runtime/nDisplay/Source/DisplayClusterConfigurator/Private/Views/Details/DisplayClusterRootActorDetailsCustomization.h b/Engine/Plugins/Runtime/nDisplay/Source/DisplayClusterConfigurator/Private/Views/Details/DisplayClusterRootActorDetailsCustomization.h
--- a/Engine/Plugins/Runtime/nDisplay/Source/DisplayClusterConfigurator/Private/Views/Details/DisplayClusterRootActorDetailsCustomization.h
+++ b/Engine/Plugins/Runtime/nDisplay/Source/DisplayClusterConfigurator/Private/Views/Details/DisplayClusterRootActorDetailsCustomization.h
@@ -45,7 +45,6 @@ private:
 
 	TArray<TSharedPtr<FString>> NodeIdOptions;
 	TSharedPtr<FString> NodeIdOptionNone;
-	TSharedPtr<FString> NodeIdOptionAll;
 	TSharedPtr<SSearchableComboBox> NodeIdComboBox;
 
 	bool bMultipleObjectsSelected;
diff --git a/Engine/Plugins/Runtime/nDisplay/Source/DisplayClusterConfigurator/Private/Views/OutputMapping/EdNodes/DisplayClusterConfiguratorViewportNode.cpp b/Engine/Plugins/Runtime/nDisplay/Source/DisplayClusterConfigurator/Private/Views/OutputMapping/EdNodes/DisplayClusterConfiguratorViewportNode.cpp
--- a/Engine/Plugins/Runtime/nDisplay/Source/DisplayClusterConfigurator/Private/Views/OutputMapping/EdNodes/DisplayClusterConfiguratorViewportNode.cpp
+++ b/Engine/Plugins/Runtime/nDisplay/Source/DisplayClusterConfigurator/Private/Views/OutputMapping/EdNodes/DisplayClusterConfiguratorViewportNode.cpp
@@ -228,19 +228,7 @@ UTexture* UDisplayClusterConfiguratorViewportNode::GetPreviewTexture() const
 		UDisplayClusterConfiguratorWindowNode* ParentWindow = GetParentChecked<UDisplayClusterConfiguratorWindowNode>();
 		if (UDisplayClusterPreviewComponent* PreviewComp = RootActor->GetPreviewComponent(ParentWindow->GetNodeName(), GetNodeName()))
 		{
-			const UDisplayClusterConfigurationViewport* CfgViewport = GetObjectChecked<UDisplayClusterConfigurationViewport>();
-			if (!CfgViewport->IsPreviewTextureAllowed())
-			{
-				// The preview texture may be disabled during certain operations, like a resize. Use only an old preview
-				// if possible. This can improve performance and also prevent possible DX crashes if the render target
-				// size is smaller than the viewport region size.
-				return PreviewComp->GetViewportPreviewTexture2D();
-			}
-			
-			if (UTexture2D* Texture = PreviewComp->GetOrCreateViewportPreviewTexture2D())
-			{
-				return Texture;
-			}
+			return PreviewComp->GetViewportPreviewTexture2D();
 		}
 	}
 
@@ -321,4 +309,4 @@ bool UDisplayClusterConfiguratorViewportNode::IsRotationAngleDegenerate(float An
 	}
 
 	return FMath::Abs(ClampedAngle - 45.0f) < KINDA_SMALL_NUMBER;
-}
\ No newline at end of file
+}
diff --git a/Engine/Plugins/Runtime/nDisplay/Source/DisplayClusterConfigurator/Private/Views/Viewport/DisplayClusterConfiguratorSCSEditorViewportClient.cpp b/Engine/Plugins/Runtime/nDisplay/Source/DisplayClusterConfigurator/Private/Views/Viewport/DisplayClusterConfiguratorSCSEditorViewportClient.cpp
--- a/Engine/Plugins/Runtime/nDisplay/Source/DisplayClusterConfigurator/Private/Views/Viewport/DisplayClusterConfiguratorSCSEditorViewportClient.cpp
+++ b/Engine/Plugins/Runtime/nDisplay/Source/DisplayClusterConfigurator/Private/Views/Viewport/DisplayClusterConfiguratorSCSEditorViewportClient.cpp
@@ -1083,15 +1083,15 @@ void FDisplayClusterConfiguratorSCSEditorViewportClient::SyncShowPreview()
 	if (BlueprintEditorPtr.IsValid())
 	{
 		const bool bShouldShowPreview = GetShowPreview();
-		
-		if (ADisplayClusterRootActor* Actor = Cast<ADisplayClusterRootActor>(BlueprintEditorPtr.Pin()->GetPreviewActor()))
+		if (!bShouldShowPreview)
 		{
-			const FString CorrectShowPreviewValue = bShouldShowPreview ? DisplayClusterConfigurationStrings::gui::preview::PreviewNodeAll : DisplayClusterConfigurationStrings::gui::preview::PreviewNodeNone;
-			
-			if (CorrectShowPreviewValue != Actor->PreviewNodeId)
+			if (ADisplayClusterRootActor* Actor = Cast<ADisplayClusterRootActor>(BlueprintEditorPtr.Pin()->GetPreviewActor()))
 			{
-				Actor->PreviewNodeId = CorrectShowPreviewValue;
-				Actor->UpdatePreviewComponents();
+				if (Actor->PreviewNodeId != DisplayClusterConfigurationStrings::gui::preview::PreviewNodeNone)
+				{
+					Actor->PreviewNodeId = DisplayClusterConfigurationStrings::gui::preview::PreviewNodeNone;
+					Actor->UpdatePreviewComponents();
+				}
 			}
 		}
 	}
diff --git a/Engine/Plugins/Runtime/nDisplay/Source/DisplayClusterProjection/Private/Blueprints/DisplayClusterProjectionBlueprintAPIImpl.cpp b/Engine/Plugins/Runtime/nDisplay/Source/DisplayClusterProjection/Private/Blueprints/DisplayClusterProjectionBlueprintAPIImpl.cpp
--- a/Engine/Plugins/Runtime/nDisplay/Source/DisplayClusterProjection/Private/Blueprints/DisplayClusterProjectionBlueprintAPIImpl.cpp
+++ b/Engine/Plugins/Runtime/nDisplay/Source/DisplayClusterProjection/Private/Blueprints/DisplayClusterProjectionBlueprintAPIImpl.cpp
@@ -30,6 +30,7 @@ void UDisplayClusterProjectionBlueprintAPIImpl::CameraPolicySetCamera(const FStr
 			// @todo: Add extra settings latter to BP call if required
 			FDisplayClusterProjectionCameraPolicySettings Settings;
 			Settings.FOVMultiplier = FOVMultiplier;
+			Settings.bCameraOverrideDefaults = true;
 
 			IDisplayClusterProjection::Get().CameraPolicySetCamera(Viewport->GetProjectionPolicy(), NewCamera, Settings);
 		}
diff --git a/Engine/Plugins/Runtime/nDisplay/Source/DisplayClusterProjection/Private/Policy/Camera/DisplayClusterProjectionCameraPolicy.cpp b/Engine/Plugins/Runtime/nDisplay/Source/DisplayClusterProjection/Private/Policy/Camera/DisplayClusterProjectionCameraPolicy.cpp
--- a/Engine/Plugins/Runtime/nDisplay/Source/DisplayClusterProjection/Private/Policy/Camera/DisplayClusterProjectionCameraPolicy.cpp
+++ b/Engine/Plugins/Runtime/nDisplay/Source/DisplayClusterProjection/Private/Policy/Camera/DisplayClusterProjectionCameraPolicy.cpp
@@ -116,7 +116,7 @@ bool FDisplayClusterProjectionCameraPolicy::ImplGetProjectionMatrix(const float
 
 	// Clamp camera fov to valid range [1.f, 178.f]
 	const float ClampedCameraFOV = FMath::Clamp(ScaledCameraFOV, 1.f, 178.f);
-	if (ClampedCameraFOV != ScaledCameraFOV && !IsEditorOperationMode())
+	if (ClampedCameraFOV != ScaledCameraFOV && !IsEditorOperationMode(InViewport))
 	{
 		UE_LOG(LogDisplayClusterProjectionCamera, Warning, TEXT("CameraFOV clamped: '%d' -> '%d'. (FieldOfView='%d', FOVMultiplier='%d'"), ScaledCameraFOV, ClampedCameraFOV, InCameraFOV, CameraSettings.FOVMultiplier);
 	}
@@ -164,19 +164,33 @@ bool FDisplayClusterProjectionCameraPolicy::GetProjectionMatrix(IDisplayClusterV
 //////////////////////////////////////////////////////////////////////////////////////////////
 void FDisplayClusterProjectionCameraPolicy::SetCamera(UCameraComponent* NewCamera, const FDisplayClusterProjectionCameraPolicySettings& InCameraSettings)
 {
+	if(CameraSettings.bCameraOverrideDefaults && InCameraSettings.bCameraOverrideDefaults == false)
+	{
+		// Ignore default camera updates (UE-137222)
+		return;
+	}
+
+	CameraSettings = InCameraSettings;
+
 	if (NewCamera)
 	{
 		UE_LOG(LogDisplayClusterProjectionCamera, Verbose, TEXT("New camera set: %s"), *NewCamera->GetFullName());
+
 		CameraRef.SetSceneComponent(NewCamera);
 	}
 	else
 	{
-		CameraRef.ResetSceneComponent();
-		if (!IsEditorOperationMode())
+		if (InCameraSettings.bCameraOverrideDefaults == false)
 		{
-			UE_LOG(LogDisplayClusterProjectionCamera, Warning, TEXT("Trying to set nullptr camera pointer"));
+			if (!IsEditorOperationMode(nullptr))
+			{
+				UE_LOG(LogDisplayClusterProjectionCamera, Warning, TEXT("Trying to set nullptr camera pointer"));
+			}
 		}
-	}
 
-	CameraSettings = InCameraSettings;
+		CameraRef.ResetSceneComponent();
+
+		// After ref reset, allow to use default cameras again
+		CameraSettings.bCameraOverrideDefaults = false;
+	}
 }
diff --git a/Engine/Plugins/Runtime/nDisplay/Source/DisplayClusterProjection/Private/Policy/DisplayClusterProjectionPolicyBase.cpp b/Engine/Plugins/Runtime/nDisplay/Source/DisplayClusterProjection/Private/Policy/DisplayClusterProjectionPolicyBase.cpp
--- a/Engine/Plugins/Runtime/nDisplay/Source/DisplayClusterProjection/Private/Policy/DisplayClusterProjectionPolicyBase.cpp
+++ b/Engine/Plugins/Runtime/nDisplay/Source/DisplayClusterProjection/Private/Policy/DisplayClusterProjectionPolicyBase.cpp
@@ -25,12 +25,49 @@ FDisplayClusterProjectionPolicyBase::~FDisplayClusterProjectionPolicyBase()
 {
 }
 
-bool FDisplayClusterProjectionPolicyBase::IsEditorOperationMode()
+bool FDisplayClusterProjectionPolicyBase::IsEditorOperationMode_RenderThread(const IDisplayClusterViewportProxy* InViewportProxy)
 {
-	// Hide spam in logs when configuring VP in editor [UE-114493]
-	static const bool bIsEditorOperationMode = IDisplayCluster::Get().GetOperationMode() == EDisplayClusterOperationMode::Editor;
+	check(IsInRenderingThread());
 
-	return bIsEditorOperationMode;
+	if (const bool bIsClusterOperationMode = IDisplayCluster::Get().GetOperationMode() == EDisplayClusterOperationMode::Cluster)
+	{
+		return false;
+	}
+
+	return true;
+}
+
+bool FDisplayClusterProjectionPolicyBase::IsEditorOperationMode(class IDisplayClusterViewport* InViewport)
+{
+	check(IsInGameThread());
+
+	// Hide spam in logs when configuring VP in editor
+	if (const bool bIsClusterOperationMode = IDisplayCluster::Get().GetOperationMode() == EDisplayClusterOperationMode::Cluster)
+	{
+		return false;
+	}
+
+	// Get state from viewport world (UE-114493)
+	if (InViewport)
+	{
+		if (UWorld* CurrentWorld = InViewport->GetOwner().GetCurrentWorld())
+		{
+			if (CurrentWorld)
+			{
+				switch (CurrentWorld->WorldType)
+				{
+				case EWorldType::Editor:
+				case EWorldType::EditorPreview:
+					return true;
+
+				default:
+					break;
+				}
+			}
+		}
+	}
+
+	return false;
 }
 
 bool FDisplayClusterProjectionPolicyBase::IsConfigurationChanged(const struct FDisplayClusterConfigurationProjection* InConfigurationProjectionPolicy) const
@@ -74,7 +111,7 @@ void FDisplayClusterProjectionPolicyBase::InitializeOriginComponent(IDisplayClus
 
 				if (PolicyOriginComp == nullptr)
 				{
-					if (!IsEditorOperationMode())
+					if (!IsEditorOperationMode(InViewport))
 					{
 						UE_LOG(LogDisplayClusterProjection, Error, TEXT("No custom origin set or component '%s' not found for policy '%s'. VR root will be used."), *OriginCompId, *GetId());
 					}
@@ -86,7 +123,7 @@ void FDisplayClusterProjectionPolicyBase::InitializeOriginComponent(IDisplayClus
 
 		if (!PolicyOriginComp)
 		{
-			if (!IsEditorOperationMode())
+			if (!IsEditorOperationMode(InViewport))
 			{
 				UE_LOG(LogDisplayClusterProjection, Error, TEXT("Couldn't set origin component"));
 			}
diff --git a/Engine/Plugins/Runtime/nDisplay/Source/DisplayClusterProjection/Private/Policy/DisplayClusterProjectionPolicyBase.h b/Engine/Plugins/Runtime/nDisplay/Source/DisplayClusterProjection/Private/Policy/DisplayClusterProjectionPolicyBase.h
--- a/Engine/Plugins/Runtime/nDisplay/Source/DisplayClusterProjection/Private/Policy/DisplayClusterProjectionPolicyBase.h
+++ b/Engine/Plugins/Runtime/nDisplay/Source/DisplayClusterProjection/Private/Policy/DisplayClusterProjectionPolicyBase.h
@@ -46,7 +46,8 @@ public:
 
 	virtual bool IsConfigurationChanged(const struct FDisplayClusterConfigurationProjection* InConfigurationProjectionPolicy) const override;
 
-	static bool IsEditorOperationMode();
+	static bool IsEditorOperationMode(class IDisplayClusterViewport* InViewport);
+	static bool IsEditorOperationMode_RenderThread(const IDisplayClusterViewportProxy* InViewportProxy);
 
 protected:
 	void InitializeOriginComponent(class IDisplayClusterViewport* InViewport, const FString& OriginCopmId);
diff --git a/Engine/Plugins/Runtime/nDisplay/Source/DisplayClusterProjection/Private/Policy/Domeprojection/DisplayClusterProjectionDomeprojectionPolicyBase.cpp b/Engine/Plugins/Runtime/nDisplay/Source/DisplayClusterProjection/Private/Policy/Domeprojection/DisplayClusterProjectionDomeprojectionPolicyBase.cpp
--- a/Engine/Plugins/Runtime/nDisplay/Source/DisplayClusterProjection/Private/Policy/Domeprojection/DisplayClusterProjectionDomeprojectionPolicyBase.cpp
+++ b/Engine/Plugins/Runtime/nDisplay/Source/DisplayClusterProjection/Private/Policy/Domeprojection/DisplayClusterProjectionDomeprojectionPolicyBase.cpp
@@ -31,9 +31,9 @@ bool FDisplayClusterProjectionDomeprojectionPolicyBase::HandleStartScene(class I
 
 	// Read domeprojection config data from nDisplay config file
 	FString File;
-	if (!ReadConfigData(InViewport->GetId(), File, OriginCompId, DomeprojectionChannel))
+	if (!ReadConfigData(InViewport, File, OriginCompId, DomeprojectionChannel))
 	{
-		if (!IsEditorOperationMode())
+		if (!IsEditorOperationMode(InViewport))
 		{
 			UE_LOG(LogDisplayClusterProjectionDomeprojection, Error, TEXT("Couldn't read Domeprojection configuration from the config file"));
 		}
@@ -43,7 +43,7 @@ bool FDisplayClusterProjectionDomeprojectionPolicyBase::HandleStartScene(class I
 
 	if (File.IsEmpty())
 	{
-		if (!IsEditorOperationMode())
+		if (!IsEditorOperationMode(InViewport))
 		{
 			UE_LOG(LogDisplayClusterProjectionDomeprojection, Error, TEXT("Domeprojection configuration file is empty"));
 		}
@@ -65,9 +65,9 @@ bool FDisplayClusterProjectionDomeprojectionPolicyBase::HandleStartScene(class I
 	const int32 MaxViewsAmount = 2; // always init contexts for stereo
 
 	ViewAdapter = CreateViewAdapter(FDisplayClusterProjectionDomeprojectionViewAdapterBase::FInitParams{ MaxViewsAmount });
-	if (!(ViewAdapter && ViewAdapter->Initialize(FullFilePath)))
+	if (!(ViewAdapter && ViewAdapter->Initialize(InViewport, FullFilePath)))
 	{
-		if (!IsEditorOperationMode())
+		if (!IsEditorOperationMode(InViewport))
 		{
 			UE_LOG(LogDisplayClusterProjectionDomeprojection, Error, TEXT("An error occurred during Domeprojection viewport adapter initialization"));
 		}
@@ -110,7 +110,7 @@ bool FDisplayClusterProjectionDomeprojectionPolicyBase::CalculateView(class IDis
 	FRotator OriginSpaceViewRotation = FRotator::ZeroRotator;
 	if (!ViewAdapter->CalculateView(InViewport, InContextNum, DomeprojectionChannel, OriginSpaceViewLocation, OriginSpaceViewRotation, ViewOffset, WorldToMeters, NCP, FCP))
 	{
-		if (!IsEditorOperationMode())
+		if (!IsEditorOperationMode(InViewport))
 		{
 			UE_LOG(LogDisplayClusterProjectionDomeprojection, Warning, TEXT("Couldn't compute view info for <%s> viewport"), *InViewport->GetId());
 		}
@@ -158,8 +158,12 @@ void FDisplayClusterProjectionDomeprojectionPolicyBase::ApplyWarpBlend_RenderThr
 //////////////////////////////////////////////////////////////////////////////////////////////
 // FDisplayClusterProjectionDomeprojectionPolicy
 //////////////////////////////////////////////////////////////////////////////////////////////
-bool FDisplayClusterProjectionDomeprojectionPolicyBase::ReadConfigData(const FString& InViewportId, FString& OutFile, FString& OutOrigin, uint32& OutChannel)
+bool FDisplayClusterProjectionDomeprojectionPolicyBase::ReadConfigData(class IDisplayClusterViewport* InViewport, FString& OutFile, FString& OutOrigin, uint32& OutChannel)
 {
+	check(InViewport);
+
+	const FString InViewportId = InViewport->GetId();
+
 	// Domeprojection file (mandatory)
 	if (DisplayClusterHelpers::map::ExtractValue(GetParameters(), DisplayClusterProjectionStrings::cfg::domeprojection::File, OutFile))
 	{
@@ -167,7 +171,7 @@ bool FDisplayClusterProjectionDomeprojectionPolicyBase::ReadConfigData(const FSt
 	}
 	else
 	{
-		if (!IsEditorOperationMode())
+		if (!IsEditorOperationMode(InViewport))
 		{
 			UE_LOG(LogDisplayClusterProjectionDomeprojection, Error, TEXT("Viewport <%s>: Projection parameter '%s' not found"), *InViewportId, DisplayClusterProjectionStrings::cfg::domeprojection::File);
 		}
@@ -182,7 +186,7 @@ bool FDisplayClusterProjectionDomeprojectionPolicyBase::ReadConfigData(const FSt
 	}
 	else
 	{
-		if (!IsEditorOperationMode())
+		if (!IsEditorOperationMode(InViewport))
 		{
 			UE_LOG(LogDisplayClusterProjectionDomeprojection, Error, TEXT("Viewport <%s>: Parameter <%s> not found in the config file"), *InViewportId, DisplayClusterProjectionStrings::cfg::domeprojection::Channel);
 		}
@@ -197,7 +201,7 @@ bool FDisplayClusterProjectionDomeprojectionPolicyBase::ReadConfigData(const FSt
 	}
 	else
 	{
-		if (!IsEditorOperationMode())
+		if (!IsEditorOperationMode(InViewport))
 		{
 			UE_LOG(LogDisplayClusterProjectionDomeprojection, Log, TEXT("Viewport <%s>: No <%s> parameter found for projection %s"), *InViewportId, DisplayClusterProjectionStrings::cfg::domeprojection::Origin, *OutOrigin);
 		}
diff --git a/Engine/Plugins/Runtime/nDisplay/Source/DisplayClusterProjection/Private/Policy/Domeprojection/DisplayClusterProjectionDomeprojectionPolicyBase.h b/Engine/Plugins/Runtime/nDisplay/Source/DisplayClusterProjection/Private/Policy/Domeprojection/DisplayClusterProjectionDomeprojectionPolicyBase.h
--- a/Engine/Plugins/Runtime/nDisplay/Source/DisplayClusterProjection/Private/Policy/Domeprojection/DisplayClusterProjectionDomeprojectionPolicyBase.h
+++ b/Engine/Plugins/Runtime/nDisplay/Source/DisplayClusterProjection/Private/Policy/Domeprojection/DisplayClusterProjectionDomeprojectionPolicyBase.h
@@ -45,7 +45,7 @@ protected:
 
 private:
 	// Parse Domeprojection related data from the nDisplay config file
-	bool ReadConfigData(const FString& InViewportId, FString& OutFile, FString& OutOrigin, uint32& OutChannel);
+	bool ReadConfigData(class IDisplayClusterViewport* InViewport, FString& OutFile, FString& OutOrigin, uint32& OutChannel);
 
 private:
 	FString OriginCompId;
diff --git a/Engine/Plugins/Runtime/nDisplay/Source/DisplayClusterProjection/Private/Policy/Domeprojection/DisplayClusterProjectionDomeprojectionViewAdapterBase.h b/Engine/Plugins/Runtime/nDisplay/Source/DisplayClusterProjection/Private/Policy/Domeprojection/DisplayClusterProjectionDomeprojectionViewAdapterBase.h
--- a/Engine/Plugins/Runtime/nDisplay/Source/DisplayClusterProjection/Private/Policy/Domeprojection/DisplayClusterProjectionDomeprojectionViewAdapterBase.h
+++ b/Engine/Plugins/Runtime/nDisplay/Source/DisplayClusterProjection/Private/Policy/Domeprojection/DisplayClusterProjectionDomeprojectionViewAdapterBase.h
@@ -26,7 +26,7 @@ public:
 	virtual ~FDisplayClusterProjectionDomeprojectionViewAdapterBase() = default;
 
 public:
-	virtual bool Initialize(const FString& File) = 0;
+	virtual bool Initialize(class IDisplayClusterViewport* InViewport, const FString& File) = 0;
 	
 	virtual void Release()
 	{ }
diff --git a/Engine/Plugins/Runtime/nDisplay/Source/DisplayClusterProjection/Private/Policy/Domeprojection/Windows/DX11/DisplayClusterProjectionDomeprojectionViewAdapterDX11.cpp b/Engine/Plugins/Runtime/nDisplay/Source/DisplayClusterProjection/Private/Policy/Domeprojection/Windows/DX11/DisplayClusterProjectionDomeprojectionViewAdapterDX11.cpp
--- a/Engine/Plugins/Runtime/nDisplay/Source/DisplayClusterProjection/Private/Policy/Domeprojection/Windows/DX11/DisplayClusterProjectionDomeprojectionViewAdapterDX11.cpp
+++ b/Engine/Plugins/Runtime/nDisplay/Source/DisplayClusterProjection/Private/Policy/Domeprojection/Windows/DX11/DisplayClusterProjectionDomeprojectionViewAdapterDX11.cpp
@@ -42,13 +42,13 @@ FDisplayClusterProjectionDomeprojectionViewAdapterDX11::~FDisplayClusterProjecti
 	Views.Empty();
 }
 
-bool FDisplayClusterProjectionDomeprojectionViewAdapterDX11::Initialize(const FString& InFile)
+bool FDisplayClusterProjectionDomeprojectionViewAdapterDX11::Initialize(class IDisplayClusterViewport* InViewport, const FString& InFile)
 {
 	bool bResult = true;
 
 	for (FViewData& ViewIt : Views)
 	{
-		if (!ViewIt.Initialize(InFile, DllAccessCS))
+		if (!ViewIt.Initialize(InViewport, InFile, DllAccessCS))
 		{
 			bResult = false;
 		}
@@ -236,14 +236,14 @@ void FDisplayClusterProjectionDomeprojectionViewAdapterDX11::FViewData::Release(
 	}
 }
 
-bool FDisplayClusterProjectionDomeprojectionViewAdapterDX11::FViewData::Initialize(const FString& InFile, FCriticalSection& DllAccessCS)
+bool FDisplayClusterProjectionDomeprojectionViewAdapterDX11::FViewData::Initialize(class IDisplayClusterViewport* InViewport, const FString& InFile, FCriticalSection& DllAccessCS)
 {
 	// Initialize Domeprojection DLL API
 	FScopeLock lock(&DllAccessCS);
 
 	if (!DisplayClusterProjectionDomeprojectionLibraryDX11::Initialize())
 	{
-		if (!FDisplayClusterProjectionPolicyBase::IsEditorOperationMode())
+		if (!FDisplayClusterProjectionPolicyBase::IsEditorOperationMode(InViewport))
 		{
 			UE_LOG(LogDisplayClusterProjectionDomeprojection, Error, TEXT("Couldn't link to the Domeprojection DLL"));
 		}
@@ -253,7 +253,7 @@ bool FDisplayClusterProjectionDomeprojectionViewAdapterDX11::FViewData::Initiali
 
 	if (InFile.IsEmpty())
 	{
-		if (!FDisplayClusterProjectionPolicyBase::IsEditorOperationMode())
+		if (!FDisplayClusterProjectionPolicyBase::IsEditorOperationMode(InViewport))
 		{
 			UE_LOG(LogDisplayClusterProjectionDomeprojection, Error, TEXT("File name is empty"));
 		}
diff --git a/Engine/Plugins/Runtime/nDisplay/Source/DisplayClusterProjection/Private/Policy/Domeprojection/Windows/DX11/DisplayClusterProjectionDomeprojectionViewAdapterDX11.h b/Engine/Plugins/Runtime/nDisplay/Source/DisplayClusterProjection/Private/Policy/Domeprojection/Windows/DX11/DisplayClusterProjectionDomeprojectionViewAdapterDX11.h
--- a/Engine/Plugins/Runtime/nDisplay/Source/DisplayClusterProjection/Private/Policy/Domeprojection/Windows/DX11/DisplayClusterProjectionDomeprojectionViewAdapterDX11.h
+++ b/Engine/Plugins/Runtime/nDisplay/Source/DisplayClusterProjection/Private/Policy/Domeprojection/Windows/DX11/DisplayClusterProjectionDomeprojectionViewAdapterDX11.h
@@ -17,7 +17,7 @@ public:
 	virtual ~FDisplayClusterProjectionDomeprojectionViewAdapterDX11();
 
 public:
-	virtual bool Initialize(const FString& File) override;
+	virtual bool Initialize(class IDisplayClusterViewport* InViewport, const FString& File) override;
 
 public:
 	virtual bool CalculateView(class IDisplayClusterViewport* InViewport, const uint32 InContextNum, const uint32 Channel, FVector& InOutViewLocation, FRotator& InOutViewRotation, const FVector& ViewOffset, const float WorldToMeters, const float NCP, const float FCP) override;
@@ -34,7 +34,7 @@ private:
 	class FViewData
 	{
 	public:
-		bool Initialize(const FString& InFile, FCriticalSection& DllAccessCS);
+		bool Initialize(class IDisplayClusterViewport* InViewport, const FString& InFile, FCriticalSection& DllAccessCS);
 		void Release(FCriticalSection& DllAccessCS);
 
 	public:
diff --git a/Engine/Plugins/Runtime/nDisplay/Source/DisplayClusterProjection/Private/Policy/EasyBlend/DisplayClusterProjectionEasyBlendPolicyBase.cpp b/Engine/Plugins/Runtime/nDisplay/Source/DisplayClusterProjection/Private/Policy/EasyBlend/DisplayClusterProjectionEasyBlendPolicyBase.cpp
--- a/Engine/Plugins/Runtime/nDisplay/Source/DisplayClusterProjection/Private/Policy/EasyBlend/DisplayClusterProjectionEasyBlendPolicyBase.cpp
+++ b/Engine/Plugins/Runtime/nDisplay/Source/DisplayClusterProjection/Private/Policy/EasyBlend/DisplayClusterProjectionEasyBlendPolicyBase.cpp
@@ -52,9 +52,9 @@ bool FDisplayClusterProjectionEasyBlendPolicyBase::HandleStartScene(class IDispl
 
 	// Read easyblend config data from nDisplay config file
 	FString FilePath;
-	if (!ReadConfigData(InViewport->GetId(), FilePath, OriginCompId, EasyBlendScale))
+	if (!ReadConfigData(InViewport, FilePath, OriginCompId, EasyBlendScale))
 	{
-		if (!IsEditorOperationMode())
+		if (!IsEditorOperationMode(InViewport))
 		{
 			UE_LOG(LogDisplayClusterProjectionEasyBlend, Error, TEXT("Couldn't read EasyBlend configuration from the config file"));
 		}
@@ -63,7 +63,7 @@ bool FDisplayClusterProjectionEasyBlendPolicyBase::HandleStartScene(class IDispl
 
 	if (FilePath.IsEmpty())
 	{
-		if (!IsEditorOperationMode())
+		if (!IsEditorOperationMode(InViewport))
 		{
 			UE_LOG(LogDisplayClusterProjectionEasyBlend, Error, TEXT("EasyBlend file '' not found (Empty)"));
 		}
@@ -81,11 +81,11 @@ bool FDisplayClusterProjectionEasyBlendPolicyBase::HandleStartScene(class IDispl
 		return false;
 	}
 
-	const uint32 MaxViewsAmmount = 2;
+	const uint32 MaxViewsAmount = 2;
 
 	// Create and store nDisplay-to-EasyBlend viewport adapter
-	ViewAdapter = CreateViewAdapter(FDisplayClusterProjectionEasyBlendViewAdapterBase::FInitParams{ MaxViewsAmmount });
-	if (!ViewAdapter || !ViewAdapter->Initialize(FullPath))
+	ViewAdapter = CreateViewAdapter(FDisplayClusterProjectionEasyBlendViewAdapterBase::FInitParams{ MaxViewsAmount });
+	if (!ViewAdapter || !ViewAdapter->Initialize(InViewport, FullPath))
 	{
 		UE_LOG(LogDisplayClusterProjectionEasyBlend, Error, TEXT("An error occurred during EasyBlend viewport adapter initialization"));
 		return false;
@@ -144,7 +144,7 @@ bool FDisplayClusterProjectionEasyBlendPolicyBase::CalculateView(IDisplayCluster
 	FRotator OriginSpaceViewRotation = FRotator::ZeroRotator;
 	if (!ViewAdapter->CalculateView(InViewport, InContextNum, OriginSpaceViewLocation, OriginSpaceViewRotation, ViewOffset, WorldScale, NCP, FCP))
 	{
-		if (!IsEditorOperationMode())
+		if (!IsEditorOperationMode(InViewport))
 		{
 			UE_LOG(LogDisplayClusterProjectionEasyBlend, Warning, TEXT("Couldn't compute view info for <%s> viewport"), *InViewport->GetId());
 		}
@@ -173,8 +173,12 @@ bool FDisplayClusterProjectionEasyBlendPolicyBase::GetProjectionMatrix(IDisplayC
 //////////////////////////////////////////////////////////////////////////////////////////////
 // FDisplayClusterProjectionEasyBlendPolicy
 //////////////////////////////////////////////////////////////////////////////////////////////
-bool FDisplayClusterProjectionEasyBlendPolicyBase::ReadConfigData(const FString& InViewportId, FString& OutFile, FString& OutOrigin, float& OutGeometryScale)
+bool FDisplayClusterProjectionEasyBlendPolicyBase::ReadConfigData(class IDisplayClusterViewport* InViewport, FString& OutFile, FString& OutOrigin, float& OutGeometryScale)
 {
+	check(InViewport);
+
+	const FString InViewportId = InViewport->GetId();
+
 	// EasyBlend file (mandatory)
 	if (DisplayClusterHelpers::map::template ExtractValueFromString(GetParameters(), DisplayClusterProjectionStrings::cfg::easyblend::File, OutFile))
 	{
@@ -182,7 +186,7 @@ bool FDisplayClusterProjectionEasyBlendPolicyBase::ReadConfigData(const FString&
 	}
 	else
 	{
-		if (!IsEditorOperationMode())
+		if (!IsEditorOperationMode(InViewport))
 		{
 			UE_LOG(LogDisplayClusterProjectionEasyBlend, Error, TEXT("Viewport <%s>: Projection parameter '%s' not found"), *InViewportId, DisplayClusterProjectionStrings::cfg::easyblend::File);
 		}
@@ -196,7 +200,7 @@ bool FDisplayClusterProjectionEasyBlendPolicyBase::ReadConfigData(const FString&
 	}
 	else
 	{
-		if (!IsEditorOperationMode())
+		if (!IsEditorOperationMode(InViewport))
 		{
 			UE_LOG(LogDisplayClusterProjectionEasyBlend, Log, TEXT("Viewport <%s>: Projection parameter '%s' not found"), *InViewportId, DisplayClusterProjectionStrings::cfg::easyblend::Origin);
 		}
@@ -209,7 +213,7 @@ bool FDisplayClusterProjectionEasyBlendPolicyBase::ReadConfigData(const FString&
 	}
 	else
 	{
-		if (!IsEditorOperationMode())
+		if (!IsEditorOperationMode(InViewport))
 		{
 			UE_LOG(LogDisplayClusterProjectionEasyBlend, Log, TEXT("Viewport <%s>: Projection parameter '%s' not found"), *InViewportId, DisplayClusterProjectionStrings::cfg::easyblend::Scale);
 		}
diff --git a/Engine/Plugins/Runtime/nDisplay/Source/DisplayClusterProjection/Private/Policy/EasyBlend/DisplayClusterProjectionEasyBlendPolicyBase.h b/Engine/Plugins/Runtime/nDisplay/Source/DisplayClusterProjection/Private/Policy/EasyBlend/DisplayClusterProjectionEasyBlendPolicyBase.h
--- a/Engine/Plugins/Runtime/nDisplay/Source/DisplayClusterProjection/Private/Policy/EasyBlend/DisplayClusterProjectionEasyBlendPolicyBase.h
+++ b/Engine/Plugins/Runtime/nDisplay/Source/DisplayClusterProjection/Private/Policy/EasyBlend/DisplayClusterProjectionEasyBlendPolicyBase.h
@@ -44,7 +44,7 @@ protected:
 
 private:
 	// Parse EasyBlend related data from the nDisplay config file
-	bool ReadConfigData(const FString& InViewportId, FString& OutFile, FString& OutOrigin, float& OutGeometryScale);
+	bool ReadConfigData(class IDisplayClusterViewport* InViewport, FString& OutFile, FString& OutOrigin, float& OutGeometryScale);
 
 private:
 	FString OriginCompId;
diff --git a/Engine/Plugins/Runtime/nDisplay/Source/DisplayClusterProjection/Private/Policy/EasyBlend/DisplayClusterProjectionEasyBlendViewAdapterBase.h b/Engine/Plugins/Runtime/nDisplay/Source/DisplayClusterProjection/Private/Policy/EasyBlend/DisplayClusterProjectionEasyBlendViewAdapterBase.h
--- a/Engine/Plugins/Runtime/nDisplay/Source/DisplayClusterProjection/Private/Policy/EasyBlend/DisplayClusterProjectionEasyBlendViewAdapterBase.h
+++ b/Engine/Plugins/Runtime/nDisplay/Source/DisplayClusterProjection/Private/Policy/EasyBlend/DisplayClusterProjectionEasyBlendViewAdapterBase.h
@@ -26,7 +26,7 @@ public:
 	virtual ~FDisplayClusterProjectionEasyBlendViewAdapterBase() = default;
 
 public:
-	virtual bool Initialize(const FString& File) = 0;
+	virtual bool Initialize(class IDisplayClusterViewport* InViewport, const FString& File) = 0;
 	
 	virtual void Release()
 	{ }
diff --git a/Engine/Plugins/Runtime/nDisplay/Source/DisplayClusterProjection/Private/Policy/EasyBlend/Windows/DX11/DisplayClusterProjectionEasyBlendViewAdapterDX11.cpp b/Engine/Plugins/Runtime/nDisplay/Source/DisplayClusterProjection/Private/Policy/EasyBlend/Windows/DX11/DisplayClusterProjectionEasyBlendViewAdapterDX11.cpp
--- a/Engine/Plugins/Runtime/nDisplay/Source/DisplayClusterProjection/Private/Policy/EasyBlend/Windows/DX11/DisplayClusterProjectionEasyBlendViewAdapterDX11.cpp
+++ b/Engine/Plugins/Runtime/nDisplay/Source/DisplayClusterProjection/Private/Policy/EasyBlend/Windows/DX11/DisplayClusterProjectionEasyBlendViewAdapterDX11.cpp
@@ -59,7 +59,7 @@ FDisplayClusterProjectionEasyBlendViewAdapterDX11::~FDisplayClusterProjectionEas
 	}
 }
 
-bool FDisplayClusterProjectionEasyBlendViewAdapterDX11::Initialize(const FString& File)
+bool FDisplayClusterProjectionEasyBlendViewAdapterDX11::Initialize(class IDisplayClusterViewport* InViewport, const FString& File)
 {
 	if (!IsEasyBlendRenderingEnabled())
 	{
@@ -69,7 +69,7 @@ bool FDisplayClusterProjectionEasyBlendViewAdapterDX11::Initialize(const FString
 	// Initialize EasyBlend DLL API
 	if (!DisplayClusterProjectionEasyBlendLibraryDX11::Initialize())
 	{
-		if (!FDisplayClusterProjectionPolicyBase::IsEditorOperationMode())
+		if (!FDisplayClusterProjectionPolicyBase::IsEditorOperationMode(InViewport))
 		{
 			UE_LOG(LogDisplayClusterProjectionEasyBlend, Error, TEXT("Couldn't link to the EasyBlend DLL"));
 		}
@@ -78,7 +78,7 @@ bool FDisplayClusterProjectionEasyBlendViewAdapterDX11::Initialize(const FString
 
 	if(File.IsEmpty())
 	{
-		if (!FDisplayClusterProjectionPolicyBase::IsEditorOperationMode())
+		if (!FDisplayClusterProjectionPolicyBase::IsEditorOperationMode(InViewport))
 		{
 			UE_LOG(LogDisplayClusterProjectionEasyBlend, Error, TEXT("File is empty"));
 		}
diff --git a/Engine/Plugins/Runtime/nDisplay/Source/DisplayClusterProjection/Private/Policy/EasyBlend/Windows/DX11/DisplayClusterProjectionEasyBlendViewAdapterDX11.h b/Engine/Plugins/Runtime/nDisplay/Source/DisplayClusterProjection/Private/Policy/EasyBlend/Windows/DX11/DisplayClusterProjectionEasyBlendViewAdapterDX11.h
--- a/Engine/Plugins/Runtime/nDisplay/Source/DisplayClusterProjection/Private/Policy/EasyBlend/Windows/DX11/DisplayClusterProjectionEasyBlendViewAdapterDX11.h
+++ b/Engine/Plugins/Runtime/nDisplay/Source/DisplayClusterProjection/Private/Policy/EasyBlend/Windows/DX11/DisplayClusterProjectionEasyBlendViewAdapterDX11.h
@@ -18,7 +18,7 @@ public:
 	virtual ~FDisplayClusterProjectionEasyBlendViewAdapterDX11();
 
 public:
-	virtual bool Initialize(const FString& File) override;
+	virtual bool Initialize(class IDisplayClusterViewport* InViewport, const FString& File) override;
 
 public:
 	virtual bool CalculateView(class IDisplayClusterViewport* InViewport, const uint32 InContextNum, FVector& InOutViewLocation, FRotator& InOutViewRotation, const FVector& ViewOffset, const float WorldToMeters, const float NCP, const float FCP) override;
diff --git a/Engine/Plugins/Runtime/nDisplay/Source/DisplayClusterProjection/Private/Policy/MPCDI/DisplayClusterProjectionMPCDIPolicy.cpp b/Engine/Plugins/Runtime/nDisplay/Source/DisplayClusterProjection/Private/Policy/MPCDI/DisplayClusterProjectionMPCDIPolicy.cpp
--- a/Engine/Plugins/Runtime/nDisplay/Source/DisplayClusterProjection/Private/Policy/MPCDI/DisplayClusterProjectionMPCDIPolicy.cpp
+++ b/Engine/Plugins/Runtime/nDisplay/Source/DisplayClusterProjection/Private/Policy/MPCDI/DisplayClusterProjectionMPCDIPolicy.cpp
@@ -86,7 +86,7 @@ bool FDisplayClusterProjectionMPCDIPolicy::HandleStartScene(IDisplayClusterViewp
 		// Ignore broken MPCDI config for other attempts
 		bInvalidConfiguration = true;
 
-		if (!IsEditorOperationMode())
+		if (!IsEditorOperationMode(InViewport))
 		{
 			UE_LOG(LogDisplayClusterProjectionMPCDI, Warning, TEXT("Couldn't load MPCDI config for viewport '%s'"), *InViewport->GetId());
 		}
@@ -153,7 +153,7 @@ bool FDisplayClusterProjectionMPCDIPolicy::CalculateView(IDisplayClusterViewport
 
 	if (WarpBlendInterface.IsValid() == false || WarpBlendContexts.Num() == 0)
 	{
-		if (!IsEditorOperationMode())
+		if (!IsEditorOperationMode(InViewport))
 		{
 			UE_LOG(LogDisplayClusterProjectionMPCDI, Warning, TEXT("Invalid warp data for viewport '%s'"), *InViewport->GetId());
 		}
@@ -273,29 +273,37 @@ void FDisplayClusterProjectionMPCDIPolicy::ApplyWarpBlend_RenderThread(FRHIComma
 					const IDisplayClusterViewportProxy* RefViewportProxy = RefViewportManagerProxy.FindViewport_RenderThread(It->ViewportId);
 					if (RefViewportProxy)
 					{
-						TArray<FRHITexture2D*> RefTextures;
-						// Use for input first MipsShader texture if enabled in viewport render settings
-						if (RefViewportProxy->GetResources_RenderThread(EDisplayClusterViewportResourceType::MipsShaderResource, RefTextures) ||
-							RefViewportProxy->GetResources_RenderThread(EDisplayClusterViewportResourceType::InputShaderResource, RefTextures))
+						const FDisplayClusterViewport_RenderSettings& RefViewportRenderSettings = RefViewportProxy->GetRenderSettings_RenderThread();
+						if (RefViewportRenderSettings.bSkipRendering == false)
 						{
-							const int32 ReContextNum = FMath::Min(ContextNum, RefTextures.Num());
-							It->Texture = RefTextures[ReContextNum];
-
-							// Support stereo icvfx
-							for (FDisplayClusterShaderParameters_ICVFX::FCameraSettings& CameraIt : ShaderICVFX.Cameras)
+							TArray<FRHITexture2D*> RefTextures;
+							// Use for input first MipsShader texture if enabled in viewport render settings
+							if (RefViewportProxy->GetResources_RenderThread(EDisplayClusterViewportResourceType::MipsShaderResource, RefTextures) ||
+								RefViewportProxy->GetResources_RenderThread(EDisplayClusterViewportResourceType::InputShaderResource, RefTextures))
 							{
-								if (It->ViewportId == CameraIt.Resource.ViewportId)
+								if (ContextNum < RefTextures.Num())
 								{
-									// matched camera reference, update context for current eye
-									const FDisplayClusterViewport_Context& InContext = RefViewportProxy->GetContexts_RenderThread()[ContextNum];
+									It->Texture = RefTextures[ContextNum];
+								}
+							}
+						}
 
-									FDisplayClusterShaderParametersICVFX_CameraContext CameraContext;
-									CameraContext.CameraViewLocation = InContext.ViewLocation;
-									CameraContext.CameraViewRotation = InContext.ViewRotation;
-									CameraContext.CameraPrjMatrix = InContext.ProjectionMatrix;
+						const TArray<FDisplayClusterViewport_Context>& RefViewportContexts = RefViewportProxy->GetContexts_RenderThread();
 
-									CameraIt.UpdateCameraContext(CameraContext);
-								}
+						// Support stereo icvfx
+						for (FDisplayClusterShaderParameters_ICVFX::FCameraSettings& CameraIt : ShaderICVFX.Cameras)
+						{
+							if (It->ViewportId == CameraIt.Resource.ViewportId && ContextNum < RefViewportContexts.Num())
+							{
+								// matched camera reference, update context for current eye
+								const FDisplayClusterViewport_Context& InContext = RefViewportContexts[ContextNum];
+
+								FDisplayClusterShaderParametersICVFX_CameraContext CameraContext;
+								CameraContext.CameraViewLocation = InContext.ViewLocation;
+								CameraContext.CameraViewRotation = InContext.ViewRotation;
+								CameraContext.CameraPrjMatrix = InContext.ProjectionMatrix;
+
+								CameraIt.UpdateCameraContext(CameraContext);
 							}
 						}
 					}
@@ -310,9 +318,11 @@ void FDisplayClusterProjectionMPCDIPolicy::ApplyWarpBlend_RenderThread(FRHIComma
 				WarpBlendParameters.Src.Set(InputTextures[ContextNum], InputRects[ContextNum]);
 				WarpBlendParameters.Dest.Set(OutputTextures[ContextNum], OutputRects[ContextNum]);
 
+				WarpBlendParameters.bRenderAlphaChannel = InViewportProxy->GetRenderSettings_RenderThread().bWarpBlendRenderAlphaChannel;
+
 				if (!ShadersAPI.RenderWarpBlend_ICVFX(RHICmdList, WarpBlendParameters, ShaderICVFX))
 				{
-					if (!IsEditorOperationMode())
+					if (!IsEditorOperationMode_RenderThread(InViewportProxy))
 					{
 						UE_LOG(LogDisplayClusterProjectionMPCDI, Warning, TEXT("Couldn't apply icvfx warp&blend"));
 					}
@@ -337,11 +347,12 @@ void FDisplayClusterProjectionMPCDIPolicy::ApplyWarpBlend_RenderThread(FRHIComma
 
 		WarpBlendParameters.Src.Set(InputTextures[ContextNum], InputRects[ContextNum]);
 		WarpBlendParameters.Dest.Set(OutputTextures[ContextNum], OutputRects[ContextNum]);
-
+		
+		WarpBlendParameters.bRenderAlphaChannel = InViewportProxy->GetRenderSettings_RenderThread().bWarpBlendRenderAlphaChannel;
 
 		if (!ShadersAPI.RenderWarpBlend_MPCDI(RHICmdList, WarpBlendParameters))
 		{
-			if (!IsEditorOperationMode())
+			if (!IsEditorOperationMode_RenderThread(InViewportProxy))
 			{
 				UE_LOG(LogDisplayClusterProjectionMPCDI, Warning, TEXT("Couldn't apply mpcdi warp&blend"));
 			}
@@ -430,7 +441,7 @@ bool FDisplayClusterProjectionMPCDIPolicy::CreateWarpBlendFromConfig(IDisplayClu
 	bool bResult = false;
 
 	FConfigParser CfgData;
-	if (CfgData.ImplLoadConfig(GetParameters()))
+	if (CfgData.ImplLoadConfig(InViewport, GetParameters()))
 	{
 		IDisplayClusterShaders& ShadersAPI = IDisplayClusterShaders::Get();
 
@@ -445,7 +456,7 @@ bool FDisplayClusterProjectionMPCDIPolicy::CreateWarpBlendFromConfig(IDisplayClu
 			// Check if MPCDI file exists
 			if (CfgData.MPCDIFileName.IsEmpty())
 			{
-				if (!IsEditorOperationMode())
+				if (!IsEditorOperationMode(InViewport))
 				{
 					UE_LOG(LogDisplayClusterProjectionMPCDI, Warning, TEXT("File not found: Empty"));
 				}
diff --git a/Engine/Plugins/Runtime/nDisplay/Source/DisplayClusterProjection/Private/Policy/MPCDI/DisplayClusterProjectionMPCDIPolicy_Config.h b/Engine/Plugins/Runtime/nDisplay/Source/DisplayClusterProjection/Private/Policy/MPCDI/DisplayClusterProjectionMPCDIPolicy_Config.h
--- a/Engine/Plugins/Runtime/nDisplay/Source/DisplayClusterProjection/Private/Policy/MPCDI/DisplayClusterProjectionMPCDIPolicy_Config.h
+++ b/Engine/Plugins/Runtime/nDisplay/Source/DisplayClusterProjection/Private/Policy/MPCDI/DisplayClusterProjectionMPCDIPolicy_Config.h
@@ -37,7 +37,7 @@ struct FConfigParser
 
 	bool     bEnablePreview = false;
 
-	inline bool ImplLoadConfig(const TMap<FString, FString>& InConfigParameters)
+	inline bool ImplLoadConfig(class IDisplayClusterViewport* InViewport, const TMap<FString, FString>& InConfigParameters)
 	{
 		FString MPCDITypeKey;
 		if (DisplayClusterHelpers::map::template ExtractValue(InConfigParameters, DisplayClusterProjectionStrings::cfg::mpcdi::MPCDITypeKey, MPCDITypeKey))
@@ -47,7 +47,7 @@ struct FConfigParser
 
 		if (MPCDITypeKey.IsEmpty())
 		{
-			if (!FDisplayClusterProjectionPolicyBase::IsEditorOperationMode())
+			if (!FDisplayClusterProjectionPolicyBase::IsEditorOperationMode(InViewport))
 			{
 				UE_LOG(LogDisplayClusterProjectionMPCDI, Error, TEXT("Undefined mpcdi type key '%s'='%s'"), DisplayClusterProjectionStrings::cfg::mpcdi::MPCDITypeKey, *MPCDITypeKey);
 			}
@@ -56,12 +56,12 @@ struct FConfigParser
 
 		if (MPCDITypeKey.Compare(DisplayClusterProjectionStrings::cfg::mpcdi::TypeMPCDI) == 0)
 		{
-			return ImplLoadMPCDIConfig(InConfigParameters) && ImplLoadBase(InConfigParameters);
+			return ImplLoadMPCDIConfig(InViewport, InConfigParameters) && ImplLoadBase(InViewport, InConfigParameters);
 		}
 
 		if (MPCDITypeKey.Compare(DisplayClusterProjectionStrings::cfg::mpcdi::TypePFM) == 0)
 		{
-			return ImplLoadPFMConfig(InConfigParameters) && ImplLoadBase(InConfigParameters);
+			return ImplLoadPFMConfig(InViewport, InConfigParameters) && ImplLoadBase(InViewport, InConfigParameters);
 		}
 
 		UE_LOG(LogDisplayClusterProjectionMPCDI, Error, TEXT("Unknown mpcdi type key '%s'='%s'"), DisplayClusterProjectionStrings::cfg::mpcdi::MPCDITypeKey, *MPCDITypeKey);
@@ -69,7 +69,7 @@ struct FConfigParser
 	}
 
 private:
-	inline bool ImplLoadMPCDIConfig(const TMap<FString, FString>& InConfigParameters)
+	inline bool ImplLoadMPCDIConfig(class IDisplayClusterViewport* InViewport, const TMap<FString, FString>& InConfigParameters)
 	{
 		// Filename
 		FString LocalMPCDIFileName;
@@ -87,7 +87,7 @@ private:
 		// Buffer
 		if (!DisplayClusterHelpers::map::template ExtractValue(InConfigParameters, DisplayClusterProjectionStrings::cfg::mpcdi::Buffer, BufferId))
 		{
-			if (!FDisplayClusterProjectionPolicyBase::IsEditorOperationMode())
+			if (!FDisplayClusterProjectionPolicyBase::IsEditorOperationMode(InViewport))
 			{
 				UE_LOG(LogDisplayClusterProjectionMPCDI, Error, TEXT("Argument '%s' not found in the config file"), DisplayClusterProjectionStrings::cfg::mpcdi::Buffer);
 			}
@@ -103,7 +103,7 @@ private:
 		// Region
 		if (!DisplayClusterHelpers::map::template ExtractValue(InConfigParameters, DisplayClusterProjectionStrings::cfg::mpcdi::Region, RegionId))
 		{
-			if (!FDisplayClusterProjectionPolicyBase::IsEditorOperationMode())
+			if (!FDisplayClusterProjectionPolicyBase::IsEditorOperationMode(InViewport))
 			{
 				UE_LOG(LogDisplayClusterProjectionMPCDI, Error, TEXT("Argument '%s' not found in the config file"), DisplayClusterProjectionStrings::cfg::mpcdi::Region);
 			}
@@ -119,7 +119,7 @@ private:
 		return true;
 	}
 
-	inline bool ImplLoadPFMConfig(const TMap<FString, FString>& InConfigParameters)
+	inline bool ImplLoadPFMConfig(class IDisplayClusterViewport* InViewport, const TMap<FString, FString>& InConfigParameters)
 	{
 		// PFM file
 		FString LocalPFMFile;
@@ -199,7 +199,7 @@ private:
 		return true;
 	}
 
-	inline bool ImplLoadBase(const TMap<FString, FString>& InConfigParameters)
+	inline bool ImplLoadBase(class IDisplayClusterViewport* InViewport, const TMap<FString, FString>& InConfigParameters)
 	{
 		// Origin node (optional)
 		if (DisplayClusterHelpers::map::template ExtractValue(InConfigParameters, DisplayClusterProjectionStrings::cfg::mpcdi::Origin, OriginType))
@@ -208,7 +208,7 @@ private:
 		}
 		else
 		{
-			if (!FDisplayClusterProjectionPolicyBase::IsEditorOperationMode())
+			if (!FDisplayClusterProjectionPolicyBase::IsEditorOperationMode(InViewport))
 			{
 				UE_LOG(LogDisplayClusterProjectionMPCDI, Log, TEXT("No origin node found. VR root will be used as default."));
 			}
diff --git a/Engine/Plugins/Runtime/nDisplay/Source/DisplayClusterProjection/Private/Policy/Manual/DisplayClusterProjectionManualPolicy.cpp b/Engine/Plugins/Runtime/nDisplay/Source/DisplayClusterProjection/Private/Policy/Manual/DisplayClusterProjectionManualPolicy.cpp
--- a/Engine/Plugins/Runtime/nDisplay/Source/DisplayClusterProjection/Private/Policy/Manual/DisplayClusterProjectionManualPolicy.cpp
+++ b/Engine/Plugins/Runtime/nDisplay/Source/DisplayClusterProjection/Private/Policy/Manual/DisplayClusterProjectionManualPolicy.cpp
@@ -38,11 +38,11 @@ bool FDisplayClusterProjectionManualPolicy::HandleStartScene(class IDisplayClust
 	FString DataTypeInString;
 	if(DisplayClusterHelpers::map::template ExtractValueFromString(GetParameters(), FString(DisplayClusterProjectionStrings::cfg::manual::Type), DataTypeInString))
 	{
-		DataType = DataTypeFromString(DataTypeInString);
+		DataType = DataTypeFromString(InViewport, DataTypeInString);
 	}
 	else
 	{
-		if (!IsEditorOperationMode())
+		if (!IsEditorOperationMode(InViewport))
 		{
 			UE_LOG(LogDisplayClusterProjectionManual, Log, TEXT("Undefined manual type '%s'"), *DataTypeInString);
 		}
@@ -53,7 +53,7 @@ bool FDisplayClusterProjectionManualPolicy::HandleStartScene(class IDisplayClust
 	// Get view rotation
 	if (!DisplayClusterHelpers::map::template ExtractValueFromString(GetParameters(), FString(DisplayClusterProjectionStrings::cfg::manual::Rotation), ViewRotation))
 	{
-		if (!IsEditorOperationMode())
+		if (!IsEditorOperationMode(InViewport))
 		{
 			UE_LOG(LogDisplayClusterProjectionManual, Log, TEXT("No rotation specified for projection policy of viewport '%s'"), *InViewport->GetId());
 		}
@@ -86,7 +86,7 @@ bool FDisplayClusterProjectionManualPolicy::HandleStartScene(class IDisplayClust
 
 			if (!ExtractAngles(AnglesLeft, FrustumAngles[0]))
 			{
-				if (!IsEditorOperationMode())
+				if (!IsEditorOperationMode(InViewport))
 				{
 					UE_LOG(LogDisplayClusterProjectionManual, Log, TEXT("Couldn't extract frustum angles from value '%s'"), *AnglesLeft);
 				}
@@ -101,7 +101,7 @@ bool FDisplayClusterProjectionManualPolicy::HandleStartScene(class IDisplayClust
 
 			if (!ExtractAngles(AnglesLeft, FrustumAngles[0]))
 			{
-				if (!IsEditorOperationMode())
+				if (!IsEditorOperationMode(InViewport))
 				{
 					UE_LOG(LogDisplayClusterProjectionManual, Log, TEXT("Couldn't extract frustum angles from value '%s'"), *AnglesLeft);
 				}
@@ -111,7 +111,7 @@ bool FDisplayClusterProjectionManualPolicy::HandleStartScene(class IDisplayClust
 
 			if (!ExtractAngles(AnglesRight, FrustumAngles[1]))
 			{
-				if (!IsEditorOperationMode())
+				if (!IsEditorOperationMode(InViewport))
 				{
 					UE_LOG(LogDisplayClusterProjectionManual, Log, TEXT("Couldn't extract frustum angles from value '%s'"), *AnglesRight);
 				}
@@ -206,7 +206,7 @@ bool FDisplayClusterProjectionManualPolicy::ExtractAngles(const FString& InAngle
 	return true;
 }
 
-FDisplayClusterProjectionManualPolicy::EManualDataType FDisplayClusterProjectionManualPolicy::DataTypeFromString(const FString& DataTypeInString) const
+FDisplayClusterProjectionManualPolicy::EManualDataType FDisplayClusterProjectionManualPolicy::DataTypeFromString(class IDisplayClusterViewport* InViewport, const FString& DataTypeInString) const
 {
 	if (DataTypeInString == DisplayClusterProjectionStrings::cfg::manual::FrustumType::Matrix)
 	{
@@ -218,7 +218,7 @@ FDisplayClusterProjectionManualPolicy::EManualDataType FDisplayClusterProjection
 	}
 	else
 	{
-		if (!IsEditorOperationMode())
+		if (!IsEditorOperationMode(InViewport))
 		{
 			UE_LOG(LogDisplayClusterProjectionManual, Log, TEXT("Undefined manual type '%s'. Matrix by default"), *DataTypeInString);
 		}
diff --git a/Engine/Plugins/Runtime/nDisplay/Source/DisplayClusterProjection/Private/Policy/Manual/DisplayClusterProjectionManualPolicy.h b/Engine/Plugins/Runtime/nDisplay/Source/DisplayClusterProjection/Private/Policy/Manual/DisplayClusterProjectionManualPolicy.h
--- a/Engine/Plugins/Runtime/nDisplay/Source/DisplayClusterProjection/Private/Policy/Manual/DisplayClusterProjectionManualPolicy.h
+++ b/Engine/Plugins/Runtime/nDisplay/Source/DisplayClusterProjection/Private/Policy/Manual/DisplayClusterProjectionManualPolicy.h
@@ -52,7 +52,7 @@ protected:
 
 	virtual bool ExtractAngles(const FString& InAngles, FFrustumAngles& OutAngles);
 private:
-	EManualDataType DataTypeFromString(const FString& DataTypeInString) const;
+	EManualDataType DataTypeFromString(class IDisplayClusterViewport* InViewport, const FString& DataTypeInString) const;
 private:
 	// Current data type (matrix, frustum angle, ...)
 	EManualDataType DataType = EManualDataType::Matrix;
diff --git a/Engine/Plugins/Runtime/nDisplay/Source/DisplayClusterProjection/Private/Policy/Mesh/DisplayClusterProjectionMeshPolicy.cpp b/Engine/Plugins/Runtime/nDisplay/Source/DisplayClusterProjection/Private/Policy/Mesh/DisplayClusterProjectionMeshPolicy.cpp
--- a/Engine/Plugins/Runtime/nDisplay/Source/DisplayClusterProjection/Private/Policy/Mesh/DisplayClusterProjectionMeshPolicy.cpp
+++ b/Engine/Plugins/Runtime/nDisplay/Source/DisplayClusterProjection/Private/Policy/Mesh/DisplayClusterProjectionMeshPolicy.cpp
@@ -72,7 +72,7 @@ bool FDisplayClusterProjectionMeshPolicy::CreateWarpMeshInterface(IDisplayCluste
 
 		if (!bResult)
 		{
-			if (!IsEditorOperationMode())
+			if (!IsEditorOperationMode(InViewport))
 			{
 				UE_LOG(LogDisplayClusterProjectionMesh, Warning, TEXT("Couldn't create mesh warpblend interface"));
 			}
@@ -95,7 +95,7 @@ bool FDisplayClusterProjectionMeshPolicy::HandleStartScene(IDisplayClusterViewpo
 
 	if (!CreateWarpMeshInterface(InViewport))
 	{
-		if (!IsEditorOperationMode())
+		if (!IsEditorOperationMode(InViewport))
 		{
 			UE_LOG(LogDisplayClusterProjectionMesh, Error, TEXT("Couldn't create warp interface for viewport '%s'"), *InViewport->GetId());
 		}
@@ -133,7 +133,7 @@ bool FDisplayClusterProjectionMeshPolicy::GetWarpMeshConfiguration(IDisplayClust
 		}
 #endif
 
-		if (!IsEditorOperationMode())
+		if (!IsEditorOperationMode(InViewport))
 		{
 			UE_LOG(LogDisplayClusterProjectionMesh, Error, TEXT("No component ID '%s' specified for projection policy '%s'"), *ComponentId, *GetId());
 		}
@@ -196,7 +196,7 @@ bool FDisplayClusterProjectionMeshPolicy::GetWarpMeshConfiguration(IDisplayClust
 		OutWarpCfg.ProceduralMeshComponent = Root->GetComponentByName<UProceduralMeshComponent>(ComponentId);
 		if (OutWarpCfg.ProceduralMeshComponent == nullptr)
 		{
-			if (!IsEditorOperationMode())
+			if (!IsEditorOperationMode(InViewport))
 			{
 				UE_LOG(LogDisplayClusterProjectionMesh, Warning, TEXT("Couldn't find mesh component '%s' in RootActor"), *ComponentId);
 			}
diff --git a/Engine/Plugins/Runtime/nDisplay/Source/DisplayClusterProjection/Private/Policy/Simple/DisplayClusterProjectionSimplePolicy.cpp b/Engine/Plugins/Runtime/nDisplay/Source/DisplayClusterProjection/Private/Policy/Simple/DisplayClusterProjectionSimplePolicy.cpp
--- a/Engine/Plugins/Runtime/nDisplay/Source/DisplayClusterProjection/Private/Policy/Simple/DisplayClusterProjectionSimplePolicy.cpp
+++ b/Engine/Plugins/Runtime/nDisplay/Source/DisplayClusterProjection/Private/Policy/Simple/DisplayClusterProjectionSimplePolicy.cpp
@@ -32,7 +32,7 @@ FDisplayClusterProjectionSimplePolicy::~FDisplayClusterProjectionSimplePolicy()
 //////////////////////////////////////////////////////////////////////////////////////////////
 // IDisplayClusterProjectionPolicy
 //////////////////////////////////////////////////////////////////////////////////////////////
-bool FDisplayClusterProjectionSimplePolicy::HandleStartScene(class IDisplayClusterViewport* InViewport)
+bool FDisplayClusterProjectionSimplePolicy::HandleStartScene(IDisplayClusterViewport* InViewport)
 {
 	check(IsInGameThread());
 
@@ -41,7 +41,7 @@ bool FDisplayClusterProjectionSimplePolicy::HandleStartScene(class IDisplayClust
 	// Get assigned screen ID
 	if (!DisplayClusterHelpers::map::template ExtractValue(GetParameters(), DisplayClusterProjectionStrings::cfg::simple::Screen, ScreenId))
 	{
-		if (!IsEditorOperationMode())
+		if (!IsEditorOperationMode(InViewport))
 		{
 			UE_LOG(LogDisplayClusterProjectionSimple, Error, TEXT("No screen ID specified for projection policy of viewport '%s'"), *InViewport->GetId());
 		}
@@ -60,7 +60,7 @@ bool FDisplayClusterProjectionSimplePolicy::HandleStartScene(class IDisplayClust
 	return false;
 }
 
-void FDisplayClusterProjectionSimplePolicy::HandleEndScene(class IDisplayClusterViewport* InViewport)
+void FDisplayClusterProjectionSimplePolicy::HandleEndScene(IDisplayClusterViewport* InViewport)
 {
 	check(IsInGameThread());
 
@@ -68,7 +68,7 @@ void FDisplayClusterProjectionSimplePolicy::HandleEndScene(class IDisplayCluster
 	ScreenCompRef.ResetSceneComponent();
 }
 
-bool FDisplayClusterProjectionSimplePolicy::CalculateView(class IDisplayClusterViewport* InViewport, const uint32 InContextNum, FVector& InOutViewLocation, FRotator& InOutViewRotation, const FVector& ViewOffset, const float WorldToMeters, const float NCP, const float FCP)
+bool FDisplayClusterProjectionSimplePolicy::CalculateView(IDisplayClusterViewport* InViewport, const uint32 InContextNum, FVector& InOutViewLocation, FRotator& InOutViewRotation, const FVector& ViewOffset, const float WorldToMeters, const float NCP, const float FCP)
 {
 	check(IsInGameThread());
 
@@ -92,7 +92,7 @@ bool FDisplayClusterProjectionSimplePolicy::CalculateView(class IDisplayClusterV
 	return true;
 }
 
-bool FDisplayClusterProjectionSimplePolicy::GetProjectionMatrix(class IDisplayClusterViewport* InViewport, const uint32 InContextNum, FMatrix& OutPrjMatrix)
+bool FDisplayClusterProjectionSimplePolicy::GetProjectionMatrix(IDisplayClusterViewport* InViewport, const uint32 InContextNum, FMatrix& OutPrjMatrix)
 {
 	check(IsInGameThread());
 
@@ -177,7 +177,7 @@ bool FDisplayClusterProjectionSimplePolicy::GetProjectionMatrix(class IDisplayCl
 //////////////////////////////////////////////////////////////////////////////////////////////
 // FDisplayClusterProjectionSimplePolicy
 //////////////////////////////////////////////////////////////////////////////////////////////
-bool FDisplayClusterProjectionSimplePolicy::InitializeMeshData(class IDisplayClusterViewport* InViewport)
+bool FDisplayClusterProjectionSimplePolicy::InitializeMeshData(IDisplayClusterViewport* InViewport)
 {
 	UE_LOG(LogDisplayClusterProjectionSimple, Verbose, TEXT("Initializing screen geometry for viewport policy  %s"), *GetId());
 
@@ -185,7 +185,7 @@ bool FDisplayClusterProjectionSimplePolicy::InitializeMeshData(class IDisplayClu
 	ADisplayClusterRootActor* const Root = InViewport->GetOwner().GetRootActor();
 	if (!Root)
 	{
-		if (!IsEditorOperationMode())
+		if (!IsEditorOperationMode(InViewport))
 		{
 			UE_LOG(LogDisplayClusterProjectionSimple, Error, TEXT("Couldn't get a VR root object"));
 		}
@@ -197,7 +197,7 @@ bool FDisplayClusterProjectionSimplePolicy::InitializeMeshData(class IDisplayClu
 	UDisplayClusterScreenComponent* ScreenComp = Root->GetComponentByName<UDisplayClusterScreenComponent>(ScreenId);
 	if (!ScreenComp)
 	{
-		if (!IsEditorOperationMode())
+		if (!IsEditorOperationMode(InViewport))
 		{
 			UE_LOG(LogDisplayClusterProjectionSimple, Warning, TEXT("Couldn't initialize screen component"));
 		}
@@ -208,16 +208,15 @@ bool FDisplayClusterProjectionSimplePolicy::InitializeMeshData(class IDisplayClu
 	return ScreenCompRef.SetSceneComponent(ScreenComp);
 }
 
-void FDisplayClusterProjectionSimplePolicy::ReleaseMeshData(class IDisplayClusterViewport* InViewport)
+void FDisplayClusterProjectionSimplePolicy::ReleaseMeshData(IDisplayClusterViewport* InViewport)
 {
 }
 
-
 #if WITH_EDITOR
 //////////////////////////////////////////////////////////////////////////////////////////////
 // IDisplayClusterProjectionPolicyPreview
 //////////////////////////////////////////////////////////////////////////////////////////////
-UMeshComponent* FDisplayClusterProjectionSimplePolicy::GetOrCreatePreviewMeshComponent(class IDisplayClusterViewport* InViewport, bool& bOutIsRootActorComponent)
+UMeshComponent* FDisplayClusterProjectionSimplePolicy::GetOrCreatePreviewMeshComponent(IDisplayClusterViewport* InViewport, bool& bOutIsRootActorComponent)
 {
 	bOutIsRootActorComponent = true;
 
diff --git a/Engine/Plugins/Runtime/nDisplay/Source/DisplayClusterProjection/Private/Policy/VIOSO/ViosoPolicyConfiguration.cpp b/Engine/Plugins/Runtime/nDisplay/Source/DisplayClusterProjection/Private/Policy/VIOSO/ViosoPolicyConfiguration.cpp
--- a/Engine/Plugins/Runtime/nDisplay/Source/DisplayClusterProjection/Private/Policy/VIOSO/ViosoPolicyConfiguration.cpp
+++ b/Engine/Plugins/Runtime/nDisplay/Source/DisplayClusterProjection/Private/Policy/VIOSO/ViosoPolicyConfiguration.cpp
@@ -6,6 +6,7 @@
 #include "DisplayClusterProjectionLog.h"
 #include "Misc/DisplayClusterHelpers.h"
 #include "Policy/DisplayClusterProjectionPolicyBase.h"
+#include "Render/Viewport/IDisplayClusterViewport.h"
 
 //////////////////////////////////////////////////////////////////////////////////////////////
 // FViosoPolicyConfiguration
@@ -34,15 +35,19 @@ FString FViosoPolicyConfiguration::ToString() const
 	}
 }
 
-bool FViosoPolicyConfiguration::Initialize(const TMap<FString, FString>& InParameters, const FString& InViewportId)
+bool FViosoPolicyConfiguration::Initialize(const TMap<FString, FString>& InParameters, class IDisplayClusterViewport* InViewport)
 {
+	check(InViewport);
+
+	const FString InViewportId = InViewport->GetId();
+
 	if (DisplayClusterHelpers::map::template ExtractValueFromString(InParameters, DisplayClusterProjectionStrings::cfg::VIOSO::Origin, OriginCompId))
 	{
 		UE_LOG(LogDisplayClusterProjectionVIOSO, Verbose, TEXT("Viewport <%s>: Projection parameter '%s' - '%s'"), *InViewportId, DisplayClusterProjectionStrings::cfg::VIOSO::Origin, *OriginCompId);
 	}
 	else
 	{
-		if (!FDisplayClusterProjectionPolicyBase::IsEditorOperationMode())
+		if (!FDisplayClusterProjectionPolicyBase::IsEditorOperationMode(InViewport))
 		{
 			UE_LOG(LogDisplayClusterProjectionVIOSO, Error, TEXT("Viewport <%s>: Projection parameter '%s' not found"), *InViewportId, DisplayClusterProjectionStrings::cfg::VIOSO::Origin);
 		}
@@ -60,7 +65,7 @@ bool FViosoPolicyConfiguration::Initialize(const TMap<FString, FString>& InParam
 
 		if (INIFile.IsEmpty())
 		{
-			if (!FDisplayClusterProjectionPolicyBase::IsEditorOperationMode())
+			if (!FDisplayClusterProjectionPolicyBase::IsEditorOperationMode(InViewport))
 			{
 				UE_LOG(LogDisplayClusterProjectionVIOSO, Error, TEXT("Vioso INI file empty string"));
 			}
@@ -81,7 +86,7 @@ bool FViosoPolicyConfiguration::Initialize(const TMap<FString, FString>& InParam
 		}
 		else
 		{
-			if (!FDisplayClusterProjectionPolicyBase::IsEditorOperationMode())
+			if (!FDisplayClusterProjectionPolicyBase::IsEditorOperationMode(InViewport))
 			{
 				UE_LOG(LogDisplayClusterProjectionVIOSO, Error, TEXT("Viewport <%s>: Projection parameter '%s' not found"), DisplayClusterProjectionStrings::cfg::VIOSO::ChannelName);
 			}
@@ -104,7 +109,7 @@ bool FViosoPolicyConfiguration::Initialize(const TMap<FString, FString>& InParam
 
 		if (CalibrationFile.IsEmpty())
 		{
-			if (!FDisplayClusterProjectionPolicyBase::IsEditorOperationMode())
+			if (!FDisplayClusterProjectionPolicyBase::IsEditorOperationMode(InViewport))
 			{
 				UE_LOG(LogDisplayClusterProjectionVIOSO, Error, TEXT("Calibration file empty string"));
 			}
@@ -122,7 +127,7 @@ bool FViosoPolicyConfiguration::Initialize(const TMap<FString, FString>& InParam
 	}
 	else
 	{
-		if (!FDisplayClusterProjectionPolicyBase::IsEditorOperationMode())
+		if (!FDisplayClusterProjectionPolicyBase::IsEditorOperationMode(InViewport))
 		{
 			UE_LOG(LogDisplayClusterProjectionVIOSO, Error, TEXT("Viewport <%s>: Projection parameter '%s' not found"), DisplayClusterProjectionStrings::cfg::VIOSO::File);
 		}
@@ -147,7 +152,7 @@ bool FViosoPolicyConfiguration::Initialize(const TMap<FString, FString>& InParam
 		else
 		{
 			CalibrationIndex = 0;
-			if (!FDisplayClusterProjectionPolicyBase::IsEditorOperationMode())
+			if (!FDisplayClusterProjectionPolicyBase::IsEditorOperationMode(InViewport))
 			{
 				UE_LOG(LogDisplayClusterProjectionVIOSO, Error, TEXT("Viewport <%s>: Projection parameter '%s' or '%s'  not found, Use default value - '%d'"), *InViewportId, DisplayClusterProjectionStrings::cfg::VIOSO::CalibIndex, DisplayClusterProjectionStrings::cfg::VIOSO::CalibAdapter, CalibrationIndex);
 			}
@@ -158,7 +163,7 @@ bool FViosoPolicyConfiguration::Initialize(const TMap<FString, FString>& InParam
 	{
 		if (BaseMatrix == FMatrix::Identity)
 		{
-			if (!FDisplayClusterProjectionPolicyBase::IsEditorOperationMode())
+			if (!FDisplayClusterProjectionPolicyBase::IsEditorOperationMode(InViewport))
 			{
 				UE_LOG(LogDisplayClusterProjectionVIOSO, Error, TEXT("Viewport <%s>: Projection parameter '%s' values read as Identity"), *InViewportId, DisplayClusterProjectionStrings::cfg::VIOSO::BaseMatrix);
 			}
@@ -169,7 +174,7 @@ bool FViosoPolicyConfiguration::Initialize(const TMap<FString, FString>& InParam
 	else
 	{
 		BaseMatrix = FMatrix::Identity;
-		if (!FDisplayClusterProjectionPolicyBase::IsEditorOperationMode())
+		if (!FDisplayClusterProjectionPolicyBase::IsEditorOperationMode(InViewport))
 		{
 			UE_LOG(LogDisplayClusterProjectionVIOSO, Error, TEXT("Viewport <%s>: Projection parameter '%s' not found, Use default value - '%s'"), *InViewportId, DisplayClusterProjectionStrings::cfg::VIOSO::BaseMatrix, *BaseMatrix.ToString());
 		}
@@ -182,7 +187,7 @@ bool FViosoPolicyConfiguration::Initialize(const TMap<FString, FString>& InParam
 	}
 	else
 	{
-		if (!FDisplayClusterProjectionPolicyBase::IsEditorOperationMode())
+		if (!FDisplayClusterProjectionPolicyBase::IsEditorOperationMode(InViewport))
 		{
 			UE_LOG(LogDisplayClusterProjectionVIOSO, Error, TEXT("Viewport <%s>: Projection parameter '%s' not found. Assigned default value - '%f'"), *InViewportId, DisplayClusterProjectionStrings::cfg::VIOSO::Gamma, Gamma);
 		}
diff --git a/Engine/Plugins/Runtime/nDisplay/Source/DisplayClusterProjection/Private/Policy/VIOSO/ViosoPolicyConfiguration.h b/Engine/Plugins/Runtime/nDisplay/Source/DisplayClusterProjection/Private/Policy/VIOSO/ViosoPolicyConfiguration.h
--- a/Engine/Plugins/Runtime/nDisplay/Source/DisplayClusterProjection/Private/Policy/VIOSO/ViosoPolicyConfiguration.h
+++ b/Engine/Plugins/Runtime/nDisplay/Source/DisplayClusterProjection/Private/Policy/VIOSO/ViosoPolicyConfiguration.h
@@ -20,7 +20,7 @@ struct FViosoPolicyConfiguration
 	//@todo add more vioso options, if required
 	float Gamma = 1.f;
 
-	bool Initialize(const TMap<FString, FString>& InParameters, const FString& InViewportId);
+	bool Initialize(const TMap<FString, FString>& InParameters, class IDisplayClusterViewport* InViewport);
 
 	FString ToString() const;
 };
diff --git a/Engine/Plugins/Runtime/nDisplay/Source/DisplayClusterProjection/Private/Policy/VIOSO/Windows/DisplayClusterProjectionVIOSOPolicy.cpp b/Engine/Plugins/Runtime/nDisplay/Source/DisplayClusterProjection/Private/Policy/VIOSO/Windows/DisplayClusterProjectionVIOSOPolicy.cpp
--- a/Engine/Plugins/Runtime/nDisplay/Source/DisplayClusterProjection/Private/Policy/VIOSO/Windows/DisplayClusterProjectionVIOSOPolicy.cpp
+++ b/Engine/Plugins/Runtime/nDisplay/Source/DisplayClusterProjection/Private/Policy/VIOSO/Windows/DisplayClusterProjectionVIOSOPolicy.cpp
@@ -167,9 +167,9 @@ bool FDisplayClusterProjectionVIOSOPolicy::HandleStartScene(class IDisplayCluste
 	check(IsInGameThread());	
 
 	// Read VIOSO config data from nDisplay config file
-	if (!ViosoConfigData.Initialize(GetParameters(), InViewport->GetId()))
+	if (!ViosoConfigData.Initialize(GetParameters(), InViewport))
 	{
-		if (!IsEditorOperationMode())
+		if (!IsEditorOperationMode(InViewport))
 		{
 			UE_LOG(LogDisplayClusterProjectionVIOSO, Error, TEXT("Couldn't read VIOSO configuration from the config file for viewport -'%s'"), *InViewport->GetId());
 		}
@@ -235,7 +235,7 @@ bool FDisplayClusterProjectionVIOSOPolicy::CalculateView(IDisplayClusterViewport
 	FScopeLock lock(&DllAccessCS);
 	if (!Views[InContextNum].UpdateVIOSO(InViewport, InContextNum, LocalEyeOrigin, LocalRotator, WorldToMeters, NCP, FCP))
 	{
-		if (Views[InContextNum].IsValid() && !FDisplayClusterProjectionPolicyBase::IsEditorOperationMode())
+		if (Views[InContextNum].IsValid() && !FDisplayClusterProjectionPolicyBase::IsEditorOperationMode(InViewport))
 		{
 			// Vioso api used, but failed inside math. The config base matrix or vioso geometry is invalid
 			UE_LOG(LogDisplayClusterProjectionVIOSO, Error, TEXT("Couldn't Calculate View for VIOSO viewport '%s'"), *InViewport->GetId());
diff --git a/Engine/Plugins/Runtime/nDisplay/Source/DisplayClusterProjection/Public/Containers/DisplayClusterProjectionCameraPolicySettings.h b/Engine/Plugins/Runtime/nDisplay/Source/DisplayClusterProjection/Public/Containers/DisplayClusterProjectionCameraPolicySettings.h
--- a/Engine/Plugins/Runtime/nDisplay/Source/DisplayClusterProjection/Public/Containers/DisplayClusterProjectionCameraPolicySettings.h
+++ b/Engine/Plugins/Runtime/nDisplay/Source/DisplayClusterProjection/Public/Containers/DisplayClusterProjectionCameraPolicySettings.h
@@ -13,4 +13,7 @@ struct FDisplayClusterProjectionCameraPolicySettings
 
 	// Move incamera frustum on this value to fix broken lens on physic camera
 	FVector  FrustumOffset = FVector::ZeroVector;
+
+	// When enabled, ignore all next updates, until null cam ref not assigned (UE-137222)
+	bool bCameraOverrideDefaults = false;
 };
diff --git a/Engine/Plugins/Runtime/nDisplay/Source/DisplayClusterShaders/Private/Shaders/DisplayClusterShadersWarpblend_ICVFX.cpp b/Engine/Plugins/Runtime/nDisplay/Source/DisplayClusterShaders/Private/Shaders/DisplayClusterShadersWarpblend_ICVFX.cpp
--- a/Engine/Plugins/Runtime/nDisplay/Source/DisplayClusterShaders/Private/Shaders/DisplayClusterShadersWarpblend_ICVFX.cpp
+++ b/Engine/Plugins/Runtime/nDisplay/Source/DisplayClusterShaders/Private/Shaders/DisplayClusterShadersWarpblend_ICVFX.cpp
@@ -54,7 +54,8 @@ enum class EIcvfxShaderType : uint8
 namespace IcvfxShaderPermutation
 {
 	// Shared permutation for picp warp
-	class FIcvfxShaderViewportInput    : SHADER_PERMUTATION_BOOL("VIEWPORT_INPUT");
+	class FIcvfxShaderViewportInput      : SHADER_PERMUTATION_BOOL("VIEWPORT_INPUT");
+	class FIcvfxShaderViewportInputAlpha : SHADER_PERMUTATION_BOOL("VIEWPORT_INPUT_ALPHA");
 
 	class FIcvfxShaderOverlayUnder     : SHADER_PERMUTATION_BOOL("OVERLAY_UNDER");
 	class FIcvfxShaderOverlayOver      : SHADER_PERMUTATION_BOOL("OVERLAY_OVER");
@@ -76,6 +77,8 @@ namespace IcvfxShaderPermutation
 
 	using FCommonPSDomain = TShaderPermutationDomain<
 		FIcvfxShaderViewportInput,
+		FIcvfxShaderViewportInputAlpha,
+
 		FIcvfxShaderOverlayUnder,
 		FIcvfxShaderOverlayOver,
 		FIcvfxShaderOverlayAlpha,
@@ -119,6 +122,11 @@ namespace IcvfxShaderPermutation
 			{
 				return false;
 			}
+
+			if (PermutationVector.Get<FIcvfxShaderViewportInputAlpha>())
+			{
+				return false;
+			}
 		}
 
 		if (!PermutationVector.Get<FIcvfxShaderInnerCamera>())
@@ -402,6 +410,11 @@ public:
 			RenderPassData.PSParameters.InputSampler = TStaticSamplerState<SF_Trilinear, AM_Clamp, AM_Clamp, AM_Clamp>::GetRHI();
 
 			RenderPassData.PSPermutationVector.Set<IcvfxShaderPermutation::FIcvfxShaderViewportInput>(true);
+
+			if (WarpBlendParameters.bRenderAlphaChannel)
+			{
+				RenderPassData.PSPermutationVector.Set<IcvfxShaderPermutation::FIcvfxShaderViewportInputAlpha>(true);
+			}
 		}
 
 		// Vertex shader viewport rect
diff --git a/Engine/Plugins/Runtime/nDisplay/Source/DisplayClusterShaders/Private/Shaders/DisplayClusterShadersWarpblend_MPCDI.cpp b/Engine/Plugins/Runtime/nDisplay/Source/DisplayClusterShaders/Private/Shaders/DisplayClusterShadersWarpblend_MPCDI.cpp
--- a/Engine/Plugins/Runtime/nDisplay/Source/DisplayClusterShaders/Private/Shaders/DisplayClusterShadersWarpblend_MPCDI.cpp
+++ b/Engine/Plugins/Runtime/nDisplay/Source/DisplayClusterShaders/Private/Shaders/DisplayClusterShadersWarpblend_MPCDI.cpp
@@ -59,10 +59,20 @@ namespace MpcdiShaderPermutation
 	class FMpcdiShaderAlphaMapBlending : SHADER_PERMUTATION_BOOL("ALPHAMAP_BLENDING");
 	class FMpcdiShaderBetaMapBlending : SHADER_PERMUTATION_BOOL("BETAMAP_BLENDING");
 
+	class FMpcdiShaderViewportInputAlpha : SHADER_PERMUTATION_BOOL("VIEWPORT_INPUT_ALPHA");
+
 	class FMpcdiShaderMeshWarp : SHADER_PERMUTATION_BOOL("MESH_WARP");
 
-	using FCommonVSDomain = TShaderPermutationDomain<FMpcdiShaderMeshWarp>;
-	using FCommonPSDomain = TShaderPermutationDomain<FMpcdiShaderAlphaMapBlending, FMpcdiShaderBetaMapBlending, FMpcdiShaderMeshWarp>;
+	using FCommonVSDomain = TShaderPermutationDomain<
+		FMpcdiShaderMeshWarp
+	>;
+
+	using FCommonPSDomain = TShaderPermutationDomain<
+		FMpcdiShaderAlphaMapBlending,
+		FMpcdiShaderBetaMapBlending,
+		FMpcdiShaderViewportInputAlpha,
+		FMpcdiShaderMeshWarp
+	>;
 
 	bool ShouldCompileCommonPSPermutation(const FGlobalShaderPermutationParameters& Parameters, const FCommonPSDomain& PermutationVector)
 	{
@@ -254,6 +264,11 @@ public:
 
 		if (WarpBlendParameters.WarpInterface.IsValid())
 		{
+			if (WarpBlendParameters.bRenderAlphaChannel)
+			{
+				RenderPassData.PSPermutationVector.Set<MpcdiShaderPermutation::FMpcdiShaderViewportInputAlpha>(true);
+			}
+
 			switch (WarpBlendParameters.WarpInterface->GetWarpGeometryType())
 			{
 				case EDisplayClusterWarpGeometryType::WarpMesh:
diff --git a/Engine/Plugins/Runtime/nDisplay/Source/DisplayClusterShaders/Private/WarpBlend/Math/DisplayClusterWarpBlendMath_Frustum.cpp b/Engine/Plugins/Runtime/nDisplay/Source/DisplayClusterShaders/Private/WarpBlend/Math/DisplayClusterWarpBlendMath_Frustum.cpp
--- a/Engine/Plugins/Runtime/nDisplay/Source/DisplayClusterShaders/Private/WarpBlend/Math/DisplayClusterWarpBlendMath_Frustum.cpp
+++ b/Engine/Plugins/Runtime/nDisplay/Source/DisplayClusterShaders/Private/WarpBlend/Math/DisplayClusterWarpBlendMath_Frustum.cpp
@@ -56,11 +56,11 @@ bool FDisplayClusterWarpBlendMath_Frustum::ImplCalcFrustum_FULL_WarpMap()
 
 	bool bAllPointsInFrustum = true;
 
-	const int32 PointsAmmount = GeometryContext.GeometryProxy.WarpMapTexture->GetTotalPoints();
+	const int32 PointsAmount = GeometryContext.GeometryProxy.WarpMapTexture->GetTotalPoints();
 	const FVector4f* SourcePts = (FVector4f*)(GeometryContext.GeometryProxy.WarpMapTexture->GetData());
 
 	// Search a camera space frustum
-	for (int32 PointIndex = 0; PointIndex < PointsAmmount; ++PointIndex)
+	for (int32 PointIndex = 0; PointIndex < PointsAmount; ++PointIndex)
 	{
 		const FVector4 Pts(SourcePts[PointIndex]);
 		if (GetProjectionClip(Pts) == false)
@@ -144,7 +144,7 @@ bool FDisplayClusterWarpBlendMath_Frustum::ImplCalcFrustum_LOD_WarpMap()
 		GeometryContext.GeometryProxy.UpdateGeometryLOD(FIntPoint(WarpMapLODRatio));
 	}
 
-	const int32 PointsAmmount = GeometryContext.GeometryProxy.WarpMapTexture->GetTotalPoints();
+	const int32 PointsAmount = GeometryContext.GeometryProxy.WarpMapTexture->GetTotalPoints();
 	const FVector4f* SourcePts = (FVector4f*)(GeometryContext.GeometryProxy.WarpMapTexture->GetData());
 
 	// Search a camera space frustum
diff --git a/Engine/Plugins/Runtime/nDisplay/Source/DisplayClusterShaders/Public/ShaderParameters/DisplayClusterShaderParameters_WarpBlend.h b/Engine/Plugins/Runtime/nDisplay/Source/DisplayClusterShaders/Public/ShaderParameters/DisplayClusterShaderParameters_WarpBlend.h
--- a/Engine/Plugins/Runtime/nDisplay/Source/DisplayClusterShaders/Public/ShaderParameters/DisplayClusterShaderParameters_WarpBlend.h
+++ b/Engine/Plugins/Runtime/nDisplay/Source/DisplayClusterShaders/Public/ShaderParameters/DisplayClusterShaderParameters_WarpBlend.h
@@ -33,4 +33,7 @@ struct FDisplayClusterShaderParameters_WarpBlend
 
 	// Context data
 	FDisplayClusterWarpContext Context;
+
+	// Render alpha channel from input texture to warp output
+	bool bRenderAlphaChannel = false;
 };
