nDisplay: Fixed Vulkan RHI resource access issues
RHIVulkan: Fixed SetupCopyOrResolveRegion() function

diff --git a/Engine/Plugins/Runtime/nDisplay/Source/DisplayCluster/Private/Render/Viewport/DisplayClusterViewportManagerProxy.cpp b/Engine/Plugins/Runtime/nDisplay/Source/DisplayCluster/Private/Render/Viewport/DisplayClusterViewportManagerProxy.cpp
--- a/Engine/Plugins/Runtime/nDisplay/Source/DisplayCluster/Private/Render/Viewport/DisplayClusterViewportManagerProxy.cpp
+++ b/Engine/Plugins/Runtime/nDisplay/Source/DisplayCluster/Private/Render/Viewport/DisplayClusterViewportManagerProxy.cpp
@@ -264,6 +264,7 @@ static void ImplClearRenderTargetResource_RenderThread(FRHICommandListImmediate&
 		DrawClearQuad(RHICmdList, FLinearColor::Black);
 	}
 	RHICmdList.EndRenderPass();
+	RHICmdList.Transition(FRHITransitionInfo(InRenderTargetTexture, ERHIAccess::Unknown, ERHIAccess::SRVMask));
 }
 
 void FDisplayClusterViewportManagerProxy::ImplClearFrameTargets_RenderThread(FRHICommandListImmediate& RHICmdList) const
diff --git a/Engine/Plugins/Runtime/nDisplay/Source/DisplayCluster/Private/Render/Viewport/DisplayClusterViewportProxy.cpp b/Engine/Plugins/Runtime/nDisplay/Source/DisplayCluster/Private/Render/Viewport/DisplayClusterViewportProxy.cpp
--- a/Engine/Plugins/Runtime/nDisplay/Source/DisplayCluster/Private/Render/Viewport/DisplayClusterViewportProxy.cpp
+++ b/Engine/Plugins/Runtime/nDisplay/Source/DisplayCluster/Private/Render/Viewport/DisplayClusterViewportProxy.cpp
@@ -537,6 +537,8 @@ void ResampleCopyTextureImpl_RenderThread(FRHICommandListImmediate& RHICmdList,
 	// Texture format mismatch, use a shader to do the copy.
 	// #todo-renderpasses there's no explicit resolve here? Do we need one?
 	FRHIRenderPassInfo RPInfo(DstTexture, ERenderTargetActions::Load_Store);
+	TransitionRenderPassTargets(RHICmdList, RPInfo);
+
 	RHICmdList.BeginRenderPass(RPInfo, TEXT("DisaplyClusterRender_ResampleTexture"));
 	{
 		FIntVector SrcSizeXYZ = SrcTexture->GetSizeXYZ();
@@ -589,6 +591,7 @@ void ResampleCopyTextureImpl_RenderThread(FRHICommandListImmediate& RHICmdList,
 		FPixelShaderUtils::DrawFullscreenQuad(RHICmdList, 1);
 	}
 	RHICmdList.EndRenderPass();
+	RHICmdList.Transition(FRHITransitionInfo(DstTexture, ERHIAccess::Unknown, ERHIAccess::SRVMask));
 }
 
 void ImplResolveResource(FRHICommandListImmediate& RHICmdList, FRHITexture2D* InputResource, const FIntRect& InputRect, FRHITexture2D* OutputResource, const FIntRect& OutputRect, const bool bOutputIsMipsResource, const bool bOutputIsPreviewResource)
@@ -605,12 +608,25 @@ void ImplResolveResource(FRHICommandListImmediate& RHICmdList, FRHITexture2D* In
 		CopyInfo.Size = FIntVector(InputRect.Width(), InputRect.Height(), 0);
 		CopyInfo.SourcePosition.X = InputRect.Min.X;
 		CopyInfo.SourcePosition.Y = InputRect.Min.Y;
-
-		CopyInfo.DestMipIndex = 0;
 		CopyInfo.DestPosition.X = OutputRect.Min.X;
 		CopyInfo.DestPosition.Y = OutputRect.Min.Y;
 
+		// Copy texture
+		FRHITransitionInfo TransitionsBefore[] =
+		{
+			FRHITransitionInfo(InputResource, ERHIAccess::Unknown, ERHIAccess::CopySrc),
+			FRHITransitionInfo(OutputResource, ERHIAccess::Unknown, ERHIAccess::CopyDest)
+		};
+		RHICmdList.Transition(MakeArrayView(TransitionsBefore, UE_ARRAY_COUNT(TransitionsBefore)));
+
 		RHICmdList.CopyTexture(InputResource, OutputResource, CopyInfo);
+
+		FRHITransitionInfo TransitionsAfter[] =
+		{
+			FRHITransitionInfo(InputResource, ERHIAccess::CopySrc, ERHIAccess::SRVMask),
+			FRHITransitionInfo(OutputResource, ERHIAccess::CopyDest, ERHIAccess::SRVMask)
+		};
+		RHICmdList.Transition(MakeArrayView(TransitionsAfter, UE_ARRAY_COUNT(TransitionsAfter)));
 	}
 	else
 	{
diff --git a/Engine/Plugins/Runtime/nDisplay/Source/DisplayClusterShaders/Private/Shaders/DisplayClusterShadersPostprocess_Blur.cpp b/Engine/Plugins/Runtime/nDisplay/Source/DisplayClusterShaders/Private/Shaders/DisplayClusterShadersPostprocess_Blur.cpp
--- a/Engine/Plugins/Runtime/nDisplay/Source/DisplayClusterShaders/Private/Shaders/DisplayClusterShadersPostprocess_Blur.cpp
+++ b/Engine/Plugins/Runtime/nDisplay/Source/DisplayClusterShaders/Private/Shaders/DisplayClusterShadersPostprocess_Blur.cpp
@@ -138,8 +138,10 @@ static void PicpBlurPostProcess_RenderThread(
 	TShaderMapRef<FDirectProjectionVS>                VertexShader(GlobalShaderMap);
 	TShaderMapRef<TPicpBlurPostProcessPS<ShaderType>> PixelShader(GlobalShaderMap);
 
-	FRHIRenderPassInfo RPInfo1(OutRenderTargetableTexture, ERenderTargetActions::DontLoad_Store);
-	RHICmdList.BeginRenderPass(RPInfo1, TEXT("nDisplay_PicpPostProcessBlur"));
+	FRHIRenderPassInfo RPInfo(OutRenderTargetableTexture, ERenderTargetActions::DontLoad_Store);
+	TransitionRenderPassTargets(RHICmdList, RPInfo);
+
+	RHICmdList.BeginRenderPass(RPInfo, TEXT("nDisplay_PicpPostProcessBlur"));
 	{
 		FGraphicsPipelineStateInitializer GraphicsPSOInit;
 		RHICmdList.ApplyCachedRenderTargets(GraphicsPSOInit);
@@ -163,6 +165,7 @@ static void PicpBlurPostProcess_RenderThread(
 		FPixelShaderUtils::DrawFullscreenQuad(RHICmdList, 1);
 	}
 	RHICmdList.EndRenderPass();
+	RHICmdList.Transition(FRHITransitionInfo(OutRenderTargetableTexture, ERHIAccess::Unknown, ERHIAccess::SRVMask));
 }
 
 DECLARE_GPU_STAT_NAMED(nDisplay_Picp_PostProcess_Compose, TEXT("nDisplay Picp_PostProcess::Compose"));
diff --git a/Engine/Plugins/Runtime/nDisplay/Source/DisplayClusterShaders/Private/Shaders/DisplayClusterShadersPostprocess_OutputRemap.cpp b/Engine/Plugins/Runtime/nDisplay/Source/DisplayClusterShaders/Private/Shaders/DisplayClusterShadersPostprocess_OutputRemap.cpp
--- a/Engine/Plugins/Runtime/nDisplay/Source/DisplayClusterShaders/Private/Shaders/DisplayClusterShadersPostprocess_OutputRemap.cpp
+++ b/Engine/Plugins/Runtime/nDisplay/Source/DisplayClusterShaders/Private/Shaders/DisplayClusterShadersPostprocess_OutputRemap.cpp
@@ -142,45 +142,47 @@ bool FDisplayClusterShadersPostprocess_OutputRemap::RenderPostprocess_OutputRema
 
 	FRHIRenderPassInfo RPInfo(InRenderTargetableDestTexture, ERenderTargetActions::Load_Store);
 	TransitionRenderPassTargets(RHICmdList, RPInfo);
+
 	RHICmdList.BeginRenderPass(RPInfo, TEXT("nDisplay_OutputRemap"));
+	{
+		const FIntPoint TargetSizeXY = InRenderTargetableDestTexture->GetSizeXY();
+		RHICmdList.SetViewport(0, 0, 0.0f, TargetSizeXY.X, TargetSizeXY.Y, 1.0f);
 
-	const FIntPoint TargetSizeXY = InRenderTargetableDestTexture->GetSizeXY();
-	RHICmdList.SetViewport(0, 0, 0.0f, TargetSizeXY.X, TargetSizeXY.Y, 1.0f);
+		const ERHIFeatureLevel::Type RenderFeatureLevel = GMaxRHIFeatureLevel;
+		const auto GlobalShaderMap = GetGlobalShaderMap(RenderFeatureLevel);
 
-	const ERHIFeatureLevel::Type RenderFeatureLevel = GMaxRHIFeatureLevel;
-	const auto GlobalShaderMap = GetGlobalShaderMap(RenderFeatureLevel);
+		TShaderMapRef<FOutputRemapVS> VertexShader(GlobalShaderMap);
+		TShaderMapRef<FOutputRemapPS> PixelShader(GlobalShaderMap);
 
-	TShaderMapRef<FOutputRemapVS> VertexShader(GlobalShaderMap);
-	TShaderMapRef<FOutputRemapPS> PixelShader(GlobalShaderMap);
+		FGraphicsPipelineStateInitializer GraphicsPSOInit;
+		// Set the graphic pipeline state.
+		RHICmdList.ApplyCachedRenderTargets(GraphicsPSOInit);
 
-	FGraphicsPipelineStateInitializer GraphicsPSOInit;
-	// Set the graphic pipeline state.
-	RHICmdList.ApplyCachedRenderTargets(GraphicsPSOInit);
+		if (MeshProxy->BeginRender_RenderThread(RHICmdList, GraphicsPSOInit))
+		{
+			GraphicsPSOInit.DepthStencilState = TStaticDepthStencilState<false, CF_Never>::GetRHI();
+			GraphicsPSOInit.RasterizerState = TStaticRasterizerState<>::GetRHI();
 
-	if (MeshProxy->BeginRender_RenderThread(RHICmdList, GraphicsPSOInit))
-	{
-		GraphicsPSOInit.DepthStencilState = TStaticDepthStencilState<false, CF_Never>::GetRHI();
-		GraphicsPSOInit.RasterizerState = TStaticRasterizerState<>::GetRHI();
+			GraphicsPSOInit.BoundShaderState.VertexShaderRHI = VertexShader.GetVertexShader();
+			GraphicsPSOInit.BoundShaderState.PixelShaderRHI = PixelShader.GetPixelShader();
 
-		GraphicsPSOInit.BoundShaderState.VertexShaderRHI = VertexShader.GetVertexShader();
-		GraphicsPSOInit.BoundShaderState.PixelShaderRHI = PixelShader.GetPixelShader();
+			GraphicsPSOInit.PrimitiveType = PT_TriangleList;
+			GraphicsPSOInit.BlendState = TStaticBlendState <>::GetRHI();
+			SetGraphicsPipelineState(RHICmdList, GraphicsPSOInit, 0);
 
-		GraphicsPSOInit.PrimitiveType = PT_TriangleList;
-		GraphicsPSOInit.BlendState = TStaticBlendState <>::GetRHI();
-		SetGraphicsPipelineState(RHICmdList, GraphicsPSOInit, 0);
+			FOutputRemapPixelShaderParameters PSParameters;
+			{
+				PSParameters.PostprocessInput0 = InSourceTexture;
+				PSParameters.PostprocessInput0Sampler = TStaticSamplerState<SF_Trilinear, AM_Clamp, AM_Clamp, AM_Clamp>::GetRHI();
+			}
+			SetShaderParameters(RHICmdList, PixelShader, PixelShader.GetPixelShader(), PSParameters);
 
-		FOutputRemapPixelShaderParameters PSParameters;
-		{
-			PSParameters.PostprocessInput0 = InSourceTexture;
-			PSParameters.PostprocessInput0Sampler = TStaticSamplerState<SF_Trilinear, AM_Clamp, AM_Clamp, AM_Clamp>::GetRHI();
+			MeshProxy->FinishRender_RenderThread(RHICmdList);
+			bResult = true;
 		}
-		SetShaderParameters(RHICmdList, PixelShader, PixelShader.GetPixelShader(), PSParameters);
-
-		MeshProxy->FinishRender_RenderThread(RHICmdList);
-		bResult = true;
 	}
-
 	RHICmdList.EndRenderPass();
+	RHICmdList.Transition(FRHITransitionInfo(InRenderTargetableDestTexture, ERHIAccess::Unknown, ERHIAccess::SRVMask));
 
 	return bResult;
 }
diff --git a/Engine/Plugins/Runtime/nDisplay/Source/DisplayClusterShaders/Private/Shaders/DisplayClusterShadersWarpblend_ICVFX.cpp b/Engine/Plugins/Runtime/nDisplay/Source/DisplayClusterShaders/Private/Shaders/DisplayClusterShadersWarpblend_ICVFX.cpp
--- a/Engine/Plugins/Runtime/nDisplay/Source/DisplayClusterShaders/Private/Shaders/DisplayClusterShadersWarpblend_ICVFX.cpp
+++ b/Engine/Plugins/Runtime/nDisplay/Source/DisplayClusterShaders/Private/Shaders/DisplayClusterShadersWarpblend_ICVFX.cpp
@@ -930,9 +930,14 @@ bool FDisplayClusterShadersWarpblend_ICVFX::RenderWarpBlend_ICVFX(FRHICommandLis
 	// Do single warp render pass
 	bool bIsRenderSuccess = false;
 	FRHIRenderPassInfo RPInfo(InWarpBlendParameters.Dest.Texture, ERenderTargetActions::Load_Store);
+	TransitionRenderPassTargets(RHICmdList, RPInfo);
+
 	RHICmdList.BeginRenderPass(RPInfo, TEXT("nDisplay_IcvfxWarpBlend"));
-	bIsRenderSuccess = IcvfxPassRenderer.Render(RHICmdList);
+	{
+		bIsRenderSuccess = IcvfxPassRenderer.Render(RHICmdList);
+	}
 	RHICmdList.EndRenderPass();
+	RHICmdList.Transition(FRHITransitionInfo(InWarpBlendParameters.Dest.Texture, ERHIAccess::Unknown, ERHIAccess::SRVMask));
 
 	return bIsRenderSuccess;
 };
diff --git a/Engine/Plugins/Runtime/nDisplay/Source/DisplayClusterShaders/Private/Shaders/DisplayClusterShadersWarpblend_MPCDI.cpp b/Engine/Plugins/Runtime/nDisplay/Source/DisplayClusterShaders/Private/Shaders/DisplayClusterShadersWarpblend_MPCDI.cpp
--- a/Engine/Plugins/Runtime/nDisplay/Source/DisplayClusterShaders/Private/Shaders/DisplayClusterShadersWarpblend_MPCDI.cpp
+++ b/Engine/Plugins/Runtime/nDisplay/Source/DisplayClusterShaders/Private/Shaders/DisplayClusterShadersWarpblend_MPCDI.cpp
@@ -583,10 +583,16 @@ bool FDisplayClusterShadersWarpblend_MPCDI::RenderWarpBlend_MPCDI(FRHICommandLis
 	SCOPED_DRAW_EVENT(RHICmdList, nDisplay_Mpcdi_WarpBlend);
 
 	// Do single-pass warp&blend render
+	bool bIsRenderSuccess = false;
 	FRHIRenderPassInfo RPInfo(InWarpBlendParameters.Dest.Texture, ERenderTargetActions::Load_Store);
+	TransitionRenderPassTargets(RHICmdList, RPInfo);
+
 	RHICmdList.BeginRenderPass(RPInfo, TEXT("nDisplay_MpcdiWarpBlend"));
-	bool bIsRenderSuccess = MpcdiPassRenderer.Render(RHICmdList);
+	{
+		bIsRenderSuccess = MpcdiPassRenderer.Render(RHICmdList);
+	}
 	RHICmdList.EndRenderPass();
+	RHICmdList.Transition(FRHITransitionInfo(InWarpBlendParameters.Dest.Texture, ERHIAccess::Unknown, ERHIAccess::SRVMask));
 
 	return bIsRenderSuccess;
 };
diff --git a/Engine/Source/Runtime/VulkanRHI/Private/VulkanRenderTarget.cpp b/Engine/Source/Runtime/VulkanRHI/Private/VulkanRenderTarget.cpp
--- a/Engine/Source/Runtime/VulkanRHI/Private/VulkanRenderTarget.cpp
+++ b/Engine/Source/Runtime/VulkanRHI/Private/VulkanRenderTarget.cpp
@@ -81,18 +81,38 @@ template<typename RegionType>
 static void SetupCopyOrResolveRegion(RegionType& Region, const FVulkanSurface& SrcSurface, const FVulkanSurface& DstSurface, const VkImageSubresourceRange& SrcRange, const VkImageSubresourceRange& DstRange, const FResolveParams& ResolveParams)
 {
 	FMemory::Memzero(Region);
-	ensure(SrcSurface.Width == DstSurface.Width && SrcSurface.Height == DstSurface.Height);
 
-	if(ResolveParams.Rect.X1 >= 0 && ResolveParams.Rect.Y1 >= 0 && ResolveParams.DestRect.X1 >= 0 && ResolveParams.DestRect.Y1 >= 0)
+	bool bCopySrcSubRect = ResolveParams.Rect.IsValid()     && (ResolveParams.Rect.X1 != 0     || ResolveParams.Rect.Y1 != 0     || ResolveParams.Rect.X2 != SrcSurface.Width     || ResolveParams.Rect.Y2 != SrcSurface.Height);
+	bool bCopyDstSubRect = ResolveParams.DestRect.IsValid() && (ResolveParams.DestRect.X1 != 0 || ResolveParams.DestRect.Y1 != 0 || ResolveParams.DestRect.X2 != DstSurface.Width || ResolveParams.DestRect.Y2 != DstSurface.Height);
+	if(bCopySrcSubRect || bCopyDstSubRect)
 	{
-		Region.srcOffset.x = ResolveParams.Rect.X1;
-		Region.srcOffset.y = ResolveParams.Rect.Y1;
-		Region.dstOffset.x = ResolveParams.DestRect.X1;
-		Region.dstOffset.y = ResolveParams.DestRect.Y1;
+		const FResolveRect& SrcSubRect = ResolveParams.Rect.IsValid()     ? ResolveParams.Rect     : FResolveRect(0, 0, SrcSurface.Width, SrcSurface.Height);
+		const FResolveRect& DstSubRect = ResolveParams.DestRect.IsValid() ? ResolveParams.DestRect : FResolveRect(0, 0, DstSurface.Width, DstSurface.Height);
+
+		const uint32 SrcSubRectWidth  = SrcSubRect.X2 - SrcSubRect.X1;
+		const uint32 SrcSubRectHeight = SrcSubRect.Y2 - SrcSubRect.Y1;
+
+		const uint32 DstSubRectWidth  = DstSubRect.X2 - DstSubRect.X1;
+		const uint32 DstSubRectHeight = DstSubRect.Y2 - DstSubRect.Y1;
+
+		ensure(SrcSubRectWidth == DstSubRectWidth && SrcSubRectHeight == DstSubRectHeight);
+
+		Region.srcOffset.x = SrcSubRect.X1;
+		Region.srcOffset.y = SrcSubRect.Y1;
+		Region.dstOffset.x = DstSubRect.X1;
+		Region.dstOffset.y = DstSubRect.Y1;
+
+		Region.extent.width = FMath::Max(1u, SrcSubRectWidth >> ResolveParams.MipIndex);
+		Region.extent.height = FMath::Max(1u, SrcSubRectHeight >> ResolveParams.MipIndex);
+	}
+	else
+	{
+		ensure(SrcSurface.Width == DstSurface.Width && SrcSurface.Height == DstSurface.Height);
+
+		Region.extent.width = FMath::Max(1u, SrcSurface.Width >> ResolveParams.MipIndex);
+		Region.extent.height = FMath::Max(1u, SrcSurface.Height >> ResolveParams.MipIndex);
 	}
 
-	Region.extent.width = FMath::Max(1u, SrcSurface.Width >> ResolveParams.MipIndex);
-	Region.extent.height = FMath::Max(1u, SrcSurface.Height >> ResolveParams.MipIndex);
 	Region.extent.depth = 1;
 	Region.srcSubresource.aspectMask = SrcSurface.GetFullAspectMask();
 	Region.srcSubresource.baseArrayLayer = SrcRange.baseArrayLayer;
