nDisplay: fixed typo. Horde Issue #145198

diff --git a/Engine/Plugins/Runtime/nDisplay/Source/DisplayCluster/Private/Render/Viewport/Configuration/DisplayClusterViewportConfigurationICVFX.cpp b/Engine/Plugins/Runtime/nDisplay/Source/DisplayCluster/Private/Render/Viewport/Configuration/DisplayClusterViewportConfigurationICVFX.cpp
--- a/Engine/Plugins/Runtime/nDisplay/Source/DisplayCluster/Private/Render/Viewport/Configuration/DisplayClusterViewportConfigurationICVFX.cpp
+++ b/Engine/Plugins/Runtime/nDisplay/Source/DisplayCluster/Private/Render/Viewport/Configuration/DisplayClusterViewportConfigurationICVFX.cpp
@@ -401,7 +401,7 @@ void FDisplayClusterViewportConfigurationICVFX::Update()
 				// Freeze render for lightcards when outer viewports freezed
 				if (StageSettings.Lightcard.bIgnoreOuterViewportsFreezingForLightcards == false)
 				{
-					const EDisplayClusterViewportRuntimeICVFXFlags LightcardViewportMask = ViewportRuntime_ICVFXLightcardColor | ViewportRuntime_ICVFXLightcard | ViewportRuntime_ICVFXLightcardColor;
+					const EDisplayClusterViewportRuntimeICVFXFlags LightcardViewportMask = ViewportRuntime_ICVFXLightcard | ViewportRuntime_ICVFXLightcardColor | ViewportRuntime_ICVFXLightcardAlpha;
 
 					for (FDisplayClusterViewport* ViewportIt : ViewportManager->ImplGetViewports())
 					{
