nDisplay: Change default inner frustum border thickness, update categories.

diff --git a/Engine/Plugins/Runtime/nDisplay/Source/DisplayClusterConfiguration/Public/DisplayClusterConfigurationTypes_ICVFX.h b/Engine/Plugins/Runtime/nDisplay/Source/DisplayClusterConfiguration/Public/DisplayClusterConfigurationTypes_ICVFX.h
--- a/Engine/Plugins/Runtime/nDisplay/Source/DisplayClusterConfiguration/Public/DisplayClusterConfigurationTypes_ICVFX.h
+++ b/Engine/Plugins/Runtime/nDisplay/Source/DisplayClusterConfiguration/Public/DisplayClusterConfigurationTypes_ICVFX.h
@@ -399,7 +399,7 @@ struct DISPLAYCLUSTERCONFIGURATION_API FDisplayClusterConfigurationICVFX_CameraB
 
 	/** Adjust border width to the top and bottom edges of the inner frustum. */
 	UPROPERTY(EditAnywhere, BlueprintReadWrite, Category = NDisplay, meta = (DisplayName = "Border Width", ClampMin = "0.0", UIMin = "0.0", ClampMax = "1.0", UIMax = "1.0"))
-	float Thickness = 0.25f;
+	float Thickness = 0.05f;
 
 	/** Adjust color of the border edges of the inner frustum. */
 	UPROPERTY(EditAnywhere, BlueprintReadWrite, Category = NDisplay, meta = (DisplayName = "Border Color"))
diff --git a/Engine/Plugins/Runtime/nDisplay/Source/DisplayClusterConfiguration/Public/DisplayClusterConfigurationTypes_Viewport.h b/Engine/Plugins/Runtime/nDisplay/Source/DisplayClusterConfiguration/Public/DisplayClusterConfigurationTypes_Viewport.h
--- a/Engine/Plugins/Runtime/nDisplay/Source/DisplayClusterConfiguration/Public/DisplayClusterConfigurationTypes_Viewport.h
+++ b/Engine/Plugins/Runtime/nDisplay/Source/DisplayClusterConfiguration/Public/DisplayClusterConfigurationTypes_Viewport.h
@@ -84,19 +84,19 @@ struct DISPLAYCLUSTERCONFIGURATION_API FDisplayClusterConfigurationViewport_Rend
 
 public:
 	/** Specify which GPU should render the second Stereo eye */
-	UPROPERTY(EditAnywhere, Category = "Configuration", meta = (DisplayName = "Stereo GPU Index"))
+	UPROPERTY(EditAnywhere, Category = "Stereo", meta = (DisplayName = "Stereo GPU Index"))
 	int StereoGPUIndex = -1;
 
 	/** Enables and sets Stereo mode */
-	UPROPERTY(EditAnywhere, BlueprintReadWrite, Category = "Configuration")
+	UPROPERTY(EditAnywhere, BlueprintReadWrite, Category = "Stereo")
 	EDisplayClusterConfigurationViewport_StereoMode StereoMode = EDisplayClusterConfigurationViewport_StereoMode::Default;
 
 	/** Adjust resolution scaling for an individual viewport.  Viewport Screen Percentage Multiplier is applied to this value. */
-	UPROPERTY(EditAnywhere, BlueprintReadWrite, Category = "Configuration", meta = (DisplayName = "Screen Percentage", ClampMin = "0.05", UIMin = "0.05", ClampMax = "10.0", UIMax = "1.0"))
+	UPROPERTY(EditAnywhere, BlueprintReadWrite, Category = "Rendering", meta = (DisplayName = "Screen Percentage", ClampMin = "0.05", UIMin = "0.05", ClampMax = "10.0", UIMax = "1.0"))
 	float BufferRatio = 1;
 
 	/** Adjust resolution scaling for an individual viewport.  Viewport Screen Percentage Multiplier is applied to this value. */
-	UPROPERTY(EditAnywhere, Category = "Configuration", meta = (ClampMin = "0.01", UIMin = "0.01", ClampMax = "1.0", UIMax = "1.0"))
+	UPROPERTY()
 	float RenderTargetRatio = 1.f;
 
 	UPROPERTY()
@@ -201,11 +201,11 @@ public:
 	FDisplayClusterConfigurationViewport_Remap ViewportRemap;
 
 	/** Allows Viewports to overlap and sets Viewport overlapping order priority */
-	UPROPERTY(EditAnywhere, BlueprintReadWrite, Category = "Configuration")
+	UPROPERTY(EditAnywhere, BlueprintReadWrite, Category = "Rendering")
 	int OverlapOrder = 0;
 
 	/** Specify which GPU should render this Viewport. "-1" is default. */
-	UPROPERTY(EditAnywhere, BlueprintReadWrite, Category = "Configuration", meta = (DisplayName = "GPU Index"))
+	UPROPERTY(EditAnywhere, BlueprintReadWrite, Category = "Rendering", meta = (DisplayName = "GPU Index"))
 	int GPUIndex = -1;
 
 	// Configure render for this viewport
