nDisplay - Readding objectpropertydata for ndisplay root component to fix wrongly exposed property.

diff --git a/Engine/Plugins/Runtime/nDisplay/Source/DisplayClusterConfigurator/Private/Views/Details/DisplayClusterRootActorDetailsCustomization.cpp b/Engine/Plugins/Runtime/nDisplay/Source/DisplayClusterConfigurator/Private/Views/Details/DisplayClusterRootActorDetailsCustomization.cpp
--- a/Engine/Plugins/Runtime/nDisplay/Source/DisplayClusterConfigurator/Private/Views/Details/DisplayClusterRootActorDetailsCustomization.cpp
+++ b/Engine/Plugins/Runtime/nDisplay/Source/DisplayClusterConfigurator/Private/Views/Details/DisplayClusterRootActorDetailsCustomization.cpp
@@ -86,7 +86,6 @@ void FDisplayClusterRootActorDetailsCustomization::CustomizeDetails(IDetailLayou
 	* This causes a crash inside the LayoutData.ClassToPropertyMap iterator.
 	* Customization called from  \\Engine\Source\Editor\PropertyEditor\Private\DetailLayoutHelpers.cpp:356
 	*/
-#if 0
 	// Manually add the transform properties' data to the layout builder's property in order to generate property handles for them.
 	{
 		TArray<UObject*> RootComponents;
@@ -105,7 +104,6 @@ void FDisplayClusterRootActorDetailsCustomization::CustomizeDetails(IDetailLayou
 		InLayoutBuilder.AddObjectPropertyData(RootComponents, USceneComponent::GetRelativeRotationPropertyName());
 		InLayoutBuilder.AddObjectPropertyData(RootComponents, USceneComponent::GetRelativeScale3DPropertyName());
 	}
-#endif
 
 	// Manually label the ICVFX category to propery format it to have the dash in "In-Camera"
 	InLayoutBuilder.EditCategory(DisplayClusterConfigurationStrings::categories::ICVFXCategory, LOCTEXT("ICVFXCategoryLabel", "In-Camera VFX"));
