nDisplay: fixed freeze preview, fixed preview node UI, fixed readonly ini treatment

diff --git a/Engine/Plugins/Runtime/nDisplay/Source/DisplayCluster/Private/Game/EngineClasses/Scene/DisplayClusterRootActorEditor.cpp b/Engine/Plugins/Runtime/nDisplay/Source/DisplayCluster/Private/Game/EngineClasses/Scene/DisplayClusterRootActorEditor.cpp
--- a/Engine/Plugins/Runtime/nDisplay/Source/DisplayCluster/Private/Game/EngineClasses/Scene/DisplayClusterRootActorEditor.cpp
+++ b/Engine/Plugins/Runtime/nDisplay/Source/DisplayCluster/Private/Game/EngineClasses/Scene/DisplayClusterRootActorEditor.cpp
@@ -772,9 +772,6 @@ void ADisplayClusterRootActor::UpdatePreviewComponents()
 		return;
 	}
 
-	// Render new preview
-	PreviewClusterNodeIndex = 0;
-
 	ImplUpdatePreviewConfiguration_Editor(DisplayClusterConfigurationStrings::gui::preview::PreviewNodeAll);
 
 	TArray<UDisplayClusterPreviewComponent*> IteratedPreviewComponents;
@@ -798,6 +795,9 @@ void ADisplayClusterRootActor::UpdatePreviewComponents()
 					check(PreviewComp);
 
 					PreviewComponents.Emplace(PreviewCompId, PreviewComp);
+
+					// Refresh preview when new component created
+					PreviewClusterNodeIndex = 0;
 				}
 
 				if (GetWorld() && !PreviewComp->IsRegistered())
diff --git a/Engine/Plugins/Runtime/nDisplay/Source/DisplayCluster/Private/Render/Viewport/DisplayClusterViewportHelpers.cpp b/Engine/Plugins/Runtime/nDisplay/Source/DisplayCluster/Private/Render/Viewport/DisplayClusterViewportHelpers.cpp
--- a/Engine/Plugins/Runtime/nDisplay/Source/DisplayCluster/Private/Render/Viewport/DisplayClusterViewportHelpers.cpp
+++ b/Engine/Plugins/Runtime/nDisplay/Source/DisplayCluster/Private/Render/Viewport/DisplayClusterViewportHelpers.cpp
@@ -184,7 +184,7 @@ float DisplayClusterViewportHelpers::GetValidSizeMultiplier(const FIntPoint& InS
 			}
 			else
 			{
-				const int32 MinDimension = MinScaledSize.GetMin();
+				const int32 MinDimension = MinScaledSize.GetMax();
 				const int32 MaxDimension = GetMaxTextureDimension();
 
 				// Get the maximum allowed multiplier value
diff --git a/Engine/Plugins/Runtime/nDisplay/Source/DisplayClusterConfigurator/Private/Views/Details/DisplayClusterRootActorDetailsCustomization.cpp b/Engine/Plugins/Runtime/nDisplay/Source/DisplayClusterConfigurator/Private/Views/Details/DisplayClusterRootActorDetailsCustomization.cpp
--- a/Engine/Plugins/Runtime/nDisplay/Source/DisplayClusterConfigurator/Private/Views/Details/DisplayClusterRootActorDetailsCustomization.cpp
+++ b/Engine/Plugins/Runtime/nDisplay/Source/DisplayClusterConfigurator/Private/Views/Details/DisplayClusterRootActorDetailsCustomization.cpp
@@ -206,7 +206,7 @@ bool FDisplayClusterRootActorDetailsCustomization::RebuildNodeIdOptionsList()
 
 	// Make sure we've got at least one cluster node in the config
 	// (None+node or None+All+node1+node2+...)
-	return NodeIdOptions.Num() >= 3;
+	return NodeIdOptions.Num() >= 2;
 }
 
 void FDisplayClusterRootActorDetailsCustomization::UpdateNodeIdSelection()
diff --git a/Engine/Plugins/Runtime/nDisplay/Source/DisplayClusterEditor/Private/Settings/DisplayClusterEditorSettings.cpp b/Engine/Plugins/Runtime/nDisplay/Source/DisplayClusterEditor/Private/Settings/DisplayClusterEditorSettings.cpp
--- a/Engine/Plugins/Runtime/nDisplay/Source/DisplayClusterEditor/Private/Settings/DisplayClusterEditorSettings.cpp
+++ b/Engine/Plugins/Runtime/nDisplay/Source/DisplayClusterEditor/Private/Settings/DisplayClusterEditorSettings.cpp
@@ -18,6 +18,7 @@ UDisplayClusterEditorSettings::UDisplayClusterEditorSettings(const FObjectInitia
 
 void UDisplayClusterEditorSettings::PostEditChangeProperty(FPropertyChangedEvent& PropertyChangedEvent)
 {
+	Super::PostEditChangeProperty(PropertyChangedEvent);
 	if (PropertyChangedEvent.Property != nullptr)
 	{
 		// Since nDisplay is Windows only, save configs that depends on nDisplay plugin assets in Windows specific config file
@@ -28,20 +29,31 @@ void UDisplayClusterEditorSettings::PostEditChangeProperty(FPropertyChangedEvent
 
 		if (PropertyName == GET_MEMBER_NAME_CHECKED(UDisplayClusterEditorSettings, bEnabled))
 		{
+			if (FPlatformFileManager::Get().GetPlatformFile().IsReadOnly(*DefaultEnginePath))
+			{
+				FPlatformFileManager::Get().GetPlatformFile().SetReadOnly(*DefaultEnginePath, false);
+			}
+
+			if (FPlatformFileManager::Get().GetPlatformFile().IsReadOnly(*DefaultGamePath))
+			{
+				FPlatformFileManager::Get().GetPlatformFile().SetReadOnly(*DefaultGamePath, false);
+			}
+
 			if (bEnabled)
 			{
 				// DefaultEngine.ini
-				GConfig->SetString(TEXT("/Script/Engine.Engine"), TEXT("GameEngine"),                  TEXT("/Script/DisplayCluster.DisplayClusterGameEngine"),         DefaultEnginePath);
-				GConfig->SetString(TEXT("/Script/Engine.Engine"), TEXT("UnrealEdEngine"),              TEXT("/Script/DisplayClusterEditor.DisplayClusterEditorEngine"), DefaultEnginePath);
-				GConfig->SetString(TEXT("/Script/Engine.Engine"), TEXT("GameViewportClientClassName"), TEXT("/Script/DisplayCluster.DisplayClusterViewportClient"),     DefaultEnginePath);
+				GConfig->SetString(TEXT("/Script/Engine.Engine"), TEXT("GameEngine"), TEXT("/Script/DisplayCluster.DisplayClusterGameEngine"), DefaultEnginePath);
+				GConfig->SetString(TEXT("/Script/Engine.Engine"), TEXT("UnrealEdEngine"), TEXT("/Script/DisplayClusterEditor.DisplayClusterEditorEngine"), DefaultEnginePath);
+				GConfig->SetString(TEXT("/Script/Engine.Engine"), TEXT("GameViewportClientClassName"), TEXT("/Script/DisplayCluster.DisplayClusterViewportClient"), DefaultEnginePath);
 
 				// DefaultGame.ini
 				GConfig->SetString(TEXT("/Script/EngineSettings.GeneralProjectSettings"), TEXT("bUseBorderlessWindow"), TEXT("True"), DefaultGamePath);
 			}
 			else
 			{
-				GConfig->SetString(TEXT("/Script/Engine.Engine"), TEXT("GameEngine"),                  TEXT("/Script/Engine.GameEngine"),         DefaultEnginePath);
-				GConfig->SetString(TEXT("/Script/Engine.Engine"), TEXT("UnrealEdEngine"),              TEXT("/Script/UnrealEd.UnrealEdEngine"),   DefaultEnginePath);
+				// DefaultEngine.ini
+				GConfig->SetString(TEXT("/Script/Engine.Engine"), TEXT("GameEngine"), TEXT("/Script/Engine.GameEngine"), DefaultEnginePath);
+				GConfig->SetString(TEXT("/Script/Engine.Engine"), TEXT("UnrealEdEngine"), TEXT("/Script/UnrealEd.UnrealEdEngine"), DefaultEnginePath);
 				GConfig->SetString(TEXT("/Script/Engine.Engine"), TEXT("GameViewportClientClassName"), TEXT("/Script/Engine.GameViewportClient"), DefaultEnginePath);
 
 				// DefaultGame.ini
@@ -52,6 +64,4 @@ void UDisplayClusterEditorSettings::PostEditChangeProperty(FPropertyChangedEvent
 			GConfig->Flush(false, DefaultGamePath);
 		}
 	}
-
-	Super::PostEditChangeProperty(PropertyChangedEvent);
 }
