ndisplay: fixed cis build error

diff --git a/Engine/Plugins/Runtime/nDisplay/Source/DisplayClusterEditor/Private/Settings/DisplayClusterEditorSettings.cpp b/Engine/Plugins/Runtime/nDisplay/Source/DisplayClusterEditor/Private/Settings/DisplayClusterEditorSettings.cpp
--- a/Engine/Plugins/Runtime/nDisplay/Source/DisplayClusterEditor/Private/Settings/DisplayClusterEditorSettings.cpp
+++ b/Engine/Plugins/Runtime/nDisplay/Source/DisplayClusterEditor/Private/Settings/DisplayClusterEditorSettings.cpp
@@ -2,6 +2,7 @@
 
 #include "DisplayClusterEditorSettings.h"
 #include "DisplayClusterEditorEngine.h"
+#include "HAL/PlatformFileManager.h"
 #include "Misc/ConfigCacheIni.h"
 
 
