nDisplay: Disabled Nanite on lightcard and custom chroma key renders by default.

Nanite is normally not used in these renders and disabling it can improve performance.

Associated cvars:
    DC.Lightcards.DisableNanite
    DC.ChromaKey.DisableNanite

diff --git a/Engine/Plugins/Runtime/nDisplay/Source/DisplayCluster/Private/Render/Viewport/DisplayClusterViewportManager.cpp b/Engine/Plugins/Runtime/nDisplay/Source/DisplayCluster/Private/Render/Viewport/DisplayClusterViewportManager.cpp
--- a/Engine/Plugins/Runtime/nDisplay/Source/DisplayCluster/Private/Render/Viewport/DisplayClusterViewportManager.cpp
+++ b/Engine/Plugins/Runtime/nDisplay/Source/DisplayCluster/Private/Render/Viewport/DisplayClusterViewportManager.cpp
@@ -30,6 +30,24 @@
 
 #include "Engine/Console.h"
 
+
+int32 GDisplayClusterLightcardsAllowNanite = 0;
+static FAutoConsoleVariableRef CVarDisplayClusterLightcardsAllowNanite(
+	TEXT("DC.Lightcards.AllowNanite"),
+	GDisplayClusterLightcardsAllowNanite,
+	TEXT("0 disables Nanite when rendering lightcards. Otherwise uses default showflag."),
+	ECVF_RenderThreadSafe
+);
+
+int32 GDisplayClusterChromaKeyAllowNanite = 0;
+static FAutoConsoleVariableRef CVarDisplayClusterChromaKeyAllowNanite(
+	TEXT("DC.ChromaKey.AllowNanite"),
+	GDisplayClusterChromaKeyAllowNanite,
+	TEXT("0 disables Nanite when rendering custom chroma keys. Otherwise uses default showflag."),
+	ECVF_RenderThreadSafe
+);
+
+
 ///////////////////////////////////////////////////////////////////////////////////////
 //          FDisplayClusterViewportManager
 ///////////////////////////////////////////////////////////////////////////////////////
@@ -552,6 +570,30 @@ FSceneViewFamily::ConstructionValues FDisplayClusterViewportManager::CreateViewF
 			break;
 	}
 
+	switch (InFrameTarget.CaptureMode)
+	{
+	case EDisplayClusterViewportCaptureMode::Chromakey:
+		
+		if (!GDisplayClusterChromaKeyAllowNanite)
+		{
+			InEngineShowFlags.SetNaniteMeshes(0);
+		}
+		break;
+
+	case EDisplayClusterViewportCaptureMode::Lightcard:
+	case EDisplayClusterViewportCaptureMode::Lightcard_OCIO:
+
+		if (!GDisplayClusterLightcardsAllowNanite)
+		{
+			InEngineShowFlags.SetNaniteMeshes(0);
+		}
+
+		break;
+
+	default:
+		break;
+	}
+
 	const FDisplayClusterRenderFrameSettings& RenderFrameSettings = GetRenderFrameSettings();
 	switch (RenderFrameSettings.RenderMode)
 	{
