nDisplay: fixed raytracing flags in chromakey/lightcards rendering, fixed preview frustum rendering flicker

diff --git a/Engine/Plugins/Runtime/nDisplay/Source/DisplayCluster/Private/Game/EngineClasses/Scene/DisplayClusterRootActorEditor.cpp b/Engine/Plugins/Runtime/nDisplay/Source/DisplayCluster/Private/Game/EngineClasses/Scene/DisplayClusterRootActorEditor.cpp
--- a/Engine/Plugins/Runtime/nDisplay/Source/DisplayCluster/Private/Game/EngineClasses/Scene/DisplayClusterRootActorEditor.cpp
+++ b/Engine/Plugins/Runtime/nDisplay/Source/DisplayCluster/Private/Game/EngineClasses/Scene/DisplayClusterRootActorEditor.cpp
@@ -514,11 +514,10 @@ void ADisplayClusterRootActor::ImplRenderPreviewFrustums_Editor()
 				UDisplayClusterICVFXCameraComponent* CineCameraComponent = Cast<UDisplayClusterICVFXCameraComponent>(ActorComponentIt);
 				if (CineCameraComponent && CineCameraComponent->IsICVFXEnabled())
 				{
-					const FString InnerFrustumID = CineCameraComponent->GetCameraUniqueId();
-					FDisplayClusterViewport* IncameraViewport = FDisplayClusterViewportConfigurationHelpers_ICVFX::ImplFindViewport(*this, InnerFrustumID, DisplayClusterViewportStrings::icvfx::camera);
-					if (IncameraViewport != nullptr)
+					// Iterate over rendered incamera viewports (whole cluster)
+					for (FDisplayClusterViewport* InCameraViewportIt : FDisplayClusterViewportConfigurationHelpers_ICVFX::PreviewGetRenderedInCameraViewports(*this, *CineCameraComponent))
 					{
-						FrustumPreviewViewports.Add(IncameraViewport);
+						FrustumPreviewViewports.Add(InCameraViewportIt);
 					}
 				}
 			}
diff --git a/Engine/Plugins/Runtime/nDisplay/Source/DisplayCluster/Private/Render/Viewport/Configuration/DisplayClusterViewportConfigurationHelpers_ICVFX.cpp b/Engine/Plugins/Runtime/nDisplay/Source/DisplayCluster/Private/Render/Viewport/Configuration/DisplayClusterViewportConfigurationHelpers_ICVFX.cpp
--- a/Engine/Plugins/Runtime/nDisplay/Source/DisplayCluster/Private/Render/Viewport/Configuration/DisplayClusterViewportConfigurationHelpers_ICVFX.cpp
+++ b/Engine/Plugins/Runtime/nDisplay/Source/DisplayCluster/Private/Render/Viewport/Configuration/DisplayClusterViewportConfigurationHelpers_ICVFX.cpp
@@ -222,6 +222,38 @@ FDisplayClusterViewport* FDisplayClusterViewportConfigurationHelpers_ICVFX::GetO
 }
 
 #if WITH_EDITOR
+TArray<FDisplayClusterViewport*> FDisplayClusterViewportConfigurationHelpers_ICVFX::PreviewGetRenderedInCameraViewports(ADisplayClusterRootActor& InRootActor, UDisplayClusterICVFXCameraComponent& InCameraComponent, bool bGetChromakey)
+{
+	TArray<FDisplayClusterViewport*> OutViewports;
+
+	// Search for rendered camera viewport on other cluster nodes
+	FDisplayClusterViewportManager* ViewportManager = FDisplayClusterViewportConfigurationHelpers_ICVFX::GetViewportManager(InRootActor);
+	if (ViewportManager != nullptr)
+	{
+		const FString ICVFXCameraId = InCameraComponent.GetCameraUniqueId();
+		const EDisplayClusterViewportRuntimeICVFXFlags InMask = bGetChromakey ? ViewportRuntime_ICVFXChromakey : ViewportRuntime_ICVFXIncamera;
+		const FString ViewportTypeId = bGetChromakey ? DisplayClusterViewportStrings::icvfx::chromakey : DisplayClusterViewportStrings::icvfx::camera;
+
+		for (FDisplayClusterViewport* ViewportIt : ViewportManager->ImplGetWholeClusterViewports_Editor())
+		{
+			if (ViewportIt != nullptr
+				&& (ViewportIt->RenderSettingsICVFX.RuntimeFlags & InMask) != 0
+				&& (ViewportIt->InputShaderResources.Num() > 0 && ViewportIt->InputShaderResources[0] != nullptr && ViewportIt->Contexts.Num() > 0)
+				&& (ViewportIt->RenderSettings.OverrideViewportId.IsEmpty()))
+			{
+				// this is incamera viewport. Check by name
+				const FString RequiredViewportId = ImplGetNameICVFX(ViewportIt->GetClusterNodeId(), ICVFXCameraId, ViewportTypeId);
+				if (RequiredViewportId.Equals(ViewportIt->GetId()))
+				{
+					OutViewports.Add(ViewportIt);
+				}
+			}
+		}
+	}
+
+	return OutViewports;
+}
+
 void FDisplayClusterViewportConfigurationHelpers_ICVFX::PreviewReuseInnerFrustumViewportWithinClusterNodes_Editor(FDisplayClusterViewport& InCameraViewport, ADisplayClusterRootActor& InRootActor, UDisplayClusterICVFXCameraComponent& InCameraComponent)
 {
 	if (!IsPreviewEnableReuseViewportInCluster_EditorImpl(InRootActor))
@@ -239,33 +271,15 @@ void FDisplayClusterViewportConfigurationHelpers_ICVFX::PreviewReuseInnerFrustum
 		return;
 	}
 
-	const FString ICVFXCameraId = InCameraComponent.GetCameraUniqueId();
-
-	// Search for rendered camera viewport on other cluster nodes
-	FDisplayClusterViewportManager* ViewportManager = FDisplayClusterViewportConfigurationHelpers_ICVFX::GetViewportManager(InRootActor);
-	if (ViewportManager != nullptr)
+	for (FDisplayClusterViewport* ViewportIt : PreviewGetRenderedInCameraViewports(InRootActor, InCameraComponent, false))
 	{
-		for (FDisplayClusterViewport* ViewportIt : ViewportManager->ImplGetWholeClusterViewports_Editor())
+		if (ViewportIt && ViewportIt->GetClusterNodeId() != InCameraViewport.GetClusterNodeId()
+			&& FDisplayClusterViewportConfigurationHelpers_OpenColorIO::IsInnerFrustumViewportSettingsEqual_Editor(*ViewportIt, InCameraViewport, InCameraComponent)
+			&& FDisplayClusterViewportConfigurationHelpers_Postprocess::IsInnerFrustumViewportSettingsEqual_Editor(*ViewportIt, InCameraViewport, InCameraComponent))
 		{
-			if (ViewportIt != nullptr
-			&& (ViewportIt->RenderSettingsICVFX.RuntimeFlags & ViewportRuntime_ICVFXIncamera) != 0 
-			&& (ViewportIt->InputShaderResources.Num() > 0 && ViewportIt->InputShaderResources[0] != nullptr && ViewportIt->Contexts.Num() > 0)
-			&& (ViewportIt->RenderSettings.OverrideViewportId.IsEmpty())
-			&& (ViewportIt->GetClusterNodeId() != InCameraViewport.GetClusterNodeId()))
-			{
-				// this is incamera viewport. Check by name
-				const FString RequiredViewportId = ImplGetNameICVFX(ViewportIt->GetClusterNodeId(), ICVFXCameraId, DisplayClusterViewportStrings::icvfx::camera);
-				if (RequiredViewportId.Equals(ViewportIt->GetId()))
-				{
-					if (FDisplayClusterViewportConfigurationHelpers_OpenColorIO::IsInnerFrustumViewportSettingsEqual_Editor(*ViewportIt, InCameraViewport, InCameraComponent)
-					&&  FDisplayClusterViewportConfigurationHelpers_Postprocess::IsInnerFrustumViewportSettingsEqual_Editor(*ViewportIt, InCameraViewport, InCameraComponent))
-					{
-						// Reuse exist viewport:
-						InCameraViewport.RenderSettings.OverrideViewportId = ViewportIt->GetId();
-						return;
-					}
-				}
-			}
+			// Reuse exist viewport:
+			InCameraViewport.RenderSettings.OverrideViewportId = ViewportIt->GetId();
+			return;
 		}
 	}
 }
@@ -287,29 +301,13 @@ void FDisplayClusterViewportConfigurationHelpers_ICVFX::PreviewReuseChromakeyVie
 		return;
 	}
 
-	const FString ICVFXCameraId = InCameraComponent.GetCameraUniqueId();
-
-	// Search for rendered camera viewport on other cluster nodes
-	FDisplayClusterViewportManager* ViewportManager = FDisplayClusterViewportConfigurationHelpers_ICVFX::GetViewportManager(InRootActor);
-	if (ViewportManager != nullptr)
+	for (FDisplayClusterViewport* ViewportIt : PreviewGetRenderedInCameraViewports(InRootActor, InCameraComponent, true))
 	{
-		for (FDisplayClusterViewport* ViewportIt : ViewportManager->ImplGetWholeClusterViewports_Editor())
+		if (ViewportIt && ViewportIt->GetClusterNodeId() != InChromakeyViewport.GetClusterNodeId())
 		{
-			if (ViewportIt != nullptr
-				&& (ViewportIt->RenderSettingsICVFX.RuntimeFlags & ViewportRuntime_ICVFXChromakey) != 0
-				&& (ViewportIt->InputShaderResources.Num() > 0 && ViewportIt->InputShaderResources[0] != nullptr && ViewportIt->Contexts.Num() > 0)
-				&& (ViewportIt->RenderSettings.OverrideViewportId.IsEmpty())
-				&& (ViewportIt->GetClusterNodeId() != InChromakeyViewport.GetClusterNodeId()))
-			{
-				// this is chromakey viewport. Check by name
-				const FString RequiredViewportId = ImplGetNameICVFX(ViewportIt->GetClusterNodeId(), ICVFXCameraId, DisplayClusterViewportStrings::icvfx::chromakey);
-				if (RequiredViewportId.Equals(ViewportIt->GetId()))
-				{
-					// Reuse exist viewport:
-					InChromakeyViewport.RenderSettings.OverrideViewportId = ViewportIt->GetId();
-					return;
-				}
-			}
+			// Reuse exist viewport from other node
+			InChromakeyViewport.RenderSettings.OverrideViewportId = ViewportIt->GetId();
+			return;
 		}
 	}
 }
diff --git a/Engine/Plugins/Runtime/nDisplay/Source/DisplayCluster/Private/Render/Viewport/Configuration/DisplayClusterViewportConfigurationHelpers_ICVFX.h b/Engine/Plugins/Runtime/nDisplay/Source/DisplayCluster/Private/Render/Viewport/Configuration/DisplayClusterViewportConfigurationHelpers_ICVFX.h
--- a/Engine/Plugins/Runtime/nDisplay/Source/DisplayCluster/Private/Render/Viewport/Configuration/DisplayClusterViewportConfigurationHelpers_ICVFX.h
+++ b/Engine/Plugins/Runtime/nDisplay/Source/DisplayCluster/Private/Render/Viewport/Configuration/DisplayClusterViewportConfigurationHelpers_ICVFX.h
@@ -40,6 +40,7 @@ public:
 	static FDisplayClusterViewport* GetOrCreateLightcardViewport(FDisplayClusterViewport& BaseViewport, ADisplayClusterRootActor& RootActor, bool bIsOpenColorIOViewportExist);
 
 #if WITH_EDITOR
+	static TArray<FDisplayClusterViewport*> PreviewGetRenderedInCameraViewports(ADisplayClusterRootActor& RootActor, UDisplayClusterICVFXCameraComponent& InCameraComponent, bool bGetChromakey = false);
 	static void PreviewReuseInnerFrustumViewportWithinClusterNodes_Editor(FDisplayClusterViewport& InCameraViewport, ADisplayClusterRootActor& RootActor, UDisplayClusterICVFXCameraComponent& InCameraComponent);
 	static void PreviewReuseChromakeyViewportWithinClusterNodes_Editor(FDisplayClusterViewport& InChromakeyViewport, ADisplayClusterRootActor& RootActor, UDisplayClusterICVFXCameraComponent& InCameraComponent);
 #endif
diff --git a/Engine/Plugins/Runtime/nDisplay/Source/DisplayCluster/Private/Render/Viewport/DisplayClusterViewport.cpp b/Engine/Plugins/Runtime/nDisplay/Source/DisplayCluster/Private/Render/Viewport/DisplayClusterViewport.cpp
--- a/Engine/Plugins/Runtime/nDisplay/Source/DisplayCluster/Private/Render/Viewport/DisplayClusterViewport.cpp
+++ b/Engine/Plugins/Runtime/nDisplay/Source/DisplayCluster/Private/Render/Viewport/DisplayClusterViewport.cpp
@@ -212,6 +212,22 @@ void FDisplayClusterViewport::SetupSceneView(uint32 ContextNum, class UWorld* Wo
 		InOutView.bAllowCrossGPUTransfer = (Contexts[ContextNum].bAllowGPUTransferOptimization == false);
 	}
 
+	// Disable raytracing for lightcard and chromakey
+#if RHI_RAYTRACING
+	switch (RenderSettings.CaptureMode)
+	{
+	case EDisplayClusterViewportCaptureMode::Chromakey:
+	case EDisplayClusterViewportCaptureMode::Lightcard:
+	case EDisplayClusterViewportCaptureMode::Lightcard_OCIO:
+
+		InOutView.bAllowRayTracing = false;
+		break;
+
+	default:
+		break;
+	}
+#endif // RHI_RAYTRACING
+
 	// Apply visibility settigns to view
 	VisibilitySettings.SetupSceneView(World, InOutView);
 
diff --git a/Engine/Plugins/Runtime/nDisplay/Source/DisplayCluster/Public/Components/DisplayClusterICVFXCameraComponent.h b/Engine/Plugins/Runtime/nDisplay/Source/DisplayCluster/Public/Components/DisplayClusterICVFXCameraComponent.h
--- a/Engine/Plugins/Runtime/nDisplay/Source/DisplayCluster/Public/Components/DisplayClusterICVFXCameraComponent.h
+++ b/Engine/Plugins/Runtime/nDisplay/Source/DisplayCluster/Public/Components/DisplayClusterICVFXCameraComponent.h
@@ -97,7 +97,7 @@ private:
 	UPROPERTY(EditAnywhere, Transient, Category = "In Camera VFX", meta = (PropertyPath = "CameraSettings.BufferRatio"))
 	FDisplayClusterEditorPropertyReference BufferRatioRef;
 
-	UPROPERTY(EditAnywhere, Transient, Category = "In Camera VFX", meta = (DisplayName = "Overscan", PropertyPath = "CameraSettings.CustomFrustum"))
+	UPROPERTY(EditAnywhere, Transient, Category = "In Camera VFX", meta = (DisplayName = "Inner Frustum Overscan", PropertyPath = "CameraSettings.CustomFrustum"))
 	FDisplayClusterEditorPropertyReference CustomFrustumRef;
 
 	UPROPERTY(EditAnywhere, Transient, Category = "In Camera VFX", meta = (PropertyPath = "CameraSettings.SoftEdge"))
