nDisplay: Correct In-Camera VFX category display name.

diff --git a/Engine/Plugins/Runtime/nDisplay/Source/DisplayClusterConfigurator/Private/Views/Details/Components/DisplayClusterICVFXCameraComponentDetailsCustomization.cpp b/Engine/Plugins/Runtime/nDisplay/Source/DisplayClusterConfigurator/Private/Views/Details/Components/DisplayClusterICVFXCameraComponentDetailsCustomization.cpp
--- a/Engine/Plugins/Runtime/nDisplay/Source/DisplayClusterConfigurator/Private/Views/Details/Components/DisplayClusterICVFXCameraComponentDetailsCustomization.cpp
+++ b/Engine/Plugins/Runtime/nDisplay/Source/DisplayClusterConfigurator/Private/Views/Details/Components/DisplayClusterICVFXCameraComponentDetailsCustomization.cpp
@@ -94,6 +94,9 @@ void FDisplayClusterICVFXCameraComponentDetailsCustomization::CustomizeDetails(I
 	// Sockets category must be hidden manually instead of through the HideCategories metadata specifier
 	InLayoutBuilder.HideCategory(TEXT("Sockets"));
 
+	// Manually label the ICVFX category to properly format it to have the dash in "In-Camera"
+	InLayoutBuilder.EditCategory(DisplayClusterConfigurationStrings::categories::ICVFXCategory, LOCTEXT("ICVFXCategoryLabel", "In-Camera VFX"));
+
 	InLayoutBuilder.SortCategories(DisplayClusterICVFXCameraComponentDetailsCustomizationUtils::SortCategories);
 
 	// Most of the properties in the camera settings are exposed through property references, so hide the camera settings property.
diff --git a/Engine/Plugins/Runtime/nDisplay/Source/DisplayClusterConfigurator/Private/Views/Details/DisplayClusterRootActorDetailsCustomization.cpp b/Engine/Plugins/Runtime/nDisplay/Source/DisplayClusterConfigurator/Private/Views/Details/DisplayClusterRootActorDetailsCustomization.cpp
--- a/Engine/Plugins/Runtime/nDisplay/Source/DisplayClusterConfigurator/Private/Views/Details/DisplayClusterRootActorDetailsCustomization.cpp
+++ b/Engine/Plugins/Runtime/nDisplay/Source/DisplayClusterConfigurator/Private/Views/Details/DisplayClusterRootActorDetailsCustomization.cpp
@@ -105,7 +105,7 @@ void FDisplayClusterRootActorDetailsCustomization::CustomizeDetails(IDetailLayou
 		InLayoutBuilder.AddObjectPropertyData(RootComponents, USceneComponent::GetRelativeScale3DPropertyName());
 	}
 
-	// Manually label the ICVFX category to propery format it to have the dash in "In-Camera"
+	// Manually label the ICVFX category to properly format it to have the dash in "In-Camera"
 	InLayoutBuilder.EditCategory(DisplayClusterConfigurationStrings::categories::ICVFXCategory, LOCTEXT("ICVFXCategoryLabel", "In-Camera VFX"));
 
 	// Customize the PreviewNodeId property to be a dropdown filled with the nodes configured on the root actor.
