nDisplay: Support updating the lightcard preview actor when a change is made to the level instance the preview actor represents. Compiling the blueprint will also properly refresh the preview actor.

diff --git a/Engine/Plugins/Runtime/nDisplay/Source/DisplayClusterLightCardEditor/Private/Viewport/DisplayClusterLightCardEditorViewportClient.cpp b/Engine/Plugins/Runtime/nDisplay/Source/DisplayClusterLightCardEditor/Private/Viewport/DisplayClusterLightCardEditorViewportClient.cpp
--- a/Engine/Plugins/Runtime/nDisplay/Source/DisplayClusterLightCardEditor/Private/Viewport/DisplayClusterLightCardEditorViewportClient.cpp
+++ b/Engine/Plugins/Runtime/nDisplay/Source/DisplayClusterLightCardEditor/Private/Viewport/DisplayClusterLightCardEditorViewportClient.cpp
@@ -88,6 +88,7 @@ FDisplayClusterLightCardEditorViewportClient::FDisplayClusterLightCardEditorView
 FDisplayClusterLightCardEditorViewportClient::~FDisplayClusterLightCardEditorViewportClient()
 {
 	EndTransaction();
+	FCoreUObjectDelegates::OnObjectPropertyChanged.RemoveAll(this);
 }
 
 FLinearColor FDisplayClusterLightCardEditorViewportClient::GetBackgroundColor() const
@@ -128,7 +129,7 @@ void FDisplayClusterLightCardEditorViewportClient::Tick(float DeltaSeconds)
 		}
 	}
 
-	if (SpawnedRootActor.IsValid() && RootActorLevelInstance.IsValid())
+	if (RootActorPreviewInstance.IsValid() && RootActorLevelInstance.IsValid())
 	{
 		// Pass the preview render targets from the level instance root actor to the preview root actor
 		UDisplayClusterConfigurationData* Config = RootActorLevelInstance->GetConfigData();
@@ -139,7 +140,7 @@ void FDisplayClusterLightCardEditorViewportClient::Tick(float DeltaSeconds)
 			for (const TPair<FString, UDisplayClusterConfigurationViewport*>& ViewportPair : Node->Viewports)
 			{
 				UDisplayClusterPreviewComponent* LevelInstancePreviewComp = RootActorLevelInstance->GetPreviewComponent(NodePair.Key, ViewportPair.Key);
-				UDisplayClusterPreviewComponent* PreviewComp = SpawnedRootActor->GetPreviewComponent(NodePair.Key, ViewportPair.Key);
+				UDisplayClusterPreviewComponent* PreviewComp = RootActorPreviewInstance->GetPreviewComponent(NodePair.Key, ViewportPair.Key);
 
 				if (PreviewComp && LevelInstancePreviewComp)
 				{
@@ -436,11 +437,12 @@ void FDisplayClusterLightCardEditorViewportClient::UpdatePreviewActor(ADisplayCl
 {
 	UWorld* PreviewWorld = PreviewScene->GetWorld();
 	check (PreviewWorld);
-
-	if (SpawnedRootActor.IsValid() && (!RootActor || !PreviewWorld->ContainsActor(RootActor)))
+	
+	if (RootActorPreviewInstance.IsValid() && (!RootActor || !PreviewWorld->ContainsActor(RootActor)))
 	{
-		PreviewWorld->EditorDestroyActor(SpawnedRootActor.Get(), false);
-		SpawnedRootActor.Reset();
+		PreviewWorld->EditorDestroyActor(RootActorPreviewInstance.Get(), false);
+		
+		RootActorPreviewInstance.Reset();
 		RootActorLevelInstance.Reset();
 
 		MeshProjectionRenderer->ClearScene();
@@ -450,31 +452,53 @@ void FDisplayClusterLightCardEditorViewportClient::UpdatePreviewActor(ADisplayCl
 	{
 		RootActorLevelInstance = RootActor;
 
-		const FVector SpawnLocation = FVector::ZeroVector;
-		const FRotator SpawnRotation = FRotator::ZeroRotator;
-		
 		FActorSpawnParameters SpawnInfo;
 		SpawnInfo.SpawnCollisionHandlingOverride = ESpawnActorCollisionHandlingMethod::AlwaysSpawn;
 		SpawnInfo.bNoFail = true;
 		SpawnInfo.ObjectFlags = RF_Transient|RF_Transactional;
 		SpawnInfo.Template = RootActor;
 		
-		SpawnedRootActor = CastChecked<ADisplayClusterRootActor>(PreviewWorld->SpawnActor(RootActor->GetClass(),
-			&SpawnLocation, &SpawnRotation, MoveTemp(SpawnInfo)));
+		RootActorPreviewInstance = CastChecked<ADisplayClusterRootActor>(PreviewWorld->SpawnActor(RootActor->GetClass(),
+			nullptr,  nullptr, MoveTemp(SpawnInfo)));
+		
+		PreviewWorld->GetCurrentLevel()->AddLoadedActor(RootActorPreviewInstance.Get());
+
+		// Copy component data and update transforms.
+		{
+			TArray<UActorComponent*> SpawnedActorComponents;
+			TArray<UActorComponent*> LevelActorComponents;
+			RootActorPreviewInstance->GetComponents(SpawnedActorComponents);
+			RootActorLevelInstance->GetComponents(LevelActorComponents);
+
+			for (UActorComponent* SpawnedComponent : SpawnedActorComponents)
+			{
+				if (UActorComponent** OriginalComponent = LevelActorComponents.FindByPredicate([SpawnedComponent](const UActorComponent* Component)
+				{
+					return Component->GetFName() == SpawnedComponent->GetFName();
+				}))
+				{
+					UEngine::CopyPropertiesForUnrelatedObjects(*OriginalComponent, SpawnedComponent);
+					SpawnedComponent->UpdateComponentToWorld();
+				}
+			}
+		}
 
 		// Spawned actor will take the transform values from the template, so manually reset them to zero here
-		SpawnedRootActor->SetActorLocation(FVector::ZeroVector);
-		SpawnedRootActor->SetActorRotation(FRotator::ZeroRotator);
+		RootActorPreviewInstance->SetActorLocation(FVector::ZeroVector);
+		RootActorPreviewInstance->SetActorRotation(FRotator::ZeroRotator);
 
 		FindProjectionOriginComponent();
 
 		// Filter out any primitives hidden in game except screen components
-		MeshProjectionRenderer->AddActor(SpawnedRootActor.Get(), [](const UPrimitiveComponent* PrimitiveComponent)
+		MeshProjectionRenderer->AddActor(RootActorPreviewInstance.Get(), [](const UPrimitiveComponent* PrimitiveComponent)
 		{
 			return !PrimitiveComponent->bHiddenInGame || PrimitiveComponent->IsA<UDisplayClusterScreenComponent>();
 		});
 
-		SpawnedRootActor->UpdatePreviewComponents();
+		RootActorPreviewInstance->UpdatePreviewComponents();
+
+		FCoreUObjectDelegates::OnObjectPropertyChanged.RemoveAll(this);
+		FCoreUObjectDelegates::OnObjectPropertyChanged.AddRaw(this, &FDisplayClusterLightCardEditorViewportClient::OnActorPropertyChanged);
 	}
 }
 
@@ -538,9 +562,31 @@ void FDisplayClusterLightCardEditorViewportClient::EndTransaction()
 	}
 }
 
+void FDisplayClusterLightCardEditorViewportClient::OnActorPropertyChanged(UObject* ObjectBeingModified,
+                                                                          FPropertyChangedEvent& PropertyChangedEvent)
+{
+	if (PropertyChangedEvent.ChangeType == EPropertyChangeType::Interactive)
+	{
+		return;
+	}
+	
+	bool bIsOurActor = ObjectBeingModified == RootActorLevelInstance;
+	if (!bIsOurActor)
+	{
+		if (const UObject* RootActorOuter = ObjectBeingModified->GetTypedOuter(ADisplayClusterRootActor::StaticClass()))
+		{
+			bIsOurActor = RootActorOuter == RootActorLevelInstance;
+		}
+	}
+	if (bIsOurActor)
+	{
+		UpdatePreviewActor(RootActorLevelInstance.Get());
+	}
+}
+
 void FDisplayClusterLightCardEditorViewportClient::GetScenePrimitiveComponents(TArray<UPrimitiveComponent*>& OutPrimitiveComponents)
 {
-	SpawnedRootActor->ForEachComponent<UPrimitiveComponent>(true, [&](UPrimitiveComponent* PrimitiveComponent)
+	RootActorPreviewInstance->ForEachComponent<UPrimitiveComponent>(true, [&](UPrimitiveComponent* PrimitiveComponent)
 	{
 		OutPrimitiveComponents.Add(PrimitiveComponent);
 	});
@@ -655,12 +701,12 @@ void FDisplayClusterLightCardEditorViewportClient::GetSceneViewInitOptions(FScen
 
 void FDisplayClusterLightCardEditorViewportClient::FindProjectionOriginComponent()
 {
-	if (SpawnedRootActor.IsValid())
+	if (RootActorPreviewInstance.IsValid())
 	{
 		if (ProjectionMode == EDisplayClusterMeshProjectionType::Perspective)
 		{
 			TArray<UDisplayClusterCameraComponent*> ViewOriginComponents;
-			SpawnedRootActor->GetComponents<UDisplayClusterCameraComponent>(ViewOriginComponents);
+			RootActorPreviewInstance->GetComponents<UDisplayClusterCameraComponent>(ViewOriginComponents);
 
 			if (ViewOriginComponents.Num())
 			{
@@ -668,12 +714,12 @@ void FDisplayClusterLightCardEditorViewportClient::FindProjectionOriginComponent
 			}
 			else
 			{
-				ProjectionOriginComponent = SpawnedRootActor->GetRootComponent();
+				ProjectionOriginComponent = RootActorPreviewInstance->GetRootComponent();
 			}
 		}
 		else if (ProjectionMode == EDisplayClusterMeshProjectionType::Azimuthal)
 		{
-			ProjectionOriginComponent = SpawnedRootActor->GetRootComponent();
+			ProjectionOriginComponent = RootActorPreviewInstance->GetRootComponent();
 		}
 	}
 	else
diff --git a/Engine/Plugins/Runtime/nDisplay/Source/DisplayClusterLightCardEditor/Private/Viewport/DisplayClusterLightCardEditorViewportClient.h b/Engine/Plugins/Runtime/nDisplay/Source/DisplayClusterLightCardEditor/Private/Viewport/DisplayClusterLightCardEditorViewportClient.h
--- a/Engine/Plugins/Runtime/nDisplay/Source/DisplayClusterLightCardEditor/Private/Viewport/DisplayClusterLightCardEditorViewportClient.h
+++ b/Engine/Plugins/Runtime/nDisplay/Source/DisplayClusterLightCardEditor/Private/Viewport/DisplayClusterLightCardEditorViewportClient.h
@@ -42,6 +42,7 @@ public:
 	void SelectActor(AActor* NewActor);
 	void ResetSelection();
 
+	/** Update the spawned preview actor from a root actor in the level. */
 	void UpdatePreviewActor(ADisplayClusterRootActor* RootActor);
 	
 	AActor* GetSelectedActor() const { return SelectedActor.Get(); }
@@ -64,6 +65,9 @@ protected:
 	/** Ends the current transaction, if one exists. */
 	void EndTransaction();
 
+	/** When a property on the actor has changed. */
+	void OnActorPropertyChanged(UObject* ObjectBeingModified, FPropertyChangedEvent& PropertyChangedEvent);
+	
 	/** Gets a list of all primitive components to be rendered in the scene */
 	void GetScenePrimitiveComponents(TArray<UPrimitiveComponent*>& OutPrimitiveComponents);
 
@@ -76,7 +80,7 @@ protected:
 private:
 	TWeakPtr<FSceneViewport> SceneViewportPtr;
 	TWeakPtr<SDisplayClusterLightCardEditor> LightCardEditorPtr;
-	TWeakObjectPtr<ADisplayClusterRootActor> SpawnedRootActor;
+	TWeakObjectPtr<ADisplayClusterRootActor> RootActorPreviewInstance;
 	TWeakObjectPtr<ADisplayClusterRootActor> RootActorLevelInstance;
 	TWeakObjectPtr<AActor> SelectedActor;
 	
diff --git a/Engine/Plugins/Runtime/nDisplay/Source/DisplayClusterOperator/Private/SDisplayClusterOperatorToolbar.cpp b/Engine/Plugins/Runtime/nDisplay/Source/DisplayClusterOperator/Private/SDisplayClusterOperatorToolbar.cpp
--- a/Engine/Plugins/Runtime/nDisplay/Source/DisplayClusterOperator/Private/SDisplayClusterOperatorToolbar.cpp
+++ b/Engine/Plugins/Runtime/nDisplay/Source/DisplayClusterOperator/Private/SDisplayClusterOperatorToolbar.cpp
@@ -19,6 +19,14 @@ SDisplayClusterOperatorToolbar::~SDisplayClusterOperatorToolbar()
 	{
 		GEngine->OnLevelActorDeleted().Remove(LevelActorDeletedHandle);
 	}
+	
+	if (ActiveRootActor.IsValid())
+	{
+		if (UBlueprint* Blueprint = UBlueprint::GetBlueprintFromClass(ActiveRootActor->GetClass()))
+		{
+			Blueprint->OnCompiled().RemoveAll(this);
+		}
+	}
 }
 
 void SDisplayClusterOperatorToolbar::Construct(const FArguments& InArgs)
@@ -93,6 +101,7 @@ TSharedPtr<FString> SDisplayClusterOperatorToolbar::FillRootActorList(const FStr
 
 void SDisplayClusterOperatorToolbar::OnRootActorChanged(TSharedPtr<FString> ItemSelected, ESelectInfo::Type SelectInfo)
 {
+	ActiveRootActorName = ItemSelected;
 	if (!ItemSelected.IsValid())
 	{
 		return;
@@ -111,10 +120,36 @@ void SDisplayClusterOperatorToolbar::OnRootActorChanged(TSharedPtr<FString> Item
 		}
 	}
 
+	if (ActiveRootActor.IsValid())
+	{
+		if (UBlueprint* Blueprint = UBlueprint::GetBlueprintFromClass(ActiveRootActor->GetClass()))
+		{
+			Blueprint->OnCompiled().RemoveAll(this);
+		}
+	}
+	
+	if (SelectedRootActor)
+	{
+		if (UBlueprint* Blueprint = UBlueprint::GetBlueprintFromClass(SelectedRootActor->GetClass()))
+		{
+			Blueprint->OnCompiled().AddRaw(this, &SDisplayClusterOperatorToolbar::OnBlueprintCompiled);
+		}
+	}
+
 	ActiveRootActor = SelectedRootActor;
 	IDisplayClusterOperator::Get().OnActiveRootActorChanged().Broadcast(SelectedRootActor);
 }
 
+void SDisplayClusterOperatorToolbar::OnBlueprintCompiled(UBlueprint* Blueprint)
+{
+	if (RootActorComboBox.IsValid() && ActiveRootActorName.IsValid())
+	{
+		// Compiling the blueprint invalidates the instance so we need to find and set the instance again.
+		RootActorComboBox->SetSelectedItem(ActiveRootActorName);
+		OnRootActorChanged(ActiveRootActorName, ESelectInfo::Direct);
+	}
+}
+
 void SDisplayClusterOperatorToolbar::OnRootActorComboBoxOpening()
 {
 	FString SelectedRootActor = "";
@@ -150,7 +185,15 @@ void SDisplayClusterOperatorToolbar::OnLevelActorDeleted(AActor* Actor)
 {
 	if (Actor == ActiveRootActor)
 	{
-		ActiveRootActor = nullptr;
+		if (Actor && Actor->GetClass()->HasAnyClassFlags(CLASS_NewerVersionExists))
+		{
+			// When a blueprint class is regenerated instances are deleted and replaced.
+			// In this case the OnCompiled() delegate will fire and refresh the actor.
+			return;
+		}
+		
+		ActiveRootActor.Reset();
+		ActiveRootActorName.Reset();
 		IDisplayClusterOperator::Get().OnActiveRootActorChanged().Broadcast(nullptr);
 		RootActorComboBox->SetSelectedItem(nullptr);
 	}
diff --git a/Engine/Plugins/Runtime/nDisplay/Source/DisplayClusterOperator/Private/SDisplayClusterOperatorToolbar.h b/Engine/Plugins/Runtime/nDisplay/Source/DisplayClusterOperator/Private/SDisplayClusterOperatorToolbar.h
--- a/Engine/Plugins/Runtime/nDisplay/Source/DisplayClusterOperator/Private/SDisplayClusterOperatorToolbar.h
+++ b/Engine/Plugins/Runtime/nDisplay/Source/DisplayClusterOperator/Private/SDisplayClusterOperatorToolbar.h
@@ -41,6 +41,9 @@ private:
 	/** Raised when the user selects a new root actor from the root actor picker dropdown */
 	void OnRootActorChanged(TSharedPtr<FString> ItemSelected, ESelectInfo::Type SelectInfo);
 
+	/** When the root actor blueprint has been compiled */
+	void OnBlueprintCompiled(UBlueprint* Blueprint);
+
 	/** Raised when the root actor picker dropdown is being opened */
 	void OnRootActorComboBoxOpening();
 
@@ -60,6 +63,9 @@ private:
 	/** The list of root actor names on the current level to display in the root actor picker dropdown */
 	TArray<TSharedPtr<FString>> RootActorList;
 
+	/** The name or label of the active actor */
+	TSharedPtr<FString> ActiveRootActorName;
+
 	/** The combo box widget that allows the uesr to pick the active root actor from */
 	TSharedPtr<SComboBox<TSharedPtr<FString>>> RootActorComboBox;
 
