nDisplay: Added light card proxies to the light card editor, which are selectable.

diff --git a/Engine/Plugins/Runtime/nDisplay/Shaders/Private/MeshProjectionHitProxy.usf b/Engine/Plugins/Runtime/nDisplay/Shaders/Private/MeshProjectionHitProxy.usf
new file mode 100644
--- /dev/null
+++ b/Engine/Plugins/Runtime/nDisplay/Shaders/Private/MeshProjectionHitProxy.usf
@@ -0,0 +1,18 @@
+// Copyright Epic Games, Inc. All Rights Reserved.
+
+#include "/Engine/Private/Common.ush"
+#include "/Engine/Generated/Material.ush"
+#include "/Engine/Generated/VertexFactory.ush"
+
+float4 HitProxyId;
+
+void Main(FVertexFactoryInterpolantsVSToPS FactoryInterpolants, in INPUT_POSITION_QUALIFIERS float4 SvPosition : SV_Position OPTIONAL_IsFrontFace, out float4 OutColor : SV_Target0)
+{
+	FMaterialPixelParameters MaterialParameters = GetMaterialPixelParameters(FactoryInterpolants, SvPosition);
+	FPixelMaterialInputs PixelMaterialInputs;
+	CalcMaterialParameters(MaterialParameters, PixelMaterialInputs, SvPosition, bIsFrontFace);
+
+	GetMaterialCoverageAndClipping(MaterialParameters, PixelMaterialInputs);
+
+	OutColor = HitProxyId;
+}
\ No newline at end of file
diff --git a/Engine/Plugins/Runtime/nDisplay/Shaders/Private/MeshProjectionSelectionOutline.usf b/Engine/Plugins/Runtime/nDisplay/Shaders/Private/MeshProjectionSelectionOutline.usf
new file mode 100644
--- /dev/null
+++ b/Engine/Plugins/Runtime/nDisplay/Shaders/Private/MeshProjectionSelectionOutline.usf
@@ -0,0 +1,221 @@
+// Copyright Epic Games, Inc. All Rights Reserved.
+
+#include "/Engine/Private/Common.ush"
+#include "/Engine/Private/ScreenPass.ush"
+
+// Note: opengl compiler doesn't like Texture2D<float>
+Texture2D EditorPrimitivesDepth;
+Texture2D<uint2> EditorPrimitivesStencil;
+
+Texture2D ColorTexture;
+SamplerState ColorSampler;
+
+Texture2D DepthTexture;
+SamplerState DepthSampler;
+
+SCREEN_PASS_TEXTURE_VIEWPORT(Color)
+SCREEN_PASS_TEXTURE_VIEWPORT(Depth)
+
+float3 OutlineColor;
+float SelectionHighlightIntensity;
+
+struct FPixelInfo
+{
+	// if there is an selected object
+	bool bObjectMask;
+
+	int ObjectId;
+	// if the current pixel is not occluded by some other object in front of it
+	float fVisible;
+	// hard version of fVisible
+	bool bVisible;
+};
+
+// @param DeviceZPlane plane in screenspace .x: ddx(DeviceZ), y: ddy(DeviceZ) z:DeviceZ 
+float ReconstructDeviceZ(float3 DeviceZPlane, float2 PixelOffset)
+{
+	return dot(DeviceZPlane, float3(PixelOffset.xy, 1));
+}
+
+// @param DeviceZPlane plane in screenspace .x: ddx(DeviceZ), y: ddy(DeviceZ) z:DeviceZ
+FPixelInfo GetPixelInfo(int2 PixelPos, int OffsetX, int OffsetY, float3 DeviceZPlane, float2 DeviceZMinMax)
+{
+	FPixelInfo Ret;
+
+	PixelPos += int2(OffsetX, OffsetY);
+
+	// more stable on the silhouette
+	float2 ReconstructionOffset = 0.5f - View.TemporalAAParams.zw;
+
+	float DeviceZ = EditorPrimitivesDepth.Load(int3(PixelPos, 0)).r;
+	int Stencil = EditorPrimitivesStencil.Load(int3(PixelPos, 0)) STENCIL_COMPONENT_SWIZZLE;
+
+	float SceneDeviceZ = ReconstructDeviceZ(DeviceZPlane, ReconstructionOffset);
+
+	// clamp SceneDeviceZ in (DeviceZMinMax.x .. DeviceZMinMax.z)
+	// this avoids flicking artifacts on the silhouette by limiting the depth reconstruction error
+	SceneDeviceZ = max(SceneDeviceZ, DeviceZMinMax.x);
+	SceneDeviceZ = min(SceneDeviceZ, DeviceZMinMax.y);
+
+	// test against far plane
+	Ret.bObjectMask = DeviceZ != 0.0f;
+
+	// outline even between multiple selected objects (best usability)
+	Ret.ObjectId = Stencil;
+
+#if (COMPILER_GLSL == 1 && FEATURE_LEVEL < FEATURE_LEVEL_SM4) || (COMPILER_METAL && FEATURE_LEVEL < FEATURE_LEVEL_SM4)
+	// Stencil read in opengl is not supported on older versions
+	Ret.ObjectId = Ret.bObjectMask ? 2 : 0;
+#endif
+
+	// Soft Bias with DeviceZ for best quality (larger bias than usual because SceneDeviceZ is only approximated)
+	const float DeviceDepthFade = 0.00005f;
+
+	// 2 to always bias over the current plane
+	Ret.fVisible = saturate(2.0f - (SceneDeviceZ - DeviceZ) / DeviceDepthFade);
+
+	Ret.bVisible = Ret.fVisible >= 0.5f;
+
+	return Ret;
+}
+
+// Computes min and max at once.
+void MinMax(inout int2 Var, int Value)
+{
+	Var.x = min(Var.x, Value);
+	Var.y = max(Var.y, Value);
+}
+
+float4 GetOutline(int2 PixelPos, FPixelInfo CenterPixelInfo, float3 DeviceZPlane, float2 DeviceZMinMax)
+{
+	float4 Color = 0;
+
+	// [0]:center, [1..4]:borders
+	FPixelInfo PixelInfo[5];
+	PixelInfo[0] = CenterPixelInfo;
+
+	// Diagonal cross is thicker than vertical/horizontal cross.
+	PixelInfo[1] = GetPixelInfo(PixelPos,  1,  1, DeviceZPlane, DeviceZMinMax);
+	PixelInfo[2] = GetPixelInfo(PixelPos, -1, -1, DeviceZPlane, DeviceZMinMax);
+	PixelInfo[3] = GetPixelInfo(PixelPos,  1, -1, DeviceZPlane, DeviceZMinMax);
+	PixelInfo[4] = GetPixelInfo(PixelPos, -1,  1, DeviceZPlane, DeviceZMinMax);
+
+	// With (.x != .y) we can detect a border around and between each object.
+	int2 BorderMinMax = int2(255,0);
+	{
+		MinMax(BorderMinMax, PixelInfo[1].ObjectId);
+		MinMax(BorderMinMax, PixelInfo[2].ObjectId);
+		MinMax(BorderMinMax, PixelInfo[3].ObjectId);
+		MinMax(BorderMinMax, PixelInfo[4].ObjectId);
+	}
+
+	int2 VisibleBorderMinMax = int2(255,0);
+	{
+		FLATTEN	if (PixelInfo[1].bVisible) MinMax(VisibleBorderMinMax, PixelInfo[1].ObjectId);
+		FLATTEN	if (PixelInfo[2].bVisible) MinMax(VisibleBorderMinMax, PixelInfo[2].ObjectId);
+		FLATTEN	if (PixelInfo[3].bVisible) MinMax(VisibleBorderMinMax, PixelInfo[3].ObjectId);
+		FLATTEN	if (PixelInfo[4].bVisible) MinMax(VisibleBorderMinMax, PixelInfo[4].ObjectId);
+	}
+
+	// this border around the object
+	bool bVisibleBorder = VisibleBorderMinMax.y != 0;
+	bool bBorder = BorderMinMax.x != BorderMinMax.y;
+	bool bInnerBorder = BorderMinMax.y == 4;
+
+	// moving diagonal lines
+	float PatternMask = ((PixelPos.x/2 + PixelPos.y/2) % 2) * 0.6f;
+
+	// the contants express the two opacity values we see when the primitive is hidden
+	float LowContrastPatternMask = lerp(0.2, 1.0f, PatternMask);
+
+	float3 SelectionColor = OutlineColor.rgb;
+	FLATTEN if (bBorder)
+	{
+		FLATTEN if (bVisibleBorder)
+		{
+			// unoccluded border
+			Color = float4(SelectionColor, 1);
+		}
+		else
+		{
+			// occluded border
+			Color = lerp(float4(0, 0, 0, 0), float4(SelectionColor, 1), LowContrastPatternMask);
+		}
+	}
+
+	FLATTEN if (bInnerBorder)
+	{
+		// even occluded object is filled
+		// occluded object is rendered differently(flickers with TemporalAA)
+		float VisibleMask = lerp(PatternMask, 1.0f, PixelInfo[0].fVisible);
+
+		// darken occluded parts
+		Color = lerp(Color, float4(SelectionColor * VisibleMask, 1), SelectionHighlightIntensity);
+	}
+
+	// highlight inner part of the object
+	if (PixelInfo[0].bObjectMask)
+	{
+		float VisibleMask = lerp(PatternMask, 1.0f, PixelInfo[0].fVisible);
+		float InnerHighlightAmount = SelectionHighlightIntensity;
+		Color = lerp(Color, float4(SelectionColor, 1), VisibleMask * InnerHighlightAmount);
+	}
+
+	return Color;
+}
+
+void Main(
+	noperspective float4 UVAndScreenPos : TEXCOORD0,
+	float4 SvPosition : SV_POSITION,
+	out float4 OutColor : SV_Target0)
+{
+	const float2 UV = UVAndScreenPos.xy;
+	const int2 Pixel = int2(UV * Color_Extent);
+
+	// Scout outwards from the center pixel to determine if any neighboring pixels are selected
+	const int ScoutStep = 2;
+	
+	const int4 StencilScout = int4(
+		EditorPrimitivesStencil.Load(int3(Pixel + int2( ScoutStep,  0), 0)) STENCIL_COMPONENT_SWIZZLE,
+		EditorPrimitivesStencil.Load(int3(Pixel + int2(-ScoutStep,  0), 0)) STENCIL_COMPONENT_SWIZZLE,
+		EditorPrimitivesStencil.Load(int3(Pixel + int2( 0,  ScoutStep), 0)) STENCIL_COMPONENT_SWIZZLE,
+		EditorPrimitivesStencil.Load(int3(Pixel + int2( 0, -ScoutStep), 0)) STENCIL_COMPONENT_SWIZZLE
+	);
+
+	const int ScoutSum = dot(saturate(StencilScout), 1);
+
+	// If this sum is zero, none of our neighbors are selected pixels and we can skip all the heavy processing.
+	if (ScoutSum > 0)
+	{
+		float Center = Texture2DSampleLevel(DepthTexture, DepthSampler, UV, 0).r;
+		float Left   = Texture2DSampleLevel(DepthTexture, DepthSampler, UV + float2(-1,  0) * Depth_ExtentInverse, 0).r;
+		float Right  = Texture2DSampleLevel(DepthTexture, DepthSampler, UV + float2( 1,  0) * Depth_ExtentInverse, 0).r;
+		float Top    = Texture2DSampleLevel(DepthTexture, DepthSampler, UV + float2( 0, -1) * Depth_ExtentInverse, 0).r;
+		float Bottom = Texture2DSampleLevel(DepthTexture, DepthSampler, UV + float2( 0,  1) * Depth_ExtentInverse, 0).r;
+
+		// This allows to reconstruct depth with a small pixel offset without many texture lookups (4xMSAA * 5 neighbors -> 20 samples)
+		// It's an approximation assuming the surface is a plane.
+		float3 DeviceZPlane;
+		DeviceZPlane.x = (Right - Left) / 2;
+		DeviceZPlane.y = (Bottom - Top) / 2;
+		DeviceZPlane.z = Center;
+
+		float2 DeviceZMinMax;
+		DeviceZMinMax.x = min(Center, min(min(Left, Right), min(Top, Bottom)));
+		DeviceZMinMax.y = max(Center, max(max(Left, Right), max(Top, Bottom)));
+
+		FPixelInfo CenterPixelInfo = GetPixelInfo(Pixel, 0, 0, DeviceZPlane, DeviceZMinMax);
+		float4 Outline = GetOutline(Pixel, CenterPixelInfo, DeviceZPlane, DeviceZMinMax);
+
+		OutColor = Texture2DSample(ColorTexture, ColorSampler, UV);
+
+		// Scene color has gamma applied. Need to remove and re-apply after linear operation.
+		OutColor.rgb = pow(OutColor.rgb, 2.2f );
+		OutColor.rgb = lerp(OutColor.rgb, Outline.rgb, Outline.a);
+		OutColor.rgb = pow(OutColor.rgb, 1.0f / 2.2f);
+	}
+	else
+	{
+		OutColor = Texture2DSample(ColorTexture, ColorSampler, UV);
+	}
+}
diff --git a/Engine/Plugins/Runtime/nDisplay/Shaders/Private/MeshProjectionShaders.usf b/Engine/Plugins/Runtime/nDisplay/Shaders/Private/MeshProjectionShaders.usf
--- a/Engine/Plugins/Runtime/nDisplay/Shaders/Private/MeshProjectionShaders.usf
+++ b/Engine/Plugins/Runtime/nDisplay/Shaders/Private/MeshProjectionShaders.usf
@@ -62,6 +62,10 @@ void MainPS(FMeshProjectionVSToPS Input, in float4 SvPosition : SV_Position OPTI
 	FPixelMaterialInputs PixelMaterialInputs;
 	CalcMaterialParameters(MaterialParameters, PixelMaterialInputs, SvPosition, bIsFrontFace);
 
+	GetMaterialCoverageAndClipping(MaterialParameters, PixelMaterialInputs);
+	
 	float3 Emissive = GetMaterialEmissive(PixelMaterialInputs);
-	OutputColor = float4(Emissive, 1);
+	float Opacity = GetMaterialOpacity(PixelMaterialInputs);
+	
+	OutputColor = float4(Emissive, Opacity);
 }
\ No newline at end of file
diff --git a/Engine/Plugins/Runtime/nDisplay/Source/DisplayClusterLightCardEditor/Private/Viewport/DisplayClusterLightCardEditorViewportClient.cpp b/Engine/Plugins/Runtime/nDisplay/Source/DisplayClusterLightCardEditor/Private/Viewport/DisplayClusterLightCardEditorViewportClient.cpp
--- a/Engine/Plugins/Runtime/nDisplay/Source/DisplayClusterLightCardEditor/Private/Viewport/DisplayClusterLightCardEditorViewportClient.cpp
+++ b/Engine/Plugins/Runtime/nDisplay/Source/DisplayClusterLightCardEditor/Private/Viewport/DisplayClusterLightCardEditorViewportClient.cpp
@@ -3,6 +3,7 @@
 #include "DisplayClusterLightCardEditorViewportClient.h"
 
 #include "DisplayClusterConfigurationTypes.h"
+#include "DisplayClusterProjectionStrings.h"
 #include "DisplayClusterRootActor.h"
 #include "Components/DisplayClusterPreviewComponent.h"
 #include "Components/DisplayClusterCameraComponent.h"
@@ -50,8 +51,8 @@ FDisplayClusterLightCardEditorViewportClient::FDisplayClusterLightCardEditorView
 	LightCardEditorPtr = InLightCardEditor;
 	
 	MeshProjectionRenderer = MakeShared<FDisplayClusterMeshProjectionRenderer>();
+	MeshProjectionRenderer->ActorSelectedDelegate = FDisplayClusterMeshProjectionRenderer::FSelection::CreateRaw(this, &FDisplayClusterLightCardEditorViewportClient::IsLightCardSelected);
 
-	SelectedActor = nullptr;
 	bDraggingActor = false;
 	ScopedTransaction = nullptr;
 	
@@ -62,6 +63,8 @@ FDisplayClusterLightCardEditorViewportClient::FDisplayClusterLightCardEditorView
 	DrawHelper.bDrawGrid = false;
 	DrawHelper.PerspectiveGridSize = HALF_WORLD_MAX1;
 
+	EngineShowFlags.SetSelectionOutline(true);
+
 	check(Widget);
 	Widget->SetSnapEnabled(true);
 	
@@ -252,7 +255,7 @@ void FDisplayClusterLightCardEditorViewportClient::Draw(FViewport* InViewport, F
 	FSceneViewInitOptions ViewInitOptions;
 	GetSceneViewInitOptions(ViewInitOptions);
 
-	MeshProjectionRenderer->Render(Canvas, GetScene(), ViewInitOptions, ProjectionMode);
+	MeshProjectionRenderer->Render(Canvas, GetScene(), ViewInitOptions, EngineShowFlags, ProjectionMode);
 
 	if (View)
 	{
@@ -311,21 +314,11 @@ void FDisplayClusterLightCardEditorViewportClient::Draw(FViewport* InViewport, F
 
 UE::Widget::EWidgetMode FDisplayClusterLightCardEditorViewportClient::GetWidgetMode() const
 {
-	if (IsActorSelected())
-	{
-		return FEditorViewportClient::GetWidgetMode();
-	}
-
 	return UE::Widget::WM_None;
 }
 
 FVector FDisplayClusterLightCardEditorViewportClient::GetWidgetLocation() const
 {
-	if (IsActorSelected())
-	{
-		return SelectedActor->GetActorLocation();
-	}
-
 	return FEditorViewportClient::GetWidgetLocation();
 }
 
@@ -386,7 +379,7 @@ bool FDisplayClusterLightCardEditorViewportClient::InputWidgetDelta(FViewport* I
 void FDisplayClusterLightCardEditorViewportClient::TrackingStarted(const FInputEventState& InInputState, bool bIsDraggingWidget,
 	bool bNudge)
 {
-	if (!bDraggingActor && bIsDraggingWidget && InInputState.IsLeftMouseButtonPressed() && IsActorSelected())
+	if (!bDraggingActor && bIsDraggingWidget && InInputState.IsLeftMouseButtonPressed() && SelectedLightCards.Num())
 	{
 		GEditor->DisableDeltaModification(true);
 		{
@@ -407,7 +400,7 @@ void FDisplayClusterLightCardEditorViewportClient::TrackingStopped()
 	bDraggingActor = false;
 	EndTransaction();
 
-	if (IsActorSelected())
+	if (SelectedLightCards.Num())
 	{
 		GEditor->DisableDeltaModification(false);
 	}
@@ -415,14 +408,97 @@ void FDisplayClusterLightCardEditorViewportClient::TrackingStopped()
 	FEditorViewportClient::TrackingStopped();
 }
 
-void FDisplayClusterLightCardEditorViewportClient::ProcessClick(FSceneView& View, HHitProxy* HitProxy, FKey Key, EInputEvent Event,
-                                                uint32 HitX, uint32 HitY)
+void FDisplayClusterLightCardEditorViewportClient::ProcessClick(FSceneView& View, HHitProxy* HitProxy, FKey Key, EInputEvent Event, uint32 HitX, uint32 HitY)
 {
-	// TODO: Add logic to check if the hit proxy was a light card, and select that light card
+	UWorld* PreviewWorld = PreviewScene->GetWorld();
+	check(PreviewWorld);
+
+	if (HitProxy)
+	{
+		if (HitProxy->IsA(HActor::StaticGetType()))
+		{
+			HActor* ActorHitProxy = static_cast<HActor*>(HitProxy);
+			if (ActorHitProxy->Actor == RootActorPreviewInstance.Get())
+			{
+				if (ActorHitProxy->PrimComponent && ActorHitProxy->PrimComponent->IsA<UStaticMeshComponent>())
+				{
+					AActor* TracedLightCard = TraceScreenForLightCard(View, HitX, HitY);
+					SelectLightCard(TracedLightCard);
+				}
+			}
+			else if (LightCardProxies.Contains(ActorHitProxy->Actor))
+			{
+				SelectLightCard(ActorHitProxy->Actor);
+			}
+			else
+			{
+				SelectLightCard(nullptr);
+			}
+		}
+	}
+	else
+	{
+		SelectLightCard(nullptr);
+	}
 	
 	FEditorViewportClient::ProcessClick(View, HitProxy, Key, Event, HitX, HitY);
 }
 
+EMouseCursor::Type FDisplayClusterLightCardEditorViewportClient::GetCursor(FViewport* InViewport, int32 X, int32 Y)
+{
+	EMouseCursor::Type MouseCursor = EMouseCursor::Default;
+
+	if (RequiredCursorVisibiltyAndAppearance.bOverrideAppearance &&
+		RequiredCursorVisibiltyAndAppearance.bHardwareCursorVisible)
+	{
+		 MouseCursor = RequiredCursorVisibiltyAndAppearance.RequiredCursor;
+	}
+	else if (!RequiredCursorVisibiltyAndAppearance.bHardwareCursorVisible)
+	{
+		MouseCursor = EMouseCursor::None;
+	}
+	else if( InViewport->IsCursorVisible() && !bWidgetAxisControlledByDrag )
+	{
+		HHitProxy* HitProxy = InViewport->GetHitProxy(X,Y);
+		if (HitProxy)
+		{
+			bShouldCheckHitProxy = true;
+
+			if (HitProxy->IsA(HActor::StaticGetType()))
+			{
+				HActor* ActorHitProxy = static_cast<HActor*>(HitProxy);
+				if (ActorHitProxy->Actor == RootActorPreviewInstance.Get())
+				{
+					if (ActorHitProxy->PrimComponent && ActorHitProxy->PrimComponent->IsA<UStaticMeshComponent>())
+					{
+						FSceneViewFamilyContext ViewFamily(FSceneViewFamily::ConstructionValues(
+							InViewport,
+							GetScene(),
+							EngineShowFlags)
+							.SetRealtimeUpdate(IsRealtime()));
+						FSceneView* View = CalcSceneView( &ViewFamily );
+
+						AActor* TracedLightCard = TraceScreenForLightCard(*View, X, Y);
+						if (TracedLightCard)
+						{
+							MouseCursor = EMouseCursor::Crosshairs;
+						}
+					}
+				}
+				else if (LightCardProxies.Contains(ActorHitProxy->Actor))
+				{
+					MouseCursor = EMouseCursor::Crosshairs;
+				}
+			}
+		}
+	}
+
+	CachedMouseX = X;
+	CachedLastMouseY = Y;
+
+	return MouseCursor;
+}
+
 void FDisplayClusterLightCardEditorViewportClient::SelectActor(AActor* NewActor)
 {
 }
@@ -433,19 +509,30 @@ void FDisplayClusterLightCardEditorViewportClient::ResetSelection()
 	SetWidgetMode(UE::Widget::WM_None);
 }
 
-void FDisplayClusterLightCardEditorViewportClient::UpdatePreviewActor(ADisplayClusterRootActor* RootActor)
+void FDisplayClusterLightCardEditorViewportClient::UpdatePreviewActor(ADisplayClusterRootActor* RootActor, bool bForce)
 {
 	UWorld* PreviewWorld = PreviewScene->GetWorld();
 	check (PreviewWorld);
 	
+	if (!bForce && RootActor == RootActorLevelInstance.Get())
+	{
+		return;
+	}
+
 	if (RootActorPreviewInstance.IsValid() && (!RootActor || !PreviewWorld->ContainsActor(RootActor)))
 	{
+		MeshProjectionRenderer->ClearScene();
 		PreviewWorld->EditorDestroyActor(RootActorPreviewInstance.Get(), false);
 		
 		RootActorPreviewInstance.Reset();
-		RootActorLevelInstance.Reset();
 
-		MeshProjectionRenderer->ClearScene();
+		for (const TWeakObjectPtr<AActor> LightCardProxy : LightCardProxies)
+		{
+			PreviewWorld->EditorDestroyActor(LightCardProxy.Get(), false);
+		}
+
+		RootActorLevelInstance.Reset();
+		LightCardProxies.Empty();
 	}
 	
 	if (RootActor)
@@ -497,6 +584,26 @@ void FDisplayClusterLightCardEditorViewportClient::UpdatePreviewActor(ADisplayCl
 
 		RootActorPreviewInstance->UpdatePreviewComponents();
 
+		TArray<TWeakObjectPtr<AActor>> LevelLightCards;
+		FindLightCardsForRootActor(RootActor, LevelLightCards);
+
+		for (const TWeakObjectPtr<AActor>& LightCard : LevelLightCards)
+		{
+			FActorSpawnParameters LightCardSpawnInfo;
+			LightCardSpawnInfo.SpawnCollisionHandlingOverride = ESpawnActorCollisionHandlingMethod::AlwaysSpawn;
+			LightCardSpawnInfo.bNoFail = true;
+			LightCardSpawnInfo.ObjectFlags = RF_Transient|RF_Transactional;
+			LightCardSpawnInfo.Template = LightCard.Get();
+		
+			AActor* LightCardProxy = PreviewWorld->SpawnActor(LightCard->GetClass(), &FVector::ZeroVector, &FRotator::ZeroRotator, MoveTemp(LightCardSpawnInfo));
+
+			LightCardProxy->SetActorLocation(LightCard->GetActorLocation() - RootActor->GetActorLocation());
+			LightCardProxy->SetActorRotation(LightCard->GetActorRotation() - RootActor->GetActorRotation());
+
+			LightCardProxies.Add(LightCardProxy);
+			MeshProjectionRenderer->AddActor(LightCardProxy);
+		}
+
 		FCoreUObjectDelegates::OnObjectPropertyChanged.RemoveAll(this);
 		FCoreUObjectDelegates::OnObjectPropertyChanged.AddRaw(this, &FDisplayClusterLightCardEditorViewportClient::OnActorPropertyChanged);
 	}
@@ -517,6 +624,9 @@ void FDisplayClusterLightCardEditorViewportClient::SetProjectionMode(EDisplayClu
 	}
 
 	FindProjectionOriginComponent();
+
+	Viewport->InvalidateHitProxy();
+	bShouldCheckHitProxy = true;
 }
 
 float FDisplayClusterLightCardEditorViewportClient::GetProjectionModeFOV(EDisplayClusterMeshProjectionType InProjectionMode) const
@@ -580,7 +690,7 @@ void FDisplayClusterLightCardEditorViewportClient::OnActorPropertyChanged(UObjec
 	}
 	if (bIsOurActor)
 	{
-		UpdatePreviewActor(RootActorLevelInstance.Get());
+		UpdatePreviewActor(RootActorLevelInstance.Get(), true);
 	}
 }
 
@@ -699,33 +809,205 @@ void FDisplayClusterLightCardEditorViewportClient::GetSceneViewInitOptions(FScen
 	OutViewInitOptions = ViewInitOptions;
 }
 
+UDisplayClusterConfigurationViewport* FDisplayClusterLightCardEditorViewportClient::FindViewportForPrimitiveComponent(UPrimitiveComponent* PrimitiveComponent)
+{
+	if (RootActorPreviewInstance.IsValid())
+	{
+		const FString PrimitiveComponentName = PrimitiveComponent->GetName();
+		UDisplayClusterConfigurationData* Config = RootActorPreviewInstance->GetConfigData();
+		
+		for (const TPair<FString, UDisplayClusterConfigurationClusterNode*>& NodePair : Config->Cluster->Nodes)
+		{
+			const UDisplayClusterConfigurationClusterNode* Node = NodePair.Value;
+			for (const TPair<FString, UDisplayClusterConfigurationViewport*>& ViewportPair : Node->Viewports)
+			{
+				UDisplayClusterConfigurationViewport* CfgViewport = ViewportPair.Value;
+
+				FString ComponentName;
+				if (CfgViewport->ProjectionPolicy.Type.Equals(DisplayClusterProjectionStrings::projection::Simple, ESearchCase::IgnoreCase)
+					&& CfgViewport->ProjectionPolicy.Parameters.Contains(DisplayClusterProjectionStrings::cfg::simple::Screen))
+				{
+					ComponentName = CfgViewport->ProjectionPolicy.Parameters[DisplayClusterProjectionStrings::cfg::simple::Screen];
+				}
+				else if (CfgViewport->ProjectionPolicy.Type.Equals(DisplayClusterProjectionStrings::projection::Mesh, ESearchCase::IgnoreCase)
+					&& CfgViewport->ProjectionPolicy.Parameters.Contains(DisplayClusterProjectionStrings::cfg::mesh::Component))
+				{
+					ComponentName = CfgViewport->ProjectionPolicy.Parameters[DisplayClusterProjectionStrings::cfg::mesh::Component];
+				}
+
+				if (ComponentName == PrimitiveComponentName)
+				{
+					return CfgViewport;
+				}
+			}
+		}
+	}
+
+	return nullptr;
+}
+
 void FDisplayClusterLightCardEditorViewportClient::FindProjectionOriginComponent()
 {
 	if (RootActorPreviewInstance.IsValid())
 	{
-		if (ProjectionMode == EDisplayClusterMeshProjectionType::Perspective)
+		TArray<UDisplayClusterCameraComponent*> ViewOriginComponents;
+		RootActorPreviewInstance->GetComponents<UDisplayClusterCameraComponent>(ViewOriginComponents);
+
+		if (ViewOriginComponents.Num())
+		{
+			ProjectionOriginComponent = ViewOriginComponents[0];
+		}
+		else
 		{
-			TArray<UDisplayClusterCameraComponent*> ViewOriginComponents;
-			RootActorPreviewInstance->GetComponents<UDisplayClusterCameraComponent>(ViewOriginComponents);
+			ProjectionOriginComponent = RootActorPreviewInstance->GetRootComponent();
+		}
+	}
+	else
+	{
+		ProjectionOriginComponent.Reset();
+	}
+}
+
+void FDisplayClusterLightCardEditorViewportClient::FindLightCardsForRootActor(ADisplayClusterRootActor* RootActor, TArray<TWeakObjectPtr<AActor>>& OutLightCards)
+{
+	if (RootActor)
+	{
+		FDisplayClusterConfigurationICVFX_VisibilityList& RootActorLightCards = RootActor->GetConfigData()->StageSettings.Lightcard.ShowOnlyList;
 
-			if (ViewOriginComponents.Num())
+		for (const TSoftObjectPtr<AActor>& LightCardActor : RootActorLightCards.Actors)
+		{
+			if (LightCardActor.IsValid())
 			{
-				ProjectionOriginComponent = ViewOriginComponents[0];
+				OutLightCards.Add(LightCardActor.Get());
 			}
-			else
+		}
+
+		// If there are any layers that are specified as light card layers, iterate over all actors in the world and 
+		// add any that are members of any of the light card layers to the list. Only add an actor once, even if it is
+		// in multiple layers
+		if (RootActorLightCards.ActorLayers.Num())
+		{
+			if (UWorld* World = RootActor->GetWorld())
 			{
-				ProjectionOriginComponent = RootActorPreviewInstance->GetRootComponent();
+				for (const TWeakObjectPtr<AActor> WeakActor : FActorRange(World))
+				{
+					if (WeakActor.IsValid())
+					{
+						for (const FActorLayer& ActorLayer : RootActorLightCards.ActorLayers)
+						{
+							if (WeakActor->Layers.Contains(ActorLayer.Name))
+							{
+								OutLightCards.Add(WeakActor);
+								break;
+							}
+						}
+					}
+				}
 			}
 		}
-		else if (ProjectionMode == EDisplayClusterMeshProjectionType::Azimuthal)
+	}
+}
+
+bool FDisplayClusterLightCardEditorViewportClient::IsLightCardSelected(const AActor* Actor)
+{
+	return SelectedLightCards.Contains(Actor);
+}
+
+void FDisplayClusterLightCardEditorViewportClient::SelectLightCard(AActor* Actor, bool bAddToSelection)
+{
+	TArray<AActor*> UpdatedActors;
+
+	if (!bAddToSelection)
+	{
+		for (const TWeakObjectPtr<AActor>& LightCard : SelectedLightCards)
 		{
-			ProjectionOriginComponent = RootActorPreviewInstance->GetRootComponent();
+			if (LightCard.IsValid())
+			{
+				UpdatedActors.Add(LightCard.Get());
+			}
 		}
+
+		SelectedLightCards.Empty();
 	}
-	else
+
+	if (Actor)
 	{
-		ProjectionOriginComponent.Reset();
+		SelectedLightCards.Add(Actor);
+		UpdatedActors.Add(Actor);
+	}
+
+	for (AActor* UpdatedActor : UpdatedActors)
+	{
+		UpdatedActor->PushSelectionToProxies();
+	}
+}
+
+AActor* FDisplayClusterLightCardEditorViewportClient::TraceScreenForLightCard(const FSceneView& View, int32 HitX, int32 HitY)
+{
+	UWorld* PreviewWorld = PreviewScene->GetWorld();
+	check(PreviewWorld);
+
+	const FVector4 ScreenPos = View.PixelToScreen(HitX, HitY, 0);
+	const FMatrix InvProjMatrix = View.ViewMatrices.GetInvProjectionMatrix();
+	const FMatrix InvViewMatrix = View.ViewMatrices.GetInvViewMatrix();
+	FVector ViewPos = FVector(InvProjMatrix.TransformFVector4(FVector4(ScreenPos.X * GNearClippingPlane, ScreenPos.Y * GNearClippingPlane, 0.0f, GNearClippingPlane)));
+	
+	// If the projection mode is azimuthal, we need to convert to the "true" view position by undoing the non-linear projection applied by the projection shader
+	if (ProjectionMode == EDisplayClusterMeshProjectionType::Azimuthal)
+	{
+		const float Rho = ViewPos.Length();
+		const FVector UnitViewPos = ViewPos.GetSafeNormal();
+
+		const FVector PlanePos = UnitViewPos / UnitViewPos.Z;
+		const FVector2D PolarCoords = FVector2D(FMath::Sqrt(PlanePos.X * PlanePos.X + PlanePos.Y * PlanePos.Y), FMath::Atan2(PlanePos.Y, PlanePos.X));
+		ViewPos = FVector(FMath::Sin(PolarCoords.X) * FMath::Cos(PolarCoords.Y), FMath::Sin(PolarCoords.X) * FMath::Sin(PolarCoords.Y), FMath::Cos(PolarCoords.X)) * Rho;
 	}
+
+	const FVector Direction = InvViewMatrix.TransformVector(ViewPos).GetSafeNormal();
+
+	const FVector CursorRayStart = View.ViewMatrices.GetViewOrigin();
+	const FVector CursorRayEnd = CursorRayStart + Direction * HALF_WORLD_MAX;
+
+	FCollisionQueryParams Param(SCENE_QUERY_STAT(DragDropTrace), true);
+
+	bool bHitLightCard = false;
+	FHitResult ScreenHitResult;
+	if (PreviewWorld->LineTraceSingleByObjectType(ScreenHitResult, CursorRayStart, CursorRayEnd, FCollisionObjectQueryParams(FCollisionObjectQueryParams::InitType::AllObjects), Param))
+	{
+		if (AActor* HitActor = ScreenHitResult.GetActor())
+		{
+			if (RootActorPreviewInstance.Get() == HitActor && ScreenHitResult.Component.IsValid())
+			{
+				if (UDisplayClusterConfigurationViewport* CfgViewport = FindViewportForPrimitiveComponent(ScreenHitResult.Component.Get()))
+				{
+					FString ViewOriginName = CfgViewport->Camera;
+					if (UDisplayClusterCameraComponent* ViewOrigin = RootActorPreviewInstance->GetComponentByName<UDisplayClusterCameraComponent>(ViewOriginName))
+					{
+						const FVector ViewOriginRayStart = ViewOrigin->GetComponentLocation();
+						const FVector ViewOriginRayEnd = ViewOriginRayStart + (ScreenHitResult.Location - ViewOriginRayStart) * HALF_WORLD_MAX;
+
+						TArray<FHitResult> HitResults;
+						if (PreviewWorld->LineTraceMultiByObjectType(HitResults, ViewOriginRayStart, ViewOriginRayEnd, FCollisionObjectQueryParams(FCollisionObjectQueryParams::InitType::AllObjects), Param))
+						{
+							for (FHitResult& HitResult : HitResults)
+							{
+								if (LightCardProxies.Contains(HitResult.GetActor()))
+								{
+									return HitResult.GetActor();
+								}
+							}
+						}
+					}
+				}
+			}
+			else if (LightCardProxies.Contains(HitActor))
+			{
+				return HitActor;
+			}
+		}
+	}
+
+	return nullptr;
 }
 
 #undef LOCTEXT_NAMESPACE
diff --git a/Engine/Plugins/Runtime/nDisplay/Source/DisplayClusterLightCardEditor/Private/Viewport/DisplayClusterLightCardEditorViewportClient.h b/Engine/Plugins/Runtime/nDisplay/Source/DisplayClusterLightCardEditor/Private/Viewport/DisplayClusterLightCardEditorViewportClient.h
--- a/Engine/Plugins/Runtime/nDisplay/Source/DisplayClusterLightCardEditor/Private/Viewport/DisplayClusterLightCardEditorViewportClient.h
+++ b/Engine/Plugins/Runtime/nDisplay/Source/DisplayClusterLightCardEditor/Private/Viewport/DisplayClusterLightCardEditorViewportClient.h
@@ -12,6 +12,7 @@
 class ADisplayClusterRootActor;
 class SDisplayClusterLightCardEditor;
 class FScopedTransaction;
+class UDisplayClusterConfigurationViewport;
 
 /** Viewport Client for the preview viewport */
 class FDisplayClusterLightCardEditorViewportClient : public FEditorViewportClient, public TSharedFromThis<FDisplayClusterLightCardEditorViewportClient>
@@ -36,6 +37,7 @@ public:
 	virtual void TrackingStarted(const FInputEventState& InInputState, bool bIsDraggingWidget, bool bNudge) override;
 	virtual void TrackingStopped() override;
 	virtual void ProcessClick(FSceneView& View, HHitProxy* HitProxy, FKey Key, EInputEvent Event, uint32 HitX, uint32 HitY) override;
+	virtual EMouseCursor::Type GetCursor(FViewport* Viewport,int32 X,int32 Y) override;
 	virtual ELevelViewportType GetViewportType() const override { return LVT_Perspective; }
 	// ~FEditorViewportClient
 
@@ -43,11 +45,7 @@ public:
 	void ResetSelection();
 
 	/** Update the spawned preview actor from a root actor in the level. */
-	void UpdatePreviewActor(ADisplayClusterRootActor* RootActor);
-	
-	AActor* GetSelectedActor() const { return SelectedActor.Get(); }
-
-	bool IsActorSelected() const { return SelectedActor.IsValid(); }
+	void UpdatePreviewActor(ADisplayClusterRootActor* RootActor, bool bForce = false);
 	
 	EDisplayClusterMeshProjectionType GetProjectionMode() const { return ProjectionMode; }
 	void SetProjectionMode(EDisplayClusterMeshProjectionType InProjectionMode);
@@ -58,7 +56,7 @@ public:
 	/** Sets the field of view of the specified projection mode */
 	void SetProjectionModeFOV(EDisplayClusterMeshProjectionType InProjectionMode, float NewFOV);
 
-protected:
+private:
 	/** Initiates a transaction. */
 	void BeginTransaction(const FText& Description);
 
@@ -74,15 +72,32 @@ protected:
 	/** Gets the scene view init options to use to create scene views for the preview scene */
 	void GetSceneViewInitOptions(FSceneViewInitOptions& OutViewInitOptions);
 
+	/** Gets the viewport that is attached to the specified primitive component */
+	UDisplayClusterConfigurationViewport* FindViewportForPrimitiveComponent(UPrimitiveComponent* PrimitiveComponent);
+
 	/** Finds a suitable primitive component on the stage actor to use as a projection origin */
 	void FindProjectionOriginComponent();
 
+	/** Gets a list of all light card actors on the level linked to the specified root actor */
+	void FindLightCardsForRootActor(ADisplayClusterRootActor* RootActor, TArray<TWeakObjectPtr<AActor>>& OutLightCards);
+
+	/** Callback to check if an light card actor is among the list of selected light card actors */
+	bool IsLightCardSelected(const AActor* Actor);
+
+	/** Adds the specified light card actor the the list of selected light cards */
+	void SelectLightCard(AActor* Actor, bool bAddToSelection = false);
+
+	/** Traces to find the light card corresponding to a click on a stage screen */
+	AActor* TraceScreenForLightCard(const FSceneView& View, int32 HitX, int32 HitY);
+
 private:
 	TWeakPtr<FSceneViewport> SceneViewportPtr;
 	TWeakPtr<SDisplayClusterLightCardEditor> LightCardEditorPtr;
 	TWeakObjectPtr<ADisplayClusterRootActor> RootActorPreviewInstance;
 	TWeakObjectPtr<ADisplayClusterRootActor> RootActorLevelInstance;
-	TWeakObjectPtr<AActor> SelectedActor;
+
+	TArray<TWeakObjectPtr<AActor>> LightCardProxies;
+	TArray<TWeakObjectPtr<AActor>> SelectedLightCards;
 	
 	/** The renderer for the viewport, which can render the meshes with a variety of projection types */
 	TSharedPtr<FDisplayClusterMeshProjectionRenderer> MeshProjectionRenderer;
diff --git a/Engine/Plugins/Runtime/nDisplay/Source/DisplayClusterLightCardEditorShaders/DisplayClusterLightCardEditorShaders.Build.cs b/Engine/Plugins/Runtime/nDisplay/Source/DisplayClusterLightCardEditorShaders/DisplayClusterLightCardEditorShaders.Build.cs
--- a/Engine/Plugins/Runtime/nDisplay/Source/DisplayClusterLightCardEditorShaders/DisplayClusterLightCardEditorShaders.Build.cs
+++ b/Engine/Plugins/Runtime/nDisplay/Source/DisplayClusterLightCardEditorShaders/DisplayClusterLightCardEditorShaders.Build.cs
@@ -17,5 +17,10 @@ public class DisplayClusterLightCardEditorShaders : ModuleRules
 				"RHI",
 				"Projects"
 			});
+
+		PrivateIncludePaths.AddRange(
+			new string[] {
+				"../../../../Source/Runtime/Renderer/Private",
+			});
 	}
 }
diff --git a/Engine/Plugins/Runtime/nDisplay/Source/DisplayClusterLightCardEditorShaders/Private/DisplayClusterMeshProjectionRenderer.cpp b/Engine/Plugins/Runtime/nDisplay/Source/DisplayClusterLightCardEditorShaders/Private/DisplayClusterMeshProjectionRenderer.cpp
--- a/Engine/Plugins/Runtime/nDisplay/Source/DisplayClusterLightCardEditorShaders/Private/DisplayClusterMeshProjectionRenderer.cpp
+++ b/Engine/Plugins/Runtime/nDisplay/Source/DisplayClusterLightCardEditorShaders/Private/DisplayClusterMeshProjectionRenderer.cpp
@@ -2,8 +2,6 @@
 
 #include "DisplayClusterMeshProjectionRenderer.h"
 
-#include "Components/PrimitiveComponent.h"
-#include "InstanceCulling/InstanceCullingContext.h"
 #include "PrimitiveSceneInfo.h"
 #include "PrimitiveSceneProxy.h"
 #include "MeshPassProcessor.h"
@@ -12,6 +10,15 @@
 #include "Shader.h"
 #include "CanvasTypes.h"
 #include "EngineModule.h"
+#include "SceneViewExtension.h"
+#include "ScreenPass.h"
+#include "Components/PrimitiveComponent.h"
+#include "InstanceCulling/InstanceCullingContext.h"
+#include "Materials/Material.h"
+
+
+//////////////////////////////////////////////////////////////////////////
+// Base Render Pass
 
 template<EDisplayClusterMeshProjectionType ProjectionType>
 class FMeshProjectionVS : public FMeshMaterialShader
@@ -77,8 +84,8 @@ public:
 		: FMeshPassProcessor(InScene, GMaxRHIFeatureLevel, InView, InDrawListContext)
 		, DrawRenderState(*InView)
 	{
-		DrawRenderState.SetDepthStencilState(TStaticDepthStencilState<false, CF_Always>::GetRHI());
-		DrawRenderState.SetBlendState(TStaticBlendState<>::GetRHI());
+		DrawRenderState.SetDepthStencilState(TStaticDepthStencilState<true, CF_DepthNearOrEqual>::GetRHI());
+		DrawRenderState.SetBlendState(TStaticBlendState<CW_RGBA, BO_Add, BF_SourceAlpha, BF_InverseSourceAlpha, BO_Add, BF_Zero, BF_One>::GetRHI());
 	}
 
 	virtual void AddMeshBatch(const FMeshBatch& RESTRICT MeshBatch, uint64 BatchElementMask, const FPrimitiveSceneProxy* RESTRICT PrimitiveSceneProxy, int32 StaticMeshId = -1) override final
@@ -132,6 +139,313 @@ private:
 	FMeshPassProcessorRenderState DrawRenderState;
 };
 
+
+//////////////////////////////////////////////////////////////////////////
+// Hit Proxy Render Pass
+
+class FMeshProjectionHitProxyShaderElementData : public FMeshMaterialShaderElementData
+{
+public:
+	FMeshProjectionHitProxyShaderElementData(FHitProxyId InBatchHitProxyId)
+		: BatchHitProxyId(InBatchHitProxyId)
+	{
+	}
+
+	FHitProxyId BatchHitProxyId;
+};
+
+class FMeshProjectionHitProxyPS : public FMeshMaterialShader
+{
+	DECLARE_SHADER_TYPE(FMeshProjectionHitProxyPS, MeshMaterial);
+
+public:
+	static bool ShouldCompilePermutation(const FMeshMaterialShaderPermutationParameters& Parameters)
+	{
+		// Only compile the hit proxy shader on desktop editor platforms
+		return IsPCPlatform(Parameters.Platform) && EnumHasAllFlags(Parameters.Flags, EShaderPermutationFlags::HasEditorOnlyData)
+			// and only compile for default materials or materials that are masked.
+			&& (Parameters.MaterialParameters.bIsSpecialEngineMaterial ||
+				!Parameters.MaterialParameters.bWritesEveryPixel ||
+				Parameters.MaterialParameters.bMaterialMayModifyMeshPosition ||
+				Parameters.MaterialParameters.bIsTwoSided);
+	}
+
+	FMeshProjectionHitProxyPS(const ShaderMetaType::CompiledShaderInitializerType& Initializer):
+		FMeshMaterialShader(Initializer)
+	{
+		HitProxyId.Bind(Initializer.ParameterMap, TEXT("HitProxyId"), SPF_Optional);
+	}
+
+	FMeshProjectionHitProxyPS() {}
+
+	void GetShaderBindings(
+		const FScene* Scene,
+		ERHIFeatureLevel::Type FeatureLevel,
+		const FPrimitiveSceneProxy* PrimitiveSceneProxy,
+		const FMaterialRenderProxy& MaterialRenderProxy,
+		const FMaterial& Material,
+		const FMeshPassProcessorRenderState& DrawRenderState,
+		const FMeshProjectionHitProxyShaderElementData& ShaderElementData,
+		FMeshDrawSingleShaderBindings& ShaderBindings) const
+	{
+		FMeshMaterialShader::GetShaderBindings(Scene, FeatureLevel, PrimitiveSceneProxy, MaterialRenderProxy, Material, DrawRenderState, ShaderElementData, ShaderBindings);
+
+		ShaderBindings.Add(HitProxyId, ShaderElementData.BatchHitProxyId.GetColor().ReinterpretAsLinear());
+	}
+
+private:
+	LAYOUT_FIELD(FShaderParameter, HitProxyId)
+};
+
+IMPLEMENT_MATERIAL_SHADER_TYPE(, FMeshProjectionHitProxyPS, TEXT("/Plugin/nDisplay/Private/MeshProjectionHitProxy.usf"), TEXT("Main") ,SF_Pixel);
+
+template<EDisplayClusterMeshProjectionType ProjectionType> 
+class FLightCardEditorHitProxyMeshPassProcessor : public FMeshPassProcessor
+{
+public:
+	FLightCardEditorHitProxyMeshPassProcessor(const FScene* InScene, const FSceneView* InView, FMeshPassDrawListContext* InDrawListContext)
+		: FMeshPassProcessor(InScene, GMaxRHIFeatureLevel, InView, InDrawListContext)
+		, DrawRenderState(*InView)
+	{
+		DrawRenderState.SetDepthStencilState(TStaticDepthStencilState<true, CF_DepthNearOrEqual>::GetRHI());
+		DrawRenderState.SetBlendState(TStaticBlendState<>::GetRHI());
+	}
+
+	virtual void AddMeshBatch(const FMeshBatch& RESTRICT MeshBatch, uint64 BatchElementMask, const FPrimitiveSceneProxy* RESTRICT PrimitiveSceneProxy, int32 StaticMeshId = -1) override final
+	{
+		const bool bDrawMeshBatch = MeshBatch.bUseForMaterial
+			&& MeshBatch.BatchHitProxyId != FHitProxyId::InvisibleHitProxyId
+			&& MeshBatch.bSelectable
+			&& PrimitiveSceneProxy
+			&& PrimitiveSceneProxy->IsSelectable();
+
+		if (bDrawMeshBatch)
+		{
+			const FMaterialRenderProxy* MaterialRenderProxy = MeshBatch.MaterialRenderProxy;
+			const FMaterial* Material = MaterialRenderProxy->GetMaterialNoFallback(FeatureLevel);
+			if (Material && Material->GetRenderingThreadShaderMap())
+			{
+				if (Material->WritesEveryPixel() && !Material->IsTwoSided() && !Material->MaterialModifiesMeshPosition_RenderThread())
+				{
+					// Default material doesn't handle masked, and doesn't have the correct bIsTwoSided setting.
+					MaterialRenderProxy = UMaterial::GetDefaultMaterial(MD_Surface)->GetRenderProxy();
+					check(MaterialRenderProxy);
+					Material = MaterialRenderProxy->GetMaterialNoFallback(FeatureLevel);
+					check(Material);
+				}
+
+				const FVertexFactory* VertexFactory = MeshBatch.VertexFactory;
+
+				TMeshProcessorShaders<FMeshProjectionVS<ProjectionType>, FMeshProjectionHitProxyPS> PassShaders;
+
+				FMaterialShaderTypes ShaderTypes;
+				ShaderTypes.AddShaderType<FMeshProjectionVS<ProjectionType>>();
+				ShaderTypes.AddShaderType<FMeshProjectionHitProxyPS>();
+
+				FMaterialShaders Shaders;
+				if (!Material->TryGetShaders(ShaderTypes, VertexFactory->GetType(), Shaders))
+				{
+					return;
+				}
+
+				Shaders.TryGetVertexShader(PassShaders.VertexShader);
+				Shaders.TryGetPixelShader(PassShaders.PixelShader);
+
+				const FMeshDrawingPolicyOverrideSettings OverrideSettings = ComputeMeshOverrideSettings(MeshBatch);
+				ERasterizerFillMode MeshFillMode = ComputeMeshFillMode(MeshBatch, *Material, OverrideSettings);
+				ERasterizerCullMode MeshCullMode = CM_None;
+
+				FMeshProjectionHitProxyShaderElementData ShaderElementData(MeshBatch.BatchHitProxyId);
+				ShaderElementData.InitializeMeshMaterialData(ViewIfDynamicMeshCommand, PrimitiveSceneProxy, MeshBatch, StaticMeshId, false);
+
+				const FMeshDrawCommandSortKey SortKey = CalculateMeshStaticSortKey(PassShaders.VertexShader, PassShaders.PixelShader);
+
+				BuildMeshDrawCommands(
+					MeshBatch,
+					BatchElementMask,
+					PrimitiveSceneProxy,
+					*MaterialRenderProxy,
+					*Material,
+					DrawRenderState,
+					PassShaders,
+					MeshFillMode,
+					MeshCullMode,
+					SortKey,
+					EMeshPassFeatures::Default,
+					ShaderElementData);
+			}
+		}
+	}
+
+private:
+	FMeshPassProcessorRenderState DrawRenderState;
+};
+
+
+//////////////////////////////////////////////////////////////////////////
+// Selection Outline Render Pass
+
+class FMeshProjectionSelectionOutlinePS : public FGlobalShader
+{
+public:
+	DECLARE_GLOBAL_SHADER(FMeshProjectionSelectionOutlinePS);
+	SHADER_USE_PARAMETER_STRUCT(FMeshProjectionSelectionOutlinePS, FGlobalShader);
+
+	BEGIN_SHADER_PARAMETER_STRUCT(FParameters, )
+		SHADER_PARAMETER_STRUCT_REF(FViewUniformShaderParameters, View)
+		SHADER_PARAMETER_STRUCT(FScreenPassTextureViewportParameters, Color)
+		SHADER_PARAMETER_STRUCT(FScreenPassTextureViewportParameters, Depth)
+		SHADER_PARAMETER_RDG_TEXTURE(Texture2D, ColorTexture)
+		SHADER_PARAMETER_SAMPLER(SamplerState, ColorSampler)
+		SHADER_PARAMETER_RDG_TEXTURE(Texture2D, DepthTexture)
+		SHADER_PARAMETER_SAMPLER(SamplerState, DepthSampler)
+		SHADER_PARAMETER_RDG_TEXTURE(Texture2D, EditorPrimitivesDepth)
+		SHADER_PARAMETER_RDG_TEXTURE_SRV(Texture2D, EditorPrimitivesStencil)
+		SHADER_PARAMETER(FVector3f, OutlineColor)
+		SHADER_PARAMETER(float, SelectionHighlightIntensity)
+		RENDER_TARGET_BINDING_SLOTS()
+	END_SHADER_PARAMETER_STRUCT()
+
+	static bool ShouldCompilePermutation(const FGlobalShaderPermutationParameters& Parameters)
+	{
+		// Only PC platforms render editor primitives.
+		return IsPCPlatform(Parameters.Platform) && EnumHasAllFlags(Parameters.Flags, EShaderPermutationFlags::HasEditorOnlyData);
+	}
+};
+
+IMPLEMENT_GLOBAL_SHADER(FMeshProjectionSelectionOutlinePS, "/Plugin/nDisplay/Private/MeshProjectionSelectionOutline.usf", "Main", SF_Pixel);
+
+template<EDisplayClusterMeshProjectionType ProjectionType> 
+class FLightCardEditorSelectionPassProcessor : public FMeshPassProcessor
+{
+public:
+	FLightCardEditorSelectionPassProcessor(const FScene* InScene, const FSceneView* InView, FMeshPassDrawListContext* InDrawListContext)
+		: FMeshPassProcessor(InScene, GMaxRHIFeatureLevel, InView, InDrawListContext)
+		, DrawRenderState(*InView)
+		, StencilValue(1)
+	{
+		DrawRenderState.SetDepthStencilState(TStaticDepthStencilState<true, CF_DepthNearOrEqual, true, CF_Always, SO_Keep, SO_Keep, SO_Replace>::GetRHI());
+		DrawRenderState.SetBlendState(TStaticBlendStateWriteMask<CW_NONE, CW_NONE, CW_NONE, CW_NONE>::GetRHI());
+	}
+
+	virtual void AddMeshBatch(const FMeshBatch& RESTRICT MeshBatch, uint64 BatchElementMask, const FPrimitiveSceneProxy* RESTRICT PrimitiveSceneProxy, int32 StaticMeshId = -1) override final
+	{
+		const bool bDrawMeshBatch = MeshBatch.bUseForMaterial
+			&& MeshBatch.bUseSelectionOutline
+			&& PrimitiveSceneProxy
+			&& PrimitiveSceneProxy->WantsSelectionOutline()
+			&& (PrimitiveSceneProxy->IsSelected() || PrimitiveSceneProxy->IsHovered());
+
+		if (bDrawMeshBatch)
+		{
+			const FMaterialRenderProxy* MaterialRenderProxy = MeshBatch.MaterialRenderProxy;
+			const FMaterial* Material = MaterialRenderProxy->GetMaterialNoFallback(FeatureLevel);
+			if (Material && Material->GetRenderingThreadShaderMap())
+			{
+				if (Material->WritesEveryPixel() && !Material->IsTwoSided() && !Material->MaterialModifiesMeshPosition_RenderThread())
+				{
+					// Default material doesn't handle masked, and doesn't have the correct bIsTwoSided setting.
+					MaterialRenderProxy = UMaterial::GetDefaultMaterial(MD_Surface)->GetRenderProxy();
+					check(MaterialRenderProxy);
+					Material = MaterialRenderProxy->GetMaterialNoFallback(FeatureLevel);
+					check(Material);
+				}
+				
+				const FVertexFactory* VertexFactory = MeshBatch.VertexFactory;
+
+				TMeshProcessorShaders<FMeshProjectionVS<ProjectionType>, FMeshProjectionPS> PassShaders;
+
+				FMaterialShaderTypes ShaderTypes;
+				ShaderTypes.AddShaderType<FMeshProjectionVS<ProjectionType>>();
+				ShaderTypes.AddShaderType<FMeshProjectionPS>();
+
+				FMaterialShaders Shaders;
+				if (!Material->TryGetShaders(ShaderTypes, VertexFactory->GetType(), Shaders))
+				{
+					return;
+				}
+
+				Shaders.TryGetVertexShader(PassShaders.VertexShader);
+				Shaders.TryGetPixelShader(PassShaders.PixelShader);
+
+				const FMeshDrawingPolicyOverrideSettings OverrideSettings = ComputeMeshOverrideSettings(MeshBatch);
+				ERasterizerFillMode MeshFillMode = ComputeMeshFillMode(MeshBatch, *Material, OverrideSettings);
+				ERasterizerCullMode MeshCullMode = CM_None;
+
+				DrawRenderState.SetStencilRef(StencilValue);
+
+				FMeshMaterialShaderElementData ShaderElementData;
+				ShaderElementData.InitializeMeshMaterialData(ViewIfDynamicMeshCommand, PrimitiveSceneProxy, MeshBatch, StaticMeshId, false);
+
+				FMeshDrawCommandSortKey SortKey{};
+
+				BuildMeshDrawCommands(
+					MeshBatch,
+					BatchElementMask,
+					PrimitiveSceneProxy,
+					*MaterialRenderProxy,
+					*Material,
+					DrawRenderState,
+					PassShaders,
+					MeshFillMode,
+					MeshCullMode,
+					SortKey,
+					EMeshPassFeatures::Default,
+					ShaderElementData);
+			}
+		}
+	}
+
+private:
+	FMeshPassProcessorRenderState DrawRenderState;
+	uint32 StencilValue;
+};
+
+namespace
+{
+	template<typename TSetupFunction>
+	void DrawScreenPass(
+		FRHICommandList& RHICmdList,
+		const FSceneView& View,
+		const FScreenPassTextureViewport& OutputViewport,
+		const FScreenPassTextureViewport& InputViewport,
+		const FScreenPassPipelineState& PipelineState,
+		TSetupFunction SetupFunction)
+	{
+		PipelineState.Validate();
+
+		const FIntRect InputRect = InputViewport.Rect;
+		const FIntPoint InputSize = InputViewport.Extent;
+		const FIntRect OutputRect = OutputViewport.Rect;
+		const FIntPoint OutputSize = OutputRect.Size();
+
+		RHICmdList.SetViewport(OutputRect.Min.X, OutputRect.Min.Y, 0.0f, OutputRect.Max.X, OutputRect.Max.Y, 1.0f);
+
+		SetScreenPassPipelineState(RHICmdList, PipelineState);
+
+		// Setting up buffers.
+		SetupFunction(RHICmdList);
+
+		FIntPoint LocalOutputPos(FIntPoint::ZeroValue);
+		FIntPoint LocalOutputSize(OutputSize);
+		EDrawRectangleFlags DrawRectangleFlags = EDRF_UseTriangleOptimization;
+
+		DrawPostProcessPass(
+			RHICmdList,
+			LocalOutputPos.X, LocalOutputPos.Y, LocalOutputSize.X, LocalOutputSize.Y,
+			InputRect.Min.X, InputRect.Min.Y, InputRect.Width(), InputRect.Height(),
+			OutputSize,
+			InputSize,
+			PipelineState.VertexShader,
+			View.StereoViewIndex,
+			false,
+			DrawRectangleFlags);
+	}
+} //! namespace
+
+//////////////////////////////////////////////////////////////////////////
+// FDisplayClusterMeshProjectionRenderer
+
 void FDisplayClusterMeshProjectionRenderer::AddActor(AActor* Actor)
 {
 	AddActor(Actor, [](const UPrimitiveComponent* PrimitiveComponent)
@@ -147,30 +461,51 @@ void FDisplayClusterMeshProjectionRenderer::AddActor(AActor* Actor, const TFunct
 		if (PrimitiveFilter(PrimitiveComponent))
 		{
 			PrimitiveComponents.Add(PrimitiveComponent);
+
+#if WITH_EDITOR
+			PrimitiveComponent->SelectionOverrideDelegate = UPrimitiveComponent::FSelectionOverride::CreateRaw(this, &FDisplayClusterMeshProjectionRenderer::IsPrimitiveComponentSelected);
+#endif
 		}
 	});
 }
 
 void FDisplayClusterMeshProjectionRenderer::RemoveActor(AActor* Actor)
 {
-	PrimitiveComponents.RemoveAll([Actor](TWeakObjectPtr<UPrimitiveComponent>& PrimitiveComponent)
+	TArray<TWeakObjectPtr<UPrimitiveComponent>> ComponentsToRemove = PrimitiveComponents.FilterByPredicate([Actor](const TWeakObjectPtr<UPrimitiveComponent>& PrimitiveComponent)
 	{
 		return !PrimitiveComponent.IsValid() || PrimitiveComponent->GetOwner() == Actor;
 	});
+
+	for (const TWeakObjectPtr<UPrimitiveComponent>& PrimitiveComponent : ComponentsToRemove)
+	{
+#if WITH_EDITOR
+		PrimitiveComponent->SelectionOverrideDelegate.Unbind();
+#endif
+
+		PrimitiveComponents.Remove(PrimitiveComponent);
+	}
 }
 
 void FDisplayClusterMeshProjectionRenderer::ClearScene()
 {
+	for (const TWeakObjectPtr<UPrimitiveComponent>& PrimitiveComponent : PrimitiveComponents)
+	{
+#if WITH_EDITOR
+		PrimitiveComponent->SelectionOverrideDelegate.Unbind();
+#endif
+	}
+
 	PrimitiveComponents.Empty();
 }
 
-void FDisplayClusterMeshProjectionRenderer::Render(FCanvas* Canvas, FSceneInterface* Scene, const FSceneViewInitOptions& ViewInitOptions, EDisplayClusterMeshProjectionType ProjectionType)
+void FDisplayClusterMeshProjectionRenderer::Render(FCanvas* Canvas, FSceneInterface* Scene, const FSceneViewInitOptions& ViewInitOptions, const FEngineShowFlags& EngineShowFlags, EDisplayClusterMeshProjectionType ProjectionType)
 {
 	Canvas->Flush_GameThread();
 	FRenderTarget* RenderTarget = Canvas->GetRenderTarget();
+	const bool bIsHitTesting = Canvas->IsHitTesting();
 
 	ENQUEUE_RENDER_COMMAND(FDrawProjectedMeshes)(
-		[RenderTarget, Scene, ViewInitOptions, ProjectionType, this](FRHICommandListImmediate& RHICmdList)
+		[RenderTarget, Scene, ViewInitOptions, EngineShowFlags, bIsHitTesting, ProjectionType, this](FRHICommandListImmediate& RHICmdList)
 		{
 			FMemMark Mark(FMemStack::Get());
 			FRDGBuilder GraphBuilder(RHICmdList);
@@ -178,7 +513,7 @@ void FDisplayClusterMeshProjectionRenderer::Render(FCanvas* Canvas, FSceneInterf
 			FSceneViewFamily ViewFamily(FSceneViewFamily::ConstructionValues(
 				RenderTarget,
 				Scene,
-				FEngineShowFlags(ESFIM_Editor))
+				EngineShowFlags)
 				.SetTime(FGameTime::GetTimeSinceAppStart())
 				.SetGammaCorrection(1.0f));
 
@@ -192,43 +527,212 @@ void FDisplayClusterMeshProjectionRenderer::Render(FCanvas* Canvas, FSceneInterf
 				ViewFamily.FrameNumber = GFrameNumber;
 			}
 
+			ViewFamily.EngineShowFlags.SetHitProxies(bIsHitTesting);
+
 			FScenePrimitiveRenderingContextScopeHelper ScenePrimitiveRenderingContextScopeHelper(GetRendererModule().BeginScenePrimitiveRendering(GraphBuilder, &ViewFamily));
 
 			FSceneViewInitOptions NewInitOptions(ViewInitOptions);
 			NewInitOptions.ViewFamily = &ViewFamily;
 
 			GetRendererModule().CreateAndInitSingleView(RHICmdList, &ViewFamily, &NewInitOptions);
-			const FSceneView* View = ViewFamily.Views[0];
+			const FViewInfo* View = (FViewInfo*)ViewFamily.Views[0];
 
 			FRDGTextureRef OutputTexture = GraphBuilder.RegisterExternalTexture(CreateRenderTarget(RenderTarget->GetRenderTargetTexture(), TEXT("ViewRenderTarget")));
+			FRenderTargetBinding OutputRenderTargetBinding(OutputTexture, ERenderTargetLoadAction::ELoad);
 
-			FMeshProjectionPassParameters* PassParameters = GraphBuilder.AllocParameters<FMeshProjectionPassParameters>();
-			PassParameters->View = View->ViewUniformBuffer;
-			PassParameters->InstanceCulling = FInstanceCullingContext::CreateDummyInstanceCullingUniformBuffer(GraphBuilder);
-			PassParameters->RenderTargets[0] = FRenderTargetBinding(OutputTexture, ERenderTargetLoadAction::ELoad);
-
-			GraphBuilder.AddPass(RDG_EVENT_NAME("MeshProjectionRenderer"), PassParameters, ERDGPassFlags::Raster | ERDGPassFlags::NeverCull, [View, ProjectionType, this](FRHICommandList& RHICmdList)
+			if (ViewFamily.EngineShowFlags.HitProxies)
 			{
-				FIntRect ViewRect = View->UnscaledViewRect;
-				RHICmdList.SetViewport(ViewRect.Min.X, ViewRect.Min.Y, 0.0f, ViewRect.Max.X, ViewRect.Max.Y, 1.0f);
+				FRDGTextureDesc Desc(FRDGTextureDesc::Create2D(OutputTexture->Desc.Extent, PF_B8G8R8A8, FClearValueBinding::Black, TexCreate_RenderTargetable | TexCreate_ShaderResource));
+				FRDGTextureRef HitProxyTexture = GraphBuilder.CreateTexture(Desc, TEXT("DisplayClusterMeshProjection.HitProxyTexture"));
+
+				const FRDGTextureDesc DepthDesc = FRDGTextureDesc::Create2D(OutputTexture->Desc.Extent, PF_DepthStencil, FClearValueBinding::DepthFar, TexCreate_DepthStencilTargetable | TexCreate_ShaderResource);
+				FRDGTextureRef HitProxyDepthTexture = GraphBuilder.CreateTexture(DepthDesc, TEXT("DisplayClusterMeshProjection.DepthTexture"));
+
+				FRenderTargetBinding HitProxyRenderTargetBinding(HitProxyTexture, ERenderTargetLoadAction::EClear);
+				FDepthStencilBinding HitProxyDepthStencilBinding(HitProxyDepthTexture, ERenderTargetLoadAction::EClear, ERenderTargetLoadAction::EClear, FExclusiveDepthStencil::DepthWrite_StencilWrite);
 
-				switch (ProjectionType)
+				AddHitProxyRenderPass(GraphBuilder, View, ProjectionType, HitProxyRenderTargetBinding, HitProxyDepthStencilBinding);
+
+				// Copy the hit proxy buffer to the viewport family's render target
 				{
-					case EDisplayClusterMeshProjectionType::Azimuthal:
-						RenderPrimitives_RenderThread<EDisplayClusterMeshProjectionType::Azimuthal>(View, RHICmdList);
-						break;
-
-					case EDisplayClusterMeshProjectionType::Perspective:
-					default:
-						RenderPrimitives_RenderThread<EDisplayClusterMeshProjectionType::Perspective>(View, RHICmdList);
-						break;
+					FCopyRectPS::FParameters* ScreenPassParameters = GraphBuilder.AllocParameters<FCopyRectPS::FParameters>();
+					ScreenPassParameters->InputTexture = HitProxyTexture;
+					ScreenPassParameters->InputSampler = TStaticSamplerState<>::GetRHI();
+					ScreenPassParameters->RenderTargets[0] = OutputRenderTargetBinding;
+
+					FGlobalShaderMap* GlobalShaderMap = GetGlobalShaderMap(GMaxRHIFeatureLevel);
+					TShaderMapRef<FScreenPassVS> ScreenPassVS(GlobalShaderMap);
+					TShaderMapRef<FCopyRectPS> CopyPixelShader(GlobalShaderMap);
+
+					FRHIBlendState* DefaultBlendState = FScreenPassPipelineState::FDefaultBlendState::GetRHI();
+					const FScreenPassTextureViewport RegionViewport(OutputRenderTargetBinding.GetTexture());
+
+					GraphBuilder.AddPass(
+						RDG_EVENT_NAME("MeshProjectionRenderer::CopyHitProxyTexture"),
+						ScreenPassParameters,
+						ERDGPassFlags::Raster,
+						[View, ScreenPassVS, CopyPixelShader, &RegionViewport, ScreenPassParameters, DefaultBlendState](FRHICommandList& RHICmdList)
+					{
+						DrawScreenPass(
+							RHICmdList,
+							*View,
+							RegionViewport,
+							RegionViewport,
+							FScreenPassPipelineState(ScreenPassVS, CopyPixelShader, DefaultBlendState),
+							[&](FRHICommandList&)
+						{
+							SetShaderParameters(RHICmdList, CopyPixelShader, CopyPixelShader.GetPixelShader(), *ScreenPassParameters);
+						});
+					});
 				}
-			});
+			}
+			else
+			{
+				const FRDGTextureDesc Desc = FRDGTextureDesc::Create2D(OutputTexture->Desc.Extent, PF_DepthStencil, FClearValueBinding::DepthFar, TexCreate_DepthStencilTargetable | TexCreate_ShaderResource);
+				FRDGTextureRef DepthTexture = GraphBuilder.CreateTexture(Desc, TEXT("DisplayClusterMeshProjection.DepthTexture"));
+
+				FDepthStencilBinding OutputDepthStencilBinding(DepthTexture, ERenderTargetLoadAction::EClear, ERenderTargetLoadAction::EClear, FExclusiveDepthStencil::DepthWrite_StencilNop);
+
+				AddBaseRenderPass(GraphBuilder, View, ProjectionType, OutputRenderTargetBinding, OutputDepthStencilBinding);
+
+#if WITH_EDITOR
+				const FRDGTextureDesc SelectionDepthDesc = FRDGTextureDesc::Create2D(OutputTexture->Desc.Extent, PF_DepthStencil, FClearValueBinding::DepthFar, TexCreate_DepthStencilTargetable | TexCreate_ShaderResource);
+				FRDGTextureRef SelectionDepthTexture = GraphBuilder.CreateTexture(SelectionDepthDesc, TEXT("DisplayClusterMeshProjection.SelectionDepthTexture"));
+				FDepthStencilBinding SelectionDepthStencilBinding(SelectionDepthTexture, ERenderTargetLoadAction::EClear, ERenderTargetLoadAction::EClear, FExclusiveDepthStencil::DepthWrite_StencilWrite);
+
+				AddSelectionDepthRenderPass(GraphBuilder, View, ProjectionType, SelectionDepthStencilBinding);
+				AddSelectionOutlineScreenPass(GraphBuilder, View, OutputRenderTargetBinding, OutputTexture, DepthTexture, SelectionDepthTexture);
+#endif
+			}
 
 			GraphBuilder.Execute();
 		});
 }
 
+// Helper macros that generate appropriate switch cases for each projection type the mesh projection renderer supports, allowing new projection types to be easily added
+#define PROJECTION_TYPE_CASE(FuncName, ProjectionType, ...) case ProjectionType: \
+	FuncName<ProjectionType>(__VA_ARGS__); \
+	break;
+
+#define SWITCH_ON_PROJECTION_TYPE(FuncName, ProjectionType, ...) switch (ProjectionType) \
+	{ \
+	PROJECTION_TYPE_CASE(FuncName, EDisplayClusterMeshProjectionType::Azimuthal, __VA_ARGS__) \
+	PROJECTION_TYPE_CASE(FuncName, EDisplayClusterMeshProjectionType::Perspective, __VA_ARGS__) \
+	}
+
+void FDisplayClusterMeshProjectionRenderer::AddBaseRenderPass(FRDGBuilder& GraphBuilder,
+	const FViewInfo* View, EDisplayClusterMeshProjectionType ProjectionType,
+	FRenderTargetBinding& OutputRenderTargetBinding,
+	FDepthStencilBinding& OutputDepthStencilBinding)
+{
+	FMeshProjectionPassParameters* MeshPassParameters = GraphBuilder.AllocParameters<FMeshProjectionPassParameters>();
+	MeshPassParameters->View = View->ViewUniformBuffer;
+	MeshPassParameters->InstanceCulling = FInstanceCullingContext::CreateDummyInstanceCullingUniformBuffer(GraphBuilder);
+	MeshPassParameters->RenderTargets[0] = OutputRenderTargetBinding;
+	MeshPassParameters->RenderTargets.DepthStencil = OutputDepthStencilBinding;
+
+	GraphBuilder.AddPass(RDG_EVENT_NAME("MeshProjectionRenderer::Base"), MeshPassParameters, ERDGPassFlags::Raster | ERDGPassFlags::NeverCull, [View, ProjectionType, this](FRHICommandList& RHICmdList)
+	{
+		FIntRect ViewRect = View->UnscaledViewRect;
+		RHICmdList.SetViewport(ViewRect.Min.X, ViewRect.Min.Y, 0.0f, ViewRect.Max.X, ViewRect.Max.Y, 1.0f);
+
+		SWITCH_ON_PROJECTION_TYPE(RenderPrimitives_RenderThread, ProjectionType, View, RHICmdList);
+	});
+}
+
+void FDisplayClusterMeshProjectionRenderer::AddHitProxyRenderPass(FRDGBuilder& GraphBuilder,
+	const FViewInfo* View,
+	EDisplayClusterMeshProjectionType ProjectionType,
+	FRenderTargetBinding& OutputRenderTargetBinding,
+	FDepthStencilBinding& OutputDepthStencilBinding)
+{
+	FMeshProjectionPassParameters* MeshPassParameters = GraphBuilder.AllocParameters<FMeshProjectionPassParameters>();
+	MeshPassParameters->View = View->ViewUniformBuffer;
+	MeshPassParameters->InstanceCulling = FInstanceCullingContext::CreateDummyInstanceCullingUniformBuffer(GraphBuilder);
+	MeshPassParameters->RenderTargets[0] = OutputRenderTargetBinding;
+	MeshPassParameters->RenderTargets.DepthStencil = OutputDepthStencilBinding;
+
+	GraphBuilder.AddPass(RDG_EVENT_NAME("MeshProjectionRenderer::HitProxies"), MeshPassParameters, ERDGPassFlags::Raster | ERDGPassFlags::NeverCull, [View, ProjectionType, this](FRHICommandList& RHICmdList)
+	{
+		FIntRect ViewRect = View->UnscaledViewRect;
+		RHICmdList.SetViewport(ViewRect.Min.X, ViewRect.Min.Y, 0.0f, ViewRect.Max.X, ViewRect.Max.Y, 1.0f);
+
+		SWITCH_ON_PROJECTION_TYPE(RenderHitProxies_RenderThread, ProjectionType, View, RHICmdList);
+	});
+}
+
+#if WITH_EDITOR
+void FDisplayClusterMeshProjectionRenderer::AddSelectionDepthRenderPass(FRDGBuilder& GraphBuilder, 
+		const FViewInfo* View,
+		EDisplayClusterMeshProjectionType ProjectionType,
+		FDepthStencilBinding& OutputDepthStencilBinding)
+{
+	FMeshProjectionPassParameters* SelectionPassParameters = GraphBuilder.AllocParameters<FMeshProjectionPassParameters>();
+	SelectionPassParameters->View = View->ViewUniformBuffer;
+	SelectionPassParameters->InstanceCulling = FInstanceCullingContext::CreateDummyInstanceCullingUniformBuffer(GraphBuilder);
+	SelectionPassParameters->RenderTargets.DepthStencil = OutputDepthStencilBinding;
+
+	GraphBuilder.AddPass(RDG_EVENT_NAME("MeshProjectionRenderer::SelectionDepth"), SelectionPassParameters, ERDGPassFlags::Raster | ERDGPassFlags::NeverCull, [View, ProjectionType, this](FRHICommandList& RHICmdList)
+	{
+		FIntRect ViewRect = View->UnscaledViewRect;
+		RHICmdList.SetViewport(ViewRect.Min.X, ViewRect.Min.Y, 0.0f, ViewRect.Max.X, ViewRect.Max.Y, 1.0f);
+
+		SWITCH_ON_PROJECTION_TYPE(RenderSelection_RenderThread, ProjectionType, View, RHICmdList);
+	});
+}
+
+void FDisplayClusterMeshProjectionRenderer::AddSelectionOutlineScreenPass(FRDGBuilder& GraphBuilder,
+	const FViewInfo* View,
+	FRenderTargetBinding& OutputRenderTargetBinding,
+	FRDGTexture* SceneColor,
+	FRDGTexture* SceneDepth,
+	FRDGTexture* SelectionDepth)
+{
+	const FScreenPassTextureViewport OutputViewport(OutputRenderTargetBinding.GetTexture());
+	const FScreenPassTextureViewport ColorViewport(SceneColor);
+	const FScreenPassTextureViewport DepthViewport(SceneDepth);
+
+	FRHISamplerState* PointClampSampler = TStaticSamplerState<SF_Point, AM_Clamp, AM_Clamp, AM_Clamp>::GetRHI();
+
+	FMeshProjectionSelectionOutlinePS::FParameters* ScreenPassParameters = GraphBuilder.AllocParameters<FMeshProjectionSelectionOutlinePS::FParameters>();
+	ScreenPassParameters->RenderTargets[0] = OutputRenderTargetBinding;
+	ScreenPassParameters->View = View->ViewUniformBuffer;
+	ScreenPassParameters->Color = GetScreenPassTextureViewportParameters(ColorViewport);
+	ScreenPassParameters->Depth = GetScreenPassTextureViewportParameters(DepthViewport);
+	ScreenPassParameters->ColorTexture = SceneColor;
+	ScreenPassParameters->ColorSampler = PointClampSampler;
+	ScreenPassParameters->DepthTexture = SceneDepth;
+	ScreenPassParameters->DepthSampler = PointClampSampler;
+	ScreenPassParameters->EditorPrimitivesDepth = SelectionDepth;
+	ScreenPassParameters->EditorPrimitivesStencil = GraphBuilder.CreateSRV(FRDGTextureSRVDesc::CreateWithPixelFormat(SelectionDepth, PF_X24_G8));
+	ScreenPassParameters->OutlineColor = FVector3f(View->SelectionOutlineColor);
+	ScreenPassParameters->SelectionHighlightIntensity = GEngine->SelectionHighlightIntensity;
+
+	TShaderMapRef<FScreenPassVS> ScreenPassVS(View->ShaderMap);
+	TShaderMapRef<FMeshProjectionSelectionOutlinePS> SelectionOutlinePS(View->ShaderMap);
+
+	FRHIBlendState* DefaultBlendState = FScreenPassPipelineState::FDefaultBlendState::GetRHI();
+
+	GraphBuilder.AddPass(
+		RDG_EVENT_NAME("MeshProjectionRenderer::SelectionScreen"),
+		ScreenPassParameters,
+		ERDGPassFlags::Raster,
+		[View, ScreenPassVS, SelectionOutlinePS, &OutputViewport, ScreenPassParameters, DefaultBlendState](FRHICommandList& RHICmdList)
+	{
+		DrawScreenPass(
+			RHICmdList,
+			*View,
+			OutputViewport,
+			OutputViewport,
+			FScreenPassPipelineState(ScreenPassVS, SelectionOutlinePS, DefaultBlendState),
+			[&](FRHICommandList&)
+		{
+			SetShaderParameters(RHICmdList, SelectionOutlinePS, SelectionOutlinePS.GetPixelShader(), *ScreenPassParameters);
+		});
+	});
+}
+#endif
+
 template<EDisplayClusterMeshProjectionType ProjectionType>
 void FDisplayClusterMeshProjectionRenderer::RenderPrimitives_RenderThread(const FSceneView* View, FRHICommandList& RHICmdList)
 {
@@ -256,6 +760,81 @@ void FDisplayClusterMeshProjectionRenderer::RenderPrimitives_RenderThread(const
 	});
 }
 
-// Explicit template specialisations for RenderPrimitives_RenderThread
-template void FDisplayClusterMeshProjectionRenderer::RenderPrimitives_RenderThread<EDisplayClusterMeshProjectionType::Perspective>(const FSceneView* View, FRHICommandList& RHICmdList);
-template void FDisplayClusterMeshProjectionRenderer::RenderPrimitives_RenderThread<EDisplayClusterMeshProjectionType::Azimuthal>(const FSceneView* View, FRHICommandList& RHICmdList);
\ No newline at end of file
+template<EDisplayClusterMeshProjectionType ProjectionType>
+void FDisplayClusterMeshProjectionRenderer::RenderHitProxies_RenderThread(const FSceneView* View, FRHICommandList& RHICmdList)
+{
+	DrawDynamicMeshPass(*View, RHICmdList, [View, this](FDynamicPassMeshDrawListContext* DynamicMeshPassContext)
+	{
+		TArray<FPrimitiveSceneProxy*> PrimitiveSceneProxies;
+		for (const TWeakObjectPtr<UPrimitiveComponent>& PrimitiveComponent : PrimitiveComponents)
+		{
+			if (PrimitiveComponent.IsValid() && PrimitiveComponent->SceneProxy)
+			{
+				PrimitiveSceneProxies.Add(PrimitiveComponent->SceneProxy);
+			}
+		}
+
+		FLightCardEditorHitProxyMeshPassProcessor<ProjectionType> MeshProcessor(nullptr, View, DynamicMeshPassContext);
+
+		for (FPrimitiveSceneProxy* PrimitiveProxy : PrimitiveSceneProxies)
+		{
+			if (const FMeshBatch* MeshBatch = PrimitiveProxy->GetPrimitiveSceneInfo()->GetMeshBatch(PrimitiveProxy->GetPrimitiveSceneInfo()->StaticMeshes.Num() - 1))
+			{
+				const uint64 BatchElementMask = ~0ull;
+				MeshProcessor.AddMeshBatch(*MeshBatch, BatchElementMask, PrimitiveProxy);
+			}
+		}
+	});
+}
+
+#if WITH_EDITOR
+template<EDisplayClusterMeshProjectionType ProjectionType>
+void FDisplayClusterMeshProjectionRenderer::RenderSelection_RenderThread(const FSceneView* View, FRHICommandList& RHICmdList)
+{
+	DrawDynamicMeshPass(*View, RHICmdList, [View, this](FDynamicPassMeshDrawListContext* DynamicMeshPassContext)
+	{
+		TArray<FPrimitiveSceneProxy*> PrimitiveSceneProxies;
+		for (const TWeakObjectPtr<UPrimitiveComponent>& PrimitiveComponent : PrimitiveComponents)
+		{
+			if (PrimitiveComponent.IsValid() && PrimitiveComponent->SceneProxy && PrimitiveComponent->SceneProxy->IsSelected())
+			{
+				PrimitiveSceneProxies.Add(PrimitiveComponent->SceneProxy);
+			}
+		}
+
+		FLightCardEditorSelectionPassProcessor<ProjectionType> MeshProcessor(nullptr, View, DynamicMeshPassContext);
+
+		for (FPrimitiveSceneProxy* PrimitiveProxy : PrimitiveSceneProxies)
+		{
+			if (const FMeshBatch* MeshBatch = PrimitiveProxy->GetPrimitiveSceneInfo()->GetMeshBatch(PrimitiveProxy->GetPrimitiveSceneInfo()->StaticMeshes.Num() - 1))
+			{
+				const uint64 BatchElementMask = ~0ull;
+				MeshProcessor.AddMeshBatch(*MeshBatch, BatchElementMask, PrimitiveProxy);
+			}
+		}
+	});
+}
+#endif
+
+bool FDisplayClusterMeshProjectionRenderer::IsPrimitiveComponentSelected(const UPrimitiveComponent* InPrimitiveComponent)
+{
+	if (ActorSelectedDelegate.IsBound())
+	{
+		return ActorSelectedDelegate.Execute(InPrimitiveComponent->GetOwner());
+	}
+
+	return false;
+}
+
+// Explicit template specializations for each projection type
+#define PROJECTION_TYPE_TEMPLATE_SPECIALIZATION(FuncName, ProjectionType, ...) template void FDisplayClusterMeshProjectionRenderer::FuncName<ProjectionType>(__VA_ARGS__);
+#define CREATE_PROJECTION_TYPE_TEMPLATE_SPECIALIZATIONS(FuncName, ...) \
+	PROJECTION_TYPE_TEMPLATE_SPECIALIZATION(FuncName, EDisplayClusterMeshProjectionType::Perspective, __VA_ARGS__) \
+	PROJECTION_TYPE_TEMPLATE_SPECIALIZATION(FuncName, EDisplayClusterMeshProjectionType::Azimuthal, __VA_ARGS__)
+
+CREATE_PROJECTION_TYPE_TEMPLATE_SPECIALIZATIONS(RenderPrimitives_RenderThread, const FSceneView* View, FRHICommandList& RHICmdList)
+CREATE_PROJECTION_TYPE_TEMPLATE_SPECIALIZATIONS(RenderHitProxies_RenderThread, const FSceneView* View, FRHICommandList& RHICmdList)
+
+#if WITH_EDITOR
+CREATE_PROJECTION_TYPE_TEMPLATE_SPECIALIZATIONS(RenderSelection_RenderThread, const FSceneView* View, FRHICommandList& RHICmdList)
+#endif
\ No newline at end of file
diff --git a/Engine/Plugins/Runtime/nDisplay/Source/DisplayClusterLightCardEditorShaders/Public/DisplayClusterMeshProjectionRenderer.h b/Engine/Plugins/Runtime/nDisplay/Source/DisplayClusterLightCardEditorShaders/Public/DisplayClusterMeshProjectionRenderer.h
--- a/Engine/Plugins/Runtime/nDisplay/Source/DisplayClusterLightCardEditorShaders/Public/DisplayClusterMeshProjectionRenderer.h
+++ b/Engine/Plugins/Runtime/nDisplay/Source/DisplayClusterLightCardEditorShaders/Public/DisplayClusterMeshProjectionRenderer.h
@@ -8,7 +8,8 @@
 class AActor;
 class UPrimitiveComponent;
 class FSceneView;
-class UPrimitiveComponent;
+class FViewInfo;
+struct FEngineShowFlags;
 
 /** Indicates which kind of projection is used by the renderer */
 enum DISPLAYCLUSTERLIGHTCARDEDITORSHADERS_API EDisplayClusterMeshProjectionType
@@ -37,13 +38,61 @@ public:
 	void ClearScene();
 
 	/** Renders its list of primitive components to the specified canvas using the desired projection type. Can be called from the game thread */
-	void Render(FCanvas* Canvas, FSceneInterface* Scene, const FSceneViewInitOptions& ViewInitOptions, EDisplayClusterMeshProjectionType  ProjectionType);
+	void Render(FCanvas* Canvas, FSceneInterface* Scene, const FSceneViewInitOptions& ViewInitOptions, const FEngineShowFlags& EngineShowFlags, EDisplayClusterMeshProjectionType  ProjectionType);
 
 private:
+	/** Adds a pass to perform the base render for the mesh projection. */
+	void AddBaseRenderPass(FRDGBuilder& GraphBuilder, 
+		const FViewInfo* View,
+		EDisplayClusterMeshProjectionType ProjectionType,
+		FRenderTargetBinding& OutputRenderTargetBinding,
+		FDepthStencilBinding& OutputDepthStencilBinding);
+
+	/** Adds a pass to perform the hit proxy render for the mesh projection */
+	void AddHitProxyRenderPass(FRDGBuilder& GraphBuilder, 
+		const FViewInfo* View,
+		EDisplayClusterMeshProjectionType ProjectionType,
+		FRenderTargetBinding& OutputRenderTargetBinding,
+		FDepthStencilBinding& OutputDepthStencilBinding);
+
+#if WITH_EDITOR
+	/** Adds a pass to perform the depth render for any selected primitives for the mesh projection */
+	void AddSelectionDepthRenderPass(FRDGBuilder& GraphBuilder, 
+		const FViewInfo* View,
+		EDisplayClusterMeshProjectionType ProjectionType,
+		FDepthStencilBinding& OutputDepthStencilBinding);
+
+	/** Adds a pass to perform the post process selection outline for the mesh projection */
+	void AddSelectionOutlineScreenPass(FRDGBuilder& GraphBuilder,
+		const FViewInfo* View,
+		FRenderTargetBinding& OutputRenderTargetBinding,
+		FRDGTexture* SceneColor,
+		FRDGTexture* SceneDepth,
+		FRDGTexture* SelectionDepth);
+#endif
+
 	/** Renders the list of primitive components using the appropriate mesh pass processor given by the template parameter */
 	template<EDisplayClusterMeshProjectionType ProjectionType>
 	void RenderPrimitives_RenderThread(const FSceneView* View, FRHICommandList& RHICmdList);
 
+	/** Renders the hit proxies of the list of primitive components using the appropriate mesh pass processor given by the template parameter */
+	template<EDisplayClusterMeshProjectionType ProjectionType>
+	void RenderHitProxies_RenderThread(const FSceneView* View, FRHICommandList& RHICmdList);
+
+#if WITH_EDITOR
+	/** Renders the list of primitive components using the appropriate mesh pass processor given by the template parameter */
+	template<EDisplayClusterMeshProjectionType ProjectionType>
+	void RenderSelection_RenderThread(const FSceneView* View, FRHICommandList& RHICmdList);
+#endif
+
+	/** Callback used to determine if a primitive component should be rendered with a selection outline */
+	bool IsPrimitiveComponentSelected(const UPrimitiveComponent* InPrimitiveComponent);
+
+public:
+	/** Delegate to determine if an actor should be rendered as selected */
+	DECLARE_DELEGATE_RetVal_OneParam(bool, FSelection, const AActor*);
+	FSelection ActorSelectedDelegate;
+
 private:
 	/** The list of primitive components to render */
 	TArray<TWeakObjectPtr<UPrimitiveComponent>> PrimitiveComponents;
