nDisplay: Fix level instance propagation to lightcard editor -- Components and config data are properly duplicated. Add support to update light card proxies from their level instances. Improve actor cleanup, fix some crashes.

diff --git a/Engine/Plugins/Runtime/nDisplay/Source/DisplayClusterLightCardEditor/Private/Viewport/DisplayClusterLightCardEditorViewportClient.cpp b/Engine/Plugins/Runtime/nDisplay/Source/DisplayClusterLightCardEditor/Private/Viewport/DisplayClusterLightCardEditorViewportClient.cpp
--- a/Engine/Plugins/Runtime/nDisplay/Source/DisplayClusterLightCardEditor/Private/Viewport/DisplayClusterLightCardEditorViewportClient.cpp
+++ b/Engine/Plugins/Runtime/nDisplay/Source/DisplayClusterLightCardEditor/Private/Viewport/DisplayClusterLightCardEditorViewportClient.cpp
@@ -132,7 +132,7 @@ void FDisplayClusterLightCardEditorViewportClient::Tick(float DeltaSeconds)
 		}
 	}
 
-	if (RootActorPreviewInstance.IsValid() && RootActorLevelInstance.IsValid())
+	if (RootActorProxy.IsValid() && RootActorLevelInstance.IsValid())
 	{
 		// Pass the preview render targets from the level instance root actor to the preview root actor
 		UDisplayClusterConfigurationData* Config = RootActorLevelInstance->GetConfigData();
@@ -143,7 +143,7 @@ void FDisplayClusterLightCardEditorViewportClient::Tick(float DeltaSeconds)
 			for (const TPair<FString, UDisplayClusterConfigurationViewport*>& ViewportPair : Node->Viewports)
 			{
 				UDisplayClusterPreviewComponent* LevelInstancePreviewComp = RootActorLevelInstance->GetPreviewComponent(NodePair.Key, ViewportPair.Key);
-				UDisplayClusterPreviewComponent* PreviewComp = RootActorPreviewInstance->GetPreviewComponent(NodePair.Key, ViewportPair.Key);
+				UDisplayClusterPreviewComponent* PreviewComp = RootActorProxy->GetPreviewComponent(NodePair.Key, ViewportPair.Key);
 
 				if (PreviewComp && LevelInstancePreviewComp)
 				{
@@ -418,7 +418,7 @@ void FDisplayClusterLightCardEditorViewportClient::ProcessClick(FSceneView& View
 		if (HitProxy->IsA(HActor::StaticGetType()))
 		{
 			HActor* ActorHitProxy = static_cast<HActor*>(HitProxy);
-			if (ActorHitProxy->Actor == RootActorPreviewInstance.Get())
+			if (ActorHitProxy->Actor == RootActorProxy.Get())
 			{
 				if (ActorHitProxy->PrimComponent && ActorHitProxy->PrimComponent->IsA<UStaticMeshComponent>())
 				{
@@ -467,7 +467,7 @@ EMouseCursor::Type FDisplayClusterLightCardEditorViewportClient::GetCursor(FView
 			if (HitProxy->IsA(HActor::StaticGetType()))
 			{
 				HActor* ActorHitProxy = static_cast<HActor*>(HitProxy);
-				if (ActorHitProxy->Actor == RootActorPreviewInstance.Get())
+				if (ActorHitProxy->Actor == RootActorProxy.Get())
 				{
 					if (ActorHitProxy->PrimComponent && ActorHitProxy->PrimComponent->IsA<UStaticMeshComponent>())
 					{
@@ -519,93 +519,79 @@ void FDisplayClusterLightCardEditorViewportClient::UpdatePreviewActor(ADisplayCl
 		return;
 	}
 
-	if (RootActorPreviewInstance.IsValid() && (!RootActor || !PreviewWorld->ContainsActor(RootActor)))
+	MeshProjectionRenderer->ClearScene();
+	
+	if (RootActorProxy.IsValid())
 	{
-		MeshProjectionRenderer->ClearScene();
-		PreviewWorld->EditorDestroyActor(RootActorPreviewInstance.Get(), false);
-		
-		RootActorPreviewInstance.Reset();
-
-		for (const TWeakObjectPtr<AActor> LightCardProxy : LightCardProxies)
+		PreviewWorld->EditorDestroyActor(RootActorProxy.Get(), false);
+		RootActorProxy.Reset();
+	}
+	
+	for (const TWeakObjectPtr<AActor>& LightCardProxy : LightCardProxies)
+	{
+		if (LightCardProxy.IsValid())
 		{
 			PreviewWorld->EditorDestroyActor(LightCardProxy.Get(), false);
 		}
-
-		RootActorLevelInstance.Reset();
-		LightCardProxies.Empty();
 	}
+
+	RootActorLevelInstance.Reset();
+	LightCardLevelInstances.Empty();
+	LightCardProxies.Empty();
 	
 	if (RootActor)
 	{
-		RootActorLevelInstance = RootActor;
-
-		FActorSpawnParameters SpawnInfo;
-		SpawnInfo.SpawnCollisionHandlingOverride = ESpawnActorCollisionHandlingMethod::AlwaysSpawn;
-		SpawnInfo.bNoFail = true;
-		SpawnInfo.ObjectFlags = RF_Transient|RF_Transactional;
-		SpawnInfo.Template = RootActor;
-		
-		RootActorPreviewInstance = CastChecked<ADisplayClusterRootActor>(PreviewWorld->SpawnActor(RootActor->GetClass(),
-			nullptr,  nullptr, MoveTemp(SpawnInfo)));
-		
-		PreviewWorld->GetCurrentLevel()->AddLoadedActor(RootActorPreviewInstance.Get());
-
-		// Copy component data and update transforms.
+		// Schedule for the next tick so CDO changes get propagated first in the event of config editor skeleton
+		// regeneration & compiles. nDisplay's custom propagation may have issues if the archetype isn't correct.
+		PreviewWorld->GetTimerManager().SetTimerForNextTick([=]()
 		{
-			TArray<UActorComponent*> SpawnedActorComponents;
-			TArray<UActorComponent*> LevelActorComponents;
-			RootActorPreviewInstance->GetComponents(SpawnedActorComponents);
-			RootActorLevelInstance->GetComponents(LevelActorComponents);
+			RootActorLevelInstance = RootActor;
 
-			for (UActorComponent* SpawnedComponent : SpawnedActorComponents)
 			{
-				if (UActorComponent** OriginalComponent = LevelActorComponents.FindByPredicate([SpawnedComponent](const UActorComponent* Component)
-				{
-					return Component->GetFName() == SpawnedComponent->GetFName();
-				}))
-				{
-					UEngine::CopyPropertiesForUnrelatedObjects(*OriginalComponent, SpawnedComponent);
-					SpawnedComponent->UpdateComponentToWorld();
-				}
+				FObjectDuplicationParameters DupeActorParameters(RootActor, PreviewWorld->GetCurrentLevel());
+				DupeActorParameters.FlagMask = RF_AllFlags & ~(RF_ArchetypeObject | RF_Transactional); // Keeps archetypes correct in config data.
+				DupeActorParameters.PortFlags = PPF_DuplicateVerbatim;
+			
+				RootActorProxy = CastChecked<ADisplayClusterRootActor>(StaticDuplicateObjectEx(DupeActorParameters));
 			}
-		}
-
-		// Spawned actor will take the transform values from the template, so manually reset them to zero here
-		RootActorPreviewInstance->SetActorLocation(FVector::ZeroVector);
-		RootActorPreviewInstance->SetActorRotation(FRotator::ZeroRotator);
-
-		FindProjectionOriginComponent();
-
-		// Filter out any primitives hidden in game except screen components
-		MeshProjectionRenderer->AddActor(RootActorPreviewInstance.Get(), [](const UPrimitiveComponent* PrimitiveComponent)
-		{
-			return !PrimitiveComponent->bHiddenInGame || PrimitiveComponent->IsA<UDisplayClusterScreenComponent>();
-		});
+			
+			PreviewWorld->GetCurrentLevel()->AddLoadedActor(RootActorProxy.Get());
 
-		RootActorPreviewInstance->UpdatePreviewComponents();
+			// Spawned actor will take the transform values from the template, so manually reset them to zero here
+			RootActorProxy->SetActorLocation(FVector::ZeroVector);
+			RootActorProxy->SetActorRotation(FRotator::ZeroRotator);
 
-		TArray<TWeakObjectPtr<AActor>> LevelLightCards;
-		FindLightCardsForRootActor(RootActor, LevelLightCards);
+			FindProjectionOriginComponent();
 
-		for (const TWeakObjectPtr<AActor>& LightCard : LevelLightCards)
-		{
-			FActorSpawnParameters LightCardSpawnInfo;
-			LightCardSpawnInfo.SpawnCollisionHandlingOverride = ESpawnActorCollisionHandlingMethod::AlwaysSpawn;
-			LightCardSpawnInfo.bNoFail = true;
-			LightCardSpawnInfo.ObjectFlags = RF_Transient|RF_Transactional;
-			LightCardSpawnInfo.Template = LightCard.Get();
-		
-			AActor* LightCardProxy = PreviewWorld->SpawnActor(LightCard->GetClass(), &FVector::ZeroVector, &FRotator::ZeroRotator, MoveTemp(LightCardSpawnInfo));
+			// Filter out any primitives hidden in game except screen components
+			MeshProjectionRenderer->AddActor(RootActorProxy.Get(), [](const UPrimitiveComponent* PrimitiveComponent)
+			{
+				return !PrimitiveComponent->bHiddenInGame || PrimitiveComponent->IsA<UDisplayClusterScreenComponent>();
+			});
 
-			LightCardProxy->SetActorLocation(LightCard->GetActorLocation() - RootActor->GetActorLocation());
-			LightCardProxy->SetActorRotation(LightCard->GetActorRotation() - RootActor->GetActorRotation());
+			RootActorProxy->UpdatePreviewComponents();
 
-			LightCardProxies.Add(LightCardProxy);
-			MeshProjectionRenderer->AddActor(LightCardProxy);
-		}
+			FindLightCardsForRootActor(RootActor, LightCardLevelInstances);
+			
+			for (const TWeakObjectPtr<AActor>& LightCard : LightCardLevelInstances)
+			{
+				FObjectDuplicationParameters DupeActorParameters(LightCard.Get(), PreviewWorld->GetCurrentLevel());
+				DupeActorParameters.FlagMask = RF_AllFlags & ~(RF_ArchetypeObject | RF_Transactional);
+				DupeActorParameters.PortFlags = PPF_DuplicateVerbatim;
+				
+				AActor* LightCardProxy = CastChecked<AActor>(StaticDuplicateObjectEx(DupeActorParameters));
+				PreviewWorld->GetCurrentLevel()->AddLoadedActor(LightCardProxy);
+				
+				LightCardProxy->SetActorLocation(LightCard->GetActorLocation() - RootActor->GetActorLocation());
+				LightCardProxy->SetActorRotation(LightCard->GetActorRotation() - RootActor->GetActorRotation());
+
+				LightCardProxies.Add(LightCardProxy);
+				MeshProjectionRenderer->AddActor(LightCardProxy);
+			}
 
-		FCoreUObjectDelegates::OnObjectPropertyChanged.RemoveAll(this);
-		FCoreUObjectDelegates::OnObjectPropertyChanged.AddRaw(this, &FDisplayClusterLightCardEditorViewportClient::OnActorPropertyChanged);
+			FCoreUObjectDelegates::OnObjectPropertyChanged.RemoveAll(this);
+			FCoreUObjectDelegates::OnObjectPropertyChanged.AddRaw(this, &FDisplayClusterLightCardEditorViewportClient::OnActorPropertyChanged);
+		});
 	}
 }
 
@@ -679,15 +665,38 @@ void FDisplayClusterLightCardEditorViewportClient::OnActorPropertyChanged(UObjec
 	{
 		return;
 	}
+
+	auto IsOurActor = [ObjectBeingModified] (UObject* ObjectToCompare) -> bool
+	{
+		if (ObjectToCompare)
+		{
+			if (ObjectBeingModified == ObjectToCompare)
+			{
+				return true;
+			}
+
+			if (const UObject* RootActorOuter = ObjectBeingModified->GetTypedOuter(ObjectToCompare->GetClass()))
+			{
+				return RootActorOuter == ObjectToCompare;
+			}
+		}
+
+		return false;
+	};
 	
-	bool bIsOurActor = ObjectBeingModified == RootActorLevelInstance;
+	bool bIsOurActor = IsOurActor(RootActorLevelInstance.Get());
 	if (!bIsOurActor)
 	{
-		if (const UObject* RootActorOuter = ObjectBeingModified->GetTypedOuter(ADisplayClusterRootActor::StaticClass()))
+		for (const TWeakObjectPtr<AActor>& LightCard : LightCardLevelInstances)
 		{
-			bIsOurActor = RootActorOuter == RootActorLevelInstance;
+			bIsOurActor = IsOurActor(LightCard.Get());
+			if (bIsOurActor)
+			{
+				break;
+			}
 		}
 	}
+	
 	if (bIsOurActor)
 	{
 		UpdatePreviewActor(RootActorLevelInstance.Get(), true);
@@ -696,7 +705,7 @@ void FDisplayClusterLightCardEditorViewportClient::OnActorPropertyChanged(UObjec
 
 void FDisplayClusterLightCardEditorViewportClient::GetScenePrimitiveComponents(TArray<UPrimitiveComponent*>& OutPrimitiveComponents)
 {
-	RootActorPreviewInstance->ForEachComponent<UPrimitiveComponent>(true, [&](UPrimitiveComponent* PrimitiveComponent)
+	RootActorProxy->ForEachComponent<UPrimitiveComponent>(true, [&](UPrimitiveComponent* PrimitiveComponent)
 	{
 		OutPrimitiveComponents.Add(PrimitiveComponent);
 	});
@@ -811,10 +820,10 @@ void FDisplayClusterLightCardEditorViewportClient::GetSceneViewInitOptions(FScen
 
 UDisplayClusterConfigurationViewport* FDisplayClusterLightCardEditorViewportClient::FindViewportForPrimitiveComponent(UPrimitiveComponent* PrimitiveComponent)
 {
-	if (RootActorPreviewInstance.IsValid())
+	if (RootActorProxy.IsValid())
 	{
 		const FString PrimitiveComponentName = PrimitiveComponent->GetName();
-		UDisplayClusterConfigurationData* Config = RootActorPreviewInstance->GetConfigData();
+		UDisplayClusterConfigurationData* Config = RootActorProxy->GetConfigData();
 		
 		for (const TPair<FString, UDisplayClusterConfigurationClusterNode*>& NodePair : Config->Cluster->Nodes)
 		{
@@ -848,10 +857,10 @@ UDisplayClusterConfigurationViewport* FDisplayClusterLightCardEditorViewportClie
 
 void FDisplayClusterLightCardEditorViewportClient::FindProjectionOriginComponent()
 {
-	if (RootActorPreviewInstance.IsValid())
+	if (RootActorProxy.IsValid())
 	{
 		TArray<UDisplayClusterCameraComponent*> ViewOriginComponents;
-		RootActorPreviewInstance->GetComponents<UDisplayClusterCameraComponent>(ViewOriginComponents);
+		RootActorProxy->GetComponents<UDisplayClusterCameraComponent>(ViewOriginComponents);
 
 		if (ViewOriginComponents.Num())
 		{
@@ -859,7 +868,7 @@ void FDisplayClusterLightCardEditorViewportClient::FindProjectionOriginComponent
 		}
 		else
 		{
-			ProjectionOriginComponent = RootActorPreviewInstance->GetRootComponent();
+			ProjectionOriginComponent = RootActorProxy->GetRootComponent();
 		}
 	}
 	else
@@ -976,12 +985,12 @@ AActor* FDisplayClusterLightCardEditorViewportClient::TraceScreenForLightCard(co
 	{
 		if (AActor* HitActor = ScreenHitResult.GetActor())
 		{
-			if (RootActorPreviewInstance.Get() == HitActor && ScreenHitResult.Component.IsValid())
+			if (RootActorProxy.Get() == HitActor && ScreenHitResult.Component.IsValid())
 			{
 				if (UDisplayClusterConfigurationViewport* CfgViewport = FindViewportForPrimitiveComponent(ScreenHitResult.Component.Get()))
 				{
 					FString ViewOriginName = CfgViewport->Camera;
-					if (UDisplayClusterCameraComponent* ViewOrigin = RootActorPreviewInstance->GetComponentByName<UDisplayClusterCameraComponent>(ViewOriginName))
+					if (UDisplayClusterCameraComponent* ViewOrigin = RootActorProxy->GetComponentByName<UDisplayClusterCameraComponent>(ViewOriginName))
 					{
 						const FVector ViewOriginRayStart = ViewOrigin->GetComponentLocation();
 						const FVector ViewOriginRayEnd = ViewOriginRayStart + (ScreenHitResult.Location - ViewOriginRayStart) * HALF_WORLD_MAX;
diff --git a/Engine/Plugins/Runtime/nDisplay/Source/DisplayClusterLightCardEditor/Private/Viewport/DisplayClusterLightCardEditorViewportClient.h b/Engine/Plugins/Runtime/nDisplay/Source/DisplayClusterLightCardEditor/Private/Viewport/DisplayClusterLightCardEditorViewportClient.h
--- a/Engine/Plugins/Runtime/nDisplay/Source/DisplayClusterLightCardEditor/Private/Viewport/DisplayClusterLightCardEditorViewportClient.h
+++ b/Engine/Plugins/Runtime/nDisplay/Source/DisplayClusterLightCardEditor/Private/Viewport/DisplayClusterLightCardEditorViewportClient.h
@@ -93,10 +93,11 @@ private:
 private:
 	TWeakPtr<FSceneViewport> SceneViewportPtr;
 	TWeakPtr<SDisplayClusterLightCardEditor> LightCardEditorPtr;
-	TWeakObjectPtr<ADisplayClusterRootActor> RootActorPreviewInstance;
+	TWeakObjectPtr<ADisplayClusterRootActor> RootActorProxy;
 	TWeakObjectPtr<ADisplayClusterRootActor> RootActorLevelInstance;
 
 	TArray<TWeakObjectPtr<AActor>> LightCardProxies;
+	TArray<TWeakObjectPtr<AActor>> LightCardLevelInstances;
 	TArray<TWeakObjectPtr<AActor>> SelectedLightCards;
 	
 	/** The renderer for the viewport, which can render the meshes with a variety of projection types */
diff --git a/Engine/Plugins/Runtime/nDisplay/Source/DisplayClusterLightCardEditorShaders/Private/DisplayClusterMeshProjectionRenderer.cpp b/Engine/Plugins/Runtime/nDisplay/Source/DisplayClusterLightCardEditorShaders/Private/DisplayClusterMeshProjectionRenderer.cpp
--- a/Engine/Plugins/Runtime/nDisplay/Source/DisplayClusterLightCardEditorShaders/Private/DisplayClusterMeshProjectionRenderer.cpp
+++ b/Engine/Plugins/Runtime/nDisplay/Source/DisplayClusterLightCardEditorShaders/Private/DisplayClusterMeshProjectionRenderer.cpp
@@ -491,7 +491,10 @@ void FDisplayClusterMeshProjectionRenderer::ClearScene()
 	for (const TWeakObjectPtr<UPrimitiveComponent>& PrimitiveComponent : PrimitiveComponents)
 	{
 #if WITH_EDITOR
-		PrimitiveComponent->SelectionOverrideDelegate.Unbind();
+		if (PrimitiveComponent.IsValid() && PrimitiveComponent->SelectionOverrideDelegate.IsBound())
+		{
+			PrimitiveComponent->SelectionOverrideDelegate.Unbind();
+		}
 #endif
 	}
 
