nDisplay: Fixed issue where light card editor was not rendering properly on DirectX 11

diff --git a/Engine/Plugins/Runtime/nDisplay/Source/DisplayClusterLightCardEditorShaders/Private/DisplayClusterMeshProjectionRenderer.cpp b/Engine/Plugins/Runtime/nDisplay/Source/DisplayClusterLightCardEditorShaders/Private/DisplayClusterMeshProjectionRenderer.cpp
--- a/Engine/Plugins/Runtime/nDisplay/Source/DisplayClusterLightCardEditorShaders/Private/DisplayClusterMeshProjectionRenderer.cpp
+++ b/Engine/Plugins/Runtime/nDisplay/Source/DisplayClusterLightCardEditorShaders/Private/DisplayClusterMeshProjectionRenderer.cpp
@@ -549,7 +549,7 @@ void FDisplayClusterMeshProjectionRenderer::Render(FCanvas* Canvas, FSceneInterf
 				FRDGTextureRef HitProxyTexture = GraphBuilder.CreateTexture(Desc, TEXT("DisplayClusterMeshProjection.HitProxyTexture"));
 
 				const FRDGTextureDesc DepthDesc = FRDGTextureDesc::Create2D(OutputTexture->Desc.Extent, PF_DepthStencil, FClearValueBinding::DepthFar, TexCreate_DepthStencilTargetable | TexCreate_ShaderResource);
-				FRDGTextureRef HitProxyDepthTexture = GraphBuilder.CreateTexture(DepthDesc, TEXT("DisplayClusterMeshProjection.DepthTexture"));
+				FRDGTextureRef HitProxyDepthTexture = GraphBuilder.CreateTexture(DepthDesc, TEXT("DisplayClusterMeshProjection.HitProxyDepthTexture"));
 
 				FRenderTargetBinding HitProxyRenderTargetBinding(HitProxyTexture, ERenderTargetLoadAction::EClear);
 				FDepthStencilBinding HitProxyDepthStencilBinding(HitProxyDepthTexture, ERenderTargetLoadAction::EClear, ERenderTargetLoadAction::EClear, FExclusiveDepthStencil::DepthWrite_StencilWrite);
@@ -591,12 +591,15 @@ void FDisplayClusterMeshProjectionRenderer::Render(FCanvas* Canvas, FSceneInterf
 			}
 			else
 			{
-				const FRDGTextureDesc Desc = FRDGTextureDesc::Create2D(OutputTexture->Desc.Extent, PF_DepthStencil, FClearValueBinding::DepthFar, TexCreate_DepthStencilTargetable | TexCreate_ShaderResource);
-				FRDGTextureRef DepthTexture = GraphBuilder.CreateTexture(Desc, TEXT("DisplayClusterMeshProjection.DepthTexture"));
+				FRDGTextureRef ColorTexture = GraphBuilder.CreateTexture(OutputTexture->Desc, TEXT("DisplayClusterMeshProjection.ColorTexture"));
 
-				FDepthStencilBinding OutputDepthStencilBinding(DepthTexture, ERenderTargetLoadAction::EClear, ERenderTargetLoadAction::EClear, FExclusiveDepthStencil::DepthWrite_StencilNop);
+				const FRDGTextureDesc DepthDesc = FRDGTextureDesc::Create2D(OutputTexture->Desc.Extent, PF_DepthStencil, FClearValueBinding::DepthFar, TexCreate_DepthStencilTargetable | TexCreate_ShaderResource);
+				FRDGTextureRef DepthTexture = GraphBuilder.CreateTexture(DepthDesc, TEXT("DisplayClusterMeshProjection.DepthTexture"));
+
+				FRenderTargetBinding ColorRenderTargetBinding(ColorTexture, ERenderTargetLoadAction::EClear);
+				FDepthStencilBinding DepthStencilBinding(DepthTexture, ERenderTargetLoadAction::EClear, ERenderTargetLoadAction::ENoAction, FExclusiveDepthStencil::DepthWrite_StencilNop);
 
-				AddBaseRenderPass(GraphBuilder, View, ProjectionType, OutputRenderTargetBinding, OutputDepthStencilBinding);
+				AddBaseRenderPass(GraphBuilder, View, ProjectionType, ColorRenderTargetBinding, DepthStencilBinding);
 
 #if WITH_EDITOR
 				const FRDGTextureDesc SelectionDepthDesc = FRDGTextureDesc::Create2D(OutputTexture->Desc.Extent, PF_DepthStencil, FClearValueBinding::DepthFar, TexCreate_DepthStencilTargetable | TexCreate_ShaderResource);
@@ -604,7 +607,40 @@ void FDisplayClusterMeshProjectionRenderer::Render(FCanvas* Canvas, FSceneInterf
 				FDepthStencilBinding SelectionDepthStencilBinding(SelectionDepthTexture, ERenderTargetLoadAction::EClear, ERenderTargetLoadAction::EClear, FExclusiveDepthStencil::DepthWrite_StencilWrite);
 
 				AddSelectionDepthRenderPass(GraphBuilder, View, ProjectionType, SelectionDepthStencilBinding);
-				AddSelectionOutlineScreenPass(GraphBuilder, View, OutputRenderTargetBinding, OutputTexture, DepthTexture, SelectionDepthTexture);
+				AddSelectionOutlineScreenPass(GraphBuilder, View, OutputRenderTargetBinding, ColorTexture, DepthTexture, SelectionDepthTexture);
+#else
+				// Copy the scene color to the output render target
+				{
+					FCopyRectPS::FParameters* ScreenPassParameters = GraphBuilder.AllocParameters<FCopyRectPS::FParameters>();
+					ScreenPassParameters->InputTexture = ColorTexture;
+					ScreenPassParameters->InputSampler = TStaticSamplerState<>::GetRHI();
+					ScreenPassParameters->RenderTargets[0] = OutputRenderTargetBinding;
+
+					FGlobalShaderMap* GlobalShaderMap = GetGlobalShaderMap(GMaxRHIFeatureLevel);
+					TShaderMapRef<FScreenPassVS> ScreenPassVS(GlobalShaderMap);
+					TShaderMapRef<FCopyRectPS> CopyPixelShader(GlobalShaderMap);
+
+					FRHIBlendState* DefaultBlendState = FScreenPassPipelineState::FDefaultBlendState::GetRHI();
+					const FScreenPassTextureViewport RegionViewport(OutputRenderTargetBinding.GetTexture());
+
+					GraphBuilder.AddPass(
+						RDG_EVENT_NAME("MeshProjectionRenderer::CopyColorTexture"),
+						ScreenPassParameters,
+						ERDGPassFlags::Raster,
+						[View, ScreenPassVS, CopyPixelShader, &RegionViewport, ScreenPassParameters, DefaultBlendState](FRHICommandList& RHICmdList)
+					{
+						DrawScreenPass(
+							RHICmdList,
+							*View,
+							RegionViewport,
+							RegionViewport,
+							FScreenPassPipelineState(ScreenPassVS, CopyPixelShader, DefaultBlendState),
+							[&](FRHICommandList&)
+						{
+							SetShaderParameters(RHICmdList, CopyPixelShader, CopyPixelShader.GetPixelShader(), *ScreenPassParameters);
+						});
+					});
+				}
 #endif
 			}
 
