nDisplay: DC screen components can be selected

diff --git a/Engine/Plugins/Runtime/nDisplay/Source/DisplayCluster/Private/Game/EngineClasses/Scene/DisplayClusterRootActor.cpp b/Engine/Plugins/Runtime/nDisplay/Source/DisplayCluster/Private/Game/EngineClasses/Scene/DisplayClusterRootActor.cpp
--- a/Engine/Plugins/Runtime/nDisplay/Source/DisplayCluster/Private/Game/EngineClasses/Scene/DisplayClusterRootActor.cpp
+++ b/Engine/Plugins/Runtime/nDisplay/Source/DisplayCluster/Private/Game/EngineClasses/Scene/DisplayClusterRootActor.cpp
@@ -488,10 +488,12 @@ bool ADisplayClusterRootActor::GetHiddenInGamePrimitives(TSet<FPrimitiveComponen
 		{
 			GetTypedPrimitives<UStaticMeshComponent>(OutPrimitives, &WarpMeshNames);
 			GetTypedPrimitives<UProceduralMeshComponent>(OutPrimitives, &WarpMeshNames);
-			GetTypedPrimitives<UDisplayClusterScreenComponent>(OutPrimitives, &WarpMeshNames);
 		}
 	}
 
+	// Always hide DC Screen component in game
+	GetTypedPrimitives<UDisplayClusterScreenComponent>(OutPrimitives);
+
 #if WITH_EDITOR
 
 	const bool bHideVisualizationComponents = !CVarShowVisualizationComponents.GetValueOnGameThread();
diff --git a/Engine/Plugins/Runtime/nDisplay/Source/DisplayCluster/Private/Game/EngineClasses/Scene/DisplayClusterScreenComponent.cpp b/Engine/Plugins/Runtime/nDisplay/Source/DisplayCluster/Private/Game/EngineClasses/Scene/DisplayClusterScreenComponent.cpp
--- a/Engine/Plugins/Runtime/nDisplay/Source/DisplayCluster/Private/Game/EngineClasses/Scene/DisplayClusterScreenComponent.cpp
+++ b/Engine/Plugins/Runtime/nDisplay/Source/DisplayCluster/Private/Game/EngineClasses/Scene/DisplayClusterScreenComponent.cpp
@@ -28,9 +28,7 @@ UDisplayClusterScreenComponent::UDisplayClusterScreenComponent(const FObjectInit
 		SetVisibility(true);
 	}
 
-#if WITH_EDITORONLY_DATA
-	SetIsVisualizationComponent(true);
-#endif
+	SetCastShadow(false);
 
 	bVisibleInReflectionCaptures = false;
 	bVisibleInRayTracing = false;
