Adding minimal interfaces to D3D11 and D3D12 RHIs
- Adding ID3D12DynamicRHI and ID3D11DynamicRHI interfaces for plugins that want to touch DirectX resources/devices directly.
- Adding specific headers for including DirectX headers: D3D12ThirdParty.h and D3D11ThirdParty.h
- Modifying plugins that were using D3D11RHIPrivate.h and D3D12RHIPrivate.h to use the interfaces instead.
- Removing plugin references to UE Modules and ThirdParty libraries that were only needed because they included the private RHI headers.
- Removing D3D11RHI/Private and D3D12RHI/Private from plugin include paths.

diff --git a/Engine/Plugins/Experimental/NNI/Source/NeuralNetworkInference/NeuralNetworkInference.Build.cs b/Engine/Plugins/Experimental/NNI/Source/NeuralNetworkInference/NeuralNetworkInference.Build.cs
--- a/Engine/Plugins/Experimental/NNI/Source/NeuralNetworkInference/NeuralNetworkInference.Build.cs
+++ b/Engine/Plugins/Experimental/NNI/Source/NeuralNetworkInference/NeuralNetworkInference.Build.cs
@@ -43,7 +43,6 @@ public class NeuralNetworkInference : ModuleRules
 				"ModelProtoFileReader",
 				"NeuralNetworkInferenceShaders",
 				"RHI",
-				"RHICore",
 				// ORT-related
 				"Projects",
 				"ThirdPartyHelperAndDLLLoader"
@@ -61,30 +60,12 @@ public class NeuralNetworkInference : ModuleRules
 				}
 			);
 
-			// Borrowed from Plugins/Media/ImgMedia (we need to use FD3D12DynamicRHI)
 			if (Target.Platform == UnrealTargetPlatform.Win64)
 			{
-				PrivateIncludePaths.AddRange(
-					new string[]{
-						//required for "D3D12RHIPrivate.h"
-						Path.Combine(EngineDirectory, "Source/Runtime/D3D12RHI/Private"),
-						Path.Combine(EngineDirectory, "Source/Runtime/D3D12RHI/Private/Windows")
-					});
-
-				PrivateDependencyModuleNames.AddRange
-					(
-					new string[] {
-						"D3D12RHI",
-						"DirectML_1_8_0"
-					}
-				);
+				PrivateDependencyModuleNames.Add("D3D12RHI");
 
 				AddEngineThirdPartyPrivateStaticDependencies(Target, "DX12");
-
-				if (Target.bBuildEditor)
-				{
-					PublicDependencyModuleNames.Add("WinPixEventRuntime");
-				}
+				AddEngineThirdPartyPrivateStaticDependencies(Target, "DirectML_1_8_0");
 			}
 		}
 
diff --git a/Engine/Plugins/Experimental/NNI/Source/NeuralNetworkInference/Private/NeuralNetworkImplBackEndUEAndORT.cpp b/Engine/Plugins/Experimental/NNI/Source/NeuralNetworkInference/Private/NeuralNetworkImplBackEndUEAndORT.cpp
--- a/Engine/Plugins/Experimental/NNI/Source/NeuralNetworkInference/Private/NeuralNetworkImplBackEndUEAndORT.cpp
+++ b/Engine/Plugins/Experimental/NNI/Source/NeuralNetworkInference/Private/NeuralNetworkImplBackEndUEAndORT.cpp
@@ -12,20 +12,7 @@
 
 #if defined(WITH_UE_AND_ORT_SUPPORT) && defined(PLATFORM_WIN64)
 	#include "HAL/CriticalSection.h"
-	#include "RHI.h"
-	#include "DynamicRHI.h"
-
-	// Disable NOMINMAX & WIN32_LEAN_AND_MEAN defines to avoid compiler warnings
-	#pragma push_macro("NOMINMAX")
-	#pragma push_macro("WIN32_LEAN_AND_MEAN")
-	#pragma push_macro("UE_MINIMAL_WINDOWS_INCLUDE")
-	#undef NOMINMAX
-	#undef WIN32_LEAN_AND_MEAN
-	#define UE_MINIMAL_WINDOWS_INCLUDE // Avoids Win64 Clang warning
-	#include "D3D12RHIPrivate.h"
-	#pragma pop_macro("UE_MINIMAL_WINDOWS_INCLUDE")
-	#pragma pop_macro("WIN32_LEAN_AND_MEAN")
-	#pragma pop_macro("NOMINMAX")
+	#include "ID3D12DynamicRHI.h"
 #endif
 
 //#define WITH_NNI_CPU_NOT_RECOMMENDED // Only for debugging purposes
@@ -51,10 +38,12 @@ NNI_THIRD_PARTY_INCLUDES_END
 
 #if defined(PLATFORM_WIN64)
 
-#if WITH_EDITOR && !UE_BUILD_SHIPPING
+#if WITH_EDITOR && defined(PLATFORM_WIN64) && !UE_BUILD_SHIPPING
+#include "Windows/AllowWindowsPlatformTypes.h"
 NNI_THIRD_PARTY_INCLUDES_START
-	#include "pix3.h"
+	#include <pix3.h>
 NNI_THIRD_PARTY_INCLUDES_END
+#include "Windows/HideWindowsPlatformTypes.h"
 	#define NNIGPUProfileMarker(Name) FNNIGPUProfiler::Instance()->Marker(Name)
 #else
 	#define NNIGPUProfileMarker(Name)
@@ -91,7 +80,7 @@ private:
 	FNNIGPUProfiler()
 	{
 #if defined(PLATFORM_WIN64) && defined(USE_PIX) && !defined(UE_BUILD_SHIPPING)
-		bIsEnabled = FD3D12DynamicRHI::GetD3DRHI()->IsPixEventEnabled();
+		bIsEnabled = GetID3D12DynamicRHI()->RHIIsPixEnabled();
 #else
 		bIsEnabled = false;
 #endif
@@ -198,7 +187,7 @@ IDMLDevice* FPrivateImplBackEndUEAndORT::FDMLDeviceList::Add(ID3D12Device* Devic
 	DML_CREATE_DEVICE_FLAGS DmlCreateFlags = DML_CREATE_DEVICE_FLAG_NONE;
 
 #if !UE_BUILD_SHIPPING
-	if (D3D12RHI_ShouldCreateWithD3DDebug()
+	if (ID3D12DynamicRHI::IsD3DDebugEnabled()
 		|| FParse::Param(FCommandLine::Get(), TEXT("d3d12gpuvalidation")) || FParse::Param(FCommandLine::Get(), TEXT("gpuvalidation")))
 	{
 		DmlCreateFlags |= DML_CREATE_DEVICE_FLAG_DEBUG;
@@ -544,35 +533,34 @@ bool UNeuralNetwork::FImplBackEndUEAndORT::ConfigureMembers(const ENeuralDeviceT
 		}
 
 		// Get adapter's D3D12 device that we would like to share with DirectML execution provider
-		FD3D12DynamicRHI*			RHI = GetDynamicRHI<FD3D12DynamicRHI>();
-		FD3D12Adapter&				RHIAdapter = RHI->GetAdapter(0);
-		const FD3D12AdapterDesc&	RHIAdapterDesc = RHIAdapter.GetDesc();
+		ID3D12DynamicRHI* RHI = GetID3D12DynamicRHI();
+		const TArray<FD3D12MinimalAdapterDesc> Adapters = RHI->RHIGetAdapterDescs();
+
+		const FD3D12MinimalAdapterDesc&	RHIAdapterDesc = Adapters[0];
 
 		UE_LOG(LogNeuralNetworkInference, Display, TEXT("%d available RHI adapters. NNI using RHI adapter %s with LUID:%0x%0x."),
-			RHI->GetNumAdapters(), RHIAdapterDesc.Desc.Description, RHIAdapterDesc.Desc.AdapterLuid.HighPart, RHIAdapterDesc.Desc.AdapterLuid.LowPart);
+			Adapters.Num(), RHIAdapterDesc.Desc.Description, RHIAdapterDesc.Desc.AdapterLuid.HighPart, RHIAdapterDesc.Desc.AdapterLuid.LowPart);
 
-		if (RHI->GetNumAdapters() > 1)
+		if (Adapters.Num() > 1)
 		{
 			UE_LOG(LogNeuralNetworkInference, Display, TEXT("All available RHI adapters:"));
-			for (int32 AdapterIndex = 0; AdapterIndex < RHI->GetNumAdapters(); ++AdapterIndex)
+			for (int32 AdapterIndex = 0; AdapterIndex < Adapters.Num(); ++AdapterIndex)
 			{
-				const FD3D12AdapterDesc& CurrAdapterDesc = RHI->GetAdapter(AdapterIndex).GetDesc();
+				const FD3D12MinimalAdapterDesc& CurrAdapterDesc = Adapters[AdapterIndex];
 				UE_LOG(LogNeuralNetworkInference, Display, TEXT("  - Adapter [%d] Name:%s with LUID:%0x%0x."),
 					AdapterIndex, CurrAdapterDesc.Desc.Description, CurrAdapterDesc.Desc.AdapterLuid.HighPart, CurrAdapterDesc.Desc.AdapterLuid.LowPart);
 			}
 		}
 
-		if (RHI->GetNumAdapters() > 1 || RHIAdapterDesc.NumDeviceNodes > 1)
+		if (Adapters.Num() > 1 || RHIAdapterDesc.NumDeviceNodes > 1)
 		{
 			UE_LOG(LogNeuralNetworkInference, Warning,
 				TEXT("FImplBackEndUEAndORT::ConfigureMembers(): There are multiple (%d) adapters and/or multiple (%d) devices, NNI is currently using only one adapter."),
-				RHI->GetNumAdapters(), RHIAdapterDesc.NumDeviceNodes);
+				Adapters.Num(), RHIAdapterDesc.NumDeviceNodes);
 		}
 
-		ID3D12Device* NativeDevice = RHIAdapter.GetD3DDevice();
-
 		// Make sure that we have one DMLDevice per D3D12 device
-		IDMLDevice* DmlDevice = FPrivateImplBackEndUEAndORT::GetDMLDeviceThreadSafe(NativeDevice);
+		IDMLDevice* DmlDevice = FPrivateImplBackEndUEAndORT::GetDMLDeviceThreadSafe(RHI->RHIGetDevice(0));
 
 		if (!DmlDevice)
 		{
@@ -582,7 +570,7 @@ bool UNeuralNetwork::FImplBackEndUEAndORT::ConfigureMembers(const ENeuralDeviceT
 
 		// Get a ID3D12CommandQueue as well
 		// TODO: Should we create our own queue?
-		ID3D12CommandQueue* NativeCmdQ = RHI->RHIGetD3DCommandQueue();
+		ID3D12CommandQueue* NativeCmdQ = RHI->RHIGetCommandQueue();
 
 		// ORT GPU (Direct ML)
 		SessionOptions->SetGraphOptimizationLevel(GraphOptimizationLevel::ORT_ENABLE_ALL); // ORT_ENABLE_ALL, ORT_ENABLE_EXTENDED, ORT_ENABLE_BASIC, ORT_DISABLE_ALL
diff --git a/Engine/Plugins/Experimental/NNI/Source/NeuralNetworkInference/Private/NeuralNetworkImplBackEndUEAndORT.h b/Engine/Plugins/Experimental/NNI/Source/NeuralNetworkInference/Private/NeuralNetworkImplBackEndUEAndORT.h
--- a/Engine/Plugins/Experimental/NNI/Source/NeuralNetworkInference/Private/NeuralNetworkImplBackEndUEAndORT.h
+++ b/Engine/Plugins/Experimental/NNI/Source/NeuralNetworkInference/Private/NeuralNetworkImplBackEndUEAndORT.h
@@ -7,8 +7,6 @@
 
 // RHI includes must happen before onnxruntime_cxx_api.h (both files include Windows.h)
 #include "HAL/CriticalSection.h"
-#include "RHI.h"
-#include "DynamicRHI.h"
 
 #if defined(WITH_UE_AND_ORT_SUPPORT) && defined(PLATFORM_WIN64)
 // Disable NOMINMAX & WIN32_LEAN_AND_MEAN defines to avoid compiler warnings
@@ -16,7 +14,7 @@
 #pragma push_macro("WIN32_LEAN_AND_MEAN")
 #undef NOMINMAX
 #undef WIN32_LEAN_AND_MEAN
-#include "D3D12RHIPrivate.h"
+#include "ID3D12DynamicRHI.h"
 #pragma pop_macro("WIN32_LEAN_AND_MEAN")
 #pragma pop_macro("NOMINMAX")
 #endif
diff --git a/Engine/Plugins/Experimental/NNI/Source/NeuralNetworkInference/Private/NeuralNetworkInferenceUtilsGPU.cpp b/Engine/Plugins/Experimental/NNI/Source/NeuralNetworkInference/Private/NeuralNetworkInferenceUtilsGPU.cpp
--- a/Engine/Plugins/Experimental/NNI/Source/NeuralNetworkInference/Private/NeuralNetworkInferenceUtilsGPU.cpp
+++ b/Engine/Plugins/Experimental/NNI/Source/NeuralNetworkInference/Private/NeuralNetworkInferenceUtilsGPU.cpp
@@ -5,27 +5,7 @@
 #include "NeuralTensorResourceArray.h"
 
 #ifdef PLATFORM_WIN64
-
-#include "RHIResources.h"
-#include "ShaderCore.h"
-
-// Disable NOMINMAX & WIN32_LEAN_AND_MEAN defines to avoid compiler warnings
-#pragma push_macro("NOMINMAX")
-#pragma push_macro("WIN32_LEAN_AND_MEAN")
-#undef NOMINMAX
-#undef WIN32_LEAN_AND_MEAN
-
-#include "D3D12RHI.h"
-#include "D3D12RHICommon.h"
-
-// @todo: We need D3D12RHIPrivate.h in order to compile correctly, i.e. to be able to use FD3D12Resource and ID3D12Resource
-#include "D3D12RHIPrivate.h"
-
-#include "D3D12Resources.h"
-
-#pragma pop_macro("WIN32_LEAN_AND_MEAN")
-#pragma pop_macro("NOMINMAX")
-
+	#include "ID3D12DynamicRHI.h"
 #endif
 
 /* FNeuralNetworkInferenceUtilsGPU public functions
@@ -79,14 +59,7 @@ bool FNeuralNetworkInferenceUtilsGPU::IsD3D12RHI()
 
 ID3D12Resource* FNeuralNetworkInferenceUtilsGPU::GetD3D12Resource(FRHIBuffer* Buffer)
 {
-	if (!IsD3D12RHI())
-	{
-		return nullptr;
-	}
-
-	FD3D12Buffer*	D3DBuffer = static_cast<FD3D12Buffer*>(Buffer);
-
-	return D3DBuffer->GetResource()->GetResource();
+	return GetID3D12DynamicRHI()->RHIGetResource(Buffer);
 }
 
 #endif
\ No newline at end of file
diff --git a/Engine/Plugins/Experimental/NNI/Source/NeuralNetworkInference/Private/NeuralNetworkInferenceUtilsGPU.h b/Engine/Plugins/Experimental/NNI/Source/NeuralNetworkInference/Private/NeuralNetworkInferenceUtilsGPU.h
--- a/Engine/Plugins/Experimental/NNI/Source/NeuralNetworkInference/Private/NeuralNetworkInferenceUtilsGPU.h
+++ b/Engine/Plugins/Experimental/NNI/Source/NeuralNetworkInference/Private/NeuralNetworkInferenceUtilsGPU.h
@@ -3,11 +3,12 @@
 #pragma once
 
 #include "CoreMinimal.h"
-#include "RenderGraphBuilder.h"
 
-#ifdef PLATFORM_WIN64
-	#include "RHI.h"
+class FRDGBuilder;
+class FRHIBuffer;
+struct FReadBuffer;
 
+#ifdef PLATFORM_WIN64
 	struct ID3D12Resource;
 #endif
 
diff --git a/Engine/Plugins/Media/HLMedia/Source/HLMedia/HLMedia.Build.cs b/Engine/Plugins/Media/HLMedia/Source/HLMedia/HLMedia.Build.cs
--- a/Engine/Plugins/Media/HLMedia/Source/HLMedia/HLMedia.Build.cs
+++ b/Engine/Plugins/Media/HLMedia/Source/HLMedia/HLMedia.Build.cs
@@ -29,7 +29,6 @@ public class HLMedia : ModuleRules
                 "Projects",
                 "RenderCore",
                 "RHI",
-				"RHICore",
 				"MediaUtils",
             });
 
@@ -45,9 +44,7 @@ public class HLMedia : ModuleRules
 
         var EngineDir = Path.GetFullPath(Target.RelativeEnginePath);
 
-        if (Target.Platform == UnrealTargetPlatform.Win64
-            ||
-            Target.Platform == UnrealTargetPlatform.HoloLens)
+        if (Target.Platform == UnrealTargetPlatform.Win64 || Target.Platform == UnrealTargetPlatform.HoloLens)
         {
             PrivateDependencyModuleNames.AddRange(new string[] 
             {
@@ -56,37 +53,23 @@ public class HLMedia : ModuleRules
 				"D3D12RHI"
             });
 
-            PrivateIncludePaths.AddRange(new string[]
-            {
-                Path.Combine(EngineDir, "Source/ThirdParty/HLMediaLibrary/inc"),
-                Path.Combine(EngineDir, "Source/Runtime/Windows/D3D11RHI/Private"),
-				Path.Combine(EngineDir, "Source/Runtime/D3D12RHI/Private"),
-				Path.Combine(EngineDir, "Source/Runtime/D3D12RHI/Public"),
-            });
-
-            // required by D3D11RHI to compile
-            AddEngineThirdPartyPrivateStaticDependencies(Target, "IntelMetricsDiscovery");
-            AddEngineThirdPartyPrivateStaticDependencies(Target, "NVAftermath");
-            AddEngineThirdPartyPrivateStaticDependencies(Target, "HLMediaLibrary");
+            PrivateIncludePaths.Add(Path.Combine(EngineDir, "Source/ThirdParty/HLMediaLibrary/inc"));
 
             PublicSystemLibraries.Add("mfplat.lib");
             PublicSystemLibraries.Add("mfreadwrite.lib");
             PublicSystemLibraries.Add("mfuuid.lib");
 			
 			// For D3D11on12
-			PublicSystemLibraries.Add("d3d11.lib");
-			AddEngineThirdPartyPrivateStaticDependencies(Target, "DX12");
+			AddEngineThirdPartyPrivateStaticDependencies(Target, "DX12", "DX11");
         }
 
         if (Target.Platform == UnrealTargetPlatform.Win64)
         {
             PrivateIncludePaths.Add(Path.Combine(EngineDir, "Source/Runtime/Windows/D3D11RHI/Private/Windows"));
-			PrivateIncludePaths.Add(Path.Combine(EngineDir, "Source/Runtime/D3D12RHI/Private/Windows"));
         }
         else if (Target.Platform == UnrealTargetPlatform.HoloLens)
         {
             PrivateIncludePaths.Add(Path.Combine(EngineDir, "Source/Runtime/Windows/D3D11RHI/Private/HoloLens"));
-			PrivateIncludePaths.Add(Path.Combine(EngineDir, "Source/Runtime/D3D12RHI/Private/HoloLens"));
         }
     }
 }
diff --git a/Engine/Plugins/Media/HLMedia/Source/HLMedia/Private/HLMediaPlayer.cpp b/Engine/Plugins/Media/HLMedia/Source/HLMedia/Private/HLMediaPlayer.cpp
--- a/Engine/Plugins/Media/HLMedia/Source/HLMedia/Private/HLMediaPlayer.cpp
+++ b/Engine/Plugins/Media/HLMedia/Source/HLMedia/Private/HLMediaPlayer.cpp
@@ -13,6 +13,13 @@
 #include "HLMediaLibrary.h"
 #include "HLMediaPlayerTracks.h"
 
+#include "ID3D12DynamicRHI.h"
+#include "Windows/AllowWindowsPlatformTypes.h"
+THIRD_PARTY_INCLUDES_START
+	#include <d3d11on12.h>
+THIRD_PARTY_INCLUDES_END
+#include "Windows/HideWindowsPlatformTypes.h"
+
 #define LOCTEXT_NAMESPACE "HLMediaPlayer"
 
 DEFINE_LOG_CATEGORY(LogHLMediaPlayer);
diff --git a/Engine/Plugins/Media/HLMedia/Source/HLMedia/Private/HLMediaPlayer.h b/Engine/Plugins/Media/HLMedia/Source/HLMedia/Private/HLMediaPlayer.h
--- a/Engine/Plugins/Media/HLMedia/Source/HLMedia/Private/HLMediaPlayer.h
+++ b/Engine/Plugins/Media/HLMedia/Source/HLMedia/Private/HLMediaPlayer.h
@@ -15,10 +15,6 @@
 
 #include "HLMediaPlayerTracks.h"
 
-#include "d3d11on12.h"
-#include "D3D12RHIPrivate.h"
-#include "D3D12Util.h"
-
 class FMediaSamples;
 class FHLMediaPlayerTracks;
 
diff --git a/Engine/Plugins/Media/HLMedia/Source/HLMedia/Private/HLMediaPlayerTracks.cpp b/Engine/Plugins/Media/HLMedia/Source/HLMedia/Private/HLMediaPlayerTracks.cpp
--- a/Engine/Plugins/Media/HLMedia/Source/HLMedia/Private/HLMediaPlayerTracks.cpp
+++ b/Engine/Plugins/Media/HLMedia/Source/HLMedia/Private/HLMediaPlayerTracks.cpp
@@ -3,14 +3,6 @@
 #include "HLMediaPlayerTracks.h"
 
 #include "MediaUtils/Public/MediaSamples.h"
-#include "RHIStaticStates.h"
-// Fix build error.
-#define INTEL_EXTENSIONS 0
-#include "D3D11RHIPrivate.h"
-#include "RenderingThread.h"
-
-#include "Engine/Engine.h"
-#include "D3D11Resources.h"
 
 using namespace HLMediaLibrary;
 
diff --git a/Engine/Plugins/Media/HLMedia/Source/HLMedia/Private/HLMediaTextureSample.h b/Engine/Plugins/Media/HLMedia/Source/HLMedia/Private/HLMediaTextureSample.h
--- a/Engine/Plugins/Media/HLMedia/Source/HLMedia/Private/HLMediaTextureSample.h
+++ b/Engine/Plugins/Media/HLMedia/Source/HLMedia/Private/HLMediaTextureSample.h
@@ -6,26 +6,9 @@
 
 #include "IMediaTextureSample.h"
 #include "MediaUtils/Public/MediaObjectPool.h"
-// Fix build error.
-#define INTEL_EXTENSIONS 0
-#include "D3D11RHIPrivate.h"
-#include "D3D11Resources.h"
-
-#define GetRenderTargetFormat GetRenderTargetFormat_D3D12
-#define FindShaderResourceDXGIFormat FindShaderResourceDXGIFormat_D3D12
-#define FindUnorderedAccessDXGIFormat FindUnorderedAccessDXGIFormat_D3D12
-#define FindDepthStencilDXGIFormat FindDepthStencilDXGIFormat_D3D12
-#define HasStencilBits HasStencilBits_D3D12
-#define FD3DGPUProfiler FD3D12GPUProfiler
-
-#include "D3D12RHIPrivate.h"
-
-#undef GetRenderTargetFormat
-#undef FindShaderResourceDXGIFormat
-#undef FindUnorderedAccessDXGIFormat
-#undef FindDepthStencilDXGIFormat
-#undef HasStencilBits
-#undef FD3DGPUProfiler
+
+#include "ID3D11DynamicRHI.h"
+#include "ID3D12DynamicRHI.h"
 
 class FHLMediaTextureSample
     : public IMediaTextureSample
@@ -42,37 +25,18 @@ public:
 		if (RHIGetInterfaceType() == ERHIInterfaceType::D3D12)
 		{
 			// MediaFoundation creates a DX11 texture.  To use this texture with DX12, open the shared handle with the DX12 RHI.
-			FD3D12DynamicRHI* DX12RHI = GetDynamicRHI<FD3D12DynamicRHI>();
 			TComPtr<ID3D12Resource> sharedMediaTexture;
-			if (FAILED(DX12RHI->GetAdapter().GetD3DDevice()->OpenSharedHandle(InSharedTextureHandle, IID_PPV_ARGS(&sharedMediaTexture))))
+			if (FAILED(GetID3D12DynamicRHI()->RHIGetDevice(0)->OpenSharedHandle(InSharedTextureHandle, IID_PPV_ARGS(&sharedMediaTexture))))
 			{
 				UE_LOG(LogHLMediaPlayer, Log, TEXT("ID3D12Device::OpenSharedHandle failed in FHLMediaTextureSample"));
 				return;
 			}
 
-			Texture = DX12RHI->RHICreateTexture2DFromResource(PF_B8G8R8A8, TexCreate_Dynamic, FClearValueBinding::None, sharedMediaTexture.Get());
+			Texture = GetID3D12DynamicRHI()->RHICreateTexture2DFromResource(PF_B8G8R8A8, TexCreate_Dynamic, FClearValueBinding::None, sharedMediaTexture.Get());
 		}
 		else
 		{
-			TArray<TRefCountPtr<ID3D11RenderTargetView>> RenderTargetViews;
-			Texture = new FD3D11Texture2D(
-				GD3D11RHI,
-				InTexture,
-				InShaderResourceView,
-				false,
-				RenderTargetViews.Num(),
-				RenderTargetViews,
-				nullptr,
-				Desc.Width,
-				Desc.Height,
-				0,
-				1,
-				1,
-				PF_B8G8R8A8,
-				false,
-				TexCreate_None,
-				false,
-				FClearValueBinding::Transparent);
+			Texture = GetID3D11DynamicRHI()->RHICreateTexture2DFromResource(PF_B8G8R8A8, TexCreate_None, FClearValueBinding::Transparent, InTexture);
 		}
     }
 
diff --git a/Engine/Plugins/Media/HardwareEncoders/Source/EncoderAMF/EncoderAMF.Build.cs b/Engine/Plugins/Media/HardwareEncoders/Source/EncoderAMF/EncoderAMF.Build.cs
--- a/Engine/Plugins/Media/HardwareEncoders/Source/EncoderAMF/EncoderAMF.Build.cs
+++ b/Engine/Plugins/Media/HardwareEncoders/Source/EncoderAMF/EncoderAMF.Build.cs
@@ -28,10 +28,7 @@ public class EncoderAMF : ModuleRules
 			"Engine",
 			"AVEncoder",
 			"VulkanRHI",
-			"Amf"
-		});
-
-		PublicDependencyModuleNames.AddRange(new string[] {
+			"Amf",
 			"RenderCore",
 			"Core",
 			"RHI"
@@ -47,7 +44,6 @@ public class EncoderAMF : ModuleRules
 
 		if (Target.IsInPlatformGroup(UnrealPlatformGroup.Windows))
 		{
-			PublicDependencyModuleNames.Add("D3D12RHI");
 			PrivateIncludePaths.Add(Path.Combine(EngineSourceDirectory, "Source/Runtime/VulkanRHI/Private/Windows"));
 		}
 		else if (Target.IsInPlatformGroup(UnrealPlatformGroup.Unix))
diff --git a/Engine/Plugins/Media/HardwareEncoders/Source/EncoderNVENC/EncoderNVENC.Build.cs b/Engine/Plugins/Media/HardwareEncoders/Source/EncoderNVENC/EncoderNVENC.Build.cs
--- a/Engine/Plugins/Media/HardwareEncoders/Source/EncoderNVENC/EncoderNVENC.Build.cs
+++ b/Engine/Plugins/Media/HardwareEncoders/Source/EncoderNVENC/EncoderNVENC.Build.cs
@@ -37,15 +37,6 @@ public class EncoderNVENC: ModuleRules
 
 		if (Target.IsInPlatformGroup(UnrealPlatformGroup.Windows))
 		{
-			PublicDependencyModuleNames.Add("D3D11RHI");
-			PublicDependencyModuleNames.Add("D3D12RHI");
-
-			PublicSystemLibraries.AddRange(new string[] {
-				"DXGI.lib",
-				"d3d11.lib",
-				"d3d12.lib"
-			});
-
 			PrivateIncludePaths.Add(Path.Combine(EngineSourceDirectory, "Source/Runtime/VulkanRHI/Private/Windows"));
 		}
 		else if (Target.IsInPlatformGroup(UnrealPlatformGroup.Unix))
diff --git a/Engine/Plugins/Media/ImgMedia/Source/ImgMedia/ImgMedia.Build.cs b/Engine/Plugins/Media/ImgMedia/Source/ImgMedia/ImgMedia.Build.cs
--- a/Engine/Plugins/Media/ImgMedia/Source/ImgMedia/ImgMedia.Build.cs
+++ b/Engine/Plugins/Media/ImgMedia/Source/ImgMedia/ImgMedia.Build.cs
@@ -27,7 +27,6 @@ namespace UnrealBuildTool.Rules
 					"RenderCore",
 					"Renderer",
 					"RHI",
-					"RHICore",
 				});
 
 			PrivateIncludePathModuleNames.AddRange(
@@ -69,19 +68,9 @@ namespace UnrealBuildTool.Rules
 			{
 				AddEngineThirdPartyPrivateStaticDependencies(Target, "DX12");
 				PrivateDependencyModuleNames.Add("D3D12RHI");
-				PrivateIncludePaths.AddRange(
-					new string[]{
-						//required for FD3D12GPUFence
-						Path.Combine(EngineDir, "Source/Runtime/D3D12RHI/Private"),
-						Path.Combine(EngineDir, "Source/Runtime/D3D12RHI/Private/Windows")
-					});
 			}
 			
-			PrivateIncludePaths.AddRange(
-				new string[] {
-					//required for FPostProcessMaterialInputs
-					Path.Combine(EngineDir, "Source/Runtime/Renderer/Private")
-				});
+			PrivateIncludePaths.Add(Path.Combine(EngineDir, "Source", "Runtime", "Renderer", "Private"));
 
 			// Is this the editor?
 			if (Target.bBuildEditor == true)
diff --git a/Engine/Plugins/Media/ImgMedia/Source/ImgMedia/Private/Readers/ExrImgMediaReaderGpu.cpp b/Engine/Plugins/Media/ImgMedia/Source/ImgMedia/Private/Readers/ExrImgMediaReaderGpu.cpp
--- a/Engine/Plugins/Media/ImgMedia/Source/ImgMedia/Private/Readers/ExrImgMediaReaderGpu.cpp
+++ b/Engine/Plugins/Media/ImgMedia/Source/ImgMedia/Private/Readers/ExrImgMediaReaderGpu.cpp
@@ -25,14 +25,7 @@
 #include "TextureResource.h"
 #include "ScreenPass.h"
 
-#include "D3D12RHIPrivate.h"
-#include "D3D12Util.h"
-#include "D3D12State.h"
-#include "D3D12Resources.h"
-#include "D3D12RHICommon.h"
-#include "D3D12Viewport.h"
-#include "D3D12ConstantBuffer.h"
-#include "D3D12DirectCommandListManager.h"
+#include "ID3D12DynamicRHI.h"
 
 #define READ_IN_CHUNKS 1
 #define EXR_ENABLE_MIPS 1
diff --git a/Engine/Plugins/Media/PixelStreaming/Source/PixelStreaming/PixelStreaming.Build.cs b/Engine/Plugins/Media/PixelStreaming/Source/PixelStreaming/PixelStreaming.Build.cs
--- a/Engine/Plugins/Media/PixelStreaming/Source/PixelStreaming/PixelStreaming.Build.cs
+++ b/Engine/Plugins/Media/PixelStreaming/Source/PixelStreaming/PixelStreaming.Build.cs
@@ -68,7 +68,6 @@ namespace UnrealBuildTool.Rules
 				"Json",
 				"RenderCore",
 				"RHI",
-				"RHICore",
 				"Slate",
 				"SlateCore",
 				"AudioMixer",
@@ -93,19 +92,10 @@ namespace UnrealBuildTool.Rules
 
 			if (Target.IsInPlatformGroup(UnrealPlatformGroup.Windows))
 			{
-				PublicDependencyModuleNames.Add("D3D11RHI");
-				PublicDependencyModuleNames.Add("D3D12RHI");
-				PrivateIncludePaths.AddRange(
-					new string[]{
-						Path.Combine(EngineDir, "Source/Runtime/D3D12RHI/Private"),
-						Path.Combine(EngineDir, "Source/Runtime/D3D12RHI/Private/Windows")
-					});
-
-				AddEngineThirdPartyPrivateStaticDependencies(Target, "DX12");
-				PublicSystemLibraries.AddRange(new string[] {
-					"DXGI.lib",
-					"d3d11.lib",
-				});
+				PrivateDependencyModuleNames.Add("D3D11RHI");
+				PrivateDependencyModuleNames.Add("D3D12RHI");
+
+				AddEngineThirdPartyPrivateStaticDependencies(Target, "DX11", "DX12");
 
 				PrivateIncludePaths.Add(Path.Combine(EngineDir, "Source/Runtime/VulkanRHI/Private/Windows"));
 			}
diff --git a/Engine/Plugins/Media/PixelStreaming/Source/PixelStreaming/Private/EncoderFrameFactory.cpp b/Engine/Plugins/Media/PixelStreaming/Source/PixelStreaming/Private/EncoderFrameFactory.cpp
--- a/Engine/Plugins/Media/PixelStreaming/Source/PixelStreaming/Private/EncoderFrameFactory.cpp
+++ b/Engine/Plugins/Media/PixelStreaming/Source/PixelStreaming/Private/EncoderFrameFactory.cpp
@@ -12,10 +12,7 @@
 	#include "VideoCommon.h"
 	#include "D3D11State.h"
 	#include "D3D11Resources.h"
-	#include "D3D12RHICommon.h"
-	#include "D3D12RHIPrivate.h"
-	#include "D3D12Resources.h"
-	#include "D3D12Texture.h"
+	#include "ID3D12DynamicRHI.h"
 	#include "Windows/AllowWindowsPlatformTypes.h"
 THIRD_PARTY_INCLUDES_START
 	#include <VersionHelpers.h>
@@ -486,13 +483,12 @@ void UE::PixelStreaming::FEncoderFrameFactory::SetTextureCUDAD3D11(TSharedPtr<AV
 
 void UE::PixelStreaming::FEncoderFrameFactory::SetTextureCUDAD3D12(TSharedPtr<AVEncoder::FVideoEncoderInputFrame> InputFrame, const FTexture2DRHIRef& Texture)
 {
-	FD3D12TextureBase* D3D12Texture = GetD3D12TextureFromRHITexture(Texture);
-	ID3D12Resource* NativeD3D12Resource = (ID3D12Resource*)Texture->GetNativeResource();
-	unsigned long long TextureMemorySize = D3D12Texture->GetMemorySize();
+	ID3D12Resource* NativeD3D12Resource = GetID3D12DynamicRHI()->RHIGetResource(Texture);
+	const int64 TextureMemorySize = GetID3D12DynamicRHI()->RHIGetResourceMemorySize(Texture);
 
 	// Because we create our texture as RenderTargetable, it is created as a committed resource, which is what our current implementation here supports.
 	// To prevent a mystery crash in future, check that our resource is a committed resource
-	check(!D3D12Texture->GetResource()->IsPlacedResource());
+	check(!GetID3D12DynamicRHI()->RHIIsResourcePlaced(Texture));
 
 	TRefCountPtr<ID3D12Device> OwnerDevice;
 	HRESULT QueryResult;
diff --git a/Engine/Plugins/Media/WmfMedia/Source/WmfMedia/Private/Wmf/WmfMediaHardwareVideoDecodingRendering.cpp b/Engine/Plugins/Media/WmfMedia/Source/WmfMedia/Private/Wmf/WmfMediaHardwareVideoDecodingRendering.cpp
--- a/Engine/Plugins/Media/WmfMedia/Source/WmfMedia/Private/Wmf/WmfMediaHardwareVideoDecodingRendering.cpp
+++ b/Engine/Plugins/Media/WmfMedia/Source/WmfMedia/Private/Wmf/WmfMediaHardwareVideoDecodingRendering.cpp
@@ -16,7 +16,7 @@
 #include "ShaderParameterUtils.h"
 #include "ProfilingDebugging/RealtimeGPUProfiler.h"
 
-#include "D3D11RHIPrivate.h"
+#include "ID3D11DynamicRHI.h"
 #include "DynamicRHI.h"
 
 #include "WmfMediaHardwareVideoDecodingTextureSample.h"
@@ -47,13 +47,12 @@ struct FRHICommandCopyResource final : public FRHICommand<FRHICommandCopyResourc
 	void Execute(FRHICommandListBase& CmdList)
 	{
 		LLM_SCOPE(ELLMTag::MediaStreaming);
-		ID3D11Device* D3D11Device = static_cast<ID3D11Device*>(GDynamicRHI->RHIGetNativeDevice());
-		ID3D11DeviceContext* D3D11DeviceContext = nullptr;
+		ID3D11Device* D3D11Device = GetID3D11DynamicRHI()->RHIGetDevice();
+		ID3D11DeviceContext* D3D11DeviceContext = GetID3D11DynamicRHI()->RHIGetDeviceContext();
 
-		D3D11Device->GetImmediateContext(&D3D11DeviceContext);
 		if (D3D11DeviceContext)
 		{
-			ID3D11Resource* DestinationTexture = reinterpret_cast<ID3D11Resource*>(SampleDestinationTexture->GetNativeResource());
+			ID3D11Resource* DestinationTexture = GetID3D11DynamicRHI()->RHIGetResource(SampleDestinationTexture);
 			if (DestinationTexture)
 			{
 				TComPtr<IDXGIResource> OtherResource(nullptr);
@@ -101,7 +100,6 @@ struct FRHICommandCopyResource final : public FRHICommand<FRHICommandCopyResourc
 					}
 				}
 			}
-			D3D11DeviceContext->Release();
 		}
 	}
 };
@@ -121,18 +119,15 @@ bool FWmfMediaHardwareVideoDecodingParameters::ConvertTextureFormat_RenderThread
 	TComPtr<ID3D11Texture2D> SampleTexture = InSample->GetSourceTexture();
 
 	ID3D11Device* D3D11Device = nullptr;
+	// Must access rendering device context to copy shared resource.
+	ID3D11DeviceContext* D3D11DeviceContext = nullptr;
+
 	if (InSample->IsBufferExternal() == false)
 	{
-		D3D11Device = static_cast<ID3D11Device*>(GDynamicRHI->RHIGetNativeDevice());
+		D3D11Device = GetID3D11DynamicRHI()->RHIGetDevice();
+		D3D11DeviceContext = GetID3D11DynamicRHI()->RHIGetDeviceContext();
 	}
-	ID3D11DeviceContext* D3D11DeviceContext = nullptr;
 			
-	// Must access rendering device context to copy shared resource.
-	if (D3D11Device != nullptr)
-	{
-		D3D11Device->GetImmediateContext(&D3D11DeviceContext);
-	}
-	
 	{
 		FRHICommandListImmediate& RHICmdList = FRHICommandListExecutor::GetImmediateCommandList();
 
@@ -275,11 +270,6 @@ bool FWmfMediaHardwareVideoDecodingParameters::ConvertTextureFormat_RenderThread
 
 		RHICmdList.DrawPrimitive(0, 2, 1);
 		RHICmdList.EndRenderPass();
-
-		if (D3D11DeviceContext != nullptr)
-		{
-			D3D11DeviceContext->Release();
-		}
 	}
 
 	return true;
diff --git a/Engine/Plugins/Media/WmfMedia/Source/WmfMedia/WmfMedia.Build.cs b/Engine/Plugins/Media/WmfMedia/Source/WmfMedia/WmfMedia.Build.cs
--- a/Engine/Plugins/Media/WmfMedia/Source/WmfMedia/WmfMedia.Build.cs
+++ b/Engine/Plugins/Media/WmfMedia/Source/WmfMedia/WmfMedia.Build.cs
@@ -27,7 +27,6 @@ namespace UnrealBuildTool.Rules
 					"Projects",
 					"RenderCore",
                     "RHI",
-					"RHICore",
                     "WmfMediaFactory"
                 });
 
@@ -42,14 +41,8 @@ namespace UnrealBuildTool.Rules
 					"WmfMedia/Private",
 					"WmfMedia/Private/Player",
 					"WmfMedia/Private/Wmf",
-                    "../../../../Source/Runtime/Windows/D3D11RHI/Private",
-                    "../../../../Source/Runtime/Windows/D3D11RHI/Private/Windows",
                 });
 
-            AddEngineThirdPartyPrivateStaticDependencies(Target, "IntelMetricsDiscovery");
-            AddEngineThirdPartyPrivateStaticDependencies(Target, "IntelExtensionsFramework");
-            AddEngineThirdPartyPrivateStaticDependencies(Target, "NVAftermath");
-
 			if (Target.bCompileAgainstEngine)
 			{
 				PrivateDependencyModuleNames.Add("Engine");
@@ -58,17 +51,13 @@ namespace UnrealBuildTool.Rules
 
 			if (Target.Platform == UnrealTargetPlatform.Win64)
 			{
-				PrivateDependencyModuleNames.AddRange(
-				  new string[] {
-					"D3D11RHI",
-				});
+				PrivateDependencyModuleNames.Add("D3D11RHI");
+				AddEngineThirdPartyPrivateStaticDependencies(Target, "DX11");
 
 				PublicDelayLoadDLLs.Add("mf.dll");
 				PublicDelayLoadDLLs.Add("mfplat.dll");
 				PublicDelayLoadDLLs.Add("mfplay.dll");
 				PublicDelayLoadDLLs.Add("shlwapi.dll");
-
-				PublicSystemLibraries.Add("dxgi.lib");
 			}
 			
 			PublicDefinitions.Add("WMFMEDIA_PLAYER_VERSION=2");
diff --git a/Engine/Plugins/Runtime/Oculus/OculusVR/Source/OculusHMD/OculusHMD.Build.cs b/Engine/Plugins/Runtime/Oculus/OculusVR/Source/OculusHMD/OculusHMD.Build.cs
--- a/Engine/Plugins/Runtime/Oculus/OculusVR/Source/OculusHMD/OculusHMD.Build.cs
+++ b/Engine/Plugins/Runtime/Oculus/OculusVR/Source/OculusHMD/OculusHMD.Build.cs
@@ -42,7 +42,6 @@ namespace UnrealBuildTool.Rules
 					"Engine",
 					"InputCore",
 					"RHI",
-					"RHICore",
 					"RenderCore",
 					"Renderer",
 					"Slate",
@@ -76,31 +75,20 @@ namespace UnrealBuildTool.Rules
 			{
 				// D3D
 				{
-					PrivateDependencyModuleNames.AddRange(
-						new string[]
-						{
-							"D3D11RHI",
-							"D3D12RHI",
-						});
+					PrivateDependencyModuleNames.AddRange(new string[] {
+						"D3D11RHI",
+						"D3D12RHI",
+					});
 
 					PrivateIncludePaths.AddRange(
 						new string[]
 						{
 							"OculusMR/Public",
-							"../../../../../Source/Runtime/Windows/D3D11RHI/Private",
-							"../../../../../Source/Runtime/Windows/D3D11RHI/Private/Windows",
-							"../../../../../Source/Runtime/D3D12RHI/Private",
-							"../../../../../Source/Runtime/D3D12RHI/Private/Windows",
 						});
 
-					AddEngineThirdPartyPrivateStaticDependencies(Target, "DX11");
-					AddEngineThirdPartyPrivateStaticDependencies(Target, "DX12");
-					AddEngineThirdPartyPrivateStaticDependencies(Target, "NVAPI");
+					AddEngineThirdPartyPrivateStaticDependencies(Target, "DX11", "DX12");
 					AddEngineThirdPartyPrivateStaticDependencies(Target, "DX11Audio");
 					AddEngineThirdPartyPrivateStaticDependencies(Target, "DirectSound");
-					AddEngineThirdPartyPrivateStaticDependencies(Target, "NVAftermath");
-					AddEngineThirdPartyPrivateStaticDependencies(Target, "IntelMetricsDiscovery");
-					AddEngineThirdPartyPrivateStaticDependencies(Target, "IntelExtensionsFramework");
 				}
 
 				// Vulkan
diff --git a/Engine/Plugins/Runtime/Oculus/OculusVR/Source/OculusHMD/Private/OculusHMDPrivateRHI.h b/Engine/Plugins/Runtime/Oculus/OculusVR/Source/OculusHMD/Private/OculusHMDPrivateRHI.h
--- a/Engine/Plugins/Runtime/Oculus/OculusVR/Source/OculusHMD/Private/OculusHMDPrivateRHI.h
+++ b/Engine/Plugins/Runtime/Oculus/OculusVR/Source/OculusHMD/Private/OculusHMDPrivateRHI.h
@@ -9,8 +9,7 @@
 //-------------------------------------------------------------------------------------------------
 
 #if OCULUS_HMD_SUPPORTED_PLATFORMS_D3D11
-#include "D3D11RHIPrivate.h"
-#include "D3D11Util.h"
+#include "ID3D11DynamicRHI.h"
 #endif // OCULUS_HMD_SUPPORTED_PLATFORMS_D3D11
 
 
@@ -19,8 +18,7 @@
 //-------------------------------------------------------------------------------------------------
 
 #if OCULUS_HMD_SUPPORTED_PLATFORMS_D3D12
-#include "D3D12RHIPrivate.h"
-#include "D3D12Util.h"
+#include "ID3D12DynamicRHI.h"
 #endif // OCULUS_HMD_SUPPORTED_PLATFORMS_D3D12
 
 
diff --git a/Engine/Plugins/Runtime/Oculus/OculusVR/Source/OculusHMD/Private/OculusHMD_CustomPresent_D3D11.cpp b/Engine/Plugins/Runtime/Oculus/OculusVR/Source/OculusHMD/Private/OculusHMD_CustomPresent_D3D11.cpp
--- a/Engine/Plugins/Runtime/Oculus/OculusVR/Source/OculusHMD/Private/OculusHMD_CustomPresent_D3D11.cpp
+++ b/Engine/Plugins/Runtime/Oculus/OculusVR/Source/OculusHMD/Private/OculusHMD_CustomPresent_D3D11.cpp
@@ -82,7 +82,7 @@ bool FD3D11CustomPresent::IsUsingCorrectDisplayAdapter() const
 
 void* FD3D11CustomPresent::GetOvrpDevice() const
 {
-	return GD3D11RHI->GetDevice();
+	return GetID3D11DynamicRHI()->RHIGetDevice();
 }
 
 
@@ -93,13 +93,13 @@ FTextureRHIRef FD3D11CustomPresent::CreateTexture_RenderThread(uint32 InSizeX, u
 	switch (InResourceType)
 	{
 	case RRT_Texture2D:
-		return GD3D11RHI->RHICreateTexture2DFromResource(InFormat, InTexCreateFlags, InBinding, (ID3D11Texture2D*) InTexture).GetReference();
+		return GetID3D11DynamicRHI()->RHICreateTexture2DFromResource(InFormat, InTexCreateFlags, InBinding, (ID3D11Texture2D*) InTexture).GetReference();
 
 	case RRT_Texture2DArray:
-		return GD3D11RHI->RHICreateTexture2DArrayFromResource(InFormat, InTexCreateFlags, InBinding, (ID3D11Texture2D*)InTexture).GetReference();
+		return GetID3D11DynamicRHI()->RHICreateTexture2DArrayFromResource(InFormat, InTexCreateFlags, InBinding, (ID3D11Texture2D*)InTexture).GetReference();
 
 	case RRT_TextureCube:
-		return GD3D11RHI->RHICreateTextureCubeFromResource(InFormat, InTexCreateFlags | TexCreate_TargetArraySlicesIndependently, InBinding, (ID3D11Texture2D*) InTexture).GetReference();
+		return GetID3D11DynamicRHI()->RHICreateTextureCubeFromResource(InFormat, InTexCreateFlags | TexCreate_TargetArraySlicesIndependently, InBinding, (ID3D11Texture2D*) InTexture).GetReference();
 
 	default:
 		return nullptr;
diff --git a/Engine/Plugins/Runtime/Oculus/OculusVR/Source/OculusHMD/Private/OculusHMD_CustomPresent_D3D12.cpp b/Engine/Plugins/Runtime/Oculus/OculusVR/Source/OculusHMD/Private/OculusHMD_CustomPresent_D3D12.cpp
--- a/Engine/Plugins/Runtime/Oculus/OculusVR/Source/OculusHMD/Private/OculusHMD_CustomPresent_D3D12.cpp
+++ b/Engine/Plugins/Runtime/Oculus/OculusVR/Source/OculusHMD/Private/OculusHMD_CustomPresent_D3D12.cpp
@@ -73,16 +73,14 @@ bool FD3D12CustomPresent::IsUsingCorrectDisplayAdapter() const
 
 void* FD3D12CustomPresent::GetOvrpDevice() const
 {
-	FD3D12DynamicRHI* DynamicRHI = GetDynamicRHI<FD3D12DynamicRHI>();
-	return DynamicRHI->RHIGetD3DCommandQueue();
+	return GetID3D12DynamicRHI()->RHIGetCommandQueue();
 }
 
-
 FTextureRHIRef FD3D12CustomPresent::CreateTexture_RenderThread(uint32 InSizeX, uint32 InSizeY, EPixelFormat InFormat, FClearValueBinding InBinding, uint32 InNumMips, uint32 InNumSamples, uint32 InNumSamplesTileMem, ERHIResourceType InResourceType, ovrpTextureHandle InTexture, ETextureCreateFlags InTexCreateFlags)
 {
 	CheckInRenderThread();
 
-	FD3D12DynamicRHI* DynamicRHI = GetDynamicRHI<FD3D12DynamicRHI>();
+	ID3D12DynamicRHI* DynamicRHI = GetID3D12DynamicRHI();
 
 	// Add TexCreate_Shared flag to indicate the textures are shared with DX11 and therefore its initial state is D3D12_RESOURCE_STATE_COMMON
 	switch (InResourceType)
diff --git a/Engine/Plugins/Runtime/Oculus/OculusVR/Source/OculusOpenXRHMD/OculusOpenXRHMD.Build.cs b/Engine/Plugins/Runtime/Oculus/OculusVR/Source/OculusOpenXRHMD/OculusOpenXRHMD.Build.cs
--- a/Engine/Plugins/Runtime/Oculus/OculusVR/Source/OculusOpenXRHMD/OculusOpenXRHMD.Build.cs
+++ b/Engine/Plugins/Runtime/Oculus/OculusVR/Source/OculusOpenXRHMD/OculusOpenXRHMD.Build.cs
@@ -43,7 +43,6 @@ namespace UnrealBuildTool.Rules
                     "Engine",
                     "InputCore",
                     "RHI",
-                    "RHICore",
                     "RenderCore",
                     "Renderer",
                     "Slate",
@@ -76,30 +75,14 @@ namespace UnrealBuildTool.Rules
             {
                 // D3D
                 {
-                    PrivateDependencyModuleNames.AddRange(
-                        new string[]
-                        {
-                            "D3D11RHI",
-                            "D3D12RHI",
-                        });
-
-                    PrivateIncludePaths.AddRange(
-                        new string[]
-                        {
-                            "../../../../../Source/Runtime/Windows/D3D11RHI/Private",
-                            "../../../../../Source/Runtime/Windows/D3D11RHI/Private/Windows",
-                            "../../../../../Source/Runtime/D3D12RHI/Private",
-                            "../../../../../Source/Runtime/D3D12RHI/Private/Windows",
-                        });
+                    PrivateDependencyModuleNames.AddRange(new string[] {
+						"D3D11RHI",
+						"D3D12RHI",
+					});
 
-                    AddEngineThirdPartyPrivateStaticDependencies(Target, "DX11");
-                    AddEngineThirdPartyPrivateStaticDependencies(Target, "DX12");
-                    AddEngineThirdPartyPrivateStaticDependencies(Target, "NVAPI");
+                    AddEngineThirdPartyPrivateStaticDependencies(Target, "DX11", "DX12");
                     AddEngineThirdPartyPrivateStaticDependencies(Target, "DX11Audio");
                     AddEngineThirdPartyPrivateStaticDependencies(Target, "DirectSound");
-                    AddEngineThirdPartyPrivateStaticDependencies(Target, "NVAftermath");
-                    AddEngineThirdPartyPrivateStaticDependencies(Target, "IntelMetricsDiscovery");
-                    AddEngineThirdPartyPrivateStaticDependencies(Target, "IntelExtensionsFramework");
                 }
 
                 // Vulkan
diff --git a/Engine/Plugins/Runtime/OpenXR/Source/OpenXRHMD/OpenXRHMD.Build.cs b/Engine/Plugins/Runtime/OpenXR/Source/OpenXRHMD/OpenXRHMD.Build.cs
--- a/Engine/Plugins/Runtime/OpenXR/Source/OpenXRHMD/OpenXRHMD.Build.cs
+++ b/Engine/Plugins/Runtime/OpenXR/Source/OpenXRHMD/OpenXRHMD.Build.cs
@@ -44,7 +44,6 @@ namespace UnrealBuildTool.Rules
                     "BuildSettings",
                     "InputCore",
 					"RHI",
-					"RHICore",
 					"RenderCore",
 					"Renderer",
 					"RenderCore",
@@ -68,21 +67,7 @@ namespace UnrealBuildTool.Rules
 					"D3D12RHI"
 				});
 
-                // Required for some private headers needed for the rendering support.
-                PrivateIncludePaths.AddRange(
-                    new string[] {
-                            Path.Combine(EngineDir, @"Source\Runtime\Windows\D3D11RHI\Private"),
-                            Path.Combine(EngineDir, @"Source\Runtime\Windows\D3D11RHI\Private\Windows"),
-							Path.Combine(EngineDir, @"Source\Runtime\D3D12RHI\Private"),
-							Path.Combine(EngineDir, @"Source\Runtime\D3D12RHI\Private\Windows")
-								});
-
-                AddEngineThirdPartyPrivateStaticDependencies(Target, "DX11");
-				AddEngineThirdPartyPrivateStaticDependencies(Target, "DX12");
-				AddEngineThirdPartyPrivateStaticDependencies(Target, "NVAPI");
-                AddEngineThirdPartyPrivateStaticDependencies(Target, "AMD_AGS");
-				AddEngineThirdPartyPrivateStaticDependencies(Target, "IntelMetricsDiscovery");
-				AddEngineThirdPartyPrivateStaticDependencies(Target, "IntelExtensionsFramework");
+                AddEngineThirdPartyPrivateStaticDependencies(Target, "DX11", "DX12");
             }
 
 			if (Target.IsInPlatformGroup(UnrealPlatformGroup.Windows) || Target.IsInPlatformGroup(UnrealPlatformGroup.Linux))
diff --git a/Engine/Plugins/Runtime/OpenXR/Source/OpenXRHMD/Private/OpenXRHMD_RenderBridge.cpp b/Engine/Plugins/Runtime/OpenXR/Source/OpenXRHMD/Private/OpenXRHMD_RenderBridge.cpp
--- a/Engine/Plugins/Runtime/OpenXR/Source/OpenXRHMD/Private/OpenXRHMD_RenderBridge.cpp
+++ b/Engine/Plugins/Runtime/OpenXR/Source/OpenXRHMD/Private/OpenXRHMD_RenderBridge.cpp
@@ -41,11 +41,9 @@ public:
 
 	virtual void* GetGraphicsBinding() override
 	{
-		FD3D11Device* Device = GD3D11RHI->GetDevice();
-
 		Binding.type = XR_TYPE_GRAPHICS_BINDING_D3D11_KHR;
 		Binding.next = nullptr;
-		Binding.device = Device;
+		Binding.device = GetID3D11DynamicRHI()->RHIGetDevice();
 		return &Binding;
 	}
 
@@ -82,14 +80,10 @@ public:
 
 	virtual void* GetGraphicsBinding() override
 	{
-		FD3D12DynamicRHI* DynamicRHI = GetDynamicRHI<FD3D12DynamicRHI>();
-		ID3D12Device* Device = DynamicRHI->GetAdapter().GetD3DDevice();
-		ID3D12CommandQueue* Queue = DynamicRHI->RHIGetD3DCommandQueue();
-
 		Binding.type = XR_TYPE_GRAPHICS_BINDING_D3D12_KHR;
 		Binding.next = nullptr;
-		Binding.device = Device;
-		Binding.queue = Queue;
+		Binding.device = GetID3D12DynamicRHI()->RHIGetDevice(0);
+		Binding.queue = GetID3D12DynamicRHI()->RHIGetCommandQueue();
 		return &Binding;
 	}
 
diff --git a/Engine/Plugins/Runtime/OpenXR/Source/OpenXRHMD/Private/OpenXRHMD_Swapchain.cpp b/Engine/Plugins/Runtime/OpenXR/Source/OpenXRHMD/Private/OpenXRHMD_Swapchain.cpp
--- a/Engine/Plugins/Runtime/OpenXR/Source/OpenXRHMD/Private/OpenXRHMD_Swapchain.cpp
+++ b/Engine/Plugins/Runtime/OpenXR/Source/OpenXRHMD/Private/OpenXRHMD_Swapchain.cpp
@@ -261,13 +261,7 @@ FXRSwapChainPtr CreateSwapchain_D3D11(XrSession InSession, uint8 Format, uint32
 {
 	TFunction<uint32(uint8)> ToPlatformFormat = [](uint8 InFormat)
 	{
-		// We need to convert typeless to typed formats to create a swapchain
-		DXGI_FORMAT PlatformFormat = (DXGI_FORMAT)GPixelFormats[InFormat].PlatformFormat;
-		PlatformFormat = ::FindDepthStencilDXGIFormat(PlatformFormat);
-
-		// UE renders a gamma-corrected image so we need to use an sRGB format if available
-		PlatformFormat = ::FindShaderResourceDXGIFormat(PlatformFormat, true);
-		return PlatformFormat;
+		return GetID3D11DynamicRHI()->RHIGetSwapChainFormat(static_cast<EPixelFormat>(InFormat));
 	};
 
 	Format = FOpenXRSwapchain::GetNearestSupportedSwapchainFormat(InSession, Format, ToPlatformFormat);
@@ -282,13 +276,15 @@ FXRSwapChainPtr CreateSwapchain_D3D11(XrSession InSession, uint8 Format, uint32
 		return nullptr;
 	}
 
+	ID3D11DynamicRHI* D3D11RHI = GetID3D11DynamicRHI();
+
 	TArray<FTextureRHIRef> TextureChain;
 	TArray<XrSwapchainImageD3D11KHR> Images = EnumerateImages<XrSwapchainImageD3D11KHR>(Swapchain, XR_TYPE_SWAPCHAIN_IMAGE_D3D11_KHR);
 	for (const auto& Image : Images)
 	{
-		TextureChain.Add(static_cast<FTextureRHIRef>(GD3D11RHI->RHICreateTexture2DArrayFromResource(GPixelFormats[Format].UnrealFormat, CreateFlags, ClearValueBinding, Image.texture)));
+		TextureChain.Add(static_cast<FTextureRHIRef>(D3D11RHI->RHICreateTexture2DArrayFromResource(GPixelFormats[Format].UnrealFormat, CreateFlags, ClearValueBinding, Image.texture)));
 	}
-	FTextureRHIRef ChainTarget = static_cast<FTextureRHIRef>(GDynamicRHI->RHICreateAliasedTexture((FTextureRHIRef&)TextureChain[0]));
+	FTextureRHIRef ChainTarget = static_cast<FTextureRHIRef>(D3D11RHI->RHICreateAliasedTexture((FTextureRHIRef&)TextureChain[0]));
 
 	return CreateXRSwapChain<FOpenXRSwapchain>(MoveTemp(TextureChain), ChainTarget, Swapchain);
 }
@@ -299,13 +295,7 @@ FXRSwapChainPtr CreateSwapchain_D3D12(XrSession InSession, uint8 Format, uint32
 {
 	TFunction<uint32(uint8)> ToPlatformFormat = [](uint8 InFormat)
 	{
-		// We need to convert typeless to typed formats to create a swapchain
-		DXGI_FORMAT PlatformFormat = (DXGI_FORMAT)GPixelFormats[InFormat].PlatformFormat;
-		PlatformFormat = D3D12RHI::FindDepthStencilDXGIFormat(PlatformFormat);
-
-		// UE renders a gamma-corrected image so we need to use an sRGB format if available
-		PlatformFormat = D3D12RHI::FindShaderResourceDXGIFormat(PlatformFormat, true);
-		return PlatformFormat;
+		return GetID3D12DynamicRHI()->RHIGetSwapChainFormat(static_cast<EPixelFormat>(InFormat));
 	};
 
 	Format = FOpenXRSwapchain::GetNearestSupportedSwapchainFormat(InSession, Format, ToPlatformFormat);
@@ -320,14 +310,14 @@ FXRSwapChainPtr CreateSwapchain_D3D12(XrSession InSession, uint8 Format, uint32
 		return nullptr;
 	}
 
-	FD3D12DynamicRHI* DynamicRHI = GetDynamicRHI<FD3D12DynamicRHI>();
+	ID3D12DynamicRHI* DynamicRHI = GetID3D12DynamicRHI();
 	TArray<FTextureRHIRef> TextureChain;
 	TArray<XrSwapchainImageD3D12KHR> Images = EnumerateImages<XrSwapchainImageD3D12KHR>(Swapchain, XR_TYPE_SWAPCHAIN_IMAGE_D3D12_KHR);
 	for (const auto& Image : Images)
 	{
 		TextureChain.Add(static_cast<FTextureRHIRef>(DynamicRHI->RHICreateTexture2DArrayFromResource(GPixelFormats[Format].UnrealFormat, CreateFlags, ClearValueBinding, Image.texture)));
 	}
-	FTextureRHIRef ChainTarget = static_cast<FTextureRHIRef>(GDynamicRHI->RHICreateAliasedTexture((FTextureRHIRef&)TextureChain[0]));
+	FTextureRHIRef ChainTarget = static_cast<FTextureRHIRef>(DynamicRHI->RHICreateAliasedTexture((FTextureRHIRef&)TextureChain[0]));
 
 	return CreateXRSwapChain<FOpenXRSwapchain>(MoveTemp(TextureChain), ChainTarget, Swapchain);
 }
diff --git a/Engine/Plugins/Runtime/OpenXR/Source/OpenXRHMD/Private/OpenXRPlatformRHI.h b/Engine/Plugins/Runtime/OpenXR/Source/OpenXRHMD/Private/OpenXRPlatformRHI.h
--- a/Engine/Plugins/Runtime/OpenXR/Source/OpenXRHMD/Private/OpenXRPlatformRHI.h
+++ b/Engine/Plugins/Runtime/OpenXR/Source/OpenXRHMD/Private/OpenXRPlatformRHI.h
@@ -26,8 +26,7 @@
 //-------------------------------------------------------------------------------------------------
 
 #ifdef XR_USE_GRAPHICS_API_D3D11
-#include "D3D11RHIPrivate.h"
-#include "D3D11Util.h"
+#include "ID3D11DynamicRHI.h"
 #endif // XR_USE_GRAPHICS_API_D3D11
 
 
@@ -36,24 +35,7 @@
 //-------------------------------------------------------------------------------------------------
 
 #ifdef XR_USE_GRAPHICS_API_D3D12
-#include "D3D12RHIPrivate.h"
-#include "D3D12Util.h"
-
-#undef GetD3D11CubeFace
-#undef VerifyD3D11Result
-#undef GetD3D11TextureFromRHITexture
-#undef FRingAllocation
-#undef GetRenderTargetFormat
-#undef ED3D11ShaderOffsetBuffer
-#undef FindShaderResourceDXGIFormat
-#undef FindUnorderedAccessDXGIFormat
-#undef FindDepthStencilDXGIFormat
-#undef HasStencilBits
-#undef FVector4VertexDeclaration
-#undef GLOBAL_CONSTANT_BUFFER_INDEX
-#undef MAX_CONSTANT_BUFFER_SLOTS
-#undef FD3DGPUProfiler
-#undef FRangeAllocator
+#include "ID3D12DynamicRHI.h"
 #endif // XR_USE_GRAPHICS_API_D3D12
 
 
diff --git a/Engine/Plugins/Runtime/Steam/SteamVR/Source/SteamVR/Private/SteamVRRender.cpp b/Engine/Plugins/Runtime/Steam/SteamVR/Source/SteamVR/Private/SteamVRRender.cpp
--- a/Engine/Plugins/Runtime/Steam/SteamVR/Source/SteamVR/Private/SteamVRRender.cpp
+++ b/Engine/Plugins/Runtime/Steam/SteamVR/Source/SteamVR/Private/SteamVRRender.cpp
@@ -24,7 +24,7 @@
 #endif
 
 #if PLATFORM_WINDOWS
-#include "D3D12RHIPrivate.h"
+#include "ID3D12DynamicRHI.h"
 #endif
 
 static TAutoConsoleVariable<int32> CUsePostPresentHandoff(TEXT("vr.SteamVR.UsePostPresentHandoff"), 0, TEXT("Whether or not to use PostPresentHandoff.  If true, more GPU time will be available, but this relies on no SceneCaptureComponent2D or WidgetComponents being active in the scene.  Otherwise, it will break async reprojection."));
@@ -261,13 +261,13 @@ void FSteamVRHMD::D3D12Bridge::FinishRendering()
 	bool bSubmitDepth = CVarEnableDepthSubmission->GetInt() > 0;
 	vr::EVRSubmitFlags Flags = bSubmitDepth ? vr::EVRSubmitFlags::Submit_TextureWithDepth : vr::EVRSubmitFlags::Submit_Default;
 
-	FD3D12DynamicRHI* D3D12RHI = GetDynamicRHI<FD3D12DynamicRHI>();
-	FD3D12Device* Device = D3D12RHI->GetAdapter().GetDevice(0);
+	ID3D12CommandQueue* CommandQueue = GetID3D12DynamicRHI()->RHIGetCommandQueue();
+	const uint32 DeviceNodeMask = GetID3D12DynamicRHI()->RHIGetDeviceNodeMask(0);
 
 	vr::D3D12TextureData_t TextureData;
 	TextureData.m_pResource = (ID3D12Resource*)SwapChain->GetTexture2D()->GetNativeResource();
-	TextureData.m_pCommandQueue = D3D12RHI->RHIGetD3DCommandQueue();
-	TextureData.m_nNodeMask = Device->GetGPUMask().GetNative();
+	TextureData.m_pCommandQueue = CommandQueue;
+	TextureData.m_nNodeMask = DeviceNodeMask;
 
 	vr::VRTextureWithDepth_t Texture;
 	Texture.handle = &TextureData;
@@ -281,8 +281,8 @@ void FSteamVRHMD::D3D12Bridge::FinishRendering()
 	{
 		// If this flag is false, the struct will be treated as a vr::Texture_t and these entries will be ignored - so we can skip this if not submitting depth
 		DepthTextureData.m_pResource = (ID3D12Resource*)DepthSwapChain->GetTexture2D()->GetNativeResource();
-		DepthTextureData.m_pCommandQueue = D3D12RHI->RHIGetD3DCommandQueue();
-		DepthTextureData.m_nNodeMask = Device->GetGPUMask().GetNative();
+		DepthTextureData.m_pCommandQueue = CommandQueue;
+		DepthTextureData.m_nNodeMask = DeviceNodeMask;
 
 		Texture.depth.handle = &DepthTextureData;
 
diff --git a/Engine/Plugins/Runtime/Steam/SteamVR/Source/SteamVR/SteamVR.Build.cs b/Engine/Plugins/Runtime/Steam/SteamVR/Source/SteamVR/SteamVR.Build.cs
--- a/Engine/Plugins/Runtime/Steam/SteamVR/Source/SteamVR/SteamVR.Build.cs
+++ b/Engine/Plugins/Runtime/Steam/SteamVR/Source/SteamVR/SteamVR.Build.cs
@@ -31,7 +31,6 @@ namespace UnrealBuildTool.Rules
 					"CoreUObject",
 					"Engine",
 					"RHI",
-					"RHICore",
 					"RenderCore",
 					"Renderer",
                     "InputCore",
@@ -62,16 +61,7 @@ namespace UnrealBuildTool.Rules
 							"D3D12RHI",
 					});
 
-				PrivateIncludePaths.AddRange(
-					new string[]
-					{
-							"../../../../../Source/Runtime/Windows/D3D11RHI/Private",
-							"../../../../../Source/Runtime/Windows/D3D11RHI/Private/Windows",
-							"../../../../../Source/Runtime/D3D12RHI/Private",
-							"../../../../../Source/Runtime/D3D12RHI/Private/Windows",
-					});
-
-				AddEngineThirdPartyPrivateStaticDependencies(Target, "DX12");
+				AddEngineThirdPartyPrivateStaticDependencies(Target, "DX11", "DX12");
 				AddEngineThirdPartyPrivateStaticDependencies(Target, "OpenVR");
 
 				AddEngineThirdPartyPrivateStaticDependencies(Target, "OpenGL");
diff --git a/Engine/Plugins/Runtime/TextureShare/Source/TextureShare/TextureShare.Build.cs b/Engine/Plugins/Runtime/TextureShare/Source/TextureShare/TextureShare.Build.cs
--- a/Engine/Plugins/Runtime/TextureShare/Source/TextureShare/TextureShare.Build.cs
+++ b/Engine/Plugins/Runtime/TextureShare/Source/TextureShare/TextureShare.Build.cs
@@ -29,7 +29,6 @@ public class TextureShare : ModuleRules
 				"Core",
 				"Engine",
 				"RHI",
-				"RHICore",
 				"Renderer",
 				"RenderCore",
 				"TextureShareCore"
diff --git a/Engine/Plugins/Runtime/TextureShare/Source/TextureShareD3D11/Private/TextureShareD3D11.cpp b/Engine/Plugins/Runtime/TextureShare/Source/TextureShareD3D11/Private/TextureShareD3D11.cpp
--- a/Engine/Plugins/Runtime/TextureShare/Source/TextureShareD3D11/Private/TextureShareD3D11.cpp
+++ b/Engine/Plugins/Runtime/TextureShare/Source/TextureShareD3D11/Private/TextureShareD3D11.cpp
@@ -3,7 +3,7 @@
 #include "TextureShareD3D11.h"
 #include "TextureShareD3D11Log.h"
 
-#include "D3D11RHIPrivate.h"
+#include "ID3D11DynamicRHI.h"
 
 inline const FString GetComErrorDescription(HRESULT Res)
 {
@@ -53,11 +53,9 @@ inline bool CreateSharedHandle(ID3D11Device* pD3D11Device, ID3D11Texture2D* pD3D
 
 bool FTextureShareD3D11::CreateRHITexture(ID3D11Texture2D* OpenedSharedResource, EPixelFormat Format, FTexture2DRHIRef& DstTexture)
 {
-	FD3D11DynamicRHI* DynamicRHI = GetDynamicRHI<FD3D11DynamicRHI>();
-
 	// Create RHI texture from D3D11 resource
 	ETextureCreateFlags TexCreateFlags = TexCreate_Shared;
-	DstTexture = DynamicRHI->RHICreateTexture2DFromResource(Format, TexCreateFlags, FClearValueBinding::None, OpenedSharedResource).GetReference();
+	DstTexture = GetID3D11DynamicRHI()->RHICreateTexture2DFromResource(Format, TexCreateFlags, FClearValueBinding::None, OpenedSharedResource).GetReference();
 	return DstTexture.IsValid();
 }
 
diff --git a/Engine/Plugins/Runtime/TextureShare/Source/TextureShareD3D11/TextureShareD3D11.Build.cs b/Engine/Plugins/Runtime/TextureShare/Source/TextureShareD3D11/TextureShareD3D11.Build.cs
--- a/Engine/Plugins/Runtime/TextureShare/Source/TextureShareD3D11/TextureShareD3D11.Build.cs
+++ b/Engine/Plugins/Runtime/TextureShare/Source/TextureShareD3D11/TextureShareD3D11.Build.cs
@@ -1,29 +1,18 @@
 // Copyright Epic Games, Inc. All Rights Reserved.
 
 using UnrealBuildTool;
-using System.IO;
 
 public class TextureShareD3D11 : ModuleRules
 {
 	public TextureShareD3D11(ReadOnlyTargetRules Target) : base(Target)
 	{
-		PrivateDependencyModuleNames.AddRange(
-						new string[]
-						{
-							"Core",
-							"Engine",
-							"RHI",
-							"RHICore",
-							"RenderCore",
-							"D3D11RHI",
-						});
-
-		PrivateIncludePaths.AddRange(
-			new string[]
-			{
-						"../../../../Source/Runtime/Windows/D3D11RHI/Private",
-						"../../../../Source/Runtime/Windows/D3D11RHI/Private/Windows",
-			});
+		PrivateDependencyModuleNames.AddRange(new string[] {
+			"Core",
+			"Engine",
+			"RHI",
+			"RenderCore",
+			"D3D11RHI",
+		});
 
 		///////////////////////////////////////////////////////////////
 		// Platform specific defines
@@ -37,12 +26,6 @@ public class TextureShareD3D11 : ModuleRules
 			Target.Platform == UnrealTargetPlatform.HoloLens)
 		{
 			AddEngineThirdPartyPrivateStaticDependencies(Target, "DX11");
-			if (Target.Platform != UnrealTargetPlatform.HoloLens)
-			{
-				AddEngineThirdPartyPrivateStaticDependencies(Target, "NVAftermath");
-				AddEngineThirdPartyPrivateStaticDependencies(Target, "IntelMetricsDiscovery");
-				AddEngineThirdPartyPrivateStaticDependencies(Target, "IntelExtensionsFramework");
-			}
 		}
 	}
 }
diff --git a/Engine/Plugins/Runtime/TextureShare/Source/TextureShareD3D12/Private/CrossGPUHeap/D3D12CrossGPUHeapFence.cpp b/Engine/Plugins/Runtime/TextureShare/Source/TextureShareD3D12/Private/CrossGPUHeap/D3D12CrossGPUHeapFence.cpp
--- a/Engine/Plugins/Runtime/TextureShare/Source/TextureShareD3D12/Private/CrossGPUHeap/D3D12CrossGPUHeapFence.cpp
+++ b/Engine/Plugins/Runtime/TextureShare/Source/TextureShareD3D12/Private/CrossGPUHeap/D3D12CrossGPUHeapFence.cpp
@@ -141,9 +141,7 @@ bool FD3D12CrossGPUFence::ImplSignal(FRHICommandListImmediate& RHICmdList, uint6
 	//@todo integrate into RHICmd. Now just flush
 	RHICmdList.ImmediateFlush(EImmediateFlushType::FlushRHIThreadFlushResources);
 
-	FD3D12DynamicRHI* DynamicRHI = GetDynamicRHI<FD3D12DynamicRHI>();
-
-	ID3D12CommandQueue* D3D12CommandQueue = DynamicRHI->RHIGetD3DCommandQueue();
+	ID3D12CommandQueue* D3D12CommandQueue = GetID3D12DynamicRHI()->RHIGetCommandQueue();
 
 	// Schedule a signal in the command queue
 	HRESULT hr = D3D12CommandQueue->Signal(GetFence(), NextFenceValue);
@@ -166,10 +164,10 @@ FString FD3D12CrossGPUFence::GetSharedHeapName(const HANDLE MasterProcessHandle)
 
 bool FD3D12CrossGPUFence::Create(const HANDLE LocalProcessHandle, HANDLE& OutFenceHandle, uint64 InitialValue)
 {
-	auto UE4D3DDevice = static_cast<ID3D12Device*>(GDynamicRHI->RHIGetNativeDevice());
+	ID3D12Device* PrimaryDevice = GetID3D12DynamicRHI()->RHIGetPrimaryDevice();
 
 	// Create fence for cross adapter resources
-	CHECK_HR_DEFAULT(UE4D3DDevice->CreateFence(
+	CHECK_HR_DEFAULT(PrimaryDevice->CreateFence(
 		InitialValue,
 		D3D12_FENCE_FLAG_SHARED | D3D12_FENCE_FLAG_SHARED_CROSS_ADAPTER,
 		IID_PPV_ARGS(Fence.GetInitReference())));
@@ -177,9 +175,9 @@ bool FD3D12CrossGPUFence::Create(const HANDLE LocalProcessHandle, HANDLE& OutFen
 	OutFenceHandle = nullptr;
 #if TE_USE_NAMEDFENCE
 	FString FenceHeapName = GetSharedHeapName(LocalProcessHandle);
-	CHECK_HR_DEFAULT(UE4D3DDevice->CreateSharedHandle(Fence, *Security, GENERIC_ALL, *FenceHeapName, &OutFenceHandle));
+	CHECK_HR_DEFAULT(PrimaryDevice->CreateSharedHandle(Fence, *Security, GENERIC_ALL, *FenceHeapName, &OutFenceHandle));
 #else
-	CHECK_HR_DEFAULT(UE4D3DDevice->CreateSharedHandle(Fence, *Security, GENERIC_ALL, nullptr, &OutFenceHandle));
+	CHECK_HR_DEFAULT(PrimaryDevice->CreateSharedHandle(Fence, *Security, GENERIC_ALL, nullptr, &OutFenceHandle));
 #endif
 
 	return true;
@@ -188,19 +186,17 @@ bool FD3D12CrossGPUFence::Create(const HANDLE LocalProcessHandle, HANDLE& OutFen
 
 bool FD3D12CrossGPUFence::Open(HANDLE FenceHandle)
 {
-	auto UE4D3DDevice = static_cast<ID3D12Device*>(GDynamicRHI->RHIGetNativeDevice());
-
-	FD3D12DynamicRHI* DynamicRHI = GetDynamicRHI<FD3D12DynamicRHI>();
+	ID3D12Device* PrimaryDevice = GetID3D12DynamicRHI()->RHIGetPrimaryDevice();
 
 	// Open shared handle on secondaryDevice device
 #if TE_USE_NAMEDFENCE
 	HANDLE NamedResourceHandle;
 	FString FenceHeapName = GetSharedHeapName(FenceProcessHandle);
-	CHECK_HR_DEFAULT(UE4D3DDevice->OpenSharedHandleByName(*FenceHeapName, GENERIC_ALL, &NamedResourceHandle));
-	CHECK_HR_DEFAULT(UE4D3DDevice->OpenSharedHandle(NamedResourceHandle, IID_PPV_ARGS(Fence.GetInitReference())));
+	CHECK_HR_DEFAULT(PrimaryDevice->OpenSharedHandleByName(*FenceHeapName, GENERIC_ALL, &NamedResourceHandle));
+	CHECK_HR_DEFAULT(PrimaryDevice->OpenSharedHandle(NamedResourceHandle, IID_PPV_ARGS(Fence.GetInitReference())));
 	CloseHandle(NamedResourceHandle);
 #else
-	CHECK_HR_DEFAULT(UE4D3DDevice->OpenSharedHandle(FenceHandle, IID_PPV_ARGS(Fence.GetInitReference())));
+	CHECK_HR_DEFAULT(PrimaryDevice->OpenSharedHandle(FenceHandle, IID_PPV_ARGS(Fence.GetInitReference())));
 	CloseHandle(FenceHandle);
 #endif
 
diff --git a/Engine/Plugins/Runtime/TextureShare/Source/TextureShareD3D12/Private/TextureShareD3D12.cpp b/Engine/Plugins/Runtime/TextureShare/Source/TextureShareD3D12/Private/TextureShareD3D12.cpp
--- a/Engine/Plugins/Runtime/TextureShare/Source/TextureShareD3D12/Private/TextureShareD3D12.cpp
+++ b/Engine/Plugins/Runtime/TextureShare/Source/TextureShareD3D12/Private/TextureShareD3D12.cpp
@@ -3,8 +3,7 @@
 #include "TextureShareD3D12.h"
 #include "TextureShareD3D12Log.h"
 
-#include "D3D12RHIPrivate.h"
-#include "D3D12Util.h"
+#include "ID3D12DynamicRHI.h"
 
 #if TEXTURESHARE_CROSSGPUHEAP
 #include "CrossGPUHeap/D3D12CrossGPUHeap.h"
@@ -31,10 +30,8 @@ inline bool CreateSharedHandle(ID3D12Device* pD3D12Device, ID3D12Resource* pD3D1
 
 bool FTextureShareD3D12::CreateRHITexture(ID3D12Resource* OpenedSharedResource, EPixelFormat Format, FTexture2DRHIRef& DstTexture)
 {
-	FD3D12DynamicRHI* DynamicRHI = GetDynamicRHI<FD3D12DynamicRHI>();
-
 	ETextureCreateFlags TexCreateFlags = TexCreate_Shared;
-	DstTexture = DynamicRHI->RHICreateTexture2DFromResource(Format, TexCreateFlags, FClearValueBinding::None, (ID3D12Resource*)OpenedSharedResource).GetReference();
+	DstTexture = GetID3D12DynamicRHI()->RHICreateTexture2DFromResource(Format, TexCreateFlags, FClearValueBinding::None, (ID3D12Resource*)OpenedSharedResource).GetReference();
 	return DstTexture.IsValid();
 }
 
@@ -47,13 +44,13 @@ bool FTextureShareD3D12::CreateSharedTexture(FIntPoint& Size, EPixelFormat Forma
 	OutRHITexture = RHICreateTexture2D(Size.X, Size.Y, Format, 1, 1, TexCreate_Shared | TexCreate_ResolveTargetable, CreateInfo);
 	if (OutRHITexture.IsValid() && OutRHITexture->IsValid())
 	{
-		auto UE4D3DDevice = static_cast<ID3D12Device*>(GDynamicRHI->RHIGetNativeDevice());
-		ID3D12Resource* ResolvedTexture = (ID3D12Resource*)OutRHITexture->GetTexture2D()->GetNativeResource();
+		ID3D12Device* D3DDevice = GetID3D12DynamicRHI()->RHIGetDevice(0);
+		ID3D12Resource* ResolvedTexture = GetID3D12DynamicRHI()->RHIGetResource(OutRHITexture);
 
 		// Create unique handle name
 		OutSharedHandleGuid = FGuid::NewGuid();
 
-		if (CreateSharedHandle(UE4D3DDevice, ResolvedTexture, UniqueHandleId, *SharedResourceSecurityAttributes, OutHandle))
+		if (CreateSharedHandle(D3DDevice, ResolvedTexture, UniqueHandleId, *SharedResourceSecurityAttributes, OutHandle))
 		{
 			return true;
 		}
diff --git a/Engine/Plugins/Runtime/TextureShare/Source/TextureShareD3D12/TextureShareD3D12.Build.cs b/Engine/Plugins/Runtime/TextureShare/Source/TextureShareD3D12/TextureShareD3D12.Build.cs
--- a/Engine/Plugins/Runtime/TextureShare/Source/TextureShareD3D12/TextureShareD3D12.Build.cs
+++ b/Engine/Plugins/Runtime/TextureShare/Source/TextureShareD3D12/TextureShareD3D12.Build.cs
@@ -1,7 +1,6 @@
 // Copyright Epic Games, Inc. All Rights Reserved.
 
 using UnrealBuildTool;
-using System.IO;
 
 public class TextureShareD3D12 : ModuleRules
 {
@@ -14,18 +13,9 @@ public class TextureShareD3D12 : ModuleRules
 							"Engine",
 							"RenderCore",
 							"RHI",
-							"RHICore",
 							"D3D12RHI",
 						});
 
-		PrivateIncludePaths.AddRange(
-			new string[]
-			{
-						"../../../../Source/Runtime/D3D12RHI/Private",
-						"../../../../Source/Runtime/D3D12RHI/Private/Windows",
-						"../../../../Source/ThirdParty/Windows/D3DX12/Include"
-			});
-
 		// Allow D3D12 Cross GPU Heap resource API (experimental)
 		PublicDefinitions.Add("TEXTURESHARE_CROSSGPUHEAP=0");
 
@@ -37,17 +27,9 @@ public class TextureShareD3D12 : ModuleRules
 			PrecompileForTargets = PrecompileTargetsType.None;
 		}
 
-		if (Target.Platform.IsInGroup(UnrealPlatformGroup.Windows) ||
-			Target.Platform == UnrealTargetPlatform.HoloLens)
+		if (Target.Platform.IsInGroup(UnrealPlatformGroup.Windows) || Target.Platform == UnrealTargetPlatform.HoloLens)
 		{
 			AddEngineThirdPartyPrivateStaticDependencies(Target, "DX12");
-			if (Target.Platform != UnrealTargetPlatform.HoloLens)
-			{
-				AddEngineThirdPartyPrivateStaticDependencies(Target, "NVAPI");
-				AddEngineThirdPartyPrivateStaticDependencies(Target, "AMD_AGS");
-				AddEngineThirdPartyPrivateStaticDependencies(Target, "NVAftermath");
-				AddEngineThirdPartyPrivateStaticDependencies(Target, "IntelMetricsDiscovery");
-			}
 		}
 	}
 }
diff --git a/Engine/Plugins/Runtime/WindowsMixedReality/Source/WindowsMixedRealityHMD/Private/WindowsMixedRealityCustomPresent.h b/Engine/Plugins/Runtime/WindowsMixedReality/Source/WindowsMixedRealityHMD/Private/WindowsMixedRealityCustomPresent.h
--- a/Engine/Plugins/Runtime/WindowsMixedReality/Source/WindowsMixedRealityHMD/Private/WindowsMixedRealityCustomPresent.h
+++ b/Engine/Plugins/Runtime/WindowsMixedReality/Source/WindowsMixedRealityHMD/Private/WindowsMixedRealityCustomPresent.h
@@ -5,8 +5,7 @@
 #include "WindowsMixedRealityPrecompiled.h"
 #include "SceneRendering.h"
 
-#include "D3D11RHIPrivate.h"
-#include "D3D11Util.h"
+#include "ID3D11DynamicRHI.h"
 
 #if WITH_WINDOWS_MIXED_REALITY
 #include "MixedRealityInterop.h"
diff --git a/Engine/Plugins/Runtime/WindowsMixedReality/Source/WindowsMixedRealityHMD/Private/WindowsMixedRealityPrecompiled.h b/Engine/Plugins/Runtime/WindowsMixedReality/Source/WindowsMixedRealityHMD/Private/WindowsMixedRealityPrecompiled.h
--- a/Engine/Plugins/Runtime/WindowsMixedReality/Source/WindowsMixedRealityHMD/Private/WindowsMixedRealityPrecompiled.h
+++ b/Engine/Plugins/Runtime/WindowsMixedReality/Source/WindowsMixedRealityHMD/Private/WindowsMixedRealityPrecompiled.h
@@ -21,7 +21,4 @@
 
 #include "SceneRendering.h"
 
-#pragma warning(disable:4668)  
-#include "D3D11RHIPrivate.h"
-#pragma warning(default:4668)
-#include "D3D11Util.h"
+#include "ID3D11DynamicRHI.h"
diff --git a/Engine/Plugins/Runtime/WindowsMixedReality/Source/WindowsMixedRealityHMD/WindowsMixedRealityHMD.Build.cs b/Engine/Plugins/Runtime/WindowsMixedReality/Source/WindowsMixedRealityHMD/WindowsMixedRealityHMD.Build.cs
--- a/Engine/Plugins/Runtime/WindowsMixedReality/Source/WindowsMixedRealityHMD/WindowsMixedRealityHMD.Build.cs
+++ b/Engine/Plugins/Runtime/WindowsMixedReality/Source/WindowsMixedRealityHMD/WindowsMixedRealityHMD.Build.cs
@@ -67,7 +67,6 @@ namespace UnrealBuildTool.Rules
                         "EngineSettings",
                         "InputCore",
 						"RHI",
-						"RHICore",
 						"RenderCore",
 						"Renderer",
 						"HeadMountedDisplay",
@@ -92,14 +91,8 @@ namespace UnrealBuildTool.Rules
                     PrivateDependencyModuleNames.Add("WindowsMixedRealityRuntimeSettings");
                 }
 
-				if (Target.Platform != UnrealTargetPlatform.HoloLens)
-				{
-					AddEngineThirdPartyPrivateStaticDependencies(Target, "NVAftermath");
-					AddEngineThirdPartyPrivateStaticDependencies(Target, "IntelMetricsDiscovery");
-					AddEngineThirdPartyPrivateStaticDependencies(Target, "IntelExtensionsFramework");
-				}
-
-                AddEngineThirdPartyPrivateStaticDependencies(Target, "WindowsMixedRealityInterop");
+				AddEngineThirdPartyPrivateStaticDependencies(Target, "DX11");
+				AddEngineThirdPartyPrivateStaticDependencies(Target, "WindowsMixedRealityInterop");
 
                 LoadMixedReality(Target);
 
@@ -107,19 +100,9 @@ namespace UnrealBuildTool.Rules
 					new string[]
 					{
 					"WindowsMixedRealityHMD/Private",
-					"../../../../Source/Runtime/Windows/D3D11RHI/Private",
 					"../../../../Source/Runtime/Renderer/Private",
 					});
 
-				if (Target.Platform == UnrealTargetPlatform.Win64)
-				{
-					PrivateIncludePaths.Add("../../../../Source/Runtime/Windows/D3D11RHI/Private/Windows");
-				}
-				else if (Target.Platform == UnrealTargetPlatform.HoloLens)
-				{
-					PrivateIncludePaths.Add("../../../../Source/Runtime/Windows/D3D11RHI/Private/HoloLens");
-				}
-
 				PCHUsage = PCHUsageMode.NoSharedPCHs;
 				PrivatePCHHeaderFile = "Private/WindowsMixedRealityPrecompiled.h";
 
diff --git a/Engine/Plugins/Runtime/WindowsMixedReality/Source/WindowsMixedRealityRHI/Private/WindowsMixedRealityPrecompiled.h b/Engine/Plugins/Runtime/WindowsMixedReality/Source/WindowsMixedRealityRHI/Private/WindowsMixedRealityPrecompiled.h
--- a/Engine/Plugins/Runtime/WindowsMixedReality/Source/WindowsMixedRealityRHI/Private/WindowsMixedRealityPrecompiled.h
+++ b/Engine/Plugins/Runtime/WindowsMixedReality/Source/WindowsMixedRealityRHI/Private/WindowsMixedRealityPrecompiled.h
@@ -22,8 +22,4 @@
 DECLARE_LOG_CATEGORY_EXTERN(LogWMRRHI, Log, All);
 
 #include "SceneRendering.h"
-
-#pragma warning(disable:4668)  
-#include "D3D11RHIPrivate.h"
-#pragma warning(default:4668)
-#include "D3D11Util.h"
+#include "ID3D11DynamicRHI.h"
diff --git a/Engine/Plugins/Runtime/WindowsMixedReality/Source/WindowsMixedRealityRHI/WindowsMixedRealityRHI.Build.cs b/Engine/Plugins/Runtime/WindowsMixedReality/Source/WindowsMixedRealityRHI/WindowsMixedRealityRHI.Build.cs
--- a/Engine/Plugins/Runtime/WindowsMixedReality/Source/WindowsMixedRealityRHI/WindowsMixedRealityRHI.Build.cs
+++ b/Engine/Plugins/Runtime/WindowsMixedReality/Source/WindowsMixedRealityRHI/WindowsMixedRealityRHI.Build.cs
@@ -53,7 +53,6 @@ namespace UnrealBuildTool.Rules
 						"Projects",
                         "HeadMountedDisplay",
                         "D3D11RHI",
-						"RHICore",
                         "Slate",
                         "SlateCore",
                         "MRMesh",
@@ -63,12 +62,6 @@ namespace UnrealBuildTool.Rules
 
                 AddEngineThirdPartyPrivateStaticDependencies(Target, "WindowsMixedRealityInterop");
 
-				if (Target.Platform != UnrealTargetPlatform.HoloLens)
-				{
-					AddEngineThirdPartyPrivateStaticDependencies(Target, "NVAftermath");
-					AddEngineThirdPartyPrivateStaticDependencies(Target, "IntelExtensionsFramework");
-				}
-
 				if (Target.bBuildEditor == true)
                 {
 					PrivateDependencyModuleNames.Add("EditorFramework");
@@ -81,20 +74,9 @@ namespace UnrealBuildTool.Rules
                     new string[]
                     {
                     "WindowsMixedRealityRHI/Private",
-                    "../../../../Source/Runtime/Windows/D3D11RHI/Private",
                     "../../../../Source/Runtime/Renderer/Private",
                     });
 
-                //TODO: needed?
-                if (Target.Platform == UnrealTargetPlatform.Win64)
-                {
-                    PrivateIncludePaths.Add("../../../../Source/Runtime/Windows/D3D11RHI/Private/Windows");
-                }
-                else if (Target.Platform == UnrealTargetPlatform.HoloLens)
-                {
-                    PrivateIncludePaths.Add("../../../../Source/Runtime/Windows/D3D11RHI/Private/HoloLens");
-                }
-
                 PCHUsage = PCHUsageMode.NoSharedPCHs;
                 PrivatePCHHeaderFile = "Private/WindowsMixedRealityPrecompiled.h";
             }
diff --git a/Engine/Plugins/Runtime/nDisplay/Source/DisplayCluster/DisplayCluster.Build.cs b/Engine/Plugins/Runtime/nDisplay/Source/DisplayCluster/DisplayCluster.Build.cs
--- a/Engine/Plugins/Runtime/nDisplay/Source/DisplayCluster/DisplayCluster.Build.cs
+++ b/Engine/Plugins/Runtime/nDisplay/Source/DisplayCluster/DisplayCluster.Build.cs
@@ -75,11 +75,7 @@ public class DisplayCluster : ModuleRules
 					"TextureShareCore",
 			});
 
-			AddEngineThirdPartyPrivateStaticDependencies(Target, "DX11");
-			AddEngineThirdPartyPrivateStaticDependencies(Target, "DX12");
-			AddEngineThirdPartyPrivateStaticDependencies(Target, "IntelExtensionsFramework");
-			AddEngineThirdPartyPrivateStaticDependencies(Target, "IntelMetricsDiscovery");
-			AddEngineThirdPartyPrivateStaticDependencies(Target, "NVAftermath");
+			AddEngineThirdPartyPrivateStaticDependencies(Target, "DX11", "DX12");
 			AddEngineThirdPartyPrivateStaticDependencies(Target, "NVAPI");
 		}
 	}
diff --git a/Engine/Plugins/Runtime/nDisplay/Source/DisplayCluster/Private/Render/Synchronization/Windows/DisplayClusterRenderSyncPolicyNvidiaWindows.cpp b/Engine/Plugins/Runtime/nDisplay/Source/DisplayCluster/Private/Render/Synchronization/Windows/DisplayClusterRenderSyncPolicyNvidiaWindows.cpp
--- a/Engine/Plugins/Runtime/nDisplay/Source/DisplayCluster/Private/Render/Synchronization/Windows/DisplayClusterRenderSyncPolicyNvidiaWindows.cpp
+++ b/Engine/Plugins/Runtime/nDisplay/Source/DisplayCluster/Private/Render/Synchronization/Windows/DisplayClusterRenderSyncPolicyNvidiaWindows.cpp
@@ -9,16 +9,14 @@
 
 #include "Engine/Engine.h"
 #include "Engine/GameViewportClient.h"
-#include "DynamicRHI.h"
 
-#pragma warning(push)
-#pragma warning (disable : 4005)
+#include "ID3D11DynamicRHI.h"
+
 #include "Windows/AllowWindowsPlatformTypes.h"
-#include "dxgi1_3.h"
-#include "d3d11.h"
-#include "nvapi.h"
+THIRD_PARTY_INCLUDES_START
+	#include <nvapi.h>
+THIRD_PARTY_INCLUDES_END
 #include "Windows/HideWindowsPlatformTypes.h"
-#pragma warning(pop)
 
 
 // TEMPORARY DIAGNOSTICS START
diff --git a/Engine/Plugins/Runtime/nDisplay/Source/DisplayClusterProjection/DisplayClusterProjection.Build.cs b/Engine/Plugins/Runtime/nDisplay/Source/DisplayClusterProjection/DisplayClusterProjection.Build.cs
--- a/Engine/Plugins/Runtime/nDisplay/Source/DisplayClusterProjection/DisplayClusterProjection.Build.cs
+++ b/Engine/Plugins/Runtime/nDisplay/Source/DisplayClusterProjection/DisplayClusterProjection.Build.cs
@@ -29,7 +29,6 @@ public class DisplayClusterProjection : ModuleRules
 				"Projects",
 				"RenderCore",
 				"RHI",
-				"RHICore",
 				"ProceduralMeshComponent",
 			});
 
@@ -40,29 +39,12 @@ public class DisplayClusterProjection : ModuleRules
 
 		if (Target.Platform == UnrealTargetPlatform.Win64)
 		{
-			PrivateDependencyModuleNames.AddRange(
-				new string[] {
-					"D3D11RHI",
-					"D3D12RHI"
+			PrivateDependencyModuleNames.AddRange(new string[] {
+				"D3D11RHI",
+				"D3D12RHI",
 			});
 
-			// Required for some private headers needed for the rendering support.
-			var EngineDir = Path.GetFullPath(Target.RelativeEnginePath);
-			PrivateIncludePaths.AddRange(
-				new string[] {
-					Path.Combine(EngineDir, @"Source\Runtime\Windows\D3D11RHI\Private"),
-					Path.Combine(EngineDir, @"Source\Runtime\Windows\D3D11RHI\Private\Windows"),
-					Path.Combine(EngineDir, @"Source\Runtime\D3D12RHI\Private"),
-					Path.Combine(EngineDir, @"Source\Runtime\D3D12RHI\Private\Windows")
-			});
-
-			AddEngineThirdPartyPrivateStaticDependencies(Target, "AMD_AGS");
-			AddEngineThirdPartyPrivateStaticDependencies(Target, "DX11");
-			AddEngineThirdPartyPrivateStaticDependencies(Target, "DX12");
-			AddEngineThirdPartyPrivateStaticDependencies(Target, "NVAftermath");
-			AddEngineThirdPartyPrivateStaticDependencies(Target, "NVAPI");
-			AddEngineThirdPartyPrivateStaticDependencies(Target, "IntelExtensionsFramework");
-			AddEngineThirdPartyPrivateStaticDependencies(Target, "IntelMetricsDiscovery");
+			AddEngineThirdPartyPrivateStaticDependencies(Target, "DX11", "DX12");
 		}
 
 		AddThirdPartyDependencies(ROTargetRules);
diff --git a/Engine/Plugins/Runtime/nDisplay/Source/DisplayClusterProjection/Private/Policy/Domeprojection/Windows/DX11/DisplayClusterProjectionDomeprojectionLibraryDX11.h b/Engine/Plugins/Runtime/nDisplay/Source/DisplayClusterProjection/Private/Policy/Domeprojection/Windows/DX11/DisplayClusterProjectionDomeprojectionLibraryDX11.h
--- a/Engine/Plugins/Runtime/nDisplay/Source/DisplayClusterProjection/Private/Policy/Domeprojection/Windows/DX11/DisplayClusterProjectionDomeprojectionLibraryDX11.h
+++ b/Engine/Plugins/Runtime/nDisplay/Source/DisplayClusterProjection/Private/Policy/Domeprojection/Windows/DX11/DisplayClusterProjectionDomeprojectionLibraryDX11.h
@@ -3,15 +3,7 @@
 #pragma once
 
 #include "CoreMinimal.h"
-
-#include "RHI.h"
-#include "RHIResources.h"
-
-#include "Windows/AllowWindowsPlatformTypes.h"
-#include "DirectX/Include/DXGI.h"
-#include "Windows/HideWindowsPlatformTypes.h"
-
-#include "D3D11RHI/Private/Windows/D3D11RHIBasePrivate.h"
+#include "ID3D11DynamicRHI.h"
 
 THIRD_PARTY_INCLUDES_START
 #include "dpTypes.h"
diff --git a/Engine/Plugins/Runtime/nDisplay/Source/DisplayClusterProjection/Private/Policy/Domeprojection/Windows/DX11/DisplayClusterProjectionDomeprojectionViewAdapterDX11.cpp b/Engine/Plugins/Runtime/nDisplay/Source/DisplayClusterProjection/Private/Policy/Domeprojection/Windows/DX11/DisplayClusterProjectionDomeprojectionViewAdapterDX11.cpp
--- a/Engine/Plugins/Runtime/nDisplay/Source/DisplayClusterProjection/Private/Policy/Domeprojection/Windows/DX11/DisplayClusterProjectionDomeprojectionViewAdapterDX11.cpp
+++ b/Engine/Plugins/Runtime/nDisplay/Source/DisplayClusterProjection/Private/Policy/Domeprojection/Windows/DX11/DisplayClusterProjectionDomeprojectionViewAdapterDX11.cpp
@@ -7,11 +7,7 @@
 #include "DisplayClusterProjectionLog.h"
 #include "Misc/DisplayClusterHelpers.h"
 
-#include "RHI.h"
-#include "RHIResources.h"
-#include "RHIUtilities.h"
-
-#include "Windows/D3D11RHI/Private/D3D11RHIPrivate.h"
+#include "ID3D11DynamicRHI.h"
 
 #include "Engine/GameViewportClient.h"
 #include "Engine/Engine.h"
@@ -61,9 +57,12 @@ bool FDisplayClusterProjectionDomeprojectionViewAdapterDX11::Initialize(class ID
 bool FDisplayClusterProjectionDomeprojectionViewAdapterDX11::CalculateView(IDisplayClusterViewport* InViewport, const uint32 InContextNum, const uint32 Channel, FVector& InOutViewLocation, FRotator& InOutViewRotation, const FVector& ViewOffset, const float WorldToMeters, const float NCP, const float FCP)
 {
 	check(GDynamicRHI);
+	if (!GDynamicRHI)
+	{
+		return false;
+	}
 
-	FD3D11DynamicRHI* d3d11RHI = GetDynamicRHI<FD3D11DynamicRHI>();
-	ID3D11Device* D3D11Device = d3d11RHI ? d3d11RHI->GetDevice() : nullptr;
+	ID3D11Device* D3D11Device = GetID3D11DynamicRHI()->RHIGetDevice();
 	if (!D3D11Device)
 	{
 		return false;
@@ -125,21 +124,15 @@ bool FDisplayClusterProjectionDomeprojectionViewAdapterDX11::GetProjectionMatrix
 
 bool FDisplayClusterProjectionDomeprojectionViewAdapterDX11::ImplApplyWarpBlend_RenderThread(FRHICommandListImmediate& RHICmdList, uint32 InContextNum, const uint32 Channel, FRHITexture2D* InputTextures, FRHITexture2D* OutputTextures)
 {
-	if (GD3D11RHI == nullptr)
+	if (!IsRHID3D11())
 	{
 		return false;
 	}
 
-	// Prepare the textures
-	FD3D11TextureBase* SrcTextureRHI = static_cast<FD3D11TextureBase*>(InputTextures->GetTextureBaseRHI());
-	FD3D11TextureBase* DstTextureRHI = static_cast<FD3D11TextureBase*>(OutputTextures->GetTextureBaseRHI());
-	if (!(SrcTextureRHI && DstTextureRHI))
-	{
-		return false;
-	}
+	ID3D11DynamicRHI* D3D11RHI = GetID3D11DynamicRHI();
 
-	ID3D11ShaderResourceView* SrcTextureSRVD3D11 = static_cast<ID3D11ShaderResourceView*>(SrcTextureRHI->GetShaderResourceView());
-	ID3D11RenderTargetView* DstTextureRTV = DstTextureRHI->GetRenderTargetView(0, -1);
+	ID3D11ShaderResourceView* SrcTextureSRVD3D11 = D3D11RHI->RHIGetShaderResourceView(InputTextures);
+	ID3D11RenderTargetView* DstTextureRTV = D3D11RHI->RHIGetRenderTargetView(OutputTextures);
 	if (!(SrcTextureSRVD3D11 && DstTextureRTV))
 	{
 		return false;
@@ -155,8 +148,8 @@ bool FDisplayClusterProjectionDomeprojectionViewAdapterDX11::ImplApplyWarpBlend_
 	RenderViewportData.TopLeftX = 0.0f;
 	RenderViewportData.TopLeftY = 0.0f;
 
-	FD3D11Device* D3D11Device = GD3D11RHI->GetDevice();
-	FD3D11DeviceContext* DeviceContext = GD3D11RHI->GetDeviceContext();
+	FD3D11Device* D3D11Device = GetID3D11DynamicRHI()->RHIGetDevice();
+	FD3D11DeviceContext* DeviceContext = GetID3D11DynamicRHI()->RHIGetDeviceContext();
 	if (D3D11Device && DeviceContext)
 	{
 		DeviceContext->RSSetViewports(1, &RenderViewportData);
@@ -262,9 +255,8 @@ bool FDisplayClusterProjectionDomeprojectionViewAdapterDX11::FViewData::Initiali
 	}
 
 	check(GDynamicRHI);
-	FD3D11DynamicRHI* d3d11RHI = GetDynamicRHI<FD3D11DynamicRHI>();
 
-	dpResult Result = DisplayClusterProjectionDomeprojectionLibraryDX11::dpCreateContextFunc(&Context, d3d11RHI->GetDevice(), DP_PLUGIN_ID);
+	dpResult Result = DisplayClusterProjectionDomeprojectionLibraryDX11::dpCreateContextFunc(&Context, GetID3D11DynamicRHI()->RHIGetDevice(), DP_PLUGIN_ID);
 	if (Result != dpNoError)
 	{
 		UE_LOG(LogDisplayClusterProjectionDomeprojection, Error, TEXT("Couldn't initialize Domeprojection context"));
diff --git a/Engine/Plugins/Runtime/nDisplay/Source/DisplayClusterProjection/Private/Policy/EasyBlend/Windows/DX11/DisplayClusterProjectionEasyBlendLibraryDX11.h b/Engine/Plugins/Runtime/nDisplay/Source/DisplayClusterProjection/Private/Policy/EasyBlend/Windows/DX11/DisplayClusterProjectionEasyBlendLibraryDX11.h
--- a/Engine/Plugins/Runtime/nDisplay/Source/DisplayClusterProjection/Private/Policy/EasyBlend/Windows/DX11/DisplayClusterProjectionEasyBlendLibraryDX11.h
+++ b/Engine/Plugins/Runtime/nDisplay/Source/DisplayClusterProjection/Private/Policy/EasyBlend/Windows/DX11/DisplayClusterProjectionEasyBlendLibraryDX11.h
@@ -3,15 +3,7 @@
 #pragma once
 
 #include "CoreMinimal.h"
-
-#include "RHI.h"
-#include "RHIResources.h"
-
-#include "Windows/AllowWindowsPlatformTypes.h"
-#include "DirectX/Include/DXGI.h"
-#include "Windows/HideWindowsPlatformTypes.h"
-
-#include "D3D11RHI/Private/Windows/D3D11RHIBasePrivate.h"
+#include "ID3D11DynamicRHI.h"
 
 THIRD_PARTY_INCLUDES_START
 #include "EasyBlendSDKDXVer.h"
diff --git a/Engine/Plugins/Runtime/nDisplay/Source/DisplayClusterProjection/Private/Policy/EasyBlend/Windows/DX11/DisplayClusterProjectionEasyBlendViewAdapterDX11.cpp b/Engine/Plugins/Runtime/nDisplay/Source/DisplayClusterProjection/Private/Policy/EasyBlend/Windows/DX11/DisplayClusterProjectionEasyBlendViewAdapterDX11.cpp
--- a/Engine/Plugins/Runtime/nDisplay/Source/DisplayClusterProjection/Private/Policy/EasyBlend/Windows/DX11/DisplayClusterProjectionEasyBlendViewAdapterDX11.cpp
+++ b/Engine/Plugins/Runtime/nDisplay/Source/DisplayClusterProjection/Private/Policy/EasyBlend/Windows/DX11/DisplayClusterProjectionEasyBlendViewAdapterDX11.cpp
@@ -7,11 +7,7 @@
 #include "DisplayClusterProjectionLog.h"
 #include "Misc/DisplayClusterHelpers.h"
 
-#include "RHI.h"
-#include "RHIResources.h"
-#include "RHIUtilities.h"
-
-#include "Windows/D3D11RHI/Private/D3D11RHIPrivate.h"
+#include "ID3D11DynamicRHI.h"
 
 #include "Engine/GameViewportClient.h"
 #include "Engine/Engine.h"
@@ -22,7 +18,6 @@
 #include "Render/Viewport/IDisplayClusterViewport.h"
 #include "Render/Viewport/IDisplayClusterViewportProxy.h"
 
-
 //------------------------------------------------------------------------------
 // FDisplayClusterProjectionEasyBlendViewAdapterDX11
 //------------------------------------------------------------------------------
@@ -140,28 +135,24 @@ void FDisplayClusterProjectionEasyBlendViewAdapterDX11::ImplInitializeResources_
 		if (IsEasyBlendRenderingEnabled())
 		{
 			FViewport* MainViewport = GEngine->GameViewport->Viewport;
-			if (GD3D11RHI && MainViewport)
+			if (IsRHID3D11() && MainViewport)
 			{
-				FD3D11Device* Device = GD3D11RHI->GetDevice();
-				FD3D11DeviceContext* DeviceContext = GD3D11RHI->GetDeviceContext();
+				FD3D11Device* Device = GetID3D11DynamicRHI()->RHIGetDevice();
+				FD3D11DeviceContext* DeviceContext = GetID3D11DynamicRHI()->RHIGetDeviceContext();
 				if (Device && DeviceContext)
 				{
-					FD3D11Viewport* Viewport = static_cast<FD3D11Viewport*>(MainViewport->GetViewportRHI().GetReference());
-					if (Viewport)
+					IDXGISwapChain* SwapChain = GetID3D11DynamicRHI()->RHIGetSwapChain(MainViewport->GetViewportRHI().GetReference());
+					if (SwapChain)
 					{
-						IDXGISwapChain* SwapChain = (IDXGISwapChain*)Viewport->GetSwapChain();
-						if (SwapChain)
+						// Create RT texture for viewport warp
+						for (FViewData& It : Views)
 						{
-							// Create RT texture for viewport warp
-							for (FViewData& It : Views)
+							// Initialize EasyBlend internals
+							check(DisplayClusterProjectionEasyBlendLibraryDX11::EasyBlendInitDeviceObjectsFunc);
+							EasyBlendSDKDXError sdkErr = DisplayClusterProjectionEasyBlendLibraryDX11::EasyBlendInitDeviceObjectsFunc(It.EasyBlendMeshData.Get(), Device, DeviceContext, SwapChain);
+							if (EasyBlendSDKDX_FAILED(sdkErr))
 							{
-								// Initialize EasyBlend internals
-								check(DisplayClusterProjectionEasyBlendLibraryDX11::EasyBlendInitDeviceObjectsFunc);
-								EasyBlendSDKDXError sdkErr = DisplayClusterProjectionEasyBlendLibraryDX11::EasyBlendInitDeviceObjectsFunc(It.EasyBlendMeshData.Get(), Device, DeviceContext, SwapChain);
-								if (EasyBlendSDKDX_FAILED(sdkErr))
-								{
-									UE_LOG(LogDisplayClusterProjectionEasyBlend, Error, TEXT("Couldn't initialize EasyBlend Device/DeviceContext/SwapChain"));
-								}
+								UE_LOG(LogDisplayClusterProjectionEasyBlend, Error, TEXT("Couldn't initialize EasyBlend Device/DeviceContext/SwapChain"));
 							}
 						}
 					}
@@ -272,19 +263,18 @@ bool FDisplayClusterProjectionEasyBlendViewAdapterDX11::ImplApplyWarpBlend_Rende
 	}
 
 	FViewport* MainViewport = GEngine->GameViewport->Viewport;
-	if (GD3D11RHI == nullptr || MainViewport == nullptr)
+	if (!IsRHID3D11() || MainViewport == nullptr)
 	{
 		return false;
 	}
 
-	// Prepare the textures
-	FD3D11TextureBase* SrcTextureRHI = static_cast<FD3D11TextureBase*>(InputTexture->GetTextureBaseRHI());
-	FD3D11TextureBase* DstTextureRHI = static_cast<FD3D11TextureBase*>(OutputTexture->GetTextureBaseRHI());
+	ID3D11DynamicRHI* D3D11RHI = GetID3D11DynamicRHI();
 
-	ID3D11RenderTargetView* DstTextureRTV = DstTextureRHI->GetRenderTargetView(0, -1);
+	// Prepare the textures
+	ID3D11RenderTargetView* DstTextureRTV = D3D11RHI->RHIGetRenderTargetView(OutputTexture);
 
-	ID3D11Texture2D* DstTextureD3D11 = static_cast<ID3D11Texture2D*>(DstTextureRHI->GetResource());
-	ID3D11Texture2D* SrcTextureD3D11 = static_cast<ID3D11Texture2D*>(SrcTextureRHI->GetResource());
+	ID3D11Texture2D* DstTextureD3D11 = (ID3D11Texture2D*)D3D11RHI->RHIGetResource(OutputTexture);
+	ID3D11Texture2D* SrcTextureD3D11 = (ID3D11Texture2D*)D3D11RHI->RHIGetResource(InputTexture);
 
 	// Setup In/Out EasyBlend textures
 	{
@@ -311,40 +301,37 @@ bool FDisplayClusterProjectionEasyBlendViewAdapterDX11::ImplApplyWarpBlend_Rende
 	RenderViewportData.TopLeftX = 0.0f;
 	RenderViewportData.TopLeftY = 0.0f;
 
-	FD3D11Device* Device = GD3D11RHI->GetDevice();
-	FD3D11DeviceContext* DeviceContext = GD3D11RHI->GetDeviceContext();
+	FD3D11Device* Device = GetID3D11DynamicRHI()->RHIGetDevice();
+	FD3D11DeviceContext* DeviceContext = GetID3D11DynamicRHI()->RHIGetDeviceContext();
 	if (Device && DeviceContext)
 	{
-		FD3D11Viewport* Viewport = static_cast<FD3D11Viewport*>(MainViewport->GetViewportRHI().GetReference());
-		if (Viewport)
+
+		IDXGISwapChain* SwapChain = D3D11RHI->RHIGetSwapChain(MainViewport->GetViewportRHI().GetReference());
+		if (SwapChain)
 		{
-			IDXGISwapChain* SwapChain = (IDXGISwapChain*)Viewport->GetSwapChain();
-			if (SwapChain)
-			{
-				DeviceContext->RSSetViewports(1, &RenderViewportData);
-				DeviceContext->OMSetRenderTargets(1, &DstTextureRTV, nullptr);
-				DeviceContext->Flush();
+			DeviceContext->RSSetViewports(1, &RenderViewportData);
+			DeviceContext->OMSetRenderTargets(1, &DstTextureRTV, nullptr);
+			DeviceContext->Flush();
 
+			{
+				FScopeLock Lock(&DllAccessCS);
+
+				// Perform warp&blend by the EasyBlend
+				check(DisplayClusterProjectionEasyBlendLibraryDX11::EasyBlendDXRenderFunc);
+				EasyBlendSDKDXError EasyBlendSDKDXError = DisplayClusterProjectionEasyBlendLibraryDX11::EasyBlendDXRenderFunc(
+					Views[ContextNum].EasyBlendMeshData.Get(),
+					Device,
+					DeviceContext,
+					SwapChain,
+					false);
+
+				if (!EasyBlendSDKDX_SUCCEEDED(EasyBlendSDKDXError))
 				{
-					FScopeLock Lock(&DllAccessCS);
-
-					// Perform warp&blend by the EasyBlend
-					check(DisplayClusterProjectionEasyBlendLibraryDX11::EasyBlendDXRenderFunc);
-					EasyBlendSDKDXError EasyBlendSDKDXError = DisplayClusterProjectionEasyBlendLibraryDX11::EasyBlendDXRenderFunc(
-						Views[ContextNum].EasyBlendMeshData.Get(),
-						Device,
-						DeviceContext,
-						SwapChain,
-						false);
-
-					if (!EasyBlendSDKDX_SUCCEEDED(EasyBlendSDKDXError))
-					{
-						UE_LOG(LogDisplayClusterProjectionEasyBlend, Error, TEXT("EasyBlend couldn't perform rendering operation"));
-						return false;
-					}
-
-					return true;
+					UE_LOG(LogDisplayClusterProjectionEasyBlend, Error, TEXT("EasyBlend couldn't perform rendering operation"));
+					return false;
 				}
+
+				return true;
 			}
 		}
 	}
diff --git a/Engine/Plugins/Runtime/nDisplay/Source/DisplayClusterProjection/Private/Policy/VIOSO/Windows/DisplayClusterProjectionVIOSOPolicy.cpp b/Engine/Plugins/Runtime/nDisplay/Source/DisplayClusterProjection/Private/Policy/VIOSO/Windows/DisplayClusterProjectionVIOSOPolicy.cpp
--- a/Engine/Plugins/Runtime/nDisplay/Source/DisplayClusterProjection/Private/Policy/VIOSO/Windows/DisplayClusterProjectionVIOSOPolicy.cpp
+++ b/Engine/Plugins/Runtime/nDisplay/Source/DisplayClusterProjection/Private/Policy/VIOSO/Windows/DisplayClusterProjectionVIOSOPolicy.cpp
@@ -26,16 +26,14 @@
 // D3D11
 //-------------------------------------------------------------------------------------------------
 #if VIOSO_USE_GRAPHICS_API_D3D11
-#include "D3D11RHIPrivate.h"
-#include "D3D11Util.h"
+#include "ID3D11DynamicRHI.h"
 #endif // VIOSO_USE_GRAPHICS_API_D3D11
 
 //-------------------------------------------------------------------------------------------------
 // D3D12
 //-------------------------------------------------------------------------------------------------
 #if VIOSO_USE_GRAPHICS_API_D3D12
-#include "D3D12RHIPrivate.h"
-#include "D3D12Util.h"
+#include "ID3D12DynamicRHI.h"
 #endif // VIOSO_USE_GRAPHICS_API_D3D12
 
 #if VIOSO_USE_GRAPHICS_API_D3D11
@@ -53,8 +51,8 @@ public:
 		DepthStencilView = NULL;
 
 
-		FD3D11DynamicRHI* d3d11RHI = GetDynamicRHI<FD3D11DynamicRHI>();
-		DeviceContext = d3d11RHI ? d3d11RHI->GetDeviceContext() : nullptr;
+		ID3D11DynamicRHI* D3D11RHI = GDynamicRHI ? GetDynamicRHI<ID3D11DynamicRHI>() : nullptr;
+		DeviceContext = D3D11RHI ? D3D11RHI->RHIGetDeviceContext() : nullptr;
 
 		if (DeviceContext)
 		{
@@ -98,8 +96,7 @@ public:
 		if (DeviceContext)
 		{
 			// Set RTV
-			FD3D11TextureBase* DestTextureD3D11RHI = static_cast<FD3D11TextureBase*>(RenderTargetTexture->GetTextureBaseRHI());
-			ID3D11RenderTargetView* DestTextureRTV = DestTextureD3D11RHI->GetRenderTargetView(0, -1);
+			ID3D11RenderTargetView* DestTextureRTV = GetID3D11DynamicRHI()->RHIGetRenderTargetView(RenderTargetTexture, 0, -1);
 			DeviceContext->OMSetRenderTargets(1, &DestTextureRTV, nullptr);
 
 			// Set viewport
@@ -396,27 +393,21 @@ bool FDisplayClusterProjectionVIOSOPolicy::FViewData::RenderVIOSO_RenderThread(F
 #ifdef VIOSO_USE_GRAPHICS_API_D3D12
 			case ERenderDevice::D3D12:
 			{
-				FD3D12Texture2D* SrcTexture2D = FD3D12DynamicRHI::ResourceCast(ShaderResourceTexture);
-				FD3D12Texture2D* DestTexture2D = FD3D12DynamicRHI::ResourceCast(RenderTargetTexture);
-
-				FD3D12RenderTargetView* RTV = DestTexture2D->GetRenderTargetView(0, 0);
-				D3D12_CPU_DESCRIPTOR_HANDLE RTVHandle = RTV->GetView();
-
-				FD3D12Device* Device = DestTexture2D->GetParentDevice();
-				FD3D12CommandListHandle& hCommandList = Device->GetDefaultCommandContext().CommandListHandle;
+				ID3D12DynamicRHI* D3D12RHI = GetID3D12DynamicRHI();
+				const D3D12_CPU_DESCRIPTOR_HANDLE RTVHandle = D3D12RHI->RHIGetRenderTargetView(RenderTargetTexture);
 
 				VWB_D3D12_RENDERINPUT RenderInput = {};
-				RenderInput.textureResource = (ID3D12Resource*)SrcTexture2D->GetNativeResource();;
-				RenderInput.renderTarget = (ID3D12Resource*)DestTexture2D->GetNativeResource();;
+				RenderInput.textureResource = D3D12RHI->RHIGetResource(ShaderResourceTexture);
+				RenderInput.renderTarget = D3D12RHI->RHIGetResource(RenderTargetTexture);
 				RenderInput.rtvHandlePtr = RTVHandle.ptr;
 
 				//experimental: add resource barrier
-				hCommandList.AddPendingResourceBarrier(DestTexture2D->GetResource(), D3D12_RESOURCE_STATE_RENDER_TARGET, D3D12_RESOURCE_BARRIER_ALL_SUBRESOURCES);
+				D3D12RHI->RHIAddPendingBarrier(RenderTargetTexture, D3D12_RESOURCE_STATE_RENDER_TARGET, D3D12_RESOURCE_BARRIER_ALL_SUBRESOURCES);
 
 				if (Warper.Render(&RenderInput, VWB_STATEMASK_DEFAULT_D3D12))
 				{
 					//experimental: add resource barrier
-					hCommandList.AddPendingResourceBarrier(DestTexture2D->GetResource(), D3D12_RESOURCE_STATE_PRESENT, D3D12_RESOURCE_BARRIER_ALL_SUBRESOURCES);
+					D3D12RHI->RHIAddPendingBarrier(RenderTargetTexture, D3D12_RESOURCE_STATE_PRESENT, D3D12_RESOURCE_BARRIER_ALL_SUBRESOURCES);
 					return true;
 				}
 				break;
@@ -468,8 +459,7 @@ bool FDisplayClusterProjectionVIOSOPolicy::FViewData::InitializeVIOSO(FRHITextur
 #if VIOSO_USE_GRAPHICS_API_D3D12
 	case ERenderDevice::D3D12:
 	{
-		FD3D12DynamicRHI* DynamicRHI = GetDynamicRHI<FD3D12DynamicRHI>();
-		ID3D12CommandQueue* D3D12CommandQueue = DynamicRHI->RHIGetD3DCommandQueue();
+		ID3D12CommandQueue* D3D12CommandQueue = GetID3D12DynamicRHI()->RHIGetCommandQueue();
 
 		if (Warper.Initialize(D3D12CommandQueue, InConfigData))
 		{
diff --git a/Engine/Source/Developer/Windows/ShaderFormatD3D/Private/D3DShaderCompiler.cpp b/Engine/Source/Developer/Windows/ShaderFormatD3D/Private/D3DShaderCompiler.cpp
--- a/Engine/Source/Developer/Windows/ShaderFormatD3D/Private/D3DShaderCompiler.cpp
+++ b/Engine/Source/Developer/Windows/ShaderFormatD3D/Private/D3DShaderCompiler.cpp
@@ -16,9 +16,6 @@ DEFINE_LOG_CATEGORY_STATIC(LogD3D11ShaderCompiler, Log, All);
 
 #define DEBUG_SHADERS 0
 
-// D3D headers.
-#define D3D_OVERLOADS 1
-
 // D3D doesn't define a mask for this, so we do so here
 #define SHADER_OPTIMIZATION_LEVEL_MASK (D3DCOMPILE_OPTIMIZATION_LEVEL0 | D3DCOMPILE_OPTIMIZATION_LEVEL1 | D3DCOMPILE_OPTIMIZATION_LEVEL2 | D3DCOMPILE_OPTIMIZATION_LEVEL3)
 
diff --git a/Engine/Source/Runtime/D3D12RHI/Private/D3D12Adapter.h b/Engine/Source/Runtime/D3D12RHI/Private/D3D12Adapter.h
--- a/Engine/Source/Runtime/D3D12RHI/Private/D3D12Adapter.h
+++ b/Engine/Source/Runtime/D3D12RHI/Private/D3D12Adapter.h
@@ -242,7 +242,7 @@ public:
 
 	FORCEINLINE FD3D12Fence* GetStagingFence()  { return StagingFence.GetReference(); }
 
-	FORCEINLINE FD3D12Device* GetDevice(uint32 GPUIndex)
+	FORCEINLINE FD3D12Device* GetDevice(uint32 GPUIndex) const
 	{
 		check(GPUIndex < GNumExplicitGPUsForRendering);
 		return Devices[GPUIndex];
diff --git a/Engine/Source/Runtime/D3D12RHI/Private/D3D12RHI.cpp b/Engine/Source/Runtime/D3D12RHI/Private/D3D12RHI.cpp
--- a/Engine/Source/Runtime/D3D12RHI/Private/D3D12RHI.cpp
+++ b/Engine/Source/Runtime/D3D12RHI/Private/D3D12RHI.cpp
@@ -405,6 +405,129 @@ void FD3D12DynamicRHI::RHIReleaseThreadOwnership()
 	// Nothing to do
 }
 
+TArray<FD3D12MinimalAdapterDesc> FD3D12DynamicRHI::RHIGetAdapterDescs() const
+{
+	TArray<FD3D12MinimalAdapterDesc> Result;
+
+	for (const TSharedPtr<FD3D12Adapter>& Adapter : ChosenAdapters)
+	{
+		FD3D12MinimalAdapterDesc Desc{};
+		Desc.Desc = Adapter->GetDesc().Desc;
+		Desc.NumDeviceNodes = Adapter->GetDesc().NumDeviceNodes;
+
+		Result.Add(Desc);
+	}
+
+	return Result;
+}
+
+bool FD3D12DynamicRHI::RHIIsPixEnabled() const
+{
+	return IsPixEventEnabled();
+}
+
+ID3D12CommandQueue* FD3D12DynamicRHI::RHIGetCommandQueue() const
+{
+	// Multi-GPU support : any code using this function needs validation.
+	return GetAdapter().GetDevice(0)->GetCommandListManager().GetD3DCommandQueue();
+}
+
+ID3D12Device* FD3D12DynamicRHI::RHIGetDevice(uint32 InIndex) const
+{
+	return GetAdapter().GetDevice(InIndex)->GetDevice();
+}
+
+uint32 FD3D12DynamicRHI::RHIGetDeviceNodeMask(uint32 InIndex) const
+{
+	return GetAdapter().GetDevice(InIndex)->GetGPUMask().GetNative();
+}
+
+ID3D12GraphicsCommandList* FD3D12DynamicRHI::RHIGetGraphicsCommandList(uint32 InDeviceIndex) const
+{
+	FD3D12Device* Device = GetRHIDevice(InDeviceIndex);
+	return Device->GetCommandContext().CommandListHandle.GraphicsCommandList();
+}
+
+DXGI_FORMAT FD3D12DynamicRHI::RHIGetSwapChainFormat(EPixelFormat InFormat) const
+{
+	const DXGI_FORMAT PlatformFormat = D3D12RHI::FindDepthStencilDXGIFormat(static_cast<DXGI_FORMAT>(GPixelFormats[InFormat].PlatformFormat));
+	return D3D12RHI::FindShaderResourceDXGIFormat(PlatformFormat, true);
+}
+
+ID3D12Resource* FD3D12DynamicRHI::RHIGetResource(FRHIBuffer* InBuffer) const
+{
+	FD3D12Buffer* D3D12Buffer = ResourceCast(InBuffer);
+	return D3D12Buffer->GetResource()->GetResource();
+}
+
+uint32 FD3D12DynamicRHI::RHIGetResourceDeviceIndex(FRHIBuffer* InBuffer) const
+{
+	FD3D12Buffer* D3D12Buffer = ResourceCast(InBuffer);
+	return D3D12Buffer->GetParentDevice()->GetGPUIndex();
+}
+
+ID3D12Resource* FD3D12DynamicRHI::RHIGetResource(FRHITexture* InTexture) const
+{
+	return (ID3D12Resource*)InTexture->GetNativeResource();
+}
+
+uint32 FD3D12DynamicRHI::RHIGetResourceDeviceIndex(FRHITexture* InTexture) const
+{
+	FD3D12TextureBase* D3D12Texture = GetD3D12TextureFromRHITexture(InTexture);
+	return D3D12Texture->GetParentDevice()->GetGPUIndex();
+}
+
+int64 FD3D12DynamicRHI::RHIGetResourceMemorySize(FRHITexture* InTexture) const
+{
+	FD3D12TextureBase* D3D12Texture = GetD3D12TextureFromRHITexture(InTexture);
+	return D3D12Texture->GetMemorySize();
+}
+
+bool FD3D12DynamicRHI::RHIIsResourcePlaced(FRHITexture* InTexture) const
+{
+	FD3D12TextureBase* D3D12Texture = GetD3D12TextureFromRHITexture(InTexture);
+	return D3D12Texture->GetResource()->IsPlacedResource();
+}
+
+D3D12_CPU_DESCRIPTOR_HANDLE FD3D12DynamicRHI::RHIGetRenderTargetView(FRHITexture* InTexture, int32 InMipIndex, int32 InArraySliceIndex) const
+{
+	FD3D12TextureBase* D3D12Texture = GetD3D12TextureFromRHITexture(InTexture);
+	FD3D12RenderTargetView* RTV = D3D12Texture->GetRenderTargetView(InMipIndex, InArraySliceIndex);
+	return RTV ? RTV->GetView() : D3D12_CPU_DESCRIPTOR_HANDLE{};
+}
+
+void FD3D12DynamicRHI::RHIFinishExternalComputeWork(uint32 InDeviceIndex, ID3D12GraphicsCommandList* InCommandList)
+{
+	FD3D12Device* Device = GetRHIDevice(InDeviceIndex);
+
+	check(InCommandList == Device->GetCommandContext().CommandListHandle.GraphicsCommandList());
+
+	Device->GetCommandContext().StateCache.ForceSetComputeRootSignature();
+	Device->GetCommandContext().StateCache.GetDescriptorCache()->SetCurrentCommandList(Device->GetCommandContext().CommandListHandle);
+}
+
+void FD3D12DynamicRHI::RHIRegisterWork(uint32 InDeviceIndex, uint32 NumPrimitives)
+{
+	FD3D12Device* Device = GetRHIDevice(InDeviceIndex);
+	if (Device->GetCommandContext().IsDefaultContext())
+	{
+		Device->RegisterGPUWork(NumPrimitives);
+	}
+}
+
+void FD3D12DynamicRHI::RHIAddPendingBarrier(FRHITexture* InTexture, D3D12_RESOURCE_STATES InState, uint32 InSubResource)
+{
+	FD3D12TextureBase* D3D12Texture = GetD3D12TextureFromRHITexture(InTexture);
+	FD3D12CommandListHandle& CommandList = D3D12Texture->GetParentDevice()->GetDefaultCommandContext().CommandListHandle;
+
+	CommandList.AddPendingResourceBarrier(D3D12Texture->GetResource(), InState, InSubResource);
+}
+
+bool ID3D12DynamicRHI::IsD3DDebugEnabled()
+{
+	return D3D12RHI_ShouldCreateWithD3DDebug();
+}
+
 void* FD3D12DynamicRHI::RHIGetNativeDevice()
 {
 	return (void*)GetAdapter().GetD3DDevice();
@@ -412,12 +535,12 @@ void* FD3D12DynamicRHI::RHIGetNativeDevice()
 
 void* FD3D12DynamicRHI::RHIGetNativeGraphicsQueue()
 {
-	return (void*)RHIGetD3DCommandQueue();
+	return (void*)RHIGetCommandQueue();
 }
 
 void* FD3D12DynamicRHI::RHIGetNativeComputeQueue()
 {
-	return (void*)RHIGetD3DCommandQueue();
+	return (void*)RHIGetCommandQueue();
 }
 
 void* FD3D12DynamicRHI::RHIGetNativeInstance()
@@ -425,7 +548,6 @@ void* FD3D12DynamicRHI::RHIGetNativeInstance()
 	return nullptr;
 }
 
-
 /**
 * Returns a supported screen resolution that most closely matches the input.
 * @param Width - Input: Desired resolution width in pixels. Output: A width that the platform supports.
@@ -539,11 +661,6 @@ void FD3D12DynamicRHI::GetBestSupportedMSAASetting(DXGI_FORMAT PlatformFormat, u
 	return;
 }
 
-uint32 FD3D12DynamicRHI::GetDebugFlags()
-{
-	return GetAdapter().GetDebugFlags();
-}
-
 bool FD3D12DynamicRHI::CheckGpuHeartbeat() const
 {
 	bool bResult = false;
diff --git a/Engine/Source/Runtime/D3D12RHI/Private/D3D12RHIPrivate.h b/Engine/Source/Runtime/D3D12RHI/Private/D3D12RHIPrivate.h
--- a/Engine/Source/Runtime/D3D12RHI/Private/D3D12RHIPrivate.h
+++ b/Engine/Source/Runtime/D3D12RHI/Private/D3D12RHIPrivate.h
@@ -29,7 +29,7 @@
 
 // Dependencies.
 #include "CoreMinimal.h"
-#include "RHI.h"
+#include "ID3D12DynamicRHI.h"
 #include "GPUProfiler.h"
 #include "ShaderCore.h"
 
@@ -265,7 +265,7 @@ struct INTCExtensionContext;
 PRAGMA_DISABLE_DEPRECATION_WARNINGS
 
 /** The interface which is implemented by the dynamically bound RHI. */
-class FD3D12DynamicRHI : public FDynamicRHI
+class FD3D12DynamicRHI : public ID3D12DynamicRHI
 {
 	friend class FD3D12CommandContext;
 
@@ -273,7 +273,7 @@ class FD3D12DynamicRHI : public FDynamicRHI
 
 public:
 
-	static D3D12RHI_API FD3D12DynamicRHI* GetD3DRHI() { return SingleD3DRHI; }
+	static FD3D12DynamicRHI* GetD3DRHI() { return SingleD3DRHI; }
 
 private:
 
@@ -293,7 +293,6 @@ public:
 	virtual void PostInit() override;
 	virtual void Shutdown() override;
 	virtual const TCHAR* GetName() override { return TEXT("D3D12"); }
-	virtual ERHIInterfaceType GetInterfaceType() const override final { return ERHIInterfaceType::D3D12; }
 
 	template<typename TRHIType>
 	static FORCEINLINE typename TD3D12ResourceTraits<TRHIType>::TConcreteType* ResourceCast(TRHIType* Resource)
@@ -434,12 +433,27 @@ public:
 	virtual IRHICommandContextContainer* RHIGetCommandContextContainer(int32 Index, int32 Num, FRHIGPUMask GPUMask)final override;
 #endif
 
-	// FD3D12DynamicRHI interface.
-	virtual uint32 GetDebugFlags();
-	virtual ID3D12CommandQueue* RHIGetD3DCommandQueue();
-	virtual FTexture2DRHIRef RHICreateTexture2DFromResource(EPixelFormat Format, ETextureCreateFlags TexCreateFlags, const FClearValueBinding& ClearValueBinding, ID3D12Resource* Resource);
-	virtual FTexture2DRHIRef RHICreateTexture2DArrayFromResource(EPixelFormat Format, ETextureCreateFlags TexCreateFlags, const FClearValueBinding& ClearValueBinding, ID3D12Resource* Resource);
-	virtual FTextureCubeRHIRef RHICreateTextureCubeFromResource(EPixelFormat Format, ETextureCreateFlags TexCreateFlags, const FClearValueBinding& ClearValueBinding, ID3D12Resource* Resource);
+	// ID3D12DynamicRHI interface.
+	virtual TArray<FD3D12MinimalAdapterDesc> RHIGetAdapterDescs() const final override;
+	virtual bool RHIIsPixEnabled() const final override;
+	virtual ID3D12CommandQueue* RHIGetCommandQueue() const final override;
+	virtual ID3D12Device* RHIGetDevice(uint32 InIndex) const final override;
+	virtual uint32 RHIGetDeviceNodeMask(uint32 InIndex) const final override;
+	virtual ID3D12GraphicsCommandList* RHIGetGraphicsCommandList(uint32 InDeviceIndex) const final override;
+	virtual DXGI_FORMAT RHIGetSwapChainFormat(EPixelFormat InFormat) const final override;
+	virtual FTexture2DRHIRef RHICreateTexture2DFromResource(EPixelFormat Format, ETextureCreateFlags TexCreateFlags, const FClearValueBinding& ClearValueBinding, ID3D12Resource* Resource) final override;
+	virtual FTexture2DArrayRHIRef RHICreateTexture2DArrayFromResource(EPixelFormat Format, ETextureCreateFlags TexCreateFlags, const FClearValueBinding& ClearValueBinding, ID3D12Resource* Resource) final override;
+	virtual FTextureCubeRHIRef RHICreateTextureCubeFromResource(EPixelFormat Format, ETextureCreateFlags TexCreateFlags, const FClearValueBinding& ClearValueBinding, ID3D12Resource* Resource) final override;
+	virtual ID3D12Resource* RHIGetResource(FRHIBuffer* InBuffer) const final override;
+	virtual uint32 RHIGetResourceDeviceIndex(FRHIBuffer* InBuffer) const final override;
+	virtual ID3D12Resource* RHIGetResource(FRHITexture* InTexture) const final override;
+	virtual uint32 RHIGetResourceDeviceIndex(FRHITexture* InTexture) const final override;
+	virtual int64 RHIGetResourceMemorySize(FRHITexture* InTexture) const final override;
+	virtual bool RHIIsResourcePlaced(FRHITexture* InTexture) const final override;
+	virtual D3D12_CPU_DESCRIPTOR_HANDLE RHIGetRenderTargetView(FRHITexture* InTexture, int32 InMipIndex = 0, int32 InArraySliceIndex = 0) const final override;
+	virtual void RHIFinishExternalComputeWork(uint32 InDeviceIndex, ID3D12GraphicsCommandList* InCommandList) final override;
+	virtual void RHIRegisterWork(uint32 InDeviceIndex, uint32 NumPrimitives) final override;
+	virtual void RHIAddPendingBarrier(FRHITexture* InTexture, D3D12_RESOURCE_STATES InState, uint32 InSubResource) final override;
 
 	//
 	// The Following functions are the _RenderThread version of the above functions. They allow the RHI to control the thread synchronization for greater efficiency.
@@ -1127,6 +1141,8 @@ public:
 	virtual bool HandleSpecialUnlock(FRHICommandListBase* RHICmdList, uint32 MipIndex, ETextureCreateFlags InFlags, const struct FD3D12TextureLayout& TextureLayout, void* RawTextureMemory) { return false; }
 
 	FD3D12Adapter& GetAdapter(uint32_t Index = 0) { return *ChosenAdapters[Index]; }
+	const FD3D12Adapter& GetAdapter(uint32_t Index = 0) const { return *ChosenAdapters[Index]; }
+
 	int32 GetNumAdapters() const { return ChosenAdapters.Num(); }
 
 	bool IsPixEventEnabled() const { return bPixEventEnabled; }
@@ -1233,7 +1249,7 @@ protected:
 	void SetupRecursiveResources();
 
 	// This should only be called by Dynamic RHI member functions
-	inline FD3D12Device* GetRHIDevice(uint32 GPUIndex)
+	inline FD3D12Device* GetRHIDevice(uint32 GPUIndex) const
 	{
 		return GetAdapter().GetDevice(GPUIndex);
 	}
diff --git a/Engine/Source/Runtime/D3D12RHI/Private/D3D12Texture.cpp b/Engine/Source/Runtime/D3D12RHI/Private/D3D12Texture.cpp
--- a/Engine/Source/Runtime/D3D12RHI/Private/D3D12Texture.cpp
+++ b/Engine/Source/Runtime/D3D12RHI/Private/D3D12Texture.cpp
@@ -3036,13 +3036,6 @@ void FD3D12DynamicRHI::RHIVirtualTextureSetFirstMipVisible(FRHITexture2D* Textur
 {
 }
 
-ID3D12CommandQueue* FD3D12DynamicRHI::RHIGetD3DCommandQueue()
-{
-	// Multi-GPU support : any code using this function needs validation.
-	return GetAdapter().GetDevice(0)->GetCommandListManager().GetD3DCommandQueue();
-}
-
-
 template<typename BaseResourceType>
 TD3D12Texture2D<BaseResourceType>* FD3D12DynamicRHI::CreateTextureFromResource(bool bTextureArray, bool bCubeTexture, EPixelFormat Format, ETextureCreateFlags TexCreateFlags, const FClearValueBinding& ClearValueBinding, ID3D12Resource* Resource)
 {
@@ -3257,7 +3250,7 @@ FTexture2DRHIRef FD3D12DynamicRHI::RHICreateTexture2DFromResource(EPixelFormat F
 	return CreateTextureFromResource<FD3D12BaseTexture2D>(false, false, Format, TexCreateFlags, ClearValueBinding, Resource);
 }
 
-FTexture2DRHIRef FD3D12DynamicRHI::RHICreateTexture2DArrayFromResource(EPixelFormat Format, ETextureCreateFlags TexCreateFlags, const FClearValueBinding& ClearValueBinding, ID3D12Resource* Resource)
+FTexture2DArrayRHIRef FD3D12DynamicRHI::RHICreateTexture2DArrayFromResource(EPixelFormat Format, ETextureCreateFlags TexCreateFlags, const FClearValueBinding& ClearValueBinding, ID3D12Resource* Resource)
 {
 	return CreateTextureFromResource<FD3D12BaseTexture2DArray>(true, false, Format, TexCreateFlags, ClearValueBinding, Resource);
 }
diff --git a/Engine/Source/Runtime/D3D12RHI/Private/HoloLens/D3D12RHIBasePrivate.h b/Engine/Source/Runtime/D3D12RHI/Private/HoloLens/D3D12RHIBasePrivate.h
--- a/Engine/Source/Runtime/D3D12RHI/Private/HoloLens/D3D12RHIBasePrivate.h
+++ b/Engine/Source/Runtime/D3D12RHI/Private/HoloLens/D3D12RHIBasePrivate.h
@@ -6,28 +6,7 @@
 
 #pragma once
 
-// Disable macro redefinition warning for compatibility with Windows SDK 8+
-#pragma warning(push)
-#pragma warning(disable : 4005)	// macro redefinition
-
-// D3D headers.
-#if PLATFORM_64BITS
-#pragma pack(push,16)
-#else
-#pragma pack(push,8)
-#endif
-#define D3D_OVERLOADS 1
-#include "HoloLens/AllowWindowsPlatformTypes.h"
-#include <d3d12.h>
-#include <d3dx12.h>
-#include <dxgidebug.h>
-#include <d3d12sdklayers.h>
-#include "HoloLens/HideWindowsPlatformTypes.h"
-
-#undef DrawText
-
-#pragma pack(pop)
-#pragma warning(pop)
+#include "HoloLens/D3D12ThirdParty.h"
 
 #define D3D12RHI_RESOURCE_FLAG_ALLOW_INDIRECT_BUFFER D3D12_RESOURCE_FLAG_NONE
 #define D3D12RHI_HEAP_FLAG_ALLOW_INDIRECT_BUFFERS		D3D12_HEAP_FLAG_NONE
diff --git a/Engine/Source/Runtime/D3D12RHI/Private/Windows/D3D12RHIBasePrivate.h b/Engine/Source/Runtime/D3D12RHI/Private/Windows/D3D12RHIBasePrivate.h
--- a/Engine/Source/Runtime/D3D12RHI/Private/Windows/D3D12RHIBasePrivate.h
+++ b/Engine/Source/Runtime/D3D12RHI/Private/Windows/D3D12RHIBasePrivate.h
@@ -6,29 +6,7 @@
 
 #pragma once
 
-// Disable macro redefinition warning for compatibility with Windows SDK 8+
-#pragma warning(push)
-#pragma warning(disable : 4005)	// macro redefinition
-
-// D3D headers.
-#if PLATFORM_64BITS
-#pragma pack(push,16)
-#else
-#pragma pack(push,8)
-#endif
-#define D3D_OVERLOADS 1
-#include "Windows/AllowWindowsPlatformTypes.h"
-THIRD_PARTY_INCLUDES_START
-#include <d3d12.h>
-#include <d3dx12.h>
-#include <d3d12sdklayers.h>
-THIRD_PARTY_INCLUDES_END
-#include "Windows/HideWindowsPlatformTypes.h"
-
-#undef DrawText
-
-#pragma pack(pop)
-#pragma warning(pop)
+#include "Windows/D3D12ThirdParty.h"
 
 #define D3D12RHI_RESOURCE_FLAG_ALLOW_INDIRECT_BUFFER D3D12_RESOURCE_FLAG_NONE
 #define D3D12RHI_HEAP_FLAG_ALLOW_INDIRECT_BUFFERS		D3D12_HEAP_FLAG_NONE
diff --git a/Engine/Source/Runtime/D3D12RHI/Public/HoloLens/D3D12ThirdParty.h b/Engine/Source/Runtime/D3D12RHI/Public/HoloLens/D3D12ThirdParty.h
new file mode 100644
--- /dev/null
+++ b/Engine/Source/Runtime/D3D12RHI/Public/HoloLens/D3D12ThirdParty.h
@@ -0,0 +1,27 @@
+// Copyright Epic Games, Inc. All Rights Reserved.
+
+#pragma once
+
+// Disable macro redefinition warning for compatibility with Windows SDK 8+
+#pragma warning(push)
+#pragma warning(disable : 4005)	// macro redefinition
+
+#if PLATFORM_64BITS
+	#pragma pack(push,16)
+#else
+	#pragma pack(push,8)
+#endif
+
+#include "Microsoft/AllowMicrosoftPlatformTypes.h"
+THIRD_PARTY_INCLUDES_START
+	#include <d3d12.h>
+	#include <d3dx12.h>
+	#include <dxgi.h>
+	#include <dxgidebug.h>
+THIRD_PARTY_INCLUDES_END
+#include "Microsoft/HideMicrosoftPlatformTypes.h"
+
+#undef DrawText
+
+#pragma pack(pop)
+#pragma warning(pop)
diff --git a/Engine/Source/Runtime/D3D12RHI/Public/ID3D12DynamicRHI.h b/Engine/Source/Runtime/D3D12RHI/Public/ID3D12DynamicRHI.h
new file mode 100644
--- /dev/null
+++ b/Engine/Source/Runtime/D3D12RHI/Public/ID3D12DynamicRHI.h
@@ -0,0 +1,64 @@
+// Copyright Epic Games, Inc. All Rights Reserved.
+
+#pragma once
+
+#include "DynamicRHI.h"
+
+#if PLATFORM_WINDOWS
+	#include "Windows/D3D12ThirdParty.h"
+#elif PLATFORM_HOLOLENS
+	#include "HoloLens/D3D12ThirdParty.h"
+#else
+	#include "D3D12ThirdParty.h"
+#endif
+
+struct FD3D12MinimalAdapterDesc
+{
+	DXGI_ADAPTER_DESC Desc{};
+	uint32 NumDeviceNodes{};
+};
+
+struct ID3D12DynamicRHI : public FDynamicRHI
+{
+	virtual ERHIInterfaceType     GetInterfaceType() const override final { return ERHIInterfaceType::D3D12; }
+
+	virtual TArray<FD3D12MinimalAdapterDesc> RHIGetAdapterDescs() const = 0;
+	virtual bool                             RHIIsPixEnabled() const = 0;
+
+	virtual ID3D12CommandQueue*        RHIGetCommandQueue() const = 0;
+	virtual ID3D12Device*              RHIGetDevice(uint32 InIndex) const = 0;
+	virtual uint32                     RHIGetDeviceNodeMask(uint32 InIndex) const = 0;
+	virtual ID3D12GraphicsCommandList* RHIGetGraphicsCommandList(uint32 InDeviceIndex) const = 0;
+	virtual DXGI_FORMAT                RHIGetSwapChainFormat(EPixelFormat InFormat) const = 0;
+
+	virtual FTexture2DRHIRef      RHICreateTexture2DFromResource(EPixelFormat Format, ETextureCreateFlags TexCreateFlags, const FClearValueBinding& ClearValueBinding, ID3D12Resource* Resource) = 0;
+	virtual FTexture2DArrayRHIRef RHICreateTexture2DArrayFromResource(EPixelFormat Format, ETextureCreateFlags TexCreateFlags, const FClearValueBinding& ClearValueBinding, ID3D12Resource* Resource) = 0;
+	virtual FTextureCubeRHIRef    RHICreateTextureCubeFromResource(EPixelFormat Format, ETextureCreateFlags TexCreateFlags, const FClearValueBinding& ClearValueBinding, ID3D12Resource* Resource) = 0;
+
+	virtual ID3D12Resource*       RHIGetResource(FRHIBuffer* InBuffer) const = 0;
+	virtual uint32                RHIGetResourceDeviceIndex(FRHIBuffer* InBuffer) const = 0;
+	virtual ID3D12Resource*       RHIGetResource(FRHITexture* InTexture) const = 0;
+	virtual uint32                RHIGetResourceDeviceIndex(FRHITexture* InTexture) const = 0;
+	virtual int64                 RHIGetResourceMemorySize(FRHITexture* InTexture) const = 0;
+	virtual bool                  RHIIsResourcePlaced(FRHITexture* InTexture) const = 0;
+
+	virtual D3D12_CPU_DESCRIPTOR_HANDLE RHIGetRenderTargetView(FRHITexture* InTexture, int32 InMipIndex = 0, int32 InArraySliceIndex = 0) const = 0;
+
+	virtual void                  RHIFinishExternalComputeWork(uint32 InDeviceIndex, ID3D12GraphicsCommandList* InCommandList) = 0;
+
+	virtual void                  RHIRegisterWork(uint32 InDeviceIndex, uint32 NumPrimitives) = 0;
+	virtual void                  RHIAddPendingBarrier(FRHITexture* InTexture, D3D12_RESOURCE_STATES InState, uint32 InSubResource) = 0;
+
+	static D3D12RHI_API bool      IsD3DDebugEnabled();
+};
+
+inline bool IsRHID3D12()
+{
+	return GDynamicRHI != nullptr && GDynamicRHI->GetInterfaceType() == ERHIInterfaceType::D3D12;
+}
+
+inline ID3D12DynamicRHI* GetID3D12DynamicRHI()
+{
+	check(GDynamicRHI->GetInterfaceType() == ERHIInterfaceType::D3D12);
+	return GetDynamicRHI<ID3D12DynamicRHI>();
+}
diff --git a/Engine/Source/Runtime/D3D12RHI/Public/Windows/D3D12ThirdParty.h b/Engine/Source/Runtime/D3D12RHI/Public/Windows/D3D12ThirdParty.h
new file mode 100644
--- /dev/null
+++ b/Engine/Source/Runtime/D3D12RHI/Public/Windows/D3D12ThirdParty.h
@@ -0,0 +1,27 @@
+// Copyright Epic Games, Inc. All Rights Reserved.
+
+#pragma once
+
+// Disable macro redefinition warning for compatibility with Windows SDK 8+
+#pragma warning(push)
+#pragma warning(disable : 4005)	// macro redefinition
+
+#if PLATFORM_64BITS
+	#pragma pack(push,16)
+#else
+	#pragma pack(push,8)
+#endif
+
+#include "Microsoft/AllowMicrosoftPlatformTypes.h"
+THIRD_PARTY_INCLUDES_START
+	#include <d3d12.h>
+	#include <d3dx12.h>
+	#include <dxgi.h>
+	#include <dxgidebug.h>
+THIRD_PARTY_INCLUDES_END
+#include "Microsoft/HideMicrosoftPlatformTypes.h"
+
+#undef DrawText
+
+#pragma pack(pop)
+#pragma warning(pop)
diff --git a/Engine/Source/Runtime/GameplayMediaEncoder/GameplayMediaEncoder.Build.cs b/Engine/Source/Runtime/GameplayMediaEncoder/GameplayMediaEncoder.Build.cs
--- a/Engine/Source/Runtime/GameplayMediaEncoder/GameplayMediaEncoder.Build.cs
+++ b/Engine/Source/Runtime/GameplayMediaEncoder/GameplayMediaEncoder.Build.cs
@@ -28,22 +28,13 @@ public class GameplayMediaEncoder : ModuleRules
 		});
 
 		if (Target.Platform == UnrealTargetPlatform.Win64)
-		{			
-			AddEngineThirdPartyPrivateStaticDependencies(Target, "NVAftermath");
-
-			if (Target.Platform == UnrealTargetPlatform.Win64)
-			{
-				PrivateDependencyModuleNames.AddRange(new string[]
-					{
-						"D3D11RHI"
-					});
-
-				PublicDelayLoadDLLs.Add("mfplat.dll");
-				PublicDelayLoadDLLs.Add("mfuuid.dll");
-				PublicDelayLoadDLLs.Add("Mfreadwrite.dll");
+		{
+			PrivateDependencyModuleNames.Add("D3D11RHI");
+			AddEngineThirdPartyPrivateStaticDependencies(Target, "DX11");
 
-				PublicSystemLibraries.Add("d3d11.lib");
-			}
+			PublicDelayLoadDLLs.Add("mfplat.dll");
+			PublicDelayLoadDLLs.Add("mfuuid.dll");
+			PublicDelayLoadDLLs.Add("Mfreadwrite.dll");
 		}
 	}
 }
diff --git a/Engine/Source/Runtime/Windows/D3D11RHI/Private/D3D11RHI.cpp b/Engine/Source/Runtime/Windows/D3D11RHI/Private/D3D11RHI.cpp
--- a/Engine/Source/Runtime/Windows/D3D11RHI/Private/D3D11RHI.cpp
+++ b/Engine/Source/Runtime/Windows/D3D11RHI/Private/D3D11RHI.cpp
@@ -761,3 +761,60 @@ void UpdateBufferStats(TRefCountPtr<ID3D11Buffer> Buffer, bool bAllocating)
 #endif
 	}
 }
+
+ID3D11Device* FD3D11DynamicRHI::RHIGetDevice() const
+{
+	return GetDevice();
+}
+
+ID3D11DeviceContext* FD3D11DynamicRHI::RHIGetDeviceContext() const
+{
+	return GetDeviceContext();
+}
+
+IDXGISwapChain* FD3D11DynamicRHI::RHIGetSwapChain(FRHIViewport* InViewport) const
+{
+	FD3D11Viewport* Viewport = static_cast<FD3D11Viewport*>(InViewport);
+	return Viewport->GetSwapChain();
+}
+
+DXGI_FORMAT FD3D11DynamicRHI::RHIGetSwapChainFormat(EPixelFormat InFormat) const
+{
+	const DXGI_FORMAT PlatformFormat = ::FindDepthStencilDXGIFormat(static_cast<DXGI_FORMAT>(GPixelFormats[InFormat].PlatformFormat));
+	return ::FindShaderResourceDXGIFormat(PlatformFormat, true);
+}
+
+ID3D11Buffer* FD3D11DynamicRHI::RHIGetResource(FRHIBuffer* InBuffer) const
+{
+	FD3D11Buffer* Buffer = ResourceCast(InBuffer);
+	return Buffer->Resource;
+}
+
+ID3D11Resource* FD3D11DynamicRHI::RHIGetResource(FRHITexture* InTexture) const
+{
+	FD3D11TextureBase* D3D11Texture = GetD3D11TextureFromRHITexture(InTexture);
+	return D3D11Texture->GetResource();
+}
+
+int64 FD3D11DynamicRHI::RHIGetResourceMemorySize(FRHITexture* InTexture) const
+{
+	FD3D11TextureBase* D3D11Texture = GetD3D11TextureFromRHITexture(InTexture);
+	return D3D11Texture->GetMemorySize();
+}
+
+ID3D11RenderTargetView* FD3D11DynamicRHI::RHIGetRenderTargetView(FRHITexture* InTexture, int32 InMipIndex, int32 InArraySliceIndex) const
+{
+	FD3D11TextureBase* D3D11Texture = GetD3D11TextureFromRHITexture(InTexture);
+	return D3D11Texture->GetRenderTargetView(InMipIndex, InArraySliceIndex);
+}
+
+ID3D11ShaderResourceView* FD3D11DynamicRHI::RHIGetShaderResourceView(FRHITexture* InTexture) const
+{
+	FD3D11TextureBase* D3D11Texture = GetD3D11TextureFromRHITexture(InTexture);
+	return D3D11Texture->GetShaderResourceView();
+}
+
+void FD3D11DynamicRHI::RHIRegisterWork(uint32 NumPrimitives)
+{
+	RegisterGPUWork(NumPrimitives);
+}
diff --git a/Engine/Source/Runtime/Windows/D3D11RHI/Private/D3D11RHIPrivate.h b/Engine/Source/Runtime/Windows/D3D11RHI/Private/D3D11RHIPrivate.h
--- a/Engine/Source/Runtime/Windows/D3D11RHI/Private/D3D11RHIPrivate.h
+++ b/Engine/Source/Runtime/Windows/D3D11RHI/Private/D3D11RHIPrivate.h
@@ -7,6 +7,7 @@
 #pragma once
 
 #include "CoreMinimal.h"
+#include "ID3D11DynamicRHI.h"
 #include "D3D11RHI.h"
 // Dependencies.
 #include "RHI.h"
@@ -18,7 +19,6 @@
 
 DECLARE_LOG_CATEGORY_EXTERN(LogD3D11RHI, Log, All);
 
-#include "Windows/D3D11RHIBasePrivate.h"
 #include "Containers/StaticArray.h"
 
 // D3D RHI public headers.
@@ -399,7 +399,7 @@ struct FD3D11Adapter
 PRAGMA_DISABLE_DEPRECATION_WARNINGS
 
 /** The interface which is implemented by the dynamically bound RHI. */
-class D3D11RHI_API FD3D11DynamicRHI : public FDynamicRHIPSOFallback, public IRHICommandContextPSOFallback
+class D3D11RHI_API FD3D11DynamicRHI : public ID3D11DynamicRHI, public IRHICommandContextPSOFallback
 {
 public:
 	typedef TMap<FD3D11LockedKey, FD3D11LockedData> FD3D11LockTracker;
@@ -420,7 +420,6 @@ public:
 	virtual void PostInit() override;
 	virtual void Shutdown() override;
 	virtual const TCHAR* GetName() override { return TEXT("D3D11"); }
-	virtual ERHIInterfaceType GetInterfaceType() const override final { return ERHIInterfaceType::D3D11; }
 
 #if PLATFORM_HOLOLENS
 	virtual void RHISuspendRendering() override;
@@ -630,10 +629,6 @@ public:
 	virtual void RHIPushEvent(const TCHAR* Name, FColor Color) final override;
 	virtual void RHIPopEvent() final override;
 
-	virtual FTexture2DRHIRef RHICreateTexture2DFromResource(EPixelFormat Format, ETextureCreateFlags TexCreateFlags, const FClearValueBinding& ClearValueBinding, ID3D11Texture2D* Resource);
-	virtual FTexture2DArrayRHIRef RHICreateTexture2DArrayFromResource(EPixelFormat Format, ETextureCreateFlags TexCreateFlags, const FClearValueBinding&, ID3D11Texture2D* Resource);
-	virtual FTextureCubeRHIRef RHICreateTextureCubeFromResource(EPixelFormat Format, ETextureCreateFlags TexCreateFlags, const FClearValueBinding& ClearValueBinding, ID3D11Texture2D* Resource);
-	
 	virtual void RHIPerFrameRHIFlushComplete() final override;
 	virtual void RHIPollRenderQueryResults() final override;
 
@@ -676,6 +671,21 @@ public:
 
 	virtual void RHICalibrateTimers() override;
 
+	// ID3D11DynamicRHI interface
+	virtual ID3D11Device*         RHIGetDevice() const final override;
+	virtual ID3D11DeviceContext*  RHIGetDeviceContext() const final override;
+	virtual IDXGISwapChain*       RHIGetSwapChain(FRHIViewport* InViewport) const final override;
+	virtual DXGI_FORMAT           RHIGetSwapChainFormat(EPixelFormat InFormat) const final override;
+	virtual FTexture2DRHIRef      RHICreateTexture2DFromResource(EPixelFormat Format, ETextureCreateFlags TexCreateFlags, const FClearValueBinding& ClearValueBinding, ID3D11Texture2D* Resource) final override;
+	virtual FTexture2DArrayRHIRef RHICreateTexture2DArrayFromResource(EPixelFormat Format, ETextureCreateFlags TexCreateFlags, const FClearValueBinding&, ID3D11Texture2D* Resource) final override;
+	virtual FTextureCubeRHIRef    RHICreateTextureCubeFromResource(EPixelFormat Format, ETextureCreateFlags TexCreateFlags, const FClearValueBinding& ClearValueBinding, ID3D11Texture2D* Resource) final override;
+	virtual ID3D11Buffer*         RHIGetResource(FRHIBuffer* InBuffer) const final override;
+	virtual ID3D11Resource*       RHIGetResource(FRHITexture* InTexture) const final override;
+	virtual int64                 RHIGetResourceMemorySize(FRHITexture* InTexture) const final override;
+	virtual ID3D11RenderTargetView* RHIGetRenderTargetView(FRHITexture* InTexture, int32 InMipIndex = 0, int32 InArraySliceIndex = -1) const final override;
+	virtual ID3D11ShaderResourceView* RHIGetShaderResourceView(FRHITexture* InTexture) const final override;
+	virtual void                  RHIRegisterWork(uint32 NumPrimitives) final override;
+
 	// Accessors.
 	ID3D11Device* GetDevice() const
 	{
diff --git a/Engine/Source/Runtime/Windows/D3D11RHI/Private/HoloLens/D3D11RHIBasePrivate.h b/Engine/Source/Runtime/Windows/D3D11RHI/Private/HoloLens/D3D11RHIBasePrivate.h
deleted file mode 100644
diff --git a/Engine/Source/Runtime/Windows/D3D11RHI/Private/Windows/D3D11RHIBasePrivate.h b/Engine/Source/Runtime/Windows/D3D11RHI/Private/Windows/D3D11RHIBasePrivate.h
deleted file mode 100644
diff --git a/Engine/Source/Runtime/Windows/D3D11RHI/Public/HoloLens/D3D11ThirdParty.h b/Engine/Source/Runtime/Windows/D3D11RHI/Public/HoloLens/D3D11ThirdParty.h
new file mode 100644
--- /dev/null
+++ b/Engine/Source/Runtime/Windows/D3D11RHI/Public/HoloLens/D3D11ThirdParty.h
@@ -0,0 +1,25 @@
+// Copyright Epic Games, Inc. All Rights Reserved.
+
+#pragma once
+
+// Disable macro redefinition warning for compatibility with Windows SDK 8+
+#pragma warning(push)
+#pragma warning(disable : 4005)	// macro redefinition
+
+#if PLATFORM_64BITS
+	#pragma pack(push,16)
+#else
+	#pragma pack(push,8)
+#endif
+
+#include "Microsoft/AllowMicrosoftPlatformTypes.h"
+THIRD_PARTY_INCLUDES_START
+	#include <d3d11_2.h>
+	#include <dxgidebug.h>
+THIRD_PARTY_INCLUDES_END
+#include "Microsoft/HideMicrosoftPlatformTypes.h"
+
+#undef DrawText
+
+#pragma pack(pop)
+#pragma warning(pop)
diff --git a/Engine/Source/Runtime/Windows/D3D11RHI/Public/ID3D11DynamicRHI.h b/Engine/Source/Runtime/Windows/D3D11RHI/Public/ID3D11DynamicRHI.h
new file mode 100644
--- /dev/null
+++ b/Engine/Source/Runtime/Windows/D3D11RHI/Public/ID3D11DynamicRHI.h
@@ -0,0 +1,47 @@
+// Copyright Epic Games, Inc. All Rights Reserved.
+#pragma once
+
+#include "DynamicRHI.h"
+
+#if PLATFORM_WINDOWS
+	#include "Windows/D3D11ThirdParty.h"
+#elif PLATFORM_HOLOLENS
+	#include "HoloLens/D3D11ThirdParty.h"
+#endif
+
+typedef ID3D11DeviceContext FD3D11DeviceContext;
+typedef ID3D11Device FD3D11Device;
+
+struct ID3D11DynamicRHI : public FDynamicRHIPSOFallback
+{
+	virtual ERHIInterfaceType     GetInterfaceType() const override final { return ERHIInterfaceType::D3D11; }
+
+	virtual ID3D11Device*         RHIGetDevice() const = 0;
+	virtual ID3D11DeviceContext*  RHIGetDeviceContext() const = 0;
+	virtual IDXGISwapChain*       RHIGetSwapChain(FRHIViewport* InViewport) const = 0;
+	virtual DXGI_FORMAT           RHIGetSwapChainFormat(EPixelFormat InFormat) const = 0;
+
+	virtual FTexture2DRHIRef      RHICreateTexture2DFromResource(EPixelFormat Format, ETextureCreateFlags TexCreateFlags, const FClearValueBinding& ClearValueBinding, ID3D11Texture2D* Resource) = 0;
+	virtual FTexture2DArrayRHIRef RHICreateTexture2DArrayFromResource(EPixelFormat Format, ETextureCreateFlags TexCreateFlags, const FClearValueBinding& ClearValueBinding, ID3D11Texture2D* Resource) = 0;
+	virtual FTextureCubeRHIRef    RHICreateTextureCubeFromResource(EPixelFormat Format, ETextureCreateFlags TexCreateFlags, const FClearValueBinding& ClearValueBinding, ID3D11Texture2D* Resource) = 0;
+
+	virtual ID3D11Buffer*         RHIGetResource(FRHIBuffer* InBuffer) const = 0;
+	virtual ID3D11Resource*       RHIGetResource(FRHITexture* InTexture) const = 0;
+	virtual int64                 RHIGetResourceMemorySize(FRHITexture* InTexture) const = 0;
+
+	virtual ID3D11RenderTargetView*   RHIGetRenderTargetView(FRHITexture* InTexture, int32 InMipIndex = 0, int32 InArraySliceIndex = -1) const = 0;
+	virtual ID3D11ShaderResourceView* RHIGetShaderResourceView(FRHITexture* InTexture) const = 0;
+
+	virtual void                  RHIRegisterWork(uint32 NumPrimitives) = 0;
+};
+
+inline bool IsRHID3D11()
+{
+	return GDynamicRHI != nullptr && GDynamicRHI->GetInterfaceType() == ERHIInterfaceType::D3D11;
+}
+
+inline ID3D11DynamicRHI* GetID3D11DynamicRHI()
+{
+	check(GDynamicRHI->GetInterfaceType() == ERHIInterfaceType::D3D11);
+	return GetDynamicRHI<ID3D11DynamicRHI>();
+}
diff --git a/Engine/Source/Runtime/Windows/D3D11RHI/Public/Windows/D3D11ThirdParty.h b/Engine/Source/Runtime/Windows/D3D11RHI/Public/Windows/D3D11ThirdParty.h
new file mode 100644
--- /dev/null
+++ b/Engine/Source/Runtime/Windows/D3D11RHI/Public/Windows/D3D11ThirdParty.h
@@ -0,0 +1,25 @@
+// Copyright Epic Games, Inc. All Rights Reserved.
+
+#pragma once
+
+// Disable macro redefinition warning for compatibility with Windows SDK 8+
+#pragma warning(push)
+#pragma warning(disable : 4005)	// macro redefinition
+
+#if PLATFORM_64BITS
+	#pragma pack(push,16)
+#else
+	#pragma pack(push,8)
+#endif
+
+#include "Microsoft/AllowMicrosoftPlatformTypes.h"
+THIRD_PARTY_INCLUDES_START
+	#include <d3d11_2.h>
+	#include <dxgidebug.h>
+THIRD_PARTY_INCLUDES_END
+#include "Microsoft/HideMicrosoftPlatformTypes.h"
+
+#undef DrawText
+
+#pragma pack(pop)
+#pragma warning(pop)
