nDisplay: Remove dependency on TextureShare and fix and required interfaces.

diff --git a/Engine/Plugins/Runtime/nDisplay/Source/DisplayCluster/DisplayCluster.Build.cs b/Engine/Plugins/Runtime/nDisplay/Source/DisplayCluster/DisplayCluster.Build.cs
--- a/Engine/Plugins/Runtime/nDisplay/Source/DisplayCluster/DisplayCluster.Build.cs
+++ b/Engine/Plugins/Runtime/nDisplay/Source/DisplayCluster/DisplayCluster.Build.cs
@@ -23,14 +23,6 @@ public class DisplayCluster : ModuleRules
 				"Engine"
 			});
 
-		if (Target.Platform == UnrealTargetPlatform.Win64)
-		{
-			PublicDependencyModuleNames.AddRange(
-			new string[] {
-				"TextureShare"
-			});
-		}
-
 		PrivateDependencyModuleNames.AddRange(
 		new string[] {
 			"HeadMountedDisplay",
@@ -71,8 +63,6 @@ public class DisplayCluster : ModuleRules
 				new string[] {
 					"D3D11RHI",
 					"D3D12RHI",
-					"TextureShare",
-					"TextureShareCore",
 			});
 
 			AddEngineThirdPartyPrivateStaticDependencies(Target, "DX11", "DX12");
diff --git a/Engine/Plugins/Runtime/nDisplay/Source/DisplayCluster/Private/Render/Viewport/Configuration/DisplayClusterViewportConfiguration.cpp b/Engine/Plugins/Runtime/nDisplay/Source/DisplayCluster/Private/Render/Viewport/Configuration/DisplayClusterViewportConfiguration.cpp
--- a/Engine/Plugins/Runtime/nDisplay/Source/DisplayCluster/Private/Render/Viewport/Configuration/DisplayClusterViewportConfiguration.cpp
+++ b/Engine/Plugins/Runtime/nDisplay/Source/DisplayCluster/Private/Render/Viewport/Configuration/DisplayClusterViewportConfiguration.cpp
@@ -104,12 +104,14 @@ bool FDisplayClusterViewportConfiguration::ImplUpdateConfiguration(EDisplayClust
 				RenderFrameSettings.PreviewMaxGPUIndex = InPreviewSettings->MaxGPUIndex;
 
 				RenderFrameSettings.bIsRenderingInEditor = true;
+				RenderFrameSettings.bIsPreviewRendering = true;
 			}
 			else
 			{
 				RenderFrameSettings.bPreviewEnablePostProcess = false;
 				RenderFrameSettings.bAllowMultiGPURenderingInEditor = false;
 				RenderFrameSettings.bIsRenderingInEditor = false;
+				RenderFrameSettings.bIsPreviewRendering = false;
 			}
 
 			ConfigurationBase.Update(InClusterNodeId);
@@ -126,8 +128,7 @@ bool FDisplayClusterViewportConfiguration::ImplUpdateConfiguration(EDisplayClust
 
 			ImplUpdateConfigurationVisibility(*RootActor, *ConfigurationData);
 
-			ConfigurationBase.UpdateClusterNodePostProcess(InClusterNodeId);
-			ConfigurationBase.UpdateTextureShare(InClusterNodeId);
+			ConfigurationBase.UpdateClusterNodePostProcess(InClusterNodeId, RenderFrameSettings);
 
 			ImplPostUpdateRenderFrameConfiguration();
 
diff --git a/Engine/Plugins/Runtime/nDisplay/Source/DisplayCluster/Private/Render/Viewport/Configuration/DisplayClusterViewportConfigurationBase.cpp b/Engine/Plugins/Runtime/nDisplay/Source/DisplayCluster/Private/Render/Viewport/Configuration/DisplayClusterViewportConfigurationBase.cpp
--- a/Engine/Plugins/Runtime/nDisplay/Source/DisplayCluster/Private/Render/Viewport/Configuration/DisplayClusterViewportConfigurationBase.cpp
+++ b/Engine/Plugins/Runtime/nDisplay/Source/DisplayCluster/Private/Render/Viewport/Configuration/DisplayClusterViewportConfigurationBase.cpp
@@ -89,7 +89,7 @@ void FDisplayClusterViewportConfigurationBase::Update(const FString& ClusterNode
 	}
 }
 
-void FDisplayClusterViewportConfigurationBase::UpdateClusterNodePostProcess(const FString& InClusterNodeId)
+void FDisplayClusterViewportConfigurationBase::UpdateClusterNodePostProcess(const FString& InClusterNodeId, const FDisplayClusterRenderFrameSettings& InRenderFrameSettings)
 {
 	const UDisplayClusterConfigurationClusterNode* ClusterNode = ConfigurationData.Cluster->GetNode(InClusterNodeId);
 	if (ClusterNode)
@@ -132,6 +132,28 @@ void FDisplayClusterViewportConfigurationBase::UpdateClusterNodePostProcess(cons
 				}
 			}
 
+			// Texture sharing is not supported in preview mode.
+			if(InRenderFrameSettings.bIsPreviewRendering == false)
+			{
+				static const FString TextureShareID(TEXT("TextureShare"));
+				if (ClusterNode->bEnableTextureShare)
+				{
+					TSharedPtr<IDisplayClusterPostProcess, ESPMode::ThreadSafe> ExistPostProcess = PPManager->FindPostProcess(TextureShareID);
+					if (!ExistPostProcess.IsValid())
+					{
+						// Helper create postprocess entry for TextureShare
+						FDisplayClusterConfigurationPostprocess TextureShareConfiguration;
+						TextureShareConfiguration.Type = TextureShareID;
+
+						PPManager->CreatePostprocess(TextureShareID, &TextureShareConfiguration);
+					}
+				}
+				else
+				{
+					PPManager->RemovePostprocess(TextureShareID);
+				}
+			}
+
 			// Update OutputRemap PP
 			{
 				const struct FDisplayClusterConfigurationFramePostProcess_OutputRemap& OutputRemapCfg = ClusterNode->OutputRemap;
@@ -194,19 +216,3 @@ bool FDisplayClusterViewportConfigurationBase::UpdateViewportConfiguration(FDisp
 
 	return true;
 }
-
-void FDisplayClusterViewportConfigurationBase::UpdateTextureShare(const FString& ClusterNodeId)
-{
-	for (FDisplayClusterViewport* ViewportIt : ViewportManager.ImplGetViewports())
-	{
-		if (ViewportIt)
-		{
-			const UDisplayClusterConfigurationViewport* ViewportCfg = ConfigurationData.GetViewport(ClusterNodeId, ViewportIt->GetId());
-			if (ViewportCfg)
-			{
-				ViewportIt->TextureShare.UpdateConfiguration(*ViewportIt, ViewportCfg->TextureShare);
-			}
-		}
-	}
-}
-
diff --git a/Engine/Plugins/Runtime/nDisplay/Source/DisplayCluster/Private/Render/Viewport/Configuration/DisplayClusterViewportConfigurationBase.h b/Engine/Plugins/Runtime/nDisplay/Source/DisplayCluster/Private/Render/Viewport/Configuration/DisplayClusterViewportConfigurationBase.h
--- a/Engine/Plugins/Runtime/nDisplay/Source/DisplayCluster/Private/Render/Viewport/Configuration/DisplayClusterViewportConfigurationBase.h
+++ b/Engine/Plugins/Runtime/nDisplay/Source/DisplayCluster/Private/Render/Viewport/Configuration/DisplayClusterViewportConfigurationBase.h
@@ -3,6 +3,7 @@
 #pragma once
 
 #include "CoreMinimal.h"
+#include "Render/Viewport/RenderFrame/DisplayClusterRenderFrameSettings.h"
 
 class FDisplayClusterViewportManager;
 class FDisplayClusterViewport;
@@ -24,8 +25,7 @@ public:
 
 public:
 	void Update(const FString& ClusterNodeId);
-	void UpdateClusterNodePostProcess(const FString& ClusterNodeId);
-	void UpdateTextureShare(const FString& ClusterNodeId);
+	void UpdateClusterNodePostProcess(const FString& ClusterNodeId, const FDisplayClusterRenderFrameSettings& InRenderFrameSettings);
 
 public:
 	static bool UpdateViewportConfiguration(FDisplayClusterViewportManager& ViewportManager, ADisplayClusterRootActor& RootActor, FDisplayClusterViewport* DesiredViewport, const UDisplayClusterConfigurationViewport* ConfigurationViewport);
diff --git a/Engine/Plugins/Runtime/nDisplay/Source/DisplayCluster/Private/Render/Viewport/DisplayClusterViewport.h b/Engine/Plugins/Runtime/nDisplay/Source/DisplayCluster/Private/Render/Viewport/DisplayClusterViewport.h
--- a/Engine/Plugins/Runtime/nDisplay/Source/DisplayCluster/Private/Render/Viewport/DisplayClusterViewport.h
+++ b/Engine/Plugins/Runtime/nDisplay/Source/DisplayCluster/Private/Render/Viewport/DisplayClusterViewport.h
@@ -11,7 +11,6 @@
 #include "Render/Viewport/Containers/DisplayClusterViewport_Enums.h"
 
 #include "Render/Viewport/DisplayClusterViewport_CustomPostProcessSettings.h"
-#include "Render/Viewport/DisplayClusterViewport_TextureShare.h"
 #include "Render/Viewport/DisplayClusterViewport_VisibilitySettings.h"
 
 #include "SceneViewExtensionContext.h"
@@ -70,6 +69,19 @@ public:
 		return RenderSettings;
 	}
 
+	virtual void SetRenderSettings(const FDisplayClusterViewport_RenderSettings& InRenderSettings) override
+	{
+		check(IsInGameThread());
+		RenderSettings = InRenderSettings;
+	}
+
+	virtual void SetContexts(TArray<FDisplayClusterViewport_Context>& InContexts) override
+	{
+		check(IsInGameThread());
+		Contexts.Empty();
+		Contexts.Append(InContexts);
+	}
+
 	virtual void CalculateProjectionMatrix(const uint32 InContextNum, float Left, float Right, float Top, float Bottom, float ZNear, float ZFar, bool bIsAnglesInput) override;
 
 	virtual bool    CalculateView(const uint32 InContextNum, FVector& InOutViewLocation, FRotator& InOutViewRotation, const FVector& ViewOffset, const float WorldToMeters, const float NCP, const float FCP) override;
@@ -199,9 +211,6 @@ public:
 	// OCIO wrapper
 	TSharedPtr<FOpenColorIODisplayExtension, ESPMode::ThreadSafe> OpenColorIODisplayExtension;
 
-	//TextureShare
-	FDisplayClusterViewport_TextureShare TextureShare;
-
 public:
 	// Projection policy instance that serves this viewport
 	TSharedPtr<IDisplayClusterProjectionPolicy, ESPMode::ThreadSafe> ProjectionPolicy;
diff --git a/Engine/Plugins/Runtime/nDisplay/Source/DisplayCluster/Private/Render/Viewport/DisplayClusterViewportManager.cpp b/Engine/Plugins/Runtime/nDisplay/Source/DisplayCluster/Private/Render/Viewport/DisplayClusterViewportManager.cpp
--- a/Engine/Plugins/Runtime/nDisplay/Source/DisplayCluster/Private/Render/Viewport/DisplayClusterViewportManager.cpp
+++ b/Engine/Plugins/Runtime/nDisplay/Source/DisplayCluster/Private/Render/Viewport/DisplayClusterViewportManager.cpp
@@ -253,7 +253,7 @@ bool FDisplayClusterViewportManager::ShouldUseFullSizeFrameTargetableResource()
 {
 	check(IsInGameThread());
 
-	if (PostProcessManager->ShouldUseFullSizeFrameTargetableResource())
+	if (PostProcessManager.IsValid() && PostProcessManager->ShouldUseFullSizeFrameTargetableResource())
 	{
 		return true;
 	}
@@ -273,7 +273,7 @@ bool FDisplayClusterViewportManager::ShouldUseAdditionalFrameTargetableResource(
 {
 	check(IsInGameThread());
 
-	if (PostProcessManager->ShouldUseAdditionalFrameTargetableResource())
+	if (PostProcessManager.IsValid() && PostProcessManager->ShouldUseAdditionalFrameTargetableResource())
 	{
 		return true;
 	}
@@ -404,6 +404,12 @@ bool FDisplayClusterViewportManager::BeginNewFrame(FViewport* InViewport, UWorld
 		}
 	}
 
+	// Before new frame
+	if (PostProcessManager.IsValid())
+	{
+		PostProcessManager->HandleSetupNewFrame();
+	}
+
 	// generate unique stereo view index for each frame
 	// Begin from 1, because INDEX_NONE use ViewState[0] in LocalPlayer.cpp:786
 	uint32 StereoViewIndex = 1;
@@ -445,9 +451,6 @@ bool FDisplayClusterViewportManager::BeginNewFrame(FViewport* InViewport, UWorld
 	{
 		if (Viewport)
 		{
-			// Update TextureShare links
-			Viewport->TextureShare.UpdateLinkSceneContextToShare(*Viewport);
-
 			// Update ViewportRemap geometry
 			Viewport->ViewportRemap.Update(*Viewport, RenderFrameSize);
 		}
@@ -461,6 +464,11 @@ bool FDisplayClusterViewportManager::BeginNewFrame(FViewport* InViewport, UWorld
 	ImplUpdatePreviewRTTResources();
 #endif /*WITH_EDITOR*/
 
+	if (PostProcessManager.IsValid())
+	{
+		PostProcessManager->HandleBeginNewFrame(OutRenderFrame);
+	}
+
 	return true;
 }
 
@@ -526,11 +534,14 @@ void FDisplayClusterViewportManager::FinalizeNewFrame()
 		}
 	}
 
-	// Update postprocess data from game thread
-	PostProcessManager->Tick();
+	if (PostProcessManager.IsValid())
+	{
+		// Update postprocess data from game thread
+		PostProcessManager->Tick();
 
-	// Send updated postprocess data to rendering thread
-	PostProcessManager->FinalizeNewFrame();
+		// Send updated postprocess data to rendering thread
+		PostProcessManager->FinalizeNewFrame();
+	}
 }
 
 FSceneViewFamily::ConstructionValues FDisplayClusterViewportManager::CreateViewFamilyConstructionValues(
diff --git a/Engine/Plugins/Runtime/nDisplay/Source/DisplayCluster/Private/Render/Viewport/DisplayClusterViewportManagerProxy.cpp b/Engine/Plugins/Runtime/nDisplay/Source/DisplayCluster/Private/Render/Viewport/DisplayClusterViewportManagerProxy.cpp
--- a/Engine/Plugins/Runtime/nDisplay/Source/DisplayCluster/Private/Render/Viewport/DisplayClusterViewportManagerProxy.cpp
+++ b/Engine/Plugins/Runtime/nDisplay/Source/DisplayCluster/Private/Render/Viewport/DisplayClusterViewportManagerProxy.cpp
@@ -187,6 +187,12 @@ void FDisplayClusterViewportManagerProxy::ImplRenderFrame(FViewport* InViewport)
 		SCOPED_GPU_STAT(RHICmdList, nDisplay_ViewportManager_RenderFrame);
 		SCOPED_DRAW_EVENT(RHICmdList, nDisplay_ViewportManager_RenderFrame);
 
+		// Handle render setup
+		if (ViewportManagerProxy->PostProcessManager.IsValid())
+		{
+			ViewportManagerProxy->PostProcessManager->HandleRenderFrameSetup_RenderThread(RHICmdList, ViewportManagerProxy);
+		}
+
 		bool bWarpBlendEnabled = ViewportManagerProxy->RenderFrameSettings.bAllowWarpBlend && CVarWarpBlendEnabled.GetValueOnRenderThread() != 0;
 
 		// mGPU not used for in-editor rendering
@@ -200,9 +206,19 @@ void FDisplayClusterViewportManagerProxy::ImplRenderFrame(FViewport* InViewport)
 		// Update viewports resources: overlay, vp-overla, blur, nummips, etc
 		ViewportManagerProxy->UpdateDeferredResources_RenderThread(RHICmdList);
 
+		if (ViewportManagerProxy->PostProcessManager.IsValid())
+		{
+			ViewportManagerProxy->PostProcessManager->HandleBeginUpdateFrameResources_RenderThread(RHICmdList, ViewportManagerProxy);
+		}
+
 		// Update the frame resources: post-processing, warping, and finally resolving everything to the frame resource
 		ViewportManagerProxy->UpdateFrameResources_RenderThread(RHICmdList, bWarpBlendEnabled);
 
+		if (ViewportManagerProxy->PostProcessManager.IsValid())
+		{
+			ViewportManagerProxy->PostProcessManager->HandleEndUpdateFrameResources_RenderThread(RHICmdList, ViewportManagerProxy);
+		}
+
 		if (InViewport)
 		{
 			if (FRHITexture2D* FrameOutputRTT = InViewport->GetRenderTargetTexture())
@@ -294,7 +310,10 @@ void FDisplayClusterViewportManagerProxy::UpdateFrameResources_RenderThread(FRHI
 	check(IsInRenderingThread());
 
 	// Do postprocess before warp&blend
-	PostProcessManager->PerformPostProcessBeforeWarpBlend_RenderThread(RHICmdList, this);
+	if (PostProcessManager.IsValid())
+	{
+		PostProcessManager->PerformPostProcessBeforeWarpBlend_RenderThread(RHICmdList, this);
+	}
 
 	// Support viewport overlap order sorting:
 	TArray<FDisplayClusterViewportProxy*> SortedViewportProxy = ClusterNodeViewportProxies;
@@ -384,7 +403,10 @@ void FDisplayClusterViewportManagerProxy::UpdateFrameResources_RenderThread(FRHI
 		}
 	}
 
-	PostProcessManager->PerformPostProcessAfterWarpBlend_RenderThread(RHICmdList, this);
+	if (PostProcessManager.IsValid())
+	{
+		PostProcessManager->PerformPostProcessAfterWarpBlend_RenderThread(RHICmdList, this);
+	}
 }
 
 void FDisplayClusterViewportManagerProxy::DoCrossGPUTransfers_RenderThread(FRHICommandListImmediate& RHICmdList) const
diff --git a/Engine/Plugins/Runtime/nDisplay/Source/DisplayCluster/Private/Render/Viewport/DisplayClusterViewportProxy.h b/Engine/Plugins/Runtime/nDisplay/Source/DisplayCluster/Private/Render/Viewport/DisplayClusterViewportProxy.h
--- a/Engine/Plugins/Runtime/nDisplay/Source/DisplayCluster/Private/Render/Viewport/DisplayClusterViewportProxy.h
+++ b/Engine/Plugins/Runtime/nDisplay/Source/DisplayCluster/Private/Render/Viewport/DisplayClusterViewportProxy.h
@@ -69,6 +69,19 @@ public:
 		return Contexts;
 	}
 
+	virtual void SetRenderSettings_RenderThread(const FDisplayClusterViewport_RenderSettings& InRenderSettings) const override
+	{
+		check(IsInRenderingThread());
+		RenderSettings = InRenderSettings;
+	}
+
+	virtual void SetContexts_RenderThread(const TArray<FDisplayClusterViewport_Context>& InContexts) const override
+	{
+		check(IsInRenderingThread());
+		Contexts.Empty();
+		Contexts.Append(InContexts);
+	}
+
 	// Apply postprocess, generate mips, etc from settings in FDisplayClusterViewporDeferredUpdateSettings
 	void UpdateDeferredResources(FRHICommandListImmediate& RHICmdList) const;
 
@@ -135,7 +148,7 @@ protected:
 	const FString ClusterNodeId;
 
 	// Viewport render params
-	FDisplayClusterViewport_RenderSettings       RenderSettings;
+	mutable FDisplayClusterViewport_RenderSettings       RenderSettings;
 	FDisplayClusterViewport_RenderSettingsICVFX  RenderSettingsICVFX;
 	FDisplayClusterViewport_PostRenderSettings   PostRenderSettings;
 
@@ -148,7 +161,7 @@ protected:
 	TSharedPtr<IDisplayClusterProjectionPolicy, ESPMode::ThreadSafe> ProjectionPolicy;
 
 	// Viewport contexts (left/center/right eyes)
-	TArray<FDisplayClusterViewport_Context> Contexts;
+	mutable TArray<FDisplayClusterViewport_Context> Contexts;
 
 	// View family render to this resources
 	TArray<FDisplayClusterViewportRenderTargetResource*> RenderTargets;
diff --git a/Engine/Plugins/Runtime/nDisplay/Source/DisplayCluster/Private/Render/Viewport/DisplayClusterViewport_TextureShare.cpp b/Engine/Plugins/Runtime/nDisplay/Source/DisplayCluster/Private/Render/Viewport/DisplayClusterViewport_TextureShare.cpp
deleted file mode 100644
diff --git a/Engine/Plugins/Runtime/nDisplay/Source/DisplayCluster/Private/Render/Viewport/DisplayClusterViewport_TextureShare.h b/Engine/Plugins/Runtime/nDisplay/Source/DisplayCluster/Private/Render/Viewport/DisplayClusterViewport_TextureShare.h
deleted file mode 100644
diff --git a/Engine/Plugins/Runtime/nDisplay/Source/DisplayCluster/Private/Render/Viewport/Postprocess/DisplayClusterViewportPostProcessManager.cpp b/Engine/Plugins/Runtime/nDisplay/Source/DisplayCluster/Private/Render/Viewport/Postprocess/DisplayClusterViewportPostProcessManager.cpp
--- a/Engine/Plugins/Runtime/nDisplay/Source/DisplayCluster/Private/Render/Viewport/Postprocess/DisplayClusterViewportPostProcessManager.cpp
+++ b/Engine/Plugins/Runtime/nDisplay/Source/DisplayCluster/Private/Render/Viewport/Postprocess/DisplayClusterViewportPostProcessManager.cpp
@@ -34,6 +34,14 @@ static TAutoConsoleVariable<int32> CVarCustomPPEnabled(
 	ECVF_RenderThreadSafe
 );
 
+// Enable new experimental PP features
+static TAutoConsoleVariable<int32> CVarEnableExperimentalTextureSharingAPI(
+	TEXT("nDisplay.render.texturesharing"),
+	0,
+	TEXT("Experimental features (0 = disabled)\n"),
+	ECVF_RenderThreadSafe
+);
+
 // Enable/disable PP round 1
 static TAutoConsoleVariable<int32> CVarPostprocessViewBeforeWarpBlend(
 	TEXT("nDisplay.render.postprocess.ViewBeforeWarpBlend"),
@@ -284,6 +292,34 @@ void FDisplayClusterViewportPostProcessManager::Tick()
 	}
 }
 
+void FDisplayClusterViewportPostProcessManager::HandleSetupNewFrame()
+{
+	if (CVarEnableExperimentalTextureSharingAPI.GetValueOnGameThread() != 0)
+	{
+		if (CVarCustomPPEnabled.GetValueOnGameThread() != 0)
+		{
+			for (const TSharedPtr<IDisplayClusterPostProcess, ESPMode::ThreadSafe >& It : PostprocessProxy)
+			{
+				It->HandleSetupNewFrame(&ViewportManager);
+			}
+		}
+	}
+}
+
+void FDisplayClusterViewportPostProcessManager::HandleBeginNewFrame(FDisplayClusterRenderFrame& InOutRenderFrame)
+{
+	if (CVarEnableExperimentalTextureSharingAPI.GetValueOnGameThread() != 0)
+	{
+		if (CVarCustomPPEnabled.GetValueOnGameThread() != 0)
+		{
+			for (const TSharedPtr<IDisplayClusterPostProcess, ESPMode::ThreadSafe >& It : PostprocessProxy)
+			{
+				It->HandleBeginNewFrame(&ViewportManager, InOutRenderFrame);
+			}
+		}
+	}
+}
+
 void FDisplayClusterViewportPostProcessManager::FinalizeNewFrame()
 {
 	check(IsInGameThread());
@@ -296,6 +332,48 @@ void FDisplayClusterViewportPostProcessManager::FinalizeNewFrame()
 	});
 }
 
+void FDisplayClusterViewportPostProcessManager::HandleRenderFrameSetup_RenderThread(FRHICommandListImmediate& RHICmdList, const FDisplayClusterViewportManagerProxy* InViewportManagerProxy)
+{
+	if (CVarEnableExperimentalTextureSharingAPI.GetValueOnRenderThread() != 0)
+	{
+		if (CVarCustomPPEnabled.GetValueOnRenderThread() != 0)
+		{
+			for (const TSharedPtr<IDisplayClusterPostProcess, ESPMode::ThreadSafe >& It : PostprocessProxy)
+			{
+				It->HandleRenderFrameSetup_RenderThread(RHICmdList, InViewportManagerProxy);
+			}
+		}
+	}
+}
+
+void FDisplayClusterViewportPostProcessManager::HandleBeginUpdateFrameResources_RenderThread(FRHICommandListImmediate& RHICmdList, const FDisplayClusterViewportManagerProxy* InViewportManagerProxy)
+{
+	if (CVarEnableExperimentalTextureSharingAPI.GetValueOnRenderThread() != 0)
+	{
+		if (CVarCustomPPEnabled.GetValueOnRenderThread() != 0)
+		{
+			for (const TSharedPtr<IDisplayClusterPostProcess, ESPMode::ThreadSafe >& It : PostprocessProxy)
+			{
+				It->HandleBeginUpdateFrameResources_RenderThread(RHICmdList, InViewportManagerProxy);
+			}
+		}
+	}
+}
+
+void FDisplayClusterViewportPostProcessManager::HandleEndUpdateFrameResources_RenderThread(FRHICommandListImmediate& RHICmdList, const FDisplayClusterViewportManagerProxy* InViewportManagerProxy)
+{
+	if (CVarEnableExperimentalTextureSharingAPI.GetValueOnRenderThread() != 0)
+	{
+		if (CVarCustomPPEnabled.GetValueOnRenderThread() != 0)
+		{
+			for (const TSharedPtr<IDisplayClusterPostProcess, ESPMode::ThreadSafe >& It : PostprocessProxy)
+			{
+				It->HandleEndUpdateFrameResources_RenderThread(RHICmdList, InViewportManagerProxy);
+			}
+		}
+	}
+}
+
 void FDisplayClusterViewportPostProcessManager::PerformPostProcessBeforeWarpBlend_RenderThread(FRHICommandListImmediate& RHICmdList, const FDisplayClusterViewportManagerProxy* InViewportManagerProxy) const
 {
 	bool bIsCustomPPEnabled = (CVarCustomPPEnabled.GetValueOnRenderThread() != 0);
diff --git a/Engine/Plugins/Runtime/nDisplay/Source/DisplayCluster/Private/Render/Viewport/Postprocess/DisplayClusterViewportPostProcessManager.h b/Engine/Plugins/Runtime/nDisplay/Source/DisplayCluster/Private/Render/Viewport/Postprocess/DisplayClusterViewportPostProcessManager.h
--- a/Engine/Plugins/Runtime/nDisplay/Source/DisplayCluster/Private/Render/Viewport/Postprocess/DisplayClusterViewportPostProcessManager.h
+++ b/Engine/Plugins/Runtime/nDisplay/Source/DisplayCluster/Private/Render/Viewport/Postprocess/DisplayClusterViewportPostProcessManager.h
@@ -38,6 +38,13 @@ public:
 	bool HandleStartScene();
 	void HandleEndScene();
 
+	void HandleSetupNewFrame();
+	void HandleBeginNewFrame(FDisplayClusterRenderFrame& InOutRenderFrame);
+
+	void HandleRenderFrameSetup_RenderThread(FRHICommandListImmediate& RHICmdList, const FDisplayClusterViewportManagerProxy* InViewportManagerProxy);
+	void HandleBeginUpdateFrameResources_RenderThread(FRHICommandListImmediate& RHICmdList, const FDisplayClusterViewportManagerProxy* InViewportManagerProxy);
+	void HandleEndUpdateFrameResources_RenderThread(FRHICommandListImmediate& RHICmdList, const FDisplayClusterViewportManagerProxy* InViewportManagerProxy);
+
 	// Send data to render thread
 	void FinalizeNewFrame();
 
diff --git a/Engine/Plugins/Runtime/nDisplay/Source/DisplayCluster/Private/Render/Viewport/RenderFrame/DisplayClusterRenderFrameSettings.h b/Engine/Plugins/Runtime/nDisplay/Source/DisplayCluster/Private/Render/Viewport/RenderFrame/DisplayClusterRenderFrameSettings.h
--- a/Engine/Plugins/Runtime/nDisplay/Source/DisplayCluster/Private/Render/Viewport/RenderFrame/DisplayClusterRenderFrameSettings.h
+++ b/Engine/Plugins/Runtime/nDisplay/Source/DisplayCluster/Private/Render/Viewport/RenderFrame/DisplayClusterRenderFrameSettings.h
@@ -54,6 +54,9 @@ struct FDisplayClusterRenderFrameSettings
 	// Render in Editor mode
 	bool bIsRenderingInEditor = false;
 
+	// Configuration used to render preview
+	bool bIsPreviewRendering = false;
+
 	// Allow mGPU in editor mode
 	bool bAllowMultiGPURenderingInEditor = false;
 	int32 PreviewMinGPUIndex = 0;
diff --git a/Engine/Plugins/Runtime/nDisplay/Source/DisplayCluster/Public/Render/PostProcess/IDisplayClusterPostProcess.h b/Engine/Plugins/Runtime/nDisplay/Source/DisplayCluster/Public/Render/PostProcess/IDisplayClusterPostProcess.h
--- a/Engine/Plugins/Runtime/nDisplay/Source/DisplayCluster/Public/Render/PostProcess/IDisplayClusterPostProcess.h
+++ b/Engine/Plugins/Runtime/nDisplay/Source/DisplayCluster/Public/Render/PostProcess/IDisplayClusterPostProcess.h
@@ -9,7 +9,9 @@
 #include "RHICommandList.h"
 
 class IDisplayClusterViewportManager;
+class IDisplayClusterViewportManagerProxy;
 class IDisplayClusterViewportProxy;
+class FDisplayClusterRenderFrame;
 
 /**
  * nDisplay post-process interface
@@ -67,6 +69,41 @@ public:
 	*/
 	virtual void HandleEndScene(IDisplayClusterViewportManager* InViewportManager) = 0;
 
+	/**
+	* Called each time a new frame setup
+	*
+	*/
+	virtual void HandleSetupNewFrame(IDisplayClusterViewportManager* InViewportManager)
+	{ }
+
+	/**
+	* Called each time a new frame begin
+	*
+	*/
+	virtual void HandleBeginNewFrame(IDisplayClusterViewportManager* InViewportManager, FDisplayClusterRenderFrame& InOutRenderFrame)
+	{ }
+
+	/**
+	* Called each time a render thread multi-viewport composing begin
+	*
+	*/
+	virtual void HandleRenderFrameSetup_RenderThread(FRHICommandListImmediate& RHICmdList, const IDisplayClusterViewportManagerProxy* InViewportManagerProxy)
+	{ }
+
+	/**
+	* Called each time a render thread before UpdateFrameResources()
+	*
+	*/
+	virtual void HandleBeginUpdateFrameResources_RenderThread(FRHICommandListImmediate& RHICmdList, const IDisplayClusterViewportManagerProxy* InViewportManagerProxy)
+	{ }
+
+	/**
+	* Called each time a render thread after UpdateFrameResources()
+	*
+	*/
+	virtual void HandleEndUpdateFrameResources_RenderThread(FRHICommandListImmediate& RHICmdList, const IDisplayClusterViewportManagerProxy* InViewportManagerProxy)
+	{ }
+
 	/**
 	* Returns if an interface implementation processes each view region before warp&blend
 	* the corresponded function will be called once per viewport
diff --git a/Engine/Plugins/Runtime/nDisplay/Source/DisplayCluster/Public/Render/Viewport/IDisplayClusterViewport.h b/Engine/Plugins/Runtime/nDisplay/Source/DisplayCluster/Public/Render/Viewport/IDisplayClusterViewport.h
--- a/Engine/Plugins/Runtime/nDisplay/Source/DisplayCluster/Public/Render/Viewport/IDisplayClusterViewport.h
+++ b/Engine/Plugins/Runtime/nDisplay/Source/DisplayCluster/Public/Render/Viewport/IDisplayClusterViewport.h
@@ -46,4 +46,7 @@ public:
 	virtual void SetupSceneView(uint32 ContextNum, class UWorld* World, FSceneViewFamily& InViewFamily, FSceneView& InView) const = 0;
 
 	virtual class IDisplayClusterViewportManager& GetOwner() const = 0;
+
+	virtual void SetRenderSettings(const FDisplayClusterViewport_RenderSettings& InRenderSettings) = 0;
+	virtual void SetContexts(TArray<FDisplayClusterViewport_Context>& InContexts) = 0;
 };
diff --git a/Engine/Plugins/Runtime/nDisplay/Source/DisplayCluster/Public/Render/Viewport/IDisplayClusterViewportProxy.h b/Engine/Plugins/Runtime/nDisplay/Source/DisplayCluster/Public/Render/Viewport/IDisplayClusterViewportProxy.h
--- a/Engine/Plugins/Runtime/nDisplay/Source/DisplayCluster/Public/Render/Viewport/IDisplayClusterViewportProxy.h
+++ b/Engine/Plugins/Runtime/nDisplay/Source/DisplayCluster/Public/Render/Viewport/IDisplayClusterViewportProxy.h
@@ -42,4 +42,7 @@ public:
 	virtual EDisplayClusterViewportResourceType   GetOutputResourceType_RenderThread() const = 0;
 
 	virtual const class IDisplayClusterViewportManagerProxy& GetOwner_RenderThread() const = 0;
+
+	virtual void SetRenderSettings_RenderThread(const FDisplayClusterViewport_RenderSettings& InRenderSettings) const = 0;
+	virtual void SetContexts_RenderThread(const TArray<FDisplayClusterViewport_Context>& InContexts) const = 0;
 };
diff --git a/Engine/Plugins/Runtime/nDisplay/Source/DisplayClusterConfiguration/DisplayClusterConfiguration.Build.cs b/Engine/Plugins/Runtime/nDisplay/Source/DisplayClusterConfiguration/DisplayClusterConfiguration.Build.cs
--- a/Engine/Plugins/Runtime/nDisplay/Source/DisplayClusterConfiguration/DisplayClusterConfiguration.Build.cs
+++ b/Engine/Plugins/Runtime/nDisplay/Source/DisplayClusterConfiguration/DisplayClusterConfiguration.Build.cs
@@ -12,7 +12,6 @@ public class DisplayClusterConfiguration : ModuleRules
 		PublicIncludePathModuleNames.AddRange(
 			new string[] {
 				"DisplayCluster",
-				"DisplayClusterPostprocess",
 				"DisplayClusterProjection",
 			});
 
@@ -23,14 +22,6 @@ public class DisplayClusterConfiguration : ModuleRules
 				"OpenColorIO"
 			});
 
-		if (Target.Platform == UnrealTargetPlatform.Win64)
-		{
-			PublicDependencyModuleNames.AddRange(
-			new string[] {
-				"TextureShare"
-			});
-		}
-
 		PrivateDependencyModuleNames.AddRange(
 			new string[] {
 				"Core",
diff --git a/Engine/Plugins/Runtime/nDisplay/Source/DisplayClusterConfiguration/Public/DisplayClusterConfigurationTypes.h b/Engine/Plugins/Runtime/nDisplay/Source/DisplayClusterConfiguration/Public/DisplayClusterConfigurationTypes.h
--- a/Engine/Plugins/Runtime/nDisplay/Source/DisplayClusterConfiguration/Public/DisplayClusterConfigurationTypes.h
+++ b/Engine/Plugins/Runtime/nDisplay/Source/DisplayClusterConfiguration/Public/DisplayClusterConfigurationTypes.h
@@ -334,6 +334,10 @@ public:
 	UPROPERTY(EditAnywhere, BlueprintReadWrite, Category = "Configuration", meta = (DisplayName = "Output Remapping"))
 	FDisplayClusterConfigurationFramePostProcess_OutputRemap OutputRemap;
 
+	/** Enables texture sharing for this cluster node */
+	UPROPERTY(EditAnywhere, BlueprintReadonly, Category = "Configuration", meta = (DisplayName = "Enable Texture Share"))
+	bool bEnableTextureShare = false;
+
 #if WITH_EDITORONLY_DATA
 	/** Locks the application window aspect ratio for easier resizing */
 	UPROPERTY(EditAnywhere, Category = "Configuration", meta = (HideProperty))
diff --git a/Engine/Plugins/Runtime/nDisplay/Source/DisplayClusterConfiguration/Public/DisplayClusterConfigurationTypes_TextureShare.h b/Engine/Plugins/Runtime/nDisplay/Source/DisplayClusterConfiguration/Public/DisplayClusterConfigurationTypes_TextureShare.h
deleted file mode 100644
diff --git a/Engine/Plugins/Runtime/nDisplay/Source/DisplayClusterConfiguration/Public/DisplayClusterConfigurationTypes_Viewport.h b/Engine/Plugins/Runtime/nDisplay/Source/DisplayClusterConfiguration/Public/DisplayClusterConfigurationTypes_Viewport.h
--- a/Engine/Plugins/Runtime/nDisplay/Source/DisplayClusterConfiguration/Public/DisplayClusterConfigurationTypes_Viewport.h
+++ b/Engine/Plugins/Runtime/nDisplay/Source/DisplayClusterConfiguration/Public/DisplayClusterConfigurationTypes_Viewport.h
@@ -10,8 +10,6 @@
 #include "DisplayClusterConfigurationStrings.h"
 #include "DisplayClusterConfigurationTypes_PostRender.h"
 #include "DisplayClusterConfigurationTypes_Postprocess.h"
-
-#include "DisplayClusterConfigurationTypes_TextureShare.h"
 #include "DisplayClusterConfigurationTypes_ViewportRemap.h"
 
 #include "OpenColorIOColorSpace.h"
@@ -183,10 +181,6 @@ public:
 	UPROPERTY(EditDefaultsOnly, BlueprintReadWrite, Category = "Configuration")
 	FDisplayClusterConfigurationProjection ProjectionPolicy;
 
-	/** Enable or disable compatibility with inter process GPU Texture share */
-	UPROPERTY(EditAnywhere, BlueprintReadOnly, Category = "Configuration", meta = (DisplayName = "Shared Texture"))
-	FDisplayClusterConfigurationTextureShare_Viewport TextureShare;
-
 #if WITH_EDITORONLY_DATA
 	UPROPERTY(EditAnywhere, BlueprintReadWrite, Category = "Preview", meta = (DisplayName = "Preview Frustum"))
 	bool bAllowPreviewFrustumRendering = false;
diff --git a/Engine/Plugins/Runtime/nDisplay/Source/DisplayClusterConfigurator/Private/Views/Details/Cluster/DisplayClusterConfiguratorViewportDetailsCustomization.cpp b/Engine/Plugins/Runtime/nDisplay/Source/DisplayClusterConfigurator/Private/Views/Details/Cluster/DisplayClusterConfiguratorViewportDetailsCustomization.cpp
--- a/Engine/Plugins/Runtime/nDisplay/Source/DisplayClusterConfigurator/Private/Views/Details/Cluster/DisplayClusterConfiguratorViewportDetailsCustomization.cpp
+++ b/Engine/Plugins/Runtime/nDisplay/Source/DisplayClusterConfigurator/Private/Views/Details/Cluster/DisplayClusterConfiguratorViewportDetailsCustomization.cpp
@@ -40,11 +40,6 @@ void FDisplayClusterConfiguratorViewportDetailsCustomization::CustomizeDetails(I
 	}
 	check(ConfigurationViewportPtr != nullptr);
 
-	// Hide the texture share property when we aren't on a Windows platform
-#if !PLATFORM_WINDOWS
-	InLayoutBuilder.HideProperty(GET_MEMBER_NAME_CHECKED(UDisplayClusterConfigurationViewport, TextureShare));
-#endif
-
 	const bool bDisplayICVFX =
 			ConfigurationViewportPtr->ProjectionPolicy.Type.Equals(DisplayClusterProjectionStrings::projection::Mesh, ESearchCase::IgnoreCase)
 		||  ConfigurationViewportPtr->ProjectionPolicy.Type.Equals(DisplayClusterProjectionStrings::projection::MPCDI, ESearchCase::IgnoreCase);
diff --git a/Engine/Plugins/Runtime/nDisplay/Source/DisplayClusterPostprocess/DisplayClusterPostprocess.Build.cs b/Engine/Plugins/Runtime/nDisplay/Source/DisplayClusterPostprocess/DisplayClusterPostprocess.Build.cs
deleted file mode 100644
diff --git a/Engine/Plugins/Runtime/nDisplay/Source/DisplayClusterPostprocess/Private/DisplayClusterPostprocessLog.cpp b/Engine/Plugins/Runtime/nDisplay/Source/DisplayClusterPostprocess/Private/DisplayClusterPostprocessLog.cpp
deleted file mode 100644
diff --git a/Engine/Plugins/Runtime/nDisplay/Source/DisplayClusterPostprocess/Private/DisplayClusterPostprocessLog.h b/Engine/Plugins/Runtime/nDisplay/Source/DisplayClusterPostprocess/Private/DisplayClusterPostprocessLog.h
deleted file mode 100644
diff --git a/Engine/Plugins/Runtime/nDisplay/Source/DisplayClusterPostprocess/Private/DisplayClusterPostprocessModule.cpp b/Engine/Plugins/Runtime/nDisplay/Source/DisplayClusterPostprocess/Private/DisplayClusterPostprocessModule.cpp
deleted file mode 100644
diff --git a/Engine/Plugins/Runtime/nDisplay/Source/DisplayClusterPostprocess/Private/DisplayClusterPostprocessModule.h b/Engine/Plugins/Runtime/nDisplay/Source/DisplayClusterPostprocess/Private/DisplayClusterPostprocessModule.h
deleted file mode 100644
diff --git a/Engine/Plugins/Runtime/nDisplay/Source/DisplayClusterPostprocess/Private/PostProcess/DisplayClusterPostprocessBase.cpp b/Engine/Plugins/Runtime/nDisplay/Source/DisplayClusterPostprocess/Private/PostProcess/DisplayClusterPostprocessBase.cpp
deleted file mode 100644
diff --git a/Engine/Plugins/Runtime/nDisplay/Source/DisplayClusterPostprocess/Private/PostProcess/DisplayClusterPostprocessBase.h b/Engine/Plugins/Runtime/nDisplay/Source/DisplayClusterPostprocess/Private/PostProcess/DisplayClusterPostprocessBase.h
deleted file mode 100644
diff --git a/Engine/Plugins/Runtime/nDisplay/Source/DisplayClusterPostprocess/Private/PostProcess/Windows/DisplayClusterPostprocessDX12CrossGPU.cpp b/Engine/Plugins/Runtime/nDisplay/Source/DisplayClusterPostprocess/Private/PostProcess/Windows/DisplayClusterPostprocessDX12CrossGPU.cpp
deleted file mode 100644
diff --git a/Engine/Plugins/Runtime/nDisplay/Source/DisplayClusterPostprocess/Private/PostProcess/Windows/DisplayClusterPostprocessDX12CrossGPU.h b/Engine/Plugins/Runtime/nDisplay/Source/DisplayClusterPostprocess/Private/PostProcess/Windows/DisplayClusterPostprocessDX12CrossGPU.h
deleted file mode 100644
diff --git a/Engine/Plugins/Runtime/nDisplay/Source/DisplayClusterPostprocess/Private/PostProcess/Windows/DisplayClusterPostprocessTextureShare.cpp b/Engine/Plugins/Runtime/nDisplay/Source/DisplayClusterPostprocess/Private/PostProcess/Windows/DisplayClusterPostprocessTextureShare.cpp
deleted file mode 100644
diff --git a/Engine/Plugins/Runtime/nDisplay/Source/DisplayClusterPostprocess/Private/PostProcess/Windows/DisplayClusterPostprocessTextureShare.h b/Engine/Plugins/Runtime/nDisplay/Source/DisplayClusterPostprocess/Private/PostProcess/Windows/DisplayClusterPostprocessTextureShare.h
deleted file mode 100644
diff --git a/Engine/Plugins/Runtime/nDisplay/Source/DisplayClusterPostprocess/Public/DisplayClusterPostprocessStrings.h b/Engine/Plugins/Runtime/nDisplay/Source/DisplayClusterPostprocess/Public/DisplayClusterPostprocessStrings.h
deleted file mode 100644
diff --git a/Engine/Plugins/Runtime/nDisplay/nDisplay.uplugin b/Engine/Plugins/Runtime/nDisplay/nDisplay.uplugin
--- a/Engine/Plugins/Runtime/nDisplay/nDisplay.uplugin
+++ b/Engine/Plugins/Runtime/nDisplay/nDisplay.uplugin
@@ -46,15 +46,6 @@
 				"Linux"
 			]
 		},
-		{
-			"Name": "DisplayClusterPostprocess",
-			"Type": "Runtime",
-			"LoadingPhase": "Default",
-			"PlatformAllowList": [
-				"Win64",
-				"Linux"
-			]
-		},
 		{
 			"Name": "DisplayClusterShaders",
 			"Type": "Runtime",
@@ -161,10 +152,6 @@
 			"Name": "ProceduralMeshComponent",
 			"Enabled": true
 		},
-		{
-			"Name": "TextureShare",
-			"Enabled": true
-		},
 		{
 			"Name": "ProceduralMeshComponent",
 			"Enabled": true
