nDisplay: Adding, removing, or compiling lightcards will refresh the preview scene. Modifying a lightcard property will maintain the details panel object. Refactor proxy regeneration from preview client to lightcard editor.

diff --git a/Engine/Plugins/Runtime/nDisplay/Source/DisplayClusterLightCardEditor/Private/SDisplayClusterLightCardEditor.cpp b/Engine/Plugins/Runtime/nDisplay/Source/DisplayClusterLightCardEditor/Private/SDisplayClusterLightCardEditor.cpp
--- a/Engine/Plugins/Runtime/nDisplay/Source/DisplayClusterLightCardEditor/Private/SDisplayClusterLightCardEditor.cpp
+++ b/Engine/Plugins/Runtime/nDisplay/Source/DisplayClusterLightCardEditor/Private/SDisplayClusterLightCardEditor.cpp
@@ -14,6 +14,7 @@
 #include "Framework/MultiBox/MultiBoxBuilder.h"
 #include "Framework/MultiBox/MultiBoxExtender.h"
 #include "Toolkits/AssetEditorToolkit.h"
+#include "Viewport/DisplayClusterLightCardEditorViewportClient.h"
 #include "Widgets/Docking/SDockTab.h"
 #include "Widgets/Layout/SSplitter.h"
 
@@ -63,11 +64,23 @@ void SDisplayClusterLightCardEditor::ExtendToolbar(FToolBarBuilder& ToolbarBuild
 SDisplayClusterLightCardEditor::~SDisplayClusterLightCardEditor()
 {
 	IDisplayClusterOperator::Get().OnActiveRootActorChanged().Remove(ActiveRootActorChangedHandle);
+	FCoreUObjectDelegates::OnObjectPropertyChanged.RemoveAll(this);
+
+	if (GEngine != nullptr)
+	{
+		GEngine->OnLevelActorDeleted().RemoveAll(this);
+	}
+
+	RemoveCompileDelegates();
 }
 
 void SDisplayClusterLightCardEditor::Construct(const FArguments& InArgs, const TSharedRef<SDockTab>& MajorTabOwner, const TSharedPtr<SWindow>& WindowOwner)
 {
 	ActiveRootActorChangedHandle = IDisplayClusterOperator::Get().OnActiveRootActorChanged().AddSP(this, &SDisplayClusterLightCardEditor::OnActiveRootActorChanged);
+	if (GEngine != nullptr)
+	{
+		GEngine->OnLevelActorDeleted().AddSP(this, &SDisplayClusterLightCardEditor::OnLevelActorDeleted);
+	}
 
 	ChildSlot                     
 	[
@@ -95,14 +108,23 @@ void SDisplayClusterLightCardEditor::Construct(const FArguments& InArgs, const T
 			]
 		]
 	];
+
+	BindCompileDelegates();
 }
 
 void SDisplayClusterLightCardEditor::OnActiveRootActorChanged(ADisplayClusterRootActor* NewRootActor)
 {
+	RemoveCompileDelegates();
+	
 	// The new root actor pointer could be null, indicating that it was deleted or the user didn't select a valid root actor
 	ActiveRootActor = NewRootActor;
 	LightCardList->SetRootActor(NewRootActor);
 	ViewportView->SetRootActor(NewRootActor);
+	
+	BindCompileDelegates();
+
+	FCoreUObjectDelegates::OnObjectPropertyChanged.RemoveAll(this);
+	FCoreUObjectDelegates::OnObjectPropertyChanged.AddSP(this, &SDisplayClusterLightCardEditor::OnActorPropertyChanged);
 }
 
 TSharedRef<SWidget> SDisplayClusterLightCardEditor::CreateLightCardListWidget()
@@ -115,4 +137,136 @@ TSharedRef<SWidget> SDisplayClusterLightCardEditor::CreateViewportWidget()
 	return SAssignNew(ViewportView, SDisplayClusterLightCardEditorViewport, SharedThis(this));
 }
 
+void SDisplayClusterLightCardEditor::RefreshPreviewActors()
+{
+	RemoveCompileDelegates();
+	
+	if (ADisplayClusterRootActor* RootActor = GetActiveRootActor().Get())
+	{
+		if (LightCardList.IsValid())
+		{
+			LightCardList->SetRootActor(RootActor);
+		}
+			
+		if (ViewportView.IsValid())
+		{
+			const bool bForce = true;
+			ViewportView->GetLightCardEditorViewportClient()->UpdatePreviewActor(RootActor, bForce);
+		}
+	}
+	
+	BindCompileDelegates();
+}
+
+void SDisplayClusterLightCardEditor::BindCompileDelegates()
+{
+	if (LightCardList.IsValid())
+	{
+		for (const TSharedPtr<SDisplayClusterLightCardList::FLightCardTreeItem>& LightCardActor : LightCardList->GetLightCardActors())
+		{
+			if (LightCardActor.IsValid() && LightCardActor->LightCardActor.IsValid())
+			{
+				if (UBlueprint* Blueprint = UBlueprint::GetBlueprintFromClass(LightCardActor->LightCardActor->GetClass()))
+				{
+					Blueprint->OnCompiled().AddSP(this, &SDisplayClusterLightCardEditor::OnBlueprintCompiled);
+				}
+			}
+		}
+	}
+}
+
+void SDisplayClusterLightCardEditor::RemoveCompileDelegates()
+{
+	if (LightCardList.IsValid())
+	{
+		for (const TSharedPtr<SDisplayClusterLightCardList::FLightCardTreeItem>& LightCardActor : LightCardList->GetLightCardActors())
+		{
+			if (LightCardActor.IsValid() && LightCardActor->LightCardActor.IsValid())
+			{
+				if (UBlueprint* Blueprint = UBlueprint::GetBlueprintFromClass(LightCardActor->LightCardActor->GetClass()))
+				{
+					Blueprint->OnCompiled().RemoveAll(this);
+				}
+			}
+		}
+	}
+}
+
+void SDisplayClusterLightCardEditor::OnActorPropertyChanged(UObject* ObjectBeingModified,
+                                                            FPropertyChangedEvent& PropertyChangedEvent)
+{
+	if (PropertyChangedEvent.ChangeType == EPropertyChangeType::Interactive)
+	{
+		return;
+	}
+
+	auto IsOurActor = [ObjectBeingModified] (UObject* ObjectToCompare) -> bool
+	{
+		if (ObjectToCompare)
+		{
+			if (ObjectBeingModified == ObjectToCompare)
+			{
+				return true;
+			}
+
+			if (const UObject* RootActorOuter = ObjectBeingModified->GetTypedOuter(ObjectToCompare->GetClass()))
+			{
+				return RootActorOuter == ObjectToCompare;
+			}
+		}
+
+		return false;
+	};
+	
+	bool bIsOurActor = IsOurActor(GetActiveRootActor().Get());
+	if (!bIsOurActor && LightCardList.IsValid())
+	{
+		for (const TSharedPtr<SDisplayClusterLightCardList::FLightCardTreeItem>& LightCard : LightCardList->GetLightCardActors())
+		{
+			bIsOurActor = IsOurActor(LightCard->LightCardActor.Get());
+			if (bIsOurActor)
+			{
+				break;
+			}
+		}
+	}
+	
+	if (bIsOurActor)
+	{
+		RefreshPreviewActors();
+	}
+}
+
+void SDisplayClusterLightCardEditor::OnLevelActorDeleted(AActor* Actor)
+{
+	if (LightCardList.IsValid() &&
+		LightCardList->GetLightCardActors().ContainsByPredicate([Actor](const TSharedPtr<SDisplayClusterLightCardList::FLightCardTreeItem>& LightCardTreeItem)
+	{
+		return LightCardTreeItem.IsValid() && Actor == LightCardTreeItem->LightCardActor.Get();
+	}))
+	{
+		if (Actor && Actor->GetClass()->HasAnyClassFlags(CLASS_NewerVersionExists))
+		{
+			// When a blueprint class is regenerated instances are deleted and replaced.
+			// In this case the OnCompiled() delegate will fire and refresh the actor.
+			return;
+		}
+		
+		if (ViewportView.IsValid())
+		{
+			ViewportView->GetLightCardEditorViewportClient()->GetWorld()->GetTimerManager().SetTimerForNextTick([=]()
+			{
+				// Schedule for next tick so available selections are properly updated once the
+				// actor is fully deleted.
+				RefreshPreviewActors();
+			});
+		}
+	}
+}
+
+void SDisplayClusterLightCardEditor::OnBlueprintCompiled(UBlueprint* Blueprint)
+{
+	RefreshPreviewActors();
+}
+
 #undef LOCTEXT_NAMESPACE
diff --git a/Engine/Plugins/Runtime/nDisplay/Source/DisplayClusterLightCardEditor/Private/SDisplayClusterLightCardEditor.h b/Engine/Plugins/Runtime/nDisplay/Source/DisplayClusterLightCardEditor/Private/SDisplayClusterLightCardEditor.h
--- a/Engine/Plugins/Runtime/nDisplay/Source/DisplayClusterLightCardEditor/Private/SDisplayClusterLightCardEditor.h
+++ b/Engine/Plugins/Runtime/nDisplay/Source/DisplayClusterLightCardEditor/Private/SDisplayClusterLightCardEditor.h
@@ -58,6 +58,24 @@ private:
 	/** Create the 3d viewport widget. */
 	TSharedRef<SWidget> CreateViewportWidget();
 
+	/** Refresh all preview actors. */
+	void RefreshPreviewActors();
+
+	/** Bind delegates to when a BP compiles. */
+	void BindCompileDelegates();
+
+	/** Remove compile delegates from a BP. */
+	void RemoveCompileDelegates();
+
+	/** When a property on the actor has changed. */
+	void OnActorPropertyChanged(UObject* ObjectBeingModified, FPropertyChangedEvent& PropertyChangedEvent);
+
+	/** Raised when the user deletes an actor from the level */
+	void OnLevelActorDeleted(AActor* Actor);
+
+	/** Raised when a supported blueprint is compiled. */
+	void OnBlueprintCompiled(UBlueprint* Blueprint);
+	
 private:
 	/** The light card list widget */
 	TSharedPtr<SDisplayClusterLightCardList> LightCardList;
diff --git a/Engine/Plugins/Runtime/nDisplay/Source/DisplayClusterLightCardEditor/Private/SDisplayClusterLightCardList.cpp b/Engine/Plugins/Runtime/nDisplay/Source/DisplayClusterLightCardEditor/Private/SDisplayClusterLightCardList.cpp
--- a/Engine/Plugins/Runtime/nDisplay/Source/DisplayClusterLightCardEditor/Private/SDisplayClusterLightCardList.cpp
+++ b/Engine/Plugins/Runtime/nDisplay/Source/DisplayClusterLightCardEditor/Private/SDisplayClusterLightCardList.cpp
@@ -149,9 +149,11 @@ void SDisplayClusterLightCardList::Construct(const FArguments& InArgs)
 void SDisplayClusterLightCardList::SetRootActor(ADisplayClusterRootActor* NewRootActor)
 {
 	RootActor = NewRootActor;
-	FillLightCardList();
-	LightCardTreeView->RebuildList();
-
+	if (FillLightCardList())
+	{
+		LightCardTreeView->RebuildList();
+	}
+	
 	for (const TSharedPtr<FLightCardTreeItem>& LightCardTreeItem : LightCardTree)
 	{
 		if (LightCardTreeItem->IsActorLayer())
@@ -161,9 +163,9 @@ void SDisplayClusterLightCardList::SetRootActor(ADisplayClusterRootActor* NewRoo
 	}
 }
 
-void SDisplayClusterLightCardList::FillLightCardList()
+bool SDisplayClusterLightCardList::FillLightCardList()
 {
-	LightCardTree.Empty();
+	TArray<TSharedPtr<FLightCardTreeItem>> OriginalTree = MoveTemp(LightCardTree);
 	LightCardActors.Empty();
 
 	if (RootActor.IsValid())
@@ -215,6 +217,29 @@ void SDisplayClusterLightCardList::FillLightCardList()
 			}
 		}
 	}
+
+	// Check if the tree items have changed.
+	{
+		if (OriginalTree.Num() != LightCardTree.Num())
+		{
+			return true;
+		}
+	
+		for (const TSharedPtr<FLightCardTreeItem>& OriginalTreeItem : OriginalTree)
+		{
+			if (!LightCardTree.ContainsByPredicate([OriginalTreeItem](const TSharedPtr<FLightCardTreeItem>& LightCardTreeItem)
+			{
+				return OriginalTreeItem.IsValid() && LightCardTreeItem.IsValid() &&
+					OriginalTreeItem->LightCardActor.Get() == LightCardTreeItem->LightCardActor.Get() &&
+						OriginalTreeItem->ActorLayer == LightCardTreeItem->ActorLayer;
+			}))
+			{
+				return true;
+			}
+		}
+	}
+
+	return false;
 }
 
 TSharedRef<ITableRow> SDisplayClusterLightCardList::GenerateTreeItemRow(TSharedPtr<FLightCardTreeItem> Item, const TSharedRef<STableViewBase>& OwnerTable)
diff --git a/Engine/Plugins/Runtime/nDisplay/Source/DisplayClusterLightCardEditor/Private/SDisplayClusterLightCardList.h b/Engine/Plugins/Runtime/nDisplay/Source/DisplayClusterLightCardEditor/Private/SDisplayClusterLightCardList.h
--- a/Engine/Plugins/Runtime/nDisplay/Source/DisplayClusterLightCardEditor/Private/SDisplayClusterLightCardList.h
+++ b/Engine/Plugins/Runtime/nDisplay/Source/DisplayClusterLightCardEditor/Private/SDisplayClusterLightCardList.h
@@ -40,8 +40,13 @@ public:
 
 	void SetRootActor(ADisplayClusterRootActor* NewRootActor);
 
+	const TArray<TSharedPtr<FLightCardTreeItem>>& GetLightCardActors() const { return LightCardActors; }
 private:
-	void FillLightCardList();
+	/**
+	 * Fill the LightCard list with available LightCards.
+	 * @return True if the list has been modified.
+	 */
+	bool FillLightCardList();
 
 	TSharedRef<ITableRow> GenerateTreeItemRow(TSharedPtr<FLightCardTreeItem> Item, const TSharedRef<STableViewBase>& OwnerTable);
 	void GetChildrenForTreeItem(TSharedPtr<FLightCardTreeItem> InItem, TArray<TSharedPtr<FLightCardTreeItem>>& OutChildren);
diff --git a/Engine/Plugins/Runtime/nDisplay/Source/DisplayClusterLightCardEditor/Private/Viewport/DisplayClusterLightCardEditorViewportClient.cpp b/Engine/Plugins/Runtime/nDisplay/Source/DisplayClusterLightCardEditor/Private/Viewport/DisplayClusterLightCardEditorViewportClient.cpp
--- a/Engine/Plugins/Runtime/nDisplay/Source/DisplayClusterLightCardEditor/Private/Viewport/DisplayClusterLightCardEditorViewportClient.cpp
+++ b/Engine/Plugins/Runtime/nDisplay/Source/DisplayClusterLightCardEditor/Private/Viewport/DisplayClusterLightCardEditorViewportClient.cpp
@@ -91,7 +91,6 @@ FDisplayClusterLightCardEditorViewportClient::FDisplayClusterLightCardEditorView
 FDisplayClusterLightCardEditorViewportClient::~FDisplayClusterLightCardEditorViewportClient()
 {
 	EndTransaction();
-	FCoreUObjectDelegates::OnObjectPropertyChanged.RemoveAll(this);
 }
 
 FLinearColor FDisplayClusterLightCardEditorViewportClient::GetBackgroundColor() const
@@ -588,9 +587,6 @@ void FDisplayClusterLightCardEditorViewportClient::UpdatePreviewActor(ADisplayCl
 				LightCardProxies.Add(LightCardProxy);
 				MeshProjectionRenderer->AddActor(LightCardProxy);
 			}
-
-			FCoreUObjectDelegates::OnObjectPropertyChanged.RemoveAll(this);
-			FCoreUObjectDelegates::OnObjectPropertyChanged.AddRaw(this, &FDisplayClusterLightCardEditorViewportClient::OnActorPropertyChanged);
 		});
 	}
 }
@@ -658,51 +654,6 @@ void FDisplayClusterLightCardEditorViewportClient::EndTransaction()
 	}
 }
 
-void FDisplayClusterLightCardEditorViewportClient::OnActorPropertyChanged(UObject* ObjectBeingModified,
-                                                                          FPropertyChangedEvent& PropertyChangedEvent)
-{
-	if (PropertyChangedEvent.ChangeType == EPropertyChangeType::Interactive)
-	{
-		return;
-	}
-
-	auto IsOurActor = [ObjectBeingModified] (UObject* ObjectToCompare) -> bool
-	{
-		if (ObjectToCompare)
-		{
-			if (ObjectBeingModified == ObjectToCompare)
-			{
-				return true;
-			}
-
-			if (const UObject* RootActorOuter = ObjectBeingModified->GetTypedOuter(ObjectToCompare->GetClass()))
-			{
-				return RootActorOuter == ObjectToCompare;
-			}
-		}
-
-		return false;
-	};
-	
-	bool bIsOurActor = IsOurActor(RootActorLevelInstance.Get());
-	if (!bIsOurActor)
-	{
-		for (const TWeakObjectPtr<AActor>& LightCard : LightCardLevelInstances)
-		{
-			bIsOurActor = IsOurActor(LightCard.Get());
-			if (bIsOurActor)
-			{
-				break;
-			}
-		}
-	}
-	
-	if (bIsOurActor)
-	{
-		UpdatePreviewActor(RootActorLevelInstance.Get(), true);
-	}
-}
-
 void FDisplayClusterLightCardEditorViewportClient::GetScenePrimitiveComponents(TArray<UPrimitiveComponent*>& OutPrimitiveComponents)
 {
 	RootActorProxy->ForEachComponent<UPrimitiveComponent>(true, [&](UPrimitiveComponent* PrimitiveComponent)
diff --git a/Engine/Plugins/Runtime/nDisplay/Source/DisplayClusterLightCardEditor/Private/Viewport/DisplayClusterLightCardEditorViewportClient.h b/Engine/Plugins/Runtime/nDisplay/Source/DisplayClusterLightCardEditor/Private/Viewport/DisplayClusterLightCardEditorViewportClient.h
--- a/Engine/Plugins/Runtime/nDisplay/Source/DisplayClusterLightCardEditor/Private/Viewport/DisplayClusterLightCardEditorViewportClient.h
+++ b/Engine/Plugins/Runtime/nDisplay/Source/DisplayClusterLightCardEditor/Private/Viewport/DisplayClusterLightCardEditorViewportClient.h
@@ -62,9 +62,6 @@ private:
 
 	/** Ends the current transaction, if one exists. */
 	void EndTransaction();
-
-	/** When a property on the actor has changed. */
-	void OnActorPropertyChanged(UObject* ObjectBeingModified, FPropertyChangedEvent& PropertyChangedEvent);
 	
 	/** Gets a list of all primitive components to be rendered in the scene */
 	void GetScenePrimitiveComponents(TArray<UPrimitiveComponent*>& OutPrimitiveComponents);
