MGPU:  Clamp user specified GPU index for Display Cluster to avoid assert.

#rnx

diff --git a/Engine/Plugins/Runtime/nDisplay/Source/DisplayCluster/Private/Render/Viewport/DisplayClusterViewport.cpp b/Engine/Plugins/Runtime/nDisplay/Source/DisplayCluster/Private/Render/Viewport/DisplayClusterViewport.cpp
--- a/Engine/Plugins/Runtime/nDisplay/Source/DisplayCluster/Private/Render/Viewport/DisplayClusterViewport.cpp
+++ b/Engine/Plugins/Runtime/nDisplay/Source/DisplayCluster/Private/Render/Viewport/DisplayClusterViewport.cpp
@@ -208,7 +208,7 @@ void FDisplayClusterViewport::SetupSceneView(uint32 ContextNum, class UWorld* Wo
 	if(Contexts[ContextNum].GPUIndex >= 0)
 	{
 		InOutView.bOverrideGPUMask = true;
-		InOutView.GPUMask = FRHIGPUMask::FromIndex(Contexts[ContextNum].GPUIndex);
+		InOutView.GPUMask = FRHIGPUMask::FromIndex(FMath::Min((uint32)Contexts[ContextNum].GPUIndex, GNumExplicitGPUsForRendering - 1));
 		InOutView.bAllowCrossGPUTransfer = (Contexts[ContextNum].bAllowGPUTransferOptimization == false);
 	}
 
