nDisplay: Fixed issue with transparency rendering in the light card editor. Also added support for multiselect for light card proxies, as well as synchronization between the selected light card proxies and the selected light cards in the light cards list and details panel.

diff --git a/Engine/Plugins/Runtime/nDisplay/Source/DisplayClusterLightCardEditor/Private/SDisplayClusterLightCardEditor.cpp b/Engine/Plugins/Runtime/nDisplay/Source/DisplayClusterLightCardEditor/Private/SDisplayClusterLightCardEditor.cpp
--- a/Engine/Plugins/Runtime/nDisplay/Source/DisplayClusterLightCardEditor/Private/SDisplayClusterLightCardEditor.cpp
+++ b/Engine/Plugins/Runtime/nDisplay/Source/DisplayClusterLightCardEditor/Private/SDisplayClusterLightCardEditor.cpp
@@ -6,6 +6,7 @@
 
 #include "IDisplayClusterOperator.h"
 #include "Viewport/DisplayClusterLightcardEditorViewport.h"
+#include "Viewport/DisplayClusterLightCardEditorViewportClient.h"
 
 #include "DisplayClusterRootActor.h"
 
@@ -112,6 +113,18 @@ void SDisplayClusterLightCardEditor::Construct(const FArguments& InArgs, const T
 	BindCompileDelegates();
 }
 
+void SDisplayClusterLightCardEditor::SelectLightCards(const TArray<AActor*>& LightCardsToSelect)
+{
+	check(LightCardList);
+	LightCardList->SelectLightCards(LightCardsToSelect);
+}
+
+void SDisplayClusterLightCardEditor::SelectLightCardProxies(const TArray<AActor*>& LightCardsToSelect)
+{
+	check(ViewportView);
+	ViewportView->GetLightCardEditorViewportClient()->SelectLightCards(LightCardsToSelect);
+}
+
 void SDisplayClusterLightCardEditor::OnActiveRootActorChanged(ADisplayClusterRootActor* NewRootActor)
 {
 	RemoveCompileDelegates();
@@ -129,7 +142,7 @@ void SDisplayClusterLightCardEditor::OnActiveRootActorChanged(ADisplayClusterRoo
 
 TSharedRef<SWidget> SDisplayClusterLightCardEditor::CreateLightCardListWidget()
 {
-	return SAssignNew(LightCardList, SDisplayClusterLightCardList);
+	return SAssignNew(LightCardList, SDisplayClusterLightCardList, SharedThis(this));
 }
 
 TSharedRef<SWidget> SDisplayClusterLightCardEditor::CreateViewportWidget()
diff --git a/Engine/Plugins/Runtime/nDisplay/Source/DisplayClusterLightCardEditor/Private/SDisplayClusterLightCardEditor.h b/Engine/Plugins/Runtime/nDisplay/Source/DisplayClusterLightCardEditor/Private/SDisplayClusterLightCardEditor.h
--- a/Engine/Plugins/Runtime/nDisplay/Source/DisplayClusterLightCardEditor/Private/SDisplayClusterLightCardEditor.h
+++ b/Engine/Plugins/Runtime/nDisplay/Source/DisplayClusterLightCardEditor/Private/SDisplayClusterLightCardEditor.h
@@ -48,6 +48,12 @@ public:
 	/** The current active root actor for this light card editor. */
 	TWeakObjectPtr<ADisplayClusterRootActor> GetActiveRootActor() const { return ActiveRootActor; }
 
+	/** Selects the specified light cards in the light card list and details panel */
+	void SelectLightCards(const TArray<AActor*>& LightCardsToSelect);
+
+	/** Selects the light card proxies that correspond to the specified light cards */
+	void SelectLightCardProxies(const TArray<AActor*>& LightCardsToSelect);
+
 private:
 	/** Raised when the active Display cluster root actor has been changed in the operator panel */
 	void OnActiveRootActorChanged(ADisplayClusterRootActor* NewRootActor);
diff --git a/Engine/Plugins/Runtime/nDisplay/Source/DisplayClusterLightCardEditor/Private/SDisplayClusterLightCardList.cpp b/Engine/Plugins/Runtime/nDisplay/Source/DisplayClusterLightCardEditor/Private/SDisplayClusterLightCardList.cpp
--- a/Engine/Plugins/Runtime/nDisplay/Source/DisplayClusterLightCardEditor/Private/SDisplayClusterLightCardList.cpp
+++ b/Engine/Plugins/Runtime/nDisplay/Source/DisplayClusterLightCardEditor/Private/SDisplayClusterLightCardList.cpp
@@ -2,6 +2,8 @@
 
 #include "SDisplayClusterLightCardList.h"
 
+#include "SDisplayClusterLightCardEditor.h"
+
 #include "DisplayClusterRootActor.h"
 #include "DisplayClusterConfigurationTypes.h"
 #include "IDisplayClusterOperator.h"
@@ -123,8 +125,10 @@ private:
 };
 
 
-void SDisplayClusterLightCardList::Construct(const FArguments& InArgs)
+void SDisplayClusterLightCardList::Construct(const FArguments& InArgs, TSharedPtr<SDisplayClusterLightCardEditor> InLightCardEditor)
 {
+	LightCardEditorPtr = InLightCardEditor;
+
 	ChildSlot
 	[
 		SAssignNew(LightCardTreeView, STreeView<TSharedPtr<FLightCardTreeItem>>)
@@ -163,6 +167,22 @@ void SDisplayClusterLightCardList::SetRootActor(ADisplayClusterRootActor* NewRoo
 	}
 }
 
+void SDisplayClusterLightCardList::SelectLightCards(const TArray<AActor*>& LightCardsToSelect)
+{
+	LightCardTreeView->ClearSelection();
+
+	TArray<TSharedPtr<FLightCardTreeItem>> SelectedTreeItems;
+	for (const TSharedPtr<FLightCardTreeItem>& TreeItem : LightCardActors)
+	{
+		if (LightCardsToSelect.Contains(TreeItem->LightCardActor))
+		{
+			SelectedTreeItems.Add(TreeItem);
+		}
+	}
+
+	LightCardTreeView->SetItemSelection(SelectedTreeItems, true);
+}
+
 bool SDisplayClusterLightCardList::FillLightCardList()
 {
 	TArray<TSharedPtr<FLightCardTreeItem>> OriginalTree = MoveTemp(LightCardTree);
@@ -266,13 +286,23 @@ void SDisplayClusterLightCardList::GetChildrenForTreeItem(TSharedPtr<FLightCardT
 
 void SDisplayClusterLightCardList::OnTreeItemSelected(TSharedPtr<FLightCardTreeItem> InItem, ESelectInfo::Type SelectInfo)
 {
-	if (InItem.IsValid() && InItem->LightCardActor.IsValid())
+	TArray<TSharedPtr<FLightCardTreeItem>> SelectedTreeItems;
+	LightCardTreeView->GetSelectedItems(SelectedTreeItems);
+
+	TArray<AActor*> SelectedLightCards;
+	for (const TSharedPtr<FLightCardTreeItem>& SelectedTreeItem : SelectedTreeItems)
 	{
-		IDisplayClusterOperator::Get().ShowDetailsForObject(InItem->LightCardActor.Get());
+		if (SelectedTreeItem->LightCardActor.IsValid())
+		{
+			SelectedLightCards.Add(SelectedTreeItem->LightCardActor.Get());
+		}
 	}
-	else
+
+	IDisplayClusterOperator::Get().ShowDetailsForObjects(*reinterpret_cast<TArray<UObject*>*>(&SelectedLightCards));
+
+	if (LightCardEditorPtr.IsValid())
 	{
-		IDisplayClusterOperator::Get().ShowDetailsForObject(nullptr);
+		LightCardEditorPtr.Pin()->SelectLightCardProxies(SelectedLightCards);
 	}
 }
 
diff --git a/Engine/Plugins/Runtime/nDisplay/Source/DisplayClusterLightCardEditor/Private/SDisplayClusterLightCardList.h b/Engine/Plugins/Runtime/nDisplay/Source/DisplayClusterLightCardEditor/Private/SDisplayClusterLightCardList.h
--- a/Engine/Plugins/Runtime/nDisplay/Source/DisplayClusterLightCardEditor/Private/SDisplayClusterLightCardList.h
+++ b/Engine/Plugins/Runtime/nDisplay/Source/DisplayClusterLightCardEditor/Private/SDisplayClusterLightCardList.h
@@ -7,6 +7,7 @@
 #include "Widgets/SCompoundWidget.h"
 
 class ADisplayClusterRootActor;
+class SDisplayClusterLightCardEditor;
 class ITableRow;
 class STableViewBase;
 
@@ -36,11 +37,14 @@ public:
 	{}
 	SLATE_END_ARGS()
 
-	void Construct(const FArguments& args);
+	void Construct(const FArguments& InArgs, TSharedPtr<SDisplayClusterLightCardEditor> InLightCardEditor);
 
 	void SetRootActor(ADisplayClusterRootActor* NewRootActor);
 
 	const TArray<TSharedPtr<FLightCardTreeItem>>& GetLightCardActors() const { return LightCardActors; }
+
+	void SelectLightCards(const TArray<AActor*>& LightCardsToSelect);
+
 private:
 	/**
 	 * Fill the LightCard list with available LightCards.
@@ -53,6 +57,9 @@ private:
 	void OnTreeItemSelected(TSharedPtr<FLightCardTreeItem> InItem, ESelectInfo::Type SelectInfo);
 
 private:
+	/** Pointer to the light card editor that owns this widget */
+	TWeakPtr<SDisplayClusterLightCardEditor> LightCardEditorPtr;
+
 	/** A hierarchical list of light card tree items to be displayed in the tree view */
 	TArray<TSharedPtr<FLightCardTreeItem>> LightCardTree;
 
diff --git a/Engine/Plugins/Runtime/nDisplay/Source/DisplayClusterLightCardEditor/Private/Viewport/DisplayClusterLightCardEditorViewportClient.cpp b/Engine/Plugins/Runtime/nDisplay/Source/DisplayClusterLightCardEditor/Private/Viewport/DisplayClusterLightCardEditorViewportClient.cpp
--- a/Engine/Plugins/Runtime/nDisplay/Source/DisplayClusterLightCardEditor/Private/Viewport/DisplayClusterLightCardEditorViewportClient.cpp
+++ b/Engine/Plugins/Runtime/nDisplay/Source/DisplayClusterLightCardEditor/Private/Viewport/DisplayClusterLightCardEditorViewportClient.cpp
@@ -412,6 +412,10 @@ void FDisplayClusterLightCardEditorViewportClient::ProcessClick(FSceneView& View
 	UWorld* PreviewWorld = PreviewScene->GetWorld();
 	check(PreviewWorld);
 
+	const bool bIsCtrlKeyDown = Viewport->KeyState(EKeys::LeftControl) || Viewport->KeyState(EKeys::RightControl);
+
+	const bool bMultiSelect = Key == EKeys::LeftMouseButton && bIsCtrlKeyDown;
+
 	if (HitProxy)
 	{
 		if (HitProxy->IsA(HActor::StaticGetType()))
@@ -422,14 +426,14 @@ void FDisplayClusterLightCardEditorViewportClient::ProcessClick(FSceneView& View
 				if (ActorHitProxy->PrimComponent && ActorHitProxy->PrimComponent->IsA<UStaticMeshComponent>())
 				{
 					AActor* TracedLightCard = TraceScreenForLightCard(View, HitX, HitY);
-					SelectLightCard(TracedLightCard);
+					SelectLightCard(TracedLightCard, bMultiSelect);
 				}
 			}
 			else if (LightCardProxies.Contains(ActorHitProxy->Actor))
 			{
-				SelectLightCard(ActorHitProxy->Actor);
+				SelectLightCard(ActorHitProxy->Actor, bMultiSelect);
 			}
-			else
+			else if (!bMultiSelect)
 			{
 				SelectLightCard(nullptr);
 			}
@@ -440,6 +444,8 @@ void FDisplayClusterLightCardEditorViewportClient::ProcessClick(FSceneView& View
 		SelectLightCard(nullptr);
 	}
 	
+	PropagateLightCardSelection();
+
 	FEditorViewportClient::ProcessClick(View, HitProxy, Key, Event, HitX, HitY);
 }
 
@@ -526,16 +532,15 @@ void FDisplayClusterLightCardEditorViewportClient::UpdatePreviewActor(ADisplayCl
 		RootActorProxy.Reset();
 	}
 	
-	for (const TWeakObjectPtr<AActor>& LightCardProxy : LightCardProxies)
+	for (const FLightCardProxy& LightCardProxy : LightCardProxies)
 	{
-		if (LightCardProxy.IsValid())
+		if (LightCardProxy.Proxy.IsValid())
 		{
-			PreviewWorld->EditorDestroyActor(LightCardProxy.Get(), false);
+			PreviewWorld->EditorDestroyActor(LightCardProxy.Proxy.Get(), false);
 		}
 	}
 
 	RootActorLevelInstance.Reset();
-	LightCardLevelInstances.Empty();
 	LightCardProxies.Empty();
 	
 	if (RootActor)
@@ -570,9 +575,10 @@ void FDisplayClusterLightCardEditorViewportClient::UpdatePreviewActor(ADisplayCl
 
 			RootActorProxy->UpdatePreviewComponents();
 
-			FindLightCardsForRootActor(RootActor, LightCardLevelInstances);
+			TArray<TWeakObjectPtr<AActor>> LightCards;
+			FindLightCardsForRootActor(RootActor, LightCards);
 			
-			for (const TWeakObjectPtr<AActor>& LightCard : LightCardLevelInstances)
+			for (const TWeakObjectPtr<AActor>& LightCard : LightCards)
 			{
 				FObjectDuplicationParameters DupeActorParameters(LightCard.Get(), PreviewWorld->GetCurrentLevel());
 				DupeActorParameters.FlagMask = RF_AllFlags & ~(RF_ArchetypeObject | RF_Transactional);
@@ -584,11 +590,30 @@ void FDisplayClusterLightCardEditorViewportClient::UpdatePreviewActor(ADisplayCl
 				LightCardProxy->SetActorLocation(LightCard->GetActorLocation() - RootActor->GetActorLocation());
 				LightCardProxy->SetActorRotation(LightCard->GetActorRotation() - RootActor->GetActorRotation());
 
-				LightCardProxies.Add(LightCardProxy);
+				LightCardProxies.Add(FLightCardProxy(LightCard.Get(), LightCardProxy));
+
 				MeshProjectionRenderer->AddActor(LightCardProxy);
 			}
 		});
 	}
+
+	Viewport->InvalidateHitProxy();
+	bShouldCheckHitProxy = true;
+}
+
+void FDisplayClusterLightCardEditorViewportClient::SelectLightCards(const TArray<AActor*>& LightCardsToSelect)
+{
+	SelectLightCard(nullptr);
+	for (AActor* LightCard : LightCardsToSelect)
+	{
+		if (FLightCardProxy* FoundProxy = LightCardProxies.FindByKey(LightCard))
+		{
+			if (FoundProxy->Proxy.IsValid())
+			{
+				SelectLightCard(FoundProxy->Proxy.Get(), true);
+			}
+		}
+	}
 }
 
 void FDisplayClusterLightCardEditorViewportClient::SetProjectionMode(EDisplayClusterMeshProjectionType InProjectionMode)
@@ -902,6 +927,23 @@ void FDisplayClusterLightCardEditorViewportClient::SelectLightCard(AActor* Actor
 	}
 }
 
+void FDisplayClusterLightCardEditorViewportClient::PropagateLightCardSelection()
+{
+	TArray<AActor*> SelectedLevelInstances;
+	for (const TWeakObjectPtr<AActor>& SelectedLightCard : SelectedLightCards)
+	{
+		if (FLightCardProxy* FoundProxy = LightCardProxies.FindByKey(SelectedLightCard.Get()))
+		{
+			if (FoundProxy->LevelInstance.IsValid())
+			{
+				SelectedLevelInstances.Add(FoundProxy->LevelInstance.Get());
+			}
+		}
+	}
+
+	LightCardEditorPtr.Pin()->SelectLightCards(SelectedLevelInstances);
+}
+
 AActor* FDisplayClusterLightCardEditorViewportClient::TraceScreenForLightCard(const FSceneView& View, int32 HitX, int32 HitY)
 {
 	UWorld* PreviewWorld = PreviewScene->GetWorld();
diff --git a/Engine/Plugins/Runtime/nDisplay/Source/DisplayClusterLightCardEditor/Private/Viewport/DisplayClusterLightCardEditorViewportClient.h b/Engine/Plugins/Runtime/nDisplay/Source/DisplayClusterLightCardEditor/Private/Viewport/DisplayClusterLightCardEditorViewportClient.h
--- a/Engine/Plugins/Runtime/nDisplay/Source/DisplayClusterLightCardEditor/Private/Viewport/DisplayClusterLightCardEditorViewportClient.h
+++ b/Engine/Plugins/Runtime/nDisplay/Source/DisplayClusterLightCardEditor/Private/Viewport/DisplayClusterLightCardEditorViewportClient.h
@@ -47,6 +47,9 @@ public:
 	/** Update the spawned preview actor from a root actor in the level. */
 	void UpdatePreviewActor(ADisplayClusterRootActor* RootActor, bool bForce = false);
 	
+	/** Selects the light card proxies that correspond to the specified light cards */
+	void SelectLightCards(const TArray<AActor*>& LightCardsToSelect);
+
 	EDisplayClusterMeshProjectionType GetProjectionMode() const { return ProjectionMode; }
 	void SetProjectionMode(EDisplayClusterMeshProjectionType InProjectionMode);
 
@@ -84,6 +87,9 @@ private:
 	/** Adds the specified light card actor the the list of selected light cards */
 	void SelectLightCard(AActor* Actor, bool bAddToSelection = false);
 
+	/** Notifies the light card editor of the currently selected light cards so that it may update other UI components to match */
+	void PropagateLightCardSelection();
+
 	/** Traces to find the light card corresponding to a click on a stage screen */
 	AActor* TraceScreenForLightCard(const FSceneView& View, int32 HitX, int32 HitY);
 
@@ -93,8 +99,20 @@ private:
 	TWeakObjectPtr<ADisplayClusterRootActor> RootActorProxy;
 	TWeakObjectPtr<ADisplayClusterRootActor> RootActorLevelInstance;
 
-	TArray<TWeakObjectPtr<AActor>> LightCardProxies;
-	TArray<TWeakObjectPtr<AActor>> LightCardLevelInstances;
+	struct FLightCardProxy
+	{
+		TWeakObjectPtr<AActor> LevelInstance;
+		TWeakObjectPtr<AActor> Proxy;
+
+		FLightCardProxy(AActor* InLevelInstance, AActor* InProxy)
+			: LevelInstance(InLevelInstance)
+			, Proxy(InProxy)
+		{ }
+
+		bool operator==(AActor* Actor) const { return (LevelInstance.IsValid() && LevelInstance.Get() == Actor) || (Proxy.IsValid() && Proxy.Get() == Actor); }
+	};
+
+	TArray<FLightCardProxy> LightCardProxies;
 	TArray<TWeakObjectPtr<AActor>> SelectedLightCards;
 	
 	/** The renderer for the viewport, which can render the meshes with a variety of projection types */
diff --git a/Engine/Plugins/Runtime/nDisplay/Source/DisplayClusterLightCardEditorShaders/Private/DisplayClusterMeshProjectionRenderer.cpp b/Engine/Plugins/Runtime/nDisplay/Source/DisplayClusterLightCardEditorShaders/Private/DisplayClusterMeshProjectionRenderer.cpp
--- a/Engine/Plugins/Runtime/nDisplay/Source/DisplayClusterLightCardEditorShaders/Private/DisplayClusterMeshProjectionRenderer.cpp
+++ b/Engine/Plugins/Runtime/nDisplay/Source/DisplayClusterLightCardEditorShaders/Private/DisplayClusterMeshProjectionRenderer.cpp
@@ -80,12 +80,21 @@ template<EDisplayClusterMeshProjectionType ProjectionType>
 class FLightCardEditorMeshPassProcessor : public FMeshPassProcessor
 {
 public:
-	FLightCardEditorMeshPassProcessor(const FScene* InScene, const FSceneView* InView, FMeshPassDrawListContext* InDrawListContext)
+	FLightCardEditorMeshPassProcessor(const FScene* InScene, const FSceneView* InView, FMeshPassDrawListContext* InDrawListContext, bool bIsTranslucencyPass = false)
 		: FMeshPassProcessor(InScene, GMaxRHIFeatureLevel, InView, InDrawListContext)
 		, DrawRenderState(*InView)
+		, bTranslucencyPass(bIsTranslucencyPass)
 	{
 		DrawRenderState.SetDepthStencilState(TStaticDepthStencilState<true, CF_DepthNearOrEqual>::GetRHI());
-		DrawRenderState.SetBlendState(TStaticBlendState<CW_RGBA, BO_Add, BF_SourceAlpha, BF_InverseSourceAlpha, BO_Add, BF_Zero, BF_One>::GetRHI());
+
+		if (bTranslucencyPass)
+		{
+			DrawRenderState.SetBlendState(TStaticBlendState<CW_RGBA, BO_Add, BF_SourceAlpha, BF_InverseSourceAlpha, BO_Add, BF_Zero, BF_One>::GetRHI());
+		}
+		else
+		{
+			DrawRenderState.SetBlendState(TStaticBlendStateWriteMask<CW_RGBA, CW_RGBA, CW_RGBA, CW_RGBA>::GetRHI());
+		}
 	}
 
 	virtual void AddMeshBatch(const FMeshBatch& RESTRICT MeshBatch, uint64 BatchElementMask, const FPrimitiveSceneProxy* RESTRICT PrimitiveSceneProxy, int32 StaticMeshId = -1) override final
@@ -94,6 +103,11 @@ public:
 		const FMaterial& Material = MeshBatch.MaterialRenderProxy->GetMaterialWithFallback(FeatureLevel, FallbackMaterialRenderProxy);
 		const FMaterialRenderProxy& MaterialRenderProxy = FallbackMaterialRenderProxy ? *FallbackMaterialRenderProxy : *MeshBatch.MaterialRenderProxy;
 
+		if (!CanDrawMeshBatch(MeshBatch, Material))
+		{
+			return;
+		}
+
 		const FVertexFactory* VertexFactory = MeshBatch.VertexFactory;
 
 		TMeshProcessorShaders<FMeshProjectionVS<ProjectionType>, FMeshProjectionPS> PassShaders;
@@ -118,7 +132,7 @@ public:
 		FMeshMaterialShaderElementData ShaderElementData;
 		ShaderElementData.InitializeMeshMaterialData(ViewIfDynamicMeshCommand, PrimitiveSceneProxy, MeshBatch, StaticMeshId, false);
 
-		FMeshDrawCommandSortKey SortKey{};
+		FMeshDrawCommandSortKey SortKey = CreateMeshSortKey(MeshBatch, PrimitiveSceneProxy, Material, PassShaders.VertexShader.GetShader(), PassShaders.PixelShader.GetShader());
 
 		BuildMeshDrawCommands(
 			MeshBatch,
@@ -135,8 +149,58 @@ public:
 			ShaderElementData);
 	}
 
+private:
+	bool CanDrawMeshBatch(const FMeshBatch& RESTRICT MeshBatch, const FMaterial& Material)
+	{
+		const bool bIsTranslucent = IsTranslucentBlendMode(Material.GetBlendMode());
+
+		if (bTranslucencyPass)
+		{
+			return bIsTranslucent;
+		}
+		else
+		{
+			return !bIsTranslucent;
+		}
+	}
+
+	FMeshDrawCommandSortKey CreateMeshSortKey(const FMeshBatch& RESTRICT MeshBatch,
+		const FPrimitiveSceneProxy* RESTRICT PrimitiveSceneProxy,
+		const FMaterial& Material,
+		const FMeshMaterialShader* VertexShader, 
+		const FMeshMaterialShader* PixelShader)
+	{
+		FMeshDrawCommandSortKey SortKey = FMeshDrawCommandSortKey::Default;
+
+		if (bTranslucencyPass)
+		{
+			uint16 SortKeyPriority = 0;
+			float DistanceOffset = 0.0f;
+
+			if (PrimitiveSceneProxy)
+			{
+				const FPrimitiveSceneInfo* PrimitiveSceneInfo = PrimitiveSceneProxy->GetPrimitiveSceneInfo();
+				SortKeyPriority = (uint16)((int32)PrimitiveSceneInfo->Proxy->GetTranslucencySortPriority() - (int32)SHRT_MIN);
+				DistanceOffset = PrimitiveSceneInfo->Proxy->GetTranslucencySortDistanceOffset();
+			}
+
+			SortKey.Translucent.MeshIdInPrimitive = MeshBatch.MeshIdInPrimitive;
+			SortKey.Translucent.Priority = SortKeyPriority;
+			SortKey.Translucent.Distance = *(uint32*)(&DistanceOffset);
+		}
+		else
+		{
+			SortKey.BasePass.VertexShaderHash = (VertexShader ? VertexShader->GetSortKey() : 0) & 0xFFFF;
+			SortKey.BasePass.PixelShaderHash = PixelShader ? PixelShader->GetSortKey() : 0;
+			SortKey.BasePass.Masked = Material.GetBlendMode() == EBlendMode::BLEND_Masked ? 1 : 0;
+		}
+
+		return SortKey;
+	}
+
 private:
 	FMeshPassProcessorRenderState DrawRenderState;
+	bool bTranslucencyPass;
 };
 
 
@@ -601,6 +665,11 @@ void FDisplayClusterMeshProjectionRenderer::Render(FCanvas* Canvas, FSceneInterf
 
 				AddBaseRenderPass(GraphBuilder, View, ProjectionType, ColorRenderTargetBinding, DepthStencilBinding);
 
+				ColorRenderTargetBinding.SetLoadAction(ERenderTargetLoadAction::ELoad);
+				DepthStencilBinding.SetDepthLoadAction(ERenderTargetLoadAction::ELoad);
+
+				AddTranslucencyRenderPass(GraphBuilder, View, ProjectionType, ColorRenderTargetBinding, DepthStencilBinding);
+
 #if WITH_EDITOR
 				const FRDGTextureDesc SelectionDepthDesc = FRDGTextureDesc::Create2D(OutputTexture->Desc.Extent, PF_DepthStencil, FClearValueBinding::DepthFar, TexCreate_DepthStencilTargetable | TexCreate_ShaderResource);
 				FRDGTextureRef SelectionDepthTexture = GraphBuilder.CreateTexture(SelectionDepthDesc, TEXT("DisplayClusterMeshProjection.SelectionDepthTexture"));
@@ -675,7 +744,27 @@ void FDisplayClusterMeshProjectionRenderer::AddBaseRenderPass(FRDGBuilder& Graph
 		FIntRect ViewRect = View->UnscaledViewRect;
 		RHICmdList.SetViewport(ViewRect.Min.X, ViewRect.Min.Y, 0.0f, ViewRect.Max.X, ViewRect.Max.Y, 1.0f);
 
-		SWITCH_ON_PROJECTION_TYPE(RenderPrimitives_RenderThread, ProjectionType, View, RHICmdList);
+		SWITCH_ON_PROJECTION_TYPE(RenderPrimitives_RenderThread, ProjectionType, View, RHICmdList, false);
+	});
+}
+
+void FDisplayClusterMeshProjectionRenderer::AddTranslucencyRenderPass(FRDGBuilder& GraphBuilder,
+	const FViewInfo* View, EDisplayClusterMeshProjectionType ProjectionType,
+	FRenderTargetBinding& OutputRenderTargetBinding,
+	FDepthStencilBinding& OutputDepthStencilBinding)
+{
+	FMeshProjectionPassParameters* MeshPassParameters = GraphBuilder.AllocParameters<FMeshProjectionPassParameters>();
+	MeshPassParameters->View = View->ViewUniformBuffer;
+	MeshPassParameters->InstanceCulling = FInstanceCullingContext::CreateDummyInstanceCullingUniformBuffer(GraphBuilder);
+	MeshPassParameters->RenderTargets[0] = OutputRenderTargetBinding;
+	MeshPassParameters->RenderTargets.DepthStencil = OutputDepthStencilBinding;
+
+	GraphBuilder.AddPass(RDG_EVENT_NAME("MeshProjectionRenderer::Translucency"), MeshPassParameters, ERDGPassFlags::Raster | ERDGPassFlags::NeverCull, [View, ProjectionType, this](FRHICommandList& RHICmdList)
+	{
+		FIntRect ViewRect = View->UnscaledViewRect;
+		RHICmdList.SetViewport(ViewRect.Min.X, ViewRect.Min.Y, 0.0f, ViewRect.Max.X, ViewRect.Max.Y, 1.0f);
+
+		SWITCH_ON_PROJECTION_TYPE(RenderPrimitives_RenderThread, ProjectionType, View, RHICmdList, true);
 	});
 }
 
@@ -773,9 +862,9 @@ void FDisplayClusterMeshProjectionRenderer::AddSelectionOutlineScreenPass(FRDGBu
 #endif
 
 template<EDisplayClusterMeshProjectionType ProjectionType>
-void FDisplayClusterMeshProjectionRenderer::RenderPrimitives_RenderThread(const FSceneView* View, FRHICommandList& RHICmdList)
+void FDisplayClusterMeshProjectionRenderer::RenderPrimitives_RenderThread(const FSceneView* View, FRHICommandList& RHICmdList, bool bTranslucencyPass)
 {
-	DrawDynamicMeshPass(*View, RHICmdList, [View, this](FDynamicPassMeshDrawListContext* DynamicMeshPassContext)
+	DrawDynamicMeshPass(*View, RHICmdList, [View, bTranslucencyPass, this](FDynamicPassMeshDrawListContext* DynamicMeshPassContext)
 	{
 		TArray<FPrimitiveSceneProxy*> PrimitiveSceneProxies;
 		for (const TWeakObjectPtr<UPrimitiveComponent>& PrimitiveComponent : PrimitiveComponents)
@@ -786,7 +875,7 @@ void FDisplayClusterMeshProjectionRenderer::RenderPrimitives_RenderThread(const
 			}
 		}
 
-		FLightCardEditorMeshPassProcessor<ProjectionType> MeshProcessor(nullptr, View, DynamicMeshPassContext);
+		FLightCardEditorMeshPassProcessor<ProjectionType> MeshProcessor(nullptr, View, DynamicMeshPassContext, bTranslucencyPass);
 
 		for (FPrimitiveSceneProxy* PrimitiveProxy : PrimitiveSceneProxies)
 		{
@@ -871,7 +960,7 @@ bool FDisplayClusterMeshProjectionRenderer::IsPrimitiveComponentSelected(const U
 	PROJECTION_TYPE_TEMPLATE_SPECIALIZATION(FuncName, EDisplayClusterMeshProjectionType::Perspective, __VA_ARGS__) \
 	PROJECTION_TYPE_TEMPLATE_SPECIALIZATION(FuncName, EDisplayClusterMeshProjectionType::Azimuthal, __VA_ARGS__)
 
-CREATE_PROJECTION_TYPE_TEMPLATE_SPECIALIZATIONS(RenderPrimitives_RenderThread, const FSceneView* View, FRHICommandList& RHICmdList)
+CREATE_PROJECTION_TYPE_TEMPLATE_SPECIALIZATIONS(RenderPrimitives_RenderThread, const FSceneView* View, FRHICommandList& RHICmdList, bool bTranslucencyPass)
 CREATE_PROJECTION_TYPE_TEMPLATE_SPECIALIZATIONS(RenderHitProxies_RenderThread, const FSceneView* View, FRHICommandList& RHICmdList)
 
 #if WITH_EDITOR
diff --git a/Engine/Plugins/Runtime/nDisplay/Source/DisplayClusterLightCardEditorShaders/Public/DisplayClusterMeshProjectionRenderer.h b/Engine/Plugins/Runtime/nDisplay/Source/DisplayClusterLightCardEditorShaders/Public/DisplayClusterMeshProjectionRenderer.h
--- a/Engine/Plugins/Runtime/nDisplay/Source/DisplayClusterLightCardEditorShaders/Public/DisplayClusterMeshProjectionRenderer.h
+++ b/Engine/Plugins/Runtime/nDisplay/Source/DisplayClusterLightCardEditorShaders/Public/DisplayClusterMeshProjectionRenderer.h
@@ -48,6 +48,13 @@ private:
 		FRenderTargetBinding& OutputRenderTargetBinding,
 		FDepthStencilBinding& OutputDepthStencilBinding);
 
+	/** Adds a pass to perform the translucency render for the mesh projection. */
+	void AddTranslucencyRenderPass(FRDGBuilder& GraphBuilder, 
+		const FViewInfo* View,
+		EDisplayClusterMeshProjectionType ProjectionType,
+		FRenderTargetBinding& OutputRenderTargetBinding,
+		FDepthStencilBinding& OutputDepthStencilBinding);
+
 	/** Adds a pass to perform the hit proxy render for the mesh projection */
 	void AddHitProxyRenderPass(FRDGBuilder& GraphBuilder, 
 		const FViewInfo* View,
@@ -73,7 +80,7 @@ private:
 
 	/** Renders the list of primitive components using the appropriate mesh pass processor given by the template parameter */
 	template<EDisplayClusterMeshProjectionType ProjectionType>
-	void RenderPrimitives_RenderThread(const FSceneView* View, FRHICommandList& RHICmdList);
+	void RenderPrimitives_RenderThread(const FSceneView* View, FRHICommandList& RHICmdList, bool bTranslucencyPass);
 
 	/** Renders the hit proxies of the list of primitive components using the appropriate mesh pass processor given by the template parameter */
 	template<EDisplayClusterMeshProjectionType ProjectionType>
