nDisplay: Fixed issue where ray trace for light card proxies when user clicks on a screen mesh failed because the screen's viewport config had an empty view origin name, and fixed shader compile error indroduced by recent changes in UE5

diff --git a/Engine/Plugins/Runtime/nDisplay/Shaders/Private/MeshProjectionShaders.usf b/Engine/Plugins/Runtime/nDisplay/Shaders/Private/MeshProjectionShaders.usf
--- a/Engine/Plugins/Runtime/nDisplay/Shaders/Private/MeshProjectionShaders.usf
+++ b/Engine/Plugins/Runtime/nDisplay/Shaders/Private/MeshProjectionShaders.usf
@@ -56,11 +56,11 @@ void PerspectiveVS(FVertexFactoryInput Input, out FMeshProjectionVSToPS Output)
 	Output.Position = mul(Output.SavedWorldPosition, ResolvedView.TranslatedWorldToClip);
 }
 
-void MainPS(FMeshProjectionVSToPS Input, in float4 SvPosition : SV_Position OPTIONAL_IsFrontFace, out float4 OutputColor : SV_Target0)
+void MainPS(FMeshProjectionVSToPS Input OPTIONAL_IsFrontFace, out float4 OutputColor : SV_Target0)
 {
-	FMaterialPixelParameters MaterialParameters = GetMaterialPixelParameters(Input.FactoryInterpolants, SvPosition);
+	FMaterialPixelParameters MaterialParameters = GetMaterialPixelParameters(Input.FactoryInterpolants, Input.Position);
 	FPixelMaterialInputs PixelMaterialInputs;
-	CalcMaterialParameters(MaterialParameters, PixelMaterialInputs, SvPosition, bIsFrontFace);
+	CalcMaterialParameters(MaterialParameters, PixelMaterialInputs, Input.Position, bIsFrontFace);
 
 	GetMaterialCoverageAndClipping(MaterialParameters, PixelMaterialInputs);
 	
diff --git a/Engine/Plugins/Runtime/nDisplay/Source/DisplayClusterLightCardEditor/Private/Viewport/DisplayClusterLightCardEditorViewportClient.cpp b/Engine/Plugins/Runtime/nDisplay/Source/DisplayClusterLightCardEditor/Private/Viewport/DisplayClusterLightCardEditorViewportClient.cpp
--- a/Engine/Plugins/Runtime/nDisplay/Source/DisplayClusterLightCardEditor/Private/Viewport/DisplayClusterLightCardEditorViewportClient.cpp
+++ b/Engine/Plugins/Runtime/nDisplay/Source/DisplayClusterLightCardEditor/Private/Viewport/DisplayClusterLightCardEditorViewportClient.cpp
@@ -983,7 +983,19 @@ AActor* FDisplayClusterLightCardEditorViewportClient::TraceScreenForLightCard(co
 				if (UDisplayClusterConfigurationViewport* CfgViewport = FindViewportForPrimitiveComponent(ScreenHitResult.Component.Get()))
 				{
 					FString ViewOriginName = CfgViewport->Camera;
-					if (UDisplayClusterCameraComponent* ViewOrigin = RootActorProxy->GetComponentByName<UDisplayClusterCameraComponent>(ViewOriginName))
+					UDisplayClusterCameraComponent* ViewOrigin = nullptr;
+
+					// If the view origin name is empty, use the first found view origin in the root actor
+					if (ViewOriginName.IsEmpty())
+					{
+						ViewOrigin = RootActorProxy->GetDefaultCamera();
+					}
+					else
+					{
+						ViewOrigin = RootActorProxy->GetComponentByName<UDisplayClusterCameraComponent>(ViewOriginName);
+					}
+
+					if (ViewOrigin)
 					{
 						const FVector ViewOriginRayStart = ViewOrigin->GetComponentLocation();
 						const FVector ViewOriginRayEnd = ViewOriginRayStart + (ScreenHitResult.Location - ViewOriginRayStart) * HALF_WORLD_MAX;
