Fixing some no unity build errors
#rnx

diff --git a/Engine/Plugins/Enterprise/DatasmithContent/Source/DatasmithContent/Private/ObjectTemplates/DatasmithSceneComponentTemplate.cpp b/Engine/Plugins/Enterprise/DatasmithContent/Source/DatasmithContent/Private/ObjectTemplates/DatasmithSceneComponentTemplate.cpp
--- a/Engine/Plugins/Enterprise/DatasmithContent/Source/DatasmithContent/Private/ObjectTemplates/DatasmithSceneComponentTemplate.cpp
+++ b/Engine/Plugins/Enterprise/DatasmithContent/Source/DatasmithContent/Private/ObjectTemplates/DatasmithSceneComponentTemplate.cpp
@@ -3,6 +3,7 @@
 #include "ObjectTemplates/DatasmithSceneComponentTemplate.h"
 
 #include "DatasmithAssetUserData.h"
+#include "Components/PrimitiveComponent.h"
 #include "Components/SceneComponent.h"
 #include "Engine/Level.h"
 #include "GameFramework/Actor.h"
diff --git a/Engine/Plugins/Runtime/nDisplay/Source/DisplayClusterProjection/Private/Policy/Domeprojection/Windows/DX11/DisplayClusterProjectionDomeprojectionLibraryDX11.h b/Engine/Plugins/Runtime/nDisplay/Source/DisplayClusterProjection/Private/Policy/Domeprojection/Windows/DX11/DisplayClusterProjectionDomeprojectionLibraryDX11.h
--- a/Engine/Plugins/Runtime/nDisplay/Source/DisplayClusterProjection/Private/Policy/Domeprojection/Windows/DX11/DisplayClusterProjectionDomeprojectionLibraryDX11.h
+++ b/Engine/Plugins/Runtime/nDisplay/Source/DisplayClusterProjection/Private/Policy/Domeprojection/Windows/DX11/DisplayClusterProjectionDomeprojectionLibraryDX11.h
@@ -3,6 +3,7 @@
 #pragma once
 
 #include "CoreMinimal.h"
+#include "RHI.h"
 #include "ID3D11DynamicRHI.h"
 
 THIRD_PARTY_INCLUDES_START
diff --git a/Engine/Source/Runtime/Renderer/Public/VT/RuntimeVirtualTextureRender.h b/Engine/Source/Runtime/Renderer/Public/VT/RuntimeVirtualTextureRender.h
--- a/Engine/Source/Runtime/Renderer/Public/VT/RuntimeVirtualTextureRender.h
+++ b/Engine/Source/Runtime/Renderer/Public/VT/RuntimeVirtualTextureRender.h
@@ -33,7 +33,7 @@ namespace RuntimeVirtualTexture
 	struct FRenderPageTarget
 	{
 		/** Physical texture to render to. */
-		FRHITexture2D* Texture = nullptr;
+		FRHITexture* Texture = nullptr;
 		/** Unordered access view of physical texture to render to. If this exists we can render to it directly instead of using RHICopyTexture(). */
 		FRHIUnorderedAccessView* UAV = nullptr;
 
