nDisplay: Improve LightCardEditor proxy refresh. LightCard updates can be done separately from the RootActor. Dragging transform sliders will now update the preview transform in real-time.

diff --git a/Engine/Plugins/Runtime/nDisplay/Source/DisplayClusterLightCardEditor/Private/DisplayClusterLightCardEditorProxyType.h b/Engine/Plugins/Runtime/nDisplay/Source/DisplayClusterLightCardEditor/Private/DisplayClusterLightCardEditorProxyType.h
new file mode 100644
--- /dev/null
+++ b/Engine/Plugins/Runtime/nDisplay/Source/DisplayClusterLightCardEditor/Private/DisplayClusterLightCardEditorProxyType.h
@@ -0,0 +1,12 @@
+// Copyright Epic Games, Inc. All Rights Reserved.
+
+#pragma once
+
+#include "CoreMinimal.h"
+
+enum class EDisplayClusterLightCardEditorProxyType : uint8
+{
+	All,
+	RootActor,
+	LightCards
+};
\ No newline at end of file
diff --git a/Engine/Plugins/Runtime/nDisplay/Source/DisplayClusterLightCardEditor/Private/SDisplayClusterLightCardEditor.cpp b/Engine/Plugins/Runtime/nDisplay/Source/DisplayClusterLightCardEditor/Private/SDisplayClusterLightCardEditor.cpp
--- a/Engine/Plugins/Runtime/nDisplay/Source/DisplayClusterLightCardEditor/Private/SDisplayClusterLightCardEditor.cpp
+++ b/Engine/Plugins/Runtime/nDisplay/Source/DisplayClusterLightCardEditor/Private/SDisplayClusterLightCardEditor.cpp
@@ -150,7 +150,7 @@ TSharedRef<SWidget> SDisplayClusterLightCardEditor::CreateViewportWidget()
 	return SAssignNew(ViewportView, SDisplayClusterLightCardEditorViewport, SharedThis(this));
 }
 
-void SDisplayClusterLightCardEditor::RefreshPreviewActors()
+void SDisplayClusterLightCardEditor::RefreshPreviewActors(EDisplayClusterLightCardEditorProxyType ProxyType)
 {
 	RemoveCompileDelegates();
 	
@@ -164,7 +164,7 @@ void SDisplayClusterLightCardEditor::RefreshPreviewActors()
 		if (ViewportView.IsValid())
 		{
 			const bool bForce = true;
-			ViewportView->GetLightCardEditorViewportClient()->UpdatePreviewActor(RootActor, bForce);
+			ViewportView->GetLightCardEditorViewportClient()->UpdatePreviewActor(RootActor, bForce, ProxyType);
 		}
 	}
 	
@@ -208,11 +208,6 @@ void SDisplayClusterLightCardEditor::RemoveCompileDelegates()
 void SDisplayClusterLightCardEditor::OnActorPropertyChanged(UObject* ObjectBeingModified,
                                                             FPropertyChangedEvent& PropertyChangedEvent)
 {
-	if (PropertyChangedEvent.ChangeType == EPropertyChangeType::Interactive)
-	{
-		return;
-	}
-
 	auto IsOurActor = [ObjectBeingModified] (UObject* ObjectToCompare) -> bool
 	{
 		if (ObjectToCompare)
@@ -230,6 +225,8 @@ void SDisplayClusterLightCardEditor::OnActorPropertyChanged(UObject* ObjectBeing
 
 		return false;
 	};
+
+	EDisplayClusterLightCardEditorProxyType ProxyType = EDisplayClusterLightCardEditorProxyType::All;
 	
 	bool bIsOurActor = IsOurActor(GetActiveRootActor().Get());
 	if (!bIsOurActor && LightCardList.IsValid())
@@ -239,6 +236,7 @@ void SDisplayClusterLightCardEditor::OnActorPropertyChanged(UObject* ObjectBeing
 			bIsOurActor = IsOurActor(LightCard->LightCardActor.Get());
 			if (bIsOurActor)
 			{
+				ProxyType = EDisplayClusterLightCardEditorProxyType::LightCards;
 				break;
 			}
 		}
@@ -246,7 +244,19 @@ void SDisplayClusterLightCardEditor::OnActorPropertyChanged(UObject* ObjectBeing
 	
 	if (bIsOurActor)
 	{
-		RefreshPreviewActors();
+		if (PropertyChangedEvent.ChangeType == EPropertyChangeType::Interactive)
+		{
+			// Real-time & efficient update when dragging a slider.
+			if (ViewportView.IsValid())
+			{
+				ViewportView->GetLightCardEditorViewportClient()->UpdateProxyTransforms();
+			}
+		}
+		else
+		{
+			// Full destroy and refresh.
+			RefreshPreviewActors(ProxyType);
+		}
 	}
 }
 
@@ -271,7 +281,7 @@ void SDisplayClusterLightCardEditor::OnLevelActorDeleted(AActor* Actor)
 			{
 				// Schedule for next tick so available selections are properly updated once the
 				// actor is fully deleted.
-				RefreshPreviewActors();
+				RefreshPreviewActors(EDisplayClusterLightCardEditorProxyType::LightCards);
 			});
 		}
 	}
@@ -279,7 +289,8 @@ void SDisplayClusterLightCardEditor::OnLevelActorDeleted(AActor* Actor)
 
 void SDisplayClusterLightCardEditor::OnBlueprintCompiled(UBlueprint* Blueprint)
 {
-	RefreshPreviewActors();
+	// Right now only LightCard blueprints are handled here.
+	RefreshPreviewActors(EDisplayClusterLightCardEditorProxyType::LightCards);
 }
 
 #undef LOCTEXT_NAMESPACE
diff --git a/Engine/Plugins/Runtime/nDisplay/Source/DisplayClusterLightCardEditor/Private/SDisplayClusterLightCardEditor.h b/Engine/Plugins/Runtime/nDisplay/Source/DisplayClusterLightCardEditor/Private/SDisplayClusterLightCardEditor.h
--- a/Engine/Plugins/Runtime/nDisplay/Source/DisplayClusterLightCardEditor/Private/SDisplayClusterLightCardEditor.h
+++ b/Engine/Plugins/Runtime/nDisplay/Source/DisplayClusterLightCardEditor/Private/SDisplayClusterLightCardEditor.h
@@ -6,6 +6,8 @@
 #include "Widgets/DeclarativeSyntaxSupport.h"
 #include "Widgets/SCompoundWidget.h"
 
+#include "DisplayClusterLightCardEditorProxyType.h"
+
 class FTabManager;
 class FLayoutExtender;
 class FSpawnTabArgs;
@@ -65,7 +67,7 @@ private:
 	TSharedRef<SWidget> CreateViewportWidget();
 
 	/** Refresh all preview actors. */
-	void RefreshPreviewActors();
+	void RefreshPreviewActors(EDisplayClusterLightCardEditorProxyType ProxyType = EDisplayClusterLightCardEditorProxyType::All);
 
 	/** Bind delegates to when a BP compiles. */
 	void BindCompileDelegates();
diff --git a/Engine/Plugins/Runtime/nDisplay/Source/DisplayClusterLightCardEditor/Private/Viewport/DisplayClusterLightCardEditorViewportClient.cpp b/Engine/Plugins/Runtime/nDisplay/Source/DisplayClusterLightCardEditor/Private/Viewport/DisplayClusterLightCardEditorViewportClient.cpp
--- a/Engine/Plugins/Runtime/nDisplay/Source/DisplayClusterLightCardEditor/Private/Viewport/DisplayClusterLightCardEditorViewportClient.cpp
+++ b/Engine/Plugins/Runtime/nDisplay/Source/DisplayClusterLightCardEditor/Private/Viewport/DisplayClusterLightCardEditorViewportClient.cpp
@@ -514,58 +514,58 @@ void FDisplayClusterLightCardEditorViewportClient::ResetSelection()
 	SetWidgetMode(UE::Widget::WM_None);
 }
 
-void FDisplayClusterLightCardEditorViewportClient::UpdatePreviewActor(ADisplayClusterRootActor* RootActor, bool bForce)
+void FDisplayClusterLightCardEditorViewportClient::UpdatePreviewActor(ADisplayClusterRootActor* RootActor, bool bForce,
+                                                                      EDisplayClusterLightCardEditorProxyType ProxyType)
 {
-	UWorld* PreviewWorld = PreviewScene->GetWorld();
-	check (PreviewWorld);
-	
 	if (!bForce && RootActor == RootActorLevelInstance.Get())
 	{
 		return;
 	}
 
-	MeshProjectionRenderer->ClearScene();
-	
-	if (RootActorProxy.IsValid())
+	auto Finalize = [this]()
 	{
-		PreviewWorld->EditorDestroyActor(RootActorProxy.Get(), false);
-		RootActorProxy.Reset();
-	}
+		Viewport->InvalidateHitProxy();
+		bShouldCheckHitProxy = true;	
+	};
 	
-	for (const FLightCardProxy& LightCardProxy : LightCardProxies)
+	if (RootActor == nullptr)
 	{
-		if (LightCardProxy.Proxy.IsValid())
-		{
-			PreviewWorld->EditorDestroyActor(LightCardProxy.Proxy.Get(), false);
-		}
+		DestroyProxies(ProxyType);
+		Finalize();
 	}
-
-	RootActorLevelInstance.Reset();
-	LightCardProxies.Empty();
-	
-	if (RootActor)
+	else
 	{
+		UWorld* PreviewWorld = PreviewScene->GetWorld();
+		check(PreviewWorld);
+		
 		// Schedule for the next tick so CDO changes get propagated first in the event of config editor skeleton
 		// regeneration & compiles. nDisplay's custom propagation may have issues if the archetype isn't correct.
 		PreviewWorld->GetTimerManager().SetTimerForNextTick([=]()
 		{
+			DestroyProxies(ProxyType);
 			RootActorLevelInstance = RootActor;
-
+			
+			if (ProxyType == EDisplayClusterLightCardEditorProxyType::All ||
+				ProxyType == EDisplayClusterLightCardEditorProxyType::RootActor)
 			{
-				FObjectDuplicationParameters DupeActorParameters(RootActor, PreviewWorld->GetCurrentLevel());
-				DupeActorParameters.FlagMask = RF_AllFlags & ~(RF_ArchetypeObject | RF_Transactional); // Keeps archetypes correct in config data.
-				DupeActorParameters.PortFlags = PPF_DuplicateVerbatim;
+				{
+					FObjectDuplicationParameters DupeActorParameters(RootActor, PreviewWorld->GetCurrentLevel());
+					DupeActorParameters.FlagMask = RF_AllFlags & ~(RF_ArchetypeObject | RF_Transactional); // Keeps archetypes correct in config data.
+					DupeActorParameters.PortFlags = PPF_DuplicateVerbatim;
 			
-				RootActorProxy = CastChecked<ADisplayClusterRootActor>(StaticDuplicateObjectEx(DupeActorParameters));
-			}
+					RootActorProxy = CastChecked<ADisplayClusterRootActor>(StaticDuplicateObjectEx(DupeActorParameters));
+				}
 			
-			PreviewWorld->GetCurrentLevel()->AddLoadedActor(RootActorProxy.Get());
+				PreviewWorld->GetCurrentLevel()->AddLoadedActor(RootActorProxy.Get());
 
-			// Spawned actor will take the transform values from the template, so manually reset them to zero here
-			RootActorProxy->SetActorLocation(FVector::ZeroVector);
-			RootActorProxy->SetActorRotation(FRotator::ZeroRotator);
+				// Spawned actor will take the transform values from the template, so manually reset them to zero here
+				RootActorProxy->SetActorLocation(FVector::ZeroVector);
+				RootActorProxy->SetActorRotation(FRotator::ZeroRotator);
 
-			FindProjectionOriginComponent();
+				FindProjectionOriginComponent();
+
+				RootActorProxy->UpdatePreviewComponents();
+			}
 
 			// Filter out any primitives hidden in game except screen components
 			MeshProjectionRenderer->AddActor(RootActorProxy.Get(), [](const UPrimitiveComponent* PrimitiveComponent)
@@ -573,32 +573,93 @@ void FDisplayClusterLightCardEditorViewportClient::UpdatePreviewActor(ADisplayCl
 				return !PrimitiveComponent->bHiddenInGame || PrimitiveComponent->IsA<UDisplayClusterScreenComponent>();
 			});
 
-			RootActorProxy->UpdatePreviewComponents();
-
-			TArray<TWeakObjectPtr<AActor>> LightCards;
-			FindLightCardsForRootActor(RootActor, LightCards);
-			
-			for (const TWeakObjectPtr<AActor>& LightCard : LightCards)
+			if (ProxyType == EDisplayClusterLightCardEditorProxyType::All ||
+				ProxyType == EDisplayClusterLightCardEditorProxyType::LightCards)
 			{
-				FObjectDuplicationParameters DupeActorParameters(LightCard.Get(), PreviewWorld->GetCurrentLevel());
-				DupeActorParameters.FlagMask = RF_AllFlags & ~(RF_ArchetypeObject | RF_Transactional);
-				DupeActorParameters.PortFlags = PPF_DuplicateVerbatim;
+				TArray<TWeakObjectPtr<AActor>> LightCards;
+				FindLightCardsForRootActor(RootActor, LightCards);
+			
+				for (const TWeakObjectPtr<AActor>& LightCard : LightCards)
+				{
+					FObjectDuplicationParameters DupeActorParameters(LightCard.Get(), PreviewWorld->GetCurrentLevel());
+					DupeActorParameters.FlagMask = RF_AllFlags & ~(RF_ArchetypeObject | RF_Transactional);
+					DupeActorParameters.PortFlags = PPF_DuplicateVerbatim;
 				
-				AActor* LightCardProxy = CastChecked<AActor>(StaticDuplicateObjectEx(DupeActorParameters));
-				PreviewWorld->GetCurrentLevel()->AddLoadedActor(LightCardProxy);
+					AActor* LightCardProxy = CastChecked<AActor>(StaticDuplicateObjectEx(DupeActorParameters));
+					PreviewWorld->GetCurrentLevel()->AddLoadedActor(LightCardProxy);
 				
-				LightCardProxy->SetActorLocation(LightCard->GetActorLocation() - RootActor->GetActorLocation());
-				LightCardProxy->SetActorRotation(LightCard->GetActorRotation() - RootActor->GetActorRotation());
+					LightCardProxy->SetActorLocation(LightCard->GetActorLocation() - RootActor->GetActorLocation());
+					LightCardProxy->SetActorRotation(LightCard->GetActorRotation() - RootActor->GetActorRotation());
 
-				LightCardProxies.Add(FLightCardProxy(LightCard.Get(), LightCardProxy));
+					LightCardProxies.Add(FLightCardProxy(LightCard.Get(), LightCardProxy));
+				}
+			}
 
-				MeshProjectionRenderer->AddActor(LightCardProxy);
+			for (const FLightCardProxy& LightCardProxy : LightCardProxies)
+			{
+				MeshProjectionRenderer->AddActor(LightCardProxy.Proxy.Get());
 			}
+
+			Finalize();
 		});
 	}
+}
 
-	Viewport->InvalidateHitProxy();
-	bShouldCheckHitProxy = true;
+void FDisplayClusterLightCardEditorViewportClient::UpdateProxyTransforms()
+{
+	if (RootActorLevelInstance.IsValid())
+	{
+		if (RootActorProxy.IsValid())
+		{
+			// Only update scale for the root actor.
+			RootActorProxy->SetActorScale3D(RootActorLevelInstance->GetActorScale3D());
+		}
+		
+		for (const FLightCardProxy& LightCardProxy : LightCardProxies)
+		{
+			if (LightCardProxy.LevelInstance.IsValid() && LightCardProxy.Proxy.IsValid())
+			{
+				LightCardProxy.Proxy->SetActorLocation(LightCardProxy.LevelInstance->GetActorLocation() - RootActorLevelInstance->GetActorLocation());
+				LightCardProxy.Proxy->SetActorRotation(LightCardProxy.LevelInstance->GetActorRotation() - RootActorLevelInstance->GetActorRotation());
+				LightCardProxy.Proxy->SetActorScale3D(LightCardProxy.LevelInstance->GetActorScale3D());
+			}
+		}
+	}
+}
+
+void FDisplayClusterLightCardEditorViewportClient::DestroyProxies(
+EDisplayClusterLightCardEditorProxyType ProxyType)
+{
+	MeshProjectionRenderer->ClearScene();
+
+	UWorld* PreviewWorld = PreviewScene->GetWorld();
+	check(PreviewWorld);
+	
+	if (ProxyType == EDisplayClusterLightCardEditorProxyType::All ||
+		ProxyType == EDisplayClusterLightCardEditorProxyType::RootActor)
+	{
+		if (RootActorProxy.IsValid())
+		{
+			PreviewWorld->EditorDestroyActor(RootActorProxy.Get(), false);
+			RootActorProxy.Reset();
+		}
+	
+		RootActorLevelInstance.Reset();
+	}
+	
+	if (ProxyType == EDisplayClusterLightCardEditorProxyType::All ||
+		ProxyType == EDisplayClusterLightCardEditorProxyType::LightCards)
+	{
+		for (const FLightCardProxy& LightCardProxy : LightCardProxies)
+		{
+			if (LightCardProxy.Proxy.IsValid())
+			{
+				PreviewWorld->EditorDestroyActor(LightCardProxy.Proxy.Get(), false);
+			}
+		}
+
+		LightCardProxies.Empty();	
+	}
 }
 
 void FDisplayClusterLightCardEditorViewportClient::SelectLightCards(const TArray<AActor*>& LightCardsToSelect)
diff --git a/Engine/Plugins/Runtime/nDisplay/Source/DisplayClusterLightCardEditor/Private/Viewport/DisplayClusterLightCardEditorViewportClient.h b/Engine/Plugins/Runtime/nDisplay/Source/DisplayClusterLightCardEditor/Private/Viewport/DisplayClusterLightCardEditorViewportClient.h
--- a/Engine/Plugins/Runtime/nDisplay/Source/DisplayClusterLightCardEditor/Private/Viewport/DisplayClusterLightCardEditorViewportClient.h
+++ b/Engine/Plugins/Runtime/nDisplay/Source/DisplayClusterLightCardEditor/Private/Viewport/DisplayClusterLightCardEditorViewportClient.h
@@ -7,6 +7,7 @@
 #include "SEditorViewport.h"
 #include "AdvancedPreviewScene.h"
 
+#include "DisplayClusterLightCardEditorProxyType.h"
 #include "DisplayClusterMeshProjectionRenderer.h"
 
 class ADisplayClusterRootActor;
@@ -44,8 +45,21 @@ public:
 	void SelectActor(AActor* NewActor);
 	void ResetSelection();
 
-	/** Update the spawned preview actor from a root actor in the level. */
-	void UpdatePreviewActor(ADisplayClusterRootActor* RootActor, bool bForce = false);
+	/**
+	 * Update the spawned preview actor from a root actor in the level
+	 * 
+	 * @param RootActor The new root actor to use. Accepts nullptr
+	 * @param bForce Force the update even if the RootActor hasn't changed
+	 * @param ProxyType The proxy type to destroy and update
+	 */
+	void UpdatePreviewActor(ADisplayClusterRootActor* RootActor, bool bForce = false,
+		EDisplayClusterLightCardEditorProxyType ProxyType = EDisplayClusterLightCardEditorProxyType::All);
+
+	/** Only update required transform values of proxies */
+	void UpdateProxyTransforms();
+
+	/** Remove proxies of the specified type */
+	void DestroyProxies(EDisplayClusterLightCardEditorProxyType ProxyType);
 	
 	/** Selects the light card proxies that correspond to the specified light cards */
 	void SelectLightCards(const TArray<AActor*>& LightCardsToSelect);
