nDisplay: fixed nDisplay/VR behaviour. Disabled HMD scene view extensions to prevent head tracking modify nDisplay viewports

diff --git a/Engine/Plugins/Runtime/nDisplay/Source/DisplayCluster/Public/DisplayClusterSceneViewExtensions.h b/Engine/Plugins/Runtime/nDisplay/Source/DisplayCluster/Public/DisplayClusterSceneViewExtensions.h
--- a/Engine/Plugins/Runtime/nDisplay/Source/DisplayCluster/Public/DisplayClusterSceneViewExtensions.h
+++ b/Engine/Plugins/Runtime/nDisplay/Source/DisplayCluster/Public/DisplayClusterSceneViewExtensions.h
@@ -11,10 +11,15 @@ class IDisplayClusterViewportManager;
 struct FDisplayClusterSceneViewExtensionContext : public FSceneViewExtensionContext
 {
 private:
-
 	//~ FSceneViewExtensionContext Interface
 	virtual FName GetRTTI() const override { return TEXT("FDisplayClusterSceneViewExtensionContext"); }
 
+	virtual bool IsHMDSupported() const override
+	{
+		// Disable all HMD extensions for nDisplay render
+		return false;
+	}
+
 public:
 
 	// The id of the nDisplay viewport being rendered.
diff --git a/Engine/Source/Runtime/Engine/Private/SceneViewExtension.cpp b/Engine/Source/Runtime/Engine/Private/SceneViewExtension.cpp
--- a/Engine/Source/Runtime/Engine/Private/SceneViewExtension.cpp
+++ b/Engine/Source/Runtime/Engine/Private/SceneViewExtension.cpp
@@ -93,7 +93,7 @@ bool FWorldSceneViewExtension::IsActiveThisFrame_Internal(const FSceneViewExtens
 
 bool FHMDSceneViewExtension::IsActiveThisFrame_Internal(const FSceneViewExtensionContext& Context) const
 {
-	return Context.IsStereoSupported();
+	return Context.IsHMDSupported() && Context.IsStereoSupported();
 }
 
 //
diff --git a/Engine/Source/Runtime/Engine/Public/SceneViewExtensionContext.h b/Engine/Source/Runtime/Engine/Public/SceneViewExtensionContext.h
--- a/Engine/Source/Runtime/Engine/Public/SceneViewExtensionContext.h
+++ b/Engine/Source/Runtime/Engine/Public/SceneViewExtensionContext.h
@@ -63,6 +63,12 @@ public:
 	{
 		return !bStereoDisabled && GEngine && GEngine->IsStereoscopic3D(Viewport);
 	}
+
+	// Return true, if HMD supported
+	virtual bool IsHMDSupported() const
+	{
+		return true;
+	}
 };
 
 
