nDisplay: fixed nDisplay/VR behaviour. Disabled HMD post-processing effects via custom view extension which disables bHMDHiddenAreaMaskActive at ViewInfo

diff --git a/Engine/Plugins/Runtime/nDisplay/Source/DisplayCluster/Private/Render/Viewport/DisplayClusterViewportManager.cpp b/Engine/Plugins/Runtime/nDisplay/Source/DisplayCluster/Private/Render/Viewport/DisplayClusterViewportManager.cpp
--- a/Engine/Plugins/Runtime/nDisplay/Source/DisplayCluster/Private/Render/Viewport/DisplayClusterViewportManager.cpp
+++ b/Engine/Plugins/Runtime/nDisplay/Source/DisplayCluster/Private/Render/Viewport/DisplayClusterViewportManager.cpp
@@ -2,6 +2,7 @@
 
 #include "Render/Viewport/DisplayClusterViewportManager.h"
 #include "Render/Viewport/DisplayClusterViewportManagerProxy.h"
+#include "Render/Viewport/DisplayClusterViewportManagerViewExtension.h"
 
 #include "IDisplayCluster.h"
 #include "Render/IDisplayClusterRenderManager.h"
@@ -380,6 +381,12 @@ bool FDisplayClusterViewportManager::BeginNewFrame(FViewport* InViewport, UWorld
 {
 	check(IsInGameThread());
 
+	if (!ViewportManagerViewExtension.IsValid())
+	{
+		// Create DC ViewExtension to handle special features
+		ViewportManagerViewExtension = FSceneViewExtensions::NewExtension<FDisplayClusterViewportManagerViewExtension>(this);
+	}
+
 	OutRenderFrame.ViewportManager = this;
 
 	// Handle world runtime update
diff --git a/Engine/Plugins/Runtime/nDisplay/Source/DisplayCluster/Private/Render/Viewport/DisplayClusterViewportManager.h b/Engine/Plugins/Runtime/nDisplay/Source/DisplayCluster/Private/Render/Viewport/DisplayClusterViewportManager.h
--- a/Engine/Plugins/Runtime/nDisplay/Source/DisplayCluster/Private/Render/Viewport/DisplayClusterViewportManager.h
+++ b/Engine/Plugins/Runtime/nDisplay/Source/DisplayCluster/Private/Render/Viewport/DisplayClusterViewportManager.h
@@ -170,4 +170,7 @@ private:
 	// Support for resetting RTT size (GROW method always grows and does not recover FPS when the viewport size or buffer ratio is changed)
 	ESceneRenderTargetResizeMethod SceneRenderTargetResizeMethod = ESceneRenderTargetResizeMethod::None;
 	int32 FrameHistoryCounter = 0;
+
+	// Handle special features
+	TSharedPtr<class FDisplayClusterViewportManagerViewExtension, ESPMode::ThreadSafe> ViewportManagerViewExtension;
 };
diff --git a/Engine/Plugins/Runtime/nDisplay/Source/DisplayCluster/Private/Render/Viewport/DisplayClusterViewportManagerViewExtension.cpp b/Engine/Plugins/Runtime/nDisplay/Source/DisplayCluster/Private/Render/Viewport/DisplayClusterViewportManagerViewExtension.cpp
new file mode 100644
--- /dev/null
+++ b/Engine/Plugins/Runtime/nDisplay/Source/DisplayCluster/Private/Render/Viewport/DisplayClusterViewportManagerViewExtension.cpp
@@ -0,0 +1,37 @@
+// Copyright Epic Games, Inc. All Rights Reserved.
+
+#include "DisplayClusterViewportManagerViewExtension.h"
+#include "DisplayClusterSceneViewExtensions.h"
+#include "SceneRendering.h"
+
+FDisplayClusterViewportManagerViewExtension::FDisplayClusterViewportManagerViewExtension(const FAutoRegister& AutoRegister, const IDisplayClusterViewportManager* InViewportManager)
+	: FSceneViewExtensionBase(AutoRegister)
+	, ViewportManager(InViewportManager)
+{ }
+
+void FDisplayClusterViewportManagerViewExtension::PreRenderView_RenderThread(FRHICommandListImmediate& RHICmdList, FSceneView& InView)
+{
+	if (InView.bIsViewInfo)
+	{
+		FViewInfo& ViewInfo = static_cast<FViewInfo&>(InView);
+
+		// UE-145088: VR HMD device post-processing cannot be applied to nDisplay rendering
+		ViewInfo.bHMDHiddenAreaMaskActive = false;
+	}
+}
+
+bool FDisplayClusterViewportManagerViewExtension::IsActiveThisFrame_Internal(const FSceneViewExtensionContext& Context) const
+{
+	static const FDisplayClusterSceneViewExtensionContext DCViewExtensionContext;
+	if (Context.IsA(MoveTempIfPossible(DCViewExtensionContext)))
+	{
+		const FDisplayClusterSceneViewExtensionContext& DisplayContext = static_cast<const FDisplayClusterSceneViewExtensionContext&>(Context);
+		if (DisplayContext.ViewportManager == ViewportManager)
+		{
+			// Apply only for DC viewports
+			return true;
+		}
+	}
+
+	return false;
+}
diff --git a/Engine/Plugins/Runtime/nDisplay/Source/DisplayCluster/Private/Render/Viewport/DisplayClusterViewportManagerViewExtension.h b/Engine/Plugins/Runtime/nDisplay/Source/DisplayCluster/Private/Render/Viewport/DisplayClusterViewportManagerViewExtension.h
new file mode 100644
--- /dev/null
+++ b/Engine/Plugins/Runtime/nDisplay/Source/DisplayCluster/Private/Render/Viewport/DisplayClusterViewportManagerViewExtension.h
@@ -0,0 +1,29 @@
+// Copyright Epic Games, Inc. All Rights Reserved.
+
+#pragma once
+#include "SceneViewExtension.h"
+
+class IDisplayClusterViewportManager;
+/** 
+ * View extension applying an DC Viewport features
+ */
+class FDisplayClusterViewportManagerViewExtension : public FSceneViewExtensionBase
+{
+public:
+	FDisplayClusterViewportManagerViewExtension(const FAutoRegister& AutoRegister, const IDisplayClusterViewportManager* InViewportManager);
+
+	//~ Begin ISceneViewExtension interface
+	virtual int32 GetPriority() const override { return -1; }
+
+	virtual void SetupViewFamily(FSceneViewFamily& InViewFamily) override {}
+	virtual void SetupView(FSceneViewFamily& InViewFamily, FSceneView& InView) override {}
+
+	virtual void BeginRenderViewFamily(FSceneViewFamily& InViewFamily) override {};
+
+	virtual void PreRenderView_RenderThread(FRHICommandListImmediate& RHICmdList, FSceneView& InView) override;
+	virtual bool IsActiveThisFrame_Internal(const FSceneViewExtensionContext& Context) const override;
+	//~End ISceneVIewExtension interface
+
+private:
+	const IDisplayClusterViewportManager* ViewportManager;
+};
