nDisplay: Added native light card class to function as a base class for light cards

diff --git a/Engine/Plugins/Runtime/nDisplay/Source/DisplayCluster/Private/Game/EngineClasses/Scene/DisplayClusterLightCardActor.cpp b/Engine/Plugins/Runtime/nDisplay/Source/DisplayCluster/Private/Game/EngineClasses/Scene/DisplayClusterLightCardActor.cpp
new file mode 100644
--- /dev/null
+++ b/Engine/Plugins/Runtime/nDisplay/Source/DisplayCluster/Private/Game/EngineClasses/Scene/DisplayClusterLightCardActor.cpp
@@ -0,0 +1,125 @@
+// Copyright Epic Games, Inc. All Rights Reserved.
+
+#include "DisplayClusterLightCardActor.h"
+
+#include "GameFramework/SpringArmComponent.h"
+#include "Materials/MaterialInstanceDynamic.h"
+#include "UObject/ConstructorHelpers.h"
+
+ADisplayClusterLightCardActor::ADisplayClusterLightCardActor(const FObjectInitializer& ObjectInitializer)
+	: Super(ObjectInitializer)
+	, DistanceFromCenter(300.f)
+	, Longitude(0.f)
+	, Latitude(30.f)
+	, Spin(0.f)
+	, Pitch(0.f)
+	, Yaw(0.f)
+	, Scale(FVector2D(1.f))
+	, Mask(EDisplayClusterLightCardMask::Circle)
+	, Texture(nullptr)
+	, Color(FLinearColor(1.f, 1.f, 1.f, 0.f))
+	, Exposure(0.f)
+	, Opacity(1.f)
+	, Feathering(0.f)
+{
+	PrimaryActorTick.bCanEverTick = true;
+
+	DefaultSceneRootComponent = CreateDefaultSubobject<USceneComponent>(TEXT("DefaultSceneRoot"));
+	SetRootComponent(DefaultSceneRootComponent);
+
+	MainSpringArmComponent = CreateDefaultSubobject<USpringArmComponent>(TEXT("MainSpringArm"));
+	MainSpringArmComponent->AttachToComponent(DefaultSceneRootComponent, FAttachmentTransformRules::KeepRelativeTransform);
+	MainSpringArmComponent->bDoCollisionTest = false;
+
+	LightCardTransformerComponent = CreateDefaultSubobject<USceneComponent>(TEXT("LightCardTransformer"));
+	LightCardTransformerComponent->AttachToComponent(MainSpringArmComponent, FAttachmentTransformRules::KeepRelativeTransform);
+
+	static ConstructorHelpers::FObjectFinder<UStaticMesh> PlaneObj(TEXT("/Engine/BasicShapes/Plane"));
+	static ConstructorHelpers::FObjectFinder<UMaterial> LightCardMatObj(TEXT("/nDisplay/LightCard/M_LightCard"));
+
+	LightCardComponent = CreateDefaultSubobject<UStaticMeshComponent>(TEXT("LightCard"));
+	LightCardComponent->AttachToComponent(LightCardTransformerComponent, FAttachmentTransformRules::KeepRelativeTransform);
+	LightCardComponent->SetCollisionProfileName(TEXT("OverlapAllDynamic"));
+	LightCardComponent->Mobility = EComponentMobility::Movable;
+	LightCardComponent->SetStaticMesh(PlaneObj.Object);
+	LightCardComponent->SetMaterial(0, LightCardMatObj.Object);
+
+	UpdateLightCardTransform();
+}
+
+void ADisplayClusterLightCardActor::OnConstruction(const FTransform& Transform)
+{
+	Super::OnConstruction(Transform);
+
+	UMaterialInterface* Material = LightCardComponent->GetMaterial(0);
+	if (Material && !Material->IsA<UMaterialInstanceDynamic>())
+	{
+		UMaterialInstanceDynamic* LightCardMatInstance = UMaterialInstanceDynamic::Create(Material, this, TEXT("LightCardMID"));
+		LightCardComponent->SetMaterial(0, LightCardMatInstance);
+
+		UpdateLightCardMaterialInstance();
+	}
+}
+
+void ADisplayClusterLightCardActor::Tick(float DeltaSeconds)
+{
+	Super::Tick(DeltaSeconds);
+	UpdateLightCardTransform();
+	UpdateLightCardMaterialInstance();
+}
+
+#if WITH_EDITOR
+
+void ADisplayClusterLightCardActor::PostEditChangeProperty(FPropertyChangedEvent& PropertyChangedEvent)
+{
+	Super::PostEditChangeProperty(PropertyChangedEvent);
+
+	if (PropertyChangedEvent.Property && (
+		PropertyChangedEvent.Property->GetFName() == GET_MEMBER_NAME_CHECKED(ADisplayClusterLightCardActor, DistanceFromCenter) ||
+		PropertyChangedEvent.Property->GetFName() == GET_MEMBER_NAME_CHECKED(ADisplayClusterLightCardActor, Longitude) ||
+		PropertyChangedEvent.Property->GetFName() == GET_MEMBER_NAME_CHECKED(ADisplayClusterLightCardActor, Latitude) ||
+		PropertyChangedEvent.Property->GetFName() == GET_MEMBER_NAME_CHECKED(ADisplayClusterLightCardActor, Spin) ||
+		PropertyChangedEvent.Property->GetFName() == GET_MEMBER_NAME_CHECKED(ADisplayClusterLightCardActor, Scale)))
+	{
+		UpdateLightCardTransform();
+	}
+
+	if (PropertyChangedEvent.Property && (
+		PropertyChangedEvent.Property->GetFName() == GET_MEMBER_NAME_CHECKED(ADisplayClusterLightCardActor, Mask) ||
+		PropertyChangedEvent.Property->GetFName() == GET_MEMBER_NAME_CHECKED(ADisplayClusterLightCardActor, Texture) ||
+		PropertyChangedEvent.Property->GetFName() == GET_MEMBER_NAME_CHECKED(ADisplayClusterLightCardActor, Color) ||
+		PropertyChangedEvent.Property->GetFName() == GET_MEMBER_NAME_CHECKED(ADisplayClusterLightCardActor, Exposure) ||
+		PropertyChangedEvent.Property->GetFName() == GET_MEMBER_NAME_CHECKED(ADisplayClusterLightCardActor, Opacity) ||
+		PropertyChangedEvent.Property->GetFName() == GET_MEMBER_NAME_CHECKED(ADisplayClusterLightCardActor, Feathering)))
+	{
+		UpdateLightCardMaterialInstance();
+	}
+}
+
+#endif
+
+void ADisplayClusterLightCardActor::UpdateLightCardTransform()
+{
+	MainSpringArmComponent->TargetArmLength = DistanceFromCenter;
+	MainSpringArmComponent->SetRelativeRotation(FRotator(-Latitude, Longitude * 180.f, 0.0));
+
+	LightCardComponent->SetRelativeRotation(FRotator(-Spin, -90.f + Yaw, 90.f + Pitch));
+	LightCardComponent->SetRelativeScale3D(FVector(Scale, 1.f));
+}
+
+void ADisplayClusterLightCardActor::UpdateLightCardMaterialInstance()
+{
+	if (UMaterialInstanceDynamic* LightCardMaterialInstance = Cast<UMaterialInstanceDynamic>(LightCardComponent->GetMaterial(0)))
+	{
+		LightCardMaterialInstance->ClearParameterValues();
+
+		LightCardMaterialInstance->SetTextureParameterValue(TEXT("Texture"), Texture);
+		LightCardMaterialInstance->SetScalarParameterValue(TEXT("UseMask"), Mask == EDisplayClusterLightCardMask::Square ? 0.f : 1.f);
+		LightCardMaterialInstance->SetScalarParameterValue(TEXT("UseTextureAlpha"), Mask == EDisplayClusterLightCardMask::UseTextureAlpha ? 1.f : 0.f);
+		LightCardMaterialInstance->SetVectorParameterValue(TEXT("CardColor"), Color);
+		LightCardMaterialInstance->SetScalarParameterValue(TEXT("Exposure"), Exposure);
+		LightCardMaterialInstance->SetScalarParameterValue(TEXT("EmissiveStrength"), Exposure);
+		LightCardMaterialInstance->SetScalarParameterValue(TEXT("Opacity"), Opacity);
+		LightCardMaterialInstance->SetScalarParameterValue(TEXT("Feather"), Feathering);
+	}
+}
\ No newline at end of file
diff --git a/Engine/Plugins/Runtime/nDisplay/Source/DisplayCluster/Public/DisplayClusterLightCardActor.h b/Engine/Plugins/Runtime/nDisplay/Source/DisplayCluster/Public/DisplayClusterLightCardActor.h
new file mode 100644
--- /dev/null
+++ b/Engine/Plugins/Runtime/nDisplay/Source/DisplayCluster/Public/DisplayClusterLightCardActor.h
@@ -0,0 +1,94 @@
+// Copyright Epic Games, Inc. All Rights Reserved.
+
+#pragma once
+
+#include "CoreMinimal.h"
+#include "GameFramework/Actor.h"
+
+#include "DisplayClusterLightCardActor.generated.h"
+
+class USceneComponent;
+class USpringArmComponent;
+class UStaticMeshComponent;
+
+UENUM(BlueprintType)
+enum class EDisplayClusterLightCardMask : uint8
+{
+	Circle,
+	Square,
+	UseTextureAlpha
+};
+
+UCLASS(Blueprintable)
+class ADisplayClusterLightCardActor : public AActor
+{
+	GENERATED_BODY()
+
+public:
+	ADisplayClusterLightCardActor(const FObjectInitializer& ObjectInitializer);
+
+	virtual void OnConstruction(const FTransform& Transform) override;
+	virtual void Tick(float DeltaSeconds) override;
+
+#if WITH_EDITOR
+	virtual void PostEditChangeProperty(FPropertyChangedEvent& PropertyChangedEvent) override;
+	virtual bool ShouldTickIfViewportsOnly() const override { return true; }
+#endif
+
+protected:
+	void UpdateLightCardTransform();
+	void UpdateLightCardMaterialInstance();
+
+public:
+	UPROPERTY(EditAnywhere, BlueprintReadWrite, Category = "Orientation")
+	float DistanceFromCenter;
+
+	UPROPERTY(EditAnywhere, BlueprintReadWrite, Category = "Orientation")
+	float Longitude;
+
+	UPROPERTY(EditAnywhere, BlueprintReadWrite, Category = "Orientation")
+	float Latitude;
+
+	UPROPERTY(EditAnywhere, BlueprintReadWrite, Category = "Orientation")
+	float Spin;
+
+	UPROPERTY(EditAnywhere, BlueprintReadWrite, Category = "Orientation")
+	float Pitch;
+
+	UPROPERTY(EditAnywhere, BlueprintReadWrite, Category = "Orientation")
+	float Yaw;
+
+	UPROPERTY(EditAnywhere, BlueprintReadWrite, Category = "Orientation")
+	FVector2D Scale;
+
+	UPROPERTY(EditAnywhere, BlueprintReadWrite, Category = "Appearance")
+	EDisplayClusterLightCardMask Mask;
+
+	UPROPERTY(EditAnywhere, BlueprintReadWrite, Category = "Appearance")
+	UTexture* Texture;
+
+	UPROPERTY(EditAnywhere, BlueprintReadWrite, Category = "Appearance")
+	FLinearColor Color;
+
+	UPROPERTY(EditAnywhere, BlueprintReadWrite, Category = "Appearance")
+	float Exposure;
+
+	UPROPERTY(EditAnywhere, BlueprintReadWrite, Category = "Appearance", meta = (UIMin = 0, ClampMin = 0, UIMax = 1, ClampMax = 1))
+	float Opacity;
+
+	UPROPERTY(EditAnywhere, BlueprintReadWrite, Category = "Appearance", meta = (UIMin = 0, ClampMin = 0))
+	float Feathering;
+
+protected:
+	UPROPERTY(EditDefaultsOnly, BlueprintReadOnly, Category = "Default")
+	TObjectPtr<USceneComponent> DefaultSceneRootComponent;
+
+	UPROPERTY(EditDefaultsOnly, BlueprintReadOnly, Category = "Default")
+	TObjectPtr<USpringArmComponent> MainSpringArmComponent;
+
+	UPROPERTY(EditDefaultsOnly, BlueprintReadOnly, Category = "Default")
+	TObjectPtr<USceneComponent> LightCardTransformerComponent;
+
+	UPROPERTY(EditDefaultsOnly, BlueprintReadOnly, Category = "Default")
+	TObjectPtr<UStaticMeshComponent> LightCardComponent;
+};
\ No newline at end of file
