nDisplay: Fixed missing include statements in new light card actor class

diff --git a/Engine/Plugins/Runtime/nDisplay/Source/DisplayCluster/Private/Game/EngineClasses/Scene/DisplayClusterLightCardActor.cpp b/Engine/Plugins/Runtime/nDisplay/Source/DisplayCluster/Private/Game/EngineClasses/Scene/DisplayClusterLightCardActor.cpp
--- a/Engine/Plugins/Runtime/nDisplay/Source/DisplayCluster/Private/Game/EngineClasses/Scene/DisplayClusterLightCardActor.cpp
+++ b/Engine/Plugins/Runtime/nDisplay/Source/DisplayCluster/Private/Game/EngineClasses/Scene/DisplayClusterLightCardActor.cpp
@@ -2,8 +2,10 @@
 
 #include "DisplayClusterLightCardActor.h"
 
+#include "Components/StaticMeshComponent.h"
 #include "GameFramework/SpringArmComponent.h"
 #include "Materials/MaterialInstanceDynamic.h"
+#include "Engine/StaticMesh.h"
 #include "UObject/ConstructorHelpers.h"
 
 ADisplayClusterLightCardActor::ADisplayClusterLightCardActor(const FObjectInitializer& ObjectInitializer)
