nDisplay: Add UI options and support for adding and removing Light Cards. Add undo/redo support for refreshing the preview.

diff --git a/Engine/Plugins/Runtime/nDisplay/Source/DisplayCluster/Public/DisplayClusterLightCardActor.h b/Engine/Plugins/Runtime/nDisplay/Source/DisplayCluster/Public/DisplayClusterLightCardActor.h
--- a/Engine/Plugins/Runtime/nDisplay/Source/DisplayCluster/Public/DisplayClusterLightCardActor.h
+++ b/Engine/Plugins/Runtime/nDisplay/Source/DisplayCluster/Public/DisplayClusterLightCardActor.h
@@ -20,7 +20,7 @@ enum class EDisplayClusterLightCardMask : uint8
 };
 
 UCLASS(Blueprintable)
-class ADisplayClusterLightCardActor : public AActor
+class DISPLAYCLUSTER_API ADisplayClusterLightCardActor : public AActor
 {
 	GENERATED_BODY()
 
diff --git a/Engine/Plugins/Runtime/nDisplay/Source/DisplayClusterLightCardEditor/DisplayClusterLightCardEditor.Build.cs b/Engine/Plugins/Runtime/nDisplay/Source/DisplayClusterLightCardEditor/DisplayClusterLightCardEditor.Build.cs
--- a/Engine/Plugins/Runtime/nDisplay/Source/DisplayClusterLightCardEditor/DisplayClusterLightCardEditor.Build.cs
+++ b/Engine/Plugins/Runtime/nDisplay/Source/DisplayClusterLightCardEditor/DisplayClusterLightCardEditor.Build.cs
@@ -27,6 +27,9 @@ public class DisplayClusterLightCardEditor : ModuleRules
 				"EditorStyle",
 				"WorkspaceMenuStructure",
 				"AdvancedPreviewScene",
+				"ToolWidgets",
+				"PropertyEditor",
+				"AppFramework"
 			});
 	}
 }
diff --git a/Engine/Plugins/Runtime/nDisplay/Source/DisplayClusterLightCardEditor/Private/DisplayClusterLightCardEditorCommands.cpp b/Engine/Plugins/Runtime/nDisplay/Source/DisplayClusterLightCardEditor/Private/DisplayClusterLightCardEditorCommands.cpp
--- a/Engine/Plugins/Runtime/nDisplay/Source/DisplayClusterLightCardEditor/Private/DisplayClusterLightCardEditorCommands.cpp
+++ b/Engine/Plugins/Runtime/nDisplay/Source/DisplayClusterLightCardEditor/Private/DisplayClusterLightCardEditorCommands.cpp
@@ -11,6 +11,10 @@ void FDisplayClusterLightCardEditorCommands::RegisterCommands()
 
 	UI_COMMAND(PerspectiveProjection, "Perspective", "A perspective projection from the stage's view origin", EUserInterfaceActionType::Button, FInputChord());
 	UI_COMMAND(AzimuthalProjection, "Azimuthal", "An azimuthal hemisphere projection from the stage's center pointed upwards", EUserInterfaceActionType::Button, FInputChord());
+
+	UI_COMMAND(AddNewLightCard, "Add New Light Card", "Add and assign a new Light Card to the actor", EUserInterfaceActionType::Button, FInputChord());
+	UI_COMMAND(AddExistingLightCard, "Add Existing Light Card", "Add an existing Light Card to the actor", EUserInterfaceActionType::Button, FInputChord());
+	UI_COMMAND(RemoveLightCard, "Remove from Actor", "Remove the Light Card from the actor but do not delete it", EUserInterfaceActionType::Button, FInputChord());
 }
 
 #undef LOCTEXT_NAMESPACE
diff --git a/Engine/Plugins/Runtime/nDisplay/Source/DisplayClusterLightCardEditor/Private/DisplayClusterLightCardEditorCommands.h b/Engine/Plugins/Runtime/nDisplay/Source/DisplayClusterLightCardEditor/Private/DisplayClusterLightCardEditorCommands.h
--- a/Engine/Plugins/Runtime/nDisplay/Source/DisplayClusterLightCardEditor/Private/DisplayClusterLightCardEditorCommands.h
+++ b/Engine/Plugins/Runtime/nDisplay/Source/DisplayClusterLightCardEditor/Private/DisplayClusterLightCardEditorCommands.h
@@ -24,4 +24,8 @@ public:
 
 	TSharedPtr<FUICommandInfo> PerspectiveProjection;
 	TSharedPtr<FUICommandInfo> AzimuthalProjection;
+
+	TSharedPtr<FUICommandInfo> AddNewLightCard;
+	TSharedPtr<FUICommandInfo> AddExistingLightCard;
+	TSharedPtr<FUICommandInfo> RemoveLightCard;
 };
diff --git a/Engine/Plugins/Runtime/nDisplay/Source/DisplayClusterLightCardEditor/Private/SDisplayClusterLightCardEditor.cpp b/Engine/Plugins/Runtime/nDisplay/Source/DisplayClusterLightCardEditor/Private/SDisplayClusterLightCardEditor.cpp
--- a/Engine/Plugins/Runtime/nDisplay/Source/DisplayClusterLightCardEditor/Private/SDisplayClusterLightCardEditor.cpp
+++ b/Engine/Plugins/Runtime/nDisplay/Source/DisplayClusterLightCardEditor/Private/SDisplayClusterLightCardEditor.cpp
@@ -15,7 +15,6 @@
 #include "Framework/MultiBox/MultiBoxBuilder.h"
 #include "Framework/MultiBox/MultiBoxExtender.h"
 #include "Toolkits/AssetEditorToolkit.h"
-#include "Viewport/DisplayClusterLightCardEditorViewportClient.h"
 #include "Widgets/Docking/SDockTab.h"
 #include "Widgets/Layout/SSplitter.h"
 
@@ -72,6 +71,11 @@ SDisplayClusterLightCardEditor::~SDisplayClusterLightCardEditor()
 		GEngine->OnLevelActorDeleted().RemoveAll(this);
 	}
 
+	if (OnObjectTransactedHandle.IsValid())
+	{
+		FCoreUObjectDelegates::OnObjectTransacted.Remove(OnObjectTransactedHandle);
+	}
+
 	RemoveCompileDelegates();
 }
 
@@ -82,6 +86,8 @@ void SDisplayClusterLightCardEditor::Construct(const FArguments& InArgs, const T
 	{
 		GEngine->OnLevelActorDeleted().AddSP(this, &SDisplayClusterLightCardEditor::OnLevelActorDeleted);
 	}
+	
+	OnObjectTransactedHandle = FCoreUObjectDelegates::OnObjectTransacted.AddSP(this, &SDisplayClusterLightCardEditor::OnObjectTransacted);
 
 	ChildSlot                     
 	[
@@ -171,6 +177,47 @@ void SDisplayClusterLightCardEditor::RefreshPreviewActors(EDisplayClusterLightCa
 	BindCompileDelegates();
 }
 
+bool SDisplayClusterLightCardEditor::IsOurObject(UObject* InObject,
+	EDisplayClusterLightCardEditorProxyType& OutProxyType) const
+{
+	auto IsOurActor = [InObject] (UObject* ObjectToCompare) -> bool
+	{
+		if (ObjectToCompare)
+		{
+			if (InObject == ObjectToCompare)
+			{
+				return true;
+			}
+
+			if (const UObject* RootActorOuter = InObject->GetTypedOuter(ObjectToCompare->GetClass()))
+			{
+				return RootActorOuter == ObjectToCompare;
+			}
+		}
+
+		return false;
+	};
+
+	EDisplayClusterLightCardEditorProxyType ProxyType = EDisplayClusterLightCardEditorProxyType::All;
+	
+	bool bIsOurActor = IsOurActor(GetActiveRootActor().Get());
+	if (!bIsOurActor && LightCardList.IsValid())
+	{
+		for (const TSharedPtr<SDisplayClusterLightCardList::FLightCardTreeItem>& LightCard : LightCardList->GetLightCardActors())
+		{
+			bIsOurActor = IsOurActor(LightCard->LightCardActor.Get());
+			if (bIsOurActor)
+			{
+				ProxyType = EDisplayClusterLightCardEditorProxyType::LightCards;
+				break;
+			}
+		}
+	}
+
+	OutProxyType = ProxyType;
+	return bIsOurActor;
+}
+
 void SDisplayClusterLightCardEditor::BindCompileDelegates()
 {
 	if (LightCardList.IsValid())
@@ -208,41 +255,8 @@ void SDisplayClusterLightCardEditor::RemoveCompileDelegates()
 void SDisplayClusterLightCardEditor::OnActorPropertyChanged(UObject* ObjectBeingModified,
                                                             FPropertyChangedEvent& PropertyChangedEvent)
 {
-	auto IsOurActor = [ObjectBeingModified] (UObject* ObjectToCompare) -> bool
-	{
-		if (ObjectToCompare)
-		{
-			if (ObjectBeingModified == ObjectToCompare)
-			{
-				return true;
-			}
-
-			if (const UObject* RootActorOuter = ObjectBeingModified->GetTypedOuter(ObjectToCompare->GetClass()))
-			{
-				return RootActorOuter == ObjectToCompare;
-			}
-		}
-
-		return false;
-	};
-
-	EDisplayClusterLightCardEditorProxyType ProxyType = EDisplayClusterLightCardEditorProxyType::All;
-	
-	bool bIsOurActor = IsOurActor(GetActiveRootActor().Get());
-	if (!bIsOurActor && LightCardList.IsValid())
-	{
-		for (const TSharedPtr<SDisplayClusterLightCardList::FLightCardTreeItem>& LightCard : LightCardList->GetLightCardActors())
-		{
-			bIsOurActor = IsOurActor(LightCard->LightCardActor.Get());
-			if (bIsOurActor)
-			{
-				ProxyType = EDisplayClusterLightCardEditorProxyType::LightCards;
-				break;
-			}
-		}
-	}
-	
-	if (bIsOurActor)
+	EDisplayClusterLightCardEditorProxyType ProxyType;
+	if (IsOurObject(ObjectBeingModified, ProxyType))
 	{
 		if (PropertyChangedEvent.ChangeType == EPropertyChangeType::Interactive)
 		{
@@ -293,4 +307,16 @@ void SDisplayClusterLightCardEditor::OnBlueprintCompiled(UBlueprint* Blueprint)
 	RefreshPreviewActors(EDisplayClusterLightCardEditorProxyType::LightCards);
 }
 
+void SDisplayClusterLightCardEditor::OnObjectTransacted(UObject* Object,
+	const FTransactionObjectEvent& TransactionObjectEvent)
+{
+	if (TransactionObjectEvent.GetEventType() == ETransactionObjectEventType::UndoRedo)
+	{
+		// Always refresh on undo because the light card actor may not inherit our C++ class
+		// so we can't easily distinguish it. This supports the case where the user deletes
+		// a LightCard actor from the level manually then undoes it.
+		RefreshPreviewActors();
+	}
+}
+
 #undef LOCTEXT_NAMESPACE
diff --git a/Engine/Plugins/Runtime/nDisplay/Source/DisplayClusterLightCardEditor/Private/SDisplayClusterLightCardEditor.h b/Engine/Plugins/Runtime/nDisplay/Source/DisplayClusterLightCardEditor/Private/SDisplayClusterLightCardEditor.h
--- a/Engine/Plugins/Runtime/nDisplay/Source/DisplayClusterLightCardEditor/Private/SDisplayClusterLightCardEditor.h
+++ b/Engine/Plugins/Runtime/nDisplay/Source/DisplayClusterLightCardEditor/Private/SDisplayClusterLightCardEditor.h
@@ -47,7 +47,7 @@ public:
 
 	void Construct(const FArguments& args, const TSharedRef<SDockTab>& MajorTabOwner, const TSharedPtr<SWindow>& WindowOwner);
 
-	/** The current active root actor for this light card editor. */
+	/** The current active root actor for this light card editor */
 	TWeakObjectPtr<ADisplayClusterRootActor> GetActiveRootActor() const { return ActiveRootActor; }
 
 	/** Selects the specified light cards in the light card list and details panel */
@@ -63,32 +63,43 @@ private:
 	/** Creates the widget used to show the list of light cards associated with the active root actor */
 	TSharedRef<SWidget> CreateLightCardListWidget();
 
-	/** Create the 3d viewport widget. */
+	/** Create the 3d viewport widget */
 	TSharedRef<SWidget> CreateViewportWidget();
 
-	/** Refresh all preview actors. */
+	/** Refresh all preview actors */
 	void RefreshPreviewActors(EDisplayClusterLightCardEditorProxyType ProxyType = EDisplayClusterLightCardEditorProxyType::All);
 
-	/** Bind delegates to when a BP compiles. */
+	/**
+	 * Check if an object is managed by us
+	 * @param InObject The object to compare
+	 * @param OutProxyType The type of the object
+	 * @return True if our object, false if not
+	 */
+	bool IsOurObject(UObject* InObject, EDisplayClusterLightCardEditorProxyType& OutProxyType) const;
+	
+	/** Bind delegates to when a BP compiles */
 	void BindCompileDelegates();
 
-	/** Remove compile delegates from a BP. */
+	/** Remove compile delegates from a BP */
 	void RemoveCompileDelegates();
 
-	/** When a property on the actor has changed. */
+	/** When a property on the actor has changed */
 	void OnActorPropertyChanged(UObject* ObjectBeingModified, FPropertyChangedEvent& PropertyChangedEvent);
 
 	/** Raised when the user deletes an actor from the level */
 	void OnLevelActorDeleted(AActor* Actor);
 
-	/** Raised when a supported blueprint is compiled. */
+	/** Raised when a supported blueprint is compiled */
 	void OnBlueprintCompiled(UBlueprint* Blueprint);
+
+	/** Raised when any object is transacted */
+	void OnObjectTransacted(UObject* Object, const FTransactionObjectEvent& TransactionObjectEvent);
 	
 private:
 	/** The light card list widget */
 	TSharedPtr<SDisplayClusterLightCardList> LightCardList;
 
-	/** The 3d viewport. */
+	/** The 3d viewport */
 	TSharedPtr<SDisplayClusterLightCardEditorViewport> ViewportView;
 
 	/** A reference to the root actor that is currently being operated on */
@@ -96,4 +107,7 @@ private:
 
 	/** Delegate handle for the OnActiveRootActorChanged delegate */
 	FDelegateHandle ActiveRootActorChangedHandle;
+
+	/** Delegate handle for when an object is transacted */
+	FDelegateHandle OnObjectTransactedHandle;
 };
\ No newline at end of file
diff --git a/Engine/Plugins/Runtime/nDisplay/Source/DisplayClusterLightCardEditor/Private/SDisplayClusterLightCardList.cpp b/Engine/Plugins/Runtime/nDisplay/Source/DisplayClusterLightCardEditor/Private/SDisplayClusterLightCardList.cpp
--- a/Engine/Plugins/Runtime/nDisplay/Source/DisplayClusterLightCardEditor/Private/SDisplayClusterLightCardList.cpp
+++ b/Engine/Plugins/Runtime/nDisplay/Source/DisplayClusterLightCardEditor/Private/SDisplayClusterLightCardList.cpp
@@ -3,16 +3,26 @@
 #include "SDisplayClusterLightCardList.h"
 
 #include "SDisplayClusterLightCardEditor.h"
+#include "DisplayClusterLightCardEditorCommands.h"
 
 #include "DisplayClusterRootActor.h"
+#include "DisplayClusterLightCardActor.h"
 #include "DisplayClusterConfigurationTypes.h"
 #include "IDisplayClusterOperator.h"
+#include "Components/DisplayClusterCameraComponent.h"
 
 #include "ClassIconFinder.h"
+#include "ScopedTransaction.h"
+#include "Selection.h"
 #include "Widgets/Images/SImage.h"
 #include "Widgets/Views/SListView.h"
 #include "Widgets/Views/STableRow.h"
 #include "Widgets/Views/STreeView.h"
+#include "SPositiveActionButton.h"
+#include "Framework/Commands/GenericCommands.h"
+#include "Framework/MultiBox/MultiBoxBuilder.h"
+#include "PropertyCustomizationHelpers.h"
+#include "Widgets/Workflow/SWizard.h"
 
 #define LOCTEXT_NAMESPACE "SDisplayClusterLightCardList"
 
@@ -129,24 +139,54 @@ void SDisplayClusterLightCardList::Construct(const FArguments& InArgs, TSharedPt
 {
 	LightCardEditorPtr = InLightCardEditor;
 
+	BindCommands();
+
 	ChildSlot
 	[
-		SAssignNew(LightCardTreeView, STreeView<TSharedPtr<FLightCardTreeItem>>)
-		.TreeItemsSource(&LightCardTree)
-		.ItemHeight(28)
-		.SelectionMode(ESelectionMode::Multi)
-		.OnGenerateRow(this, &SDisplayClusterLightCardList::GenerateTreeItemRow)
-		.OnGetChildren(this, &SDisplayClusterLightCardList::GetChildrenForTreeItem)
-		.OnSelectionChanged(this, &SDisplayClusterLightCardList::OnTreeItemSelected)
-		.HighlightParentNodesForSelection(true)
-		.HeaderRow
-		(
-			SNew(SHeaderRow)
-
-			+ SHeaderRow::Column(DisplayClusterLightCardListColumnNames::LightCardName)
-			.DefaultLabel(LOCTEXT("LightCardName", "Light Card"))
-			.FillWidth(0.8f)
-		)
+		SNew(SVerticalBox)
+		+SVerticalBox::Slot()
+		.AutoHeight()
+		[
+			SNew(SBorder)
+			.Padding(0)
+			.BorderImage(FAppStyle::Get().GetBrush("DetailsView.CategoryTop"))
+			.BorderBackgroundColor(FLinearColor(0.6f, 0.6f, 0.6f, 1.0f))
+			[
+				SNew(SHorizontalBox)
+				+SHorizontalBox::Slot()
+				.Padding(3.0f, 3.0f)
+				.AutoWidth()
+				.VAlign(VAlign_Center)
+				[
+						
+					SNew(SPositiveActionButton)
+					.Text(LOCTEXT("AddNewItem", "Add"))
+					.Icon(FAppStyle::Get().GetBrush("Icons.Plus"))
+					.ToolTipText(LOCTEXT("AddNewLightCardTooltip", "Adds a new Light Card to the display cluster"))
+					.OnGetMenuContent(this, &SDisplayClusterLightCardList::CreateAddNewMenuContent)
+				]
+			]
+		]
+		+SVerticalBox::Slot()
+		[
+			SAssignNew(LightCardTreeView, STreeView<TSharedPtr<FLightCardTreeItem>>)
+			.TreeItemsSource(&LightCardTree)
+			.ItemHeight(28)
+			.SelectionMode(ESelectionMode::Multi)
+			.OnGenerateRow(this, &SDisplayClusterLightCardList::GenerateTreeItemRow)
+			.OnGetChildren(this, &SDisplayClusterLightCardList::GetChildrenForTreeItem)
+			.OnSelectionChanged(this, &SDisplayClusterLightCardList::OnTreeItemSelected)
+			.OnContextMenuOpening(this, &SDisplayClusterLightCardList::CreateContextMenu)
+			.HighlightParentNodesForSelection(true)
+			.HeaderRow
+			(
+				SNew(SHeaderRow)
+
+				+ SHeaderRow::Column(DisplayClusterLightCardListColumnNames::LightCardName)
+				.DefaultLabel(LOCTEXT("LightCardName", "Light Card"))
+				.FillWidth(0.8f)
+			)
+		]
 	];
 }
 
@@ -183,6 +223,42 @@ void SDisplayClusterLightCardList::SelectLightCards(const TArray<AActor*>& Light
 	LightCardTreeView->SetItemSelection(SelectedTreeItems, true);
 }
 
+void SDisplayClusterLightCardList::Refresh()
+{
+	SetRootActor(RootActor.Get());
+}
+
+void SDisplayClusterLightCardList::BindCommands()
+{
+	const FDisplayClusterLightCardEditorCommands& Commands = FDisplayClusterLightCardEditorCommands::Get();
+
+	CommandList = MakeShareable(new FUICommandList);
+	
+	CommandList->MapAction(
+		Commands.AddNewLightCard,
+		FExecuteAction::CreateSP(this, &SDisplayClusterLightCardList::AddNewLightCard),
+		FCanExecuteAction::CreateSP(this, &SDisplayClusterLightCardList::CanAddLightCard),
+		FCanExecuteAction());
+
+	CommandList->MapAction(
+		Commands.AddExistingLightCard,
+		FExecuteAction::CreateSP(this, &SDisplayClusterLightCardList::AddExistingLightCard),
+		FCanExecuteAction::CreateSP(this, &SDisplayClusterLightCardList::CanAddLightCard),
+		FCanExecuteAction());
+	
+	CommandList->MapAction(
+		Commands.RemoveLightCard,
+		FExecuteAction::CreateSP(this, &SDisplayClusterLightCardList::RemoveLightCard, false),
+		FCanExecuteAction::CreateSP(this, &SDisplayClusterLightCardList::CanRemoveLightCard),
+		FCanExecuteAction());
+
+	CommandList->MapAction(
+		FGenericCommands::Get().Delete,
+		FExecuteAction::CreateSP(this, &SDisplayClusterLightCardList::RemoveLightCard, true),
+		FCanExecuteAction::CreateSP(this, &SDisplayClusterLightCardList::CanRemoveLightCard),
+		FCanExecuteAction());
+}
+
 bool SDisplayClusterLightCardList::FillLightCardList()
 {
 	TArray<TSharedPtr<FLightCardTreeItem>> OriginalTree = MoveTemp(LightCardTree);
@@ -262,6 +338,210 @@ bool SDisplayClusterLightCardList::FillLightCardList()
 	return false;
 }
 
+void SDisplayClusterLightCardList::AddNewLightCard()
+{
+	check(RootActor.IsValid());
+	
+	FScopedTransaction Transaction(LOCTEXT("AddNewLightCard", "Add New Light Card"));
+	
+	const FVector SpawnLocation = RootActor->GetDefaultCamera()->GetComponentLocation();
+	FRotator SpawnRotation = RootActor->GetDefaultCamera()->GetComponentRotation();
+	SpawnRotation.Yaw -= 180.f;
+	
+	FActorSpawnParameters SpawnParameters;
+	SpawnParameters.bNoFail = true;
+	SpawnParameters.SpawnCollisionHandlingOverride = ESpawnActorCollisionHandlingMethod::AdjustIfPossibleButAlwaysSpawn;
+	SpawnParameters.Name = TEXT("LightCard");
+	SpawnParameters.NameMode = FActorSpawnParameters::ESpawnActorNameMode::Requested;
+	
+	ADisplayClusterLightCardActor* NewActor = CastChecked<ADisplayClusterLightCardActor>(
+		RootActor->GetWorld()->SpawnActor(ADisplayClusterLightCardActor::StaticClass(),
+		&SpawnLocation, &SpawnRotation, MoveTemp(SpawnParameters)));
+
+	NewActor->SetActorLabel(NewActor->GetName());
+
+	AddLightCardToActor(NewActor);
+}
+
+void SDisplayClusterLightCardList::AddExistingLightCard()
+{
+	TSharedPtr<SWindow> PickerWindow;
+	TWeakObjectPtr<AActor> SelectedActorPtr;
+	bool bFinished = false;
+	
+	const TSharedRef<SWidget> ActorPicker = PropertyCustomizationHelpers::MakeActorPickerWithMenu(
+		nullptr,
+		false,
+		FOnShouldFilterActor::CreateLambda([&](const AActor* const InActor) -> bool // ActorFilter
+		{
+			const bool IsAllowed = InActor != nullptr && !InActor->IsChildActor() && InActor->IsA<AActor>() &&
+				!InActor->GetClass()->HasAnyClassFlags(CLASS_Interface)	&& !InActor->IsA<ADisplayClusterRootActor>();
+			
+			return IsAllowed;
+		}),
+		FOnActorSelected::CreateLambda([&](AActor* InActor) -> void // OnSet
+		{
+			SelectedActorPtr = InActor;
+		}),
+		FSimpleDelegate::CreateLambda([&]() -> void // OnClose
+		{
+		}),
+		FSimpleDelegate::CreateLambda([&]() -> void // OnUseSelected
+		{
+			if (AActor* Selection = Cast<AActor>(GEditor->GetSelectedActors()->GetTop(AActor::StaticClass())))
+			{
+				SelectedActorPtr = Selection;
+			}
+		}));
+	
+	PickerWindow = SNew(SWindow)
+	.Title(LOCTEXT("AddExistingLightCard", "Select an existing Light Card actor"))
+	.ClientSize(FVector2D(500, 525))
+	.SupportsMinimize(false) .SupportsMaximize(false)
+	[
+		SNew(SBorder)
+		.BorderImage(FEditorStyle::GetBrush("Menu.Background"))
+		[
+			SNew(SWizard)
+			.FinishButtonText(LOCTEXT("FinishAddingExistingLightCard", "Add Actor"))
+			.OnCanceled(FSimpleDelegate::CreateLambda([&]()
+			{
+				if (PickerWindow.IsValid())
+				{
+					PickerWindow->RequestDestroyWindow();
+				}
+			}))
+			.OnFinished(FSimpleDelegate::CreateLambda([&]()
+			{
+				bFinished = true;
+				if (PickerWindow.IsValid())
+				{
+					PickerWindow->RequestDestroyWindow();
+				}
+			}))
+			.CanFinish(TAttribute<bool>::CreateLambda([&]()
+			{
+				return SelectedActorPtr.IsValid();
+			}))
+			.ShowPageList(false)
+			+SWizard::Page()
+			.CanShow(true)
+			[
+				SNew(SBorder)
+				.VAlign(VAlign_Fill)
+				.BorderImage(FEditorStyle::GetBrush("Menu.Background"))
+				[
+					SNew(SVerticalBox)
+					+SVerticalBox::Slot()
+					.AutoHeight()
+					.VAlign(VAlign_Fill)
+					.HAlign(HAlign_Fill)
+					[
+						ActorPicker
+					]
+					+SVerticalBox::Slot()
+					.AutoHeight()
+					.VAlign(VAlign_Bottom)
+					.Padding(0.f, 8.f)
+					[
+						SNew(STextBlock)
+						.TextStyle(FEditorStyle::Get(), "NormalText.Important")
+						.Text_Lambda([&]
+						{
+							const FString Result = FString::Printf(TEXT("Selected Actor: %s"),
+								SelectedActorPtr.IsValid() ? *SelectedActorPtr->GetActorLabel() : TEXT(""));
+							return FText::FromString(Result);
+						})
+					]
+				]
+			]
+		]
+	];
+
+	GEditor->EditorAddModalWindow(PickerWindow.ToSharedRef());
+	if (bFinished && SelectedActorPtr.IsValid())
+	{
+		FScopedTransaction Transaction(LOCTEXT("AddExistingLightCard", "Add Existing Light Card"));
+		AddLightCardToActor(SelectedActorPtr.Get());
+	}
+
+	PickerWindow.Reset();
+	SelectedActorPtr.Reset();
+}
+
+void SDisplayClusterLightCardList::AddLightCardToActor(AActor* LightCard)
+{
+	UDisplayClusterConfigurationData* ConfigData = RootActor->GetConfigData();
+	ConfigData->Modify();
+	FDisplayClusterConfigurationICVFX_VisibilityList& RootActorLightCards = ConfigData->StageSettings.Lightcard.ShowOnlyList;
+	RootActorLightCards.Actors.AddUnique(LightCard);
+
+	Refresh();
+}
+
+bool SDisplayClusterLightCardList::CanAddLightCard() const
+{
+	return RootActor.IsValid() && RootActor->GetWorld() != nullptr;
+}
+
+void SDisplayClusterLightCardList::RemoveLightCard(bool bDeleteLightCardActor)
+{
+	TArray<TSharedPtr<FLightCardTreeItem>> SelectedTreeItems;
+	LightCardTreeView->GetSelectedItems(SelectedTreeItems);
+
+	FScopedTransaction Transaction(LOCTEXT("RemoveLightCard", "Remove Light Card(s)"));
+
+	USelection* EdSelectionManager = GEditor->GetSelectedActors();
+	UWorld* WorldToUse = nullptr;
+	
+	if (bDeleteLightCardActor)
+	{
+		EdSelectionManager->BeginBatchSelectOperation();
+		EdSelectionManager->Modify();
+		EdSelectionManager->DeselectAll();
+	}
+	
+	for (const TSharedPtr<FLightCardTreeItem>& Item : SelectedTreeItems)
+	{
+		if (Item.IsValid() && Item->LightCardActor.IsValid())
+		{
+			if (RootActor.IsValid())
+			{
+				UDisplayClusterConfigurationData* ConfigData = RootActor->GetConfigData();
+				ConfigData->Modify();
+				
+				FDisplayClusterConfigurationICVFX_VisibilityList& RootActorLightCards = ConfigData->StageSettings.Lightcard.ShowOnlyList;
+				RootActorLightCards.Actors.Remove(Item->LightCardActor.Get());
+			}
+
+			if (bDeleteLightCardActor)
+			{
+				WorldToUse = Item->LightCardActor->GetWorld();
+				GEditor->SelectActor(Item->LightCardActor.Get(), /*bSelect =*/true, /*bNotifyForActor =*/false, /*bSelectEvenIfHidden =*/true);
+			}
+		}
+	}
+
+	if (bDeleteLightCardActor)
+	{
+		EdSelectionManager->EndBatchSelectOperation();
+
+		if (WorldToUse)
+		{
+			GEditor->edactDeleteSelected(WorldToUse);
+		}
+	}
+
+	Refresh();
+}
+
+bool SDisplayClusterLightCardList::CanRemoveLightCard() const
+{
+	TArray<TSharedPtr<FLightCardTreeItem>> SelectedTreeItems;
+	LightCardTreeView->GetSelectedItems(SelectedTreeItems);
+	return SelectedTreeItems.Num() > 0;
+}
+
 TSharedRef<ITableRow> SDisplayClusterLightCardList::GenerateTreeItemRow(TSharedPtr<FLightCardTreeItem> Item, const TSharedRef<STableViewBase>& OwnerTable)
 {
 	return SNew(SLightCardTreeItemRow, OwnerTable, Item);
@@ -306,4 +586,37 @@ void SDisplayClusterLightCardList::OnTreeItemSelected(TSharedPtr<FLightCardTreeI
 	}
 }
 
-#undef LOCTEXT_NAMESPACE
\ No newline at end of file
+FReply SDisplayClusterLightCardList::OnKeyDown(const FGeometry& MyGeometry, const FKeyEvent& InKeyEvent)
+{
+	if (CommandList->ProcessCommandBindings(InKeyEvent))
+	{
+		return FReply::Handled();
+	}
+	return FReply::Unhandled();
+}
+
+TSharedRef<SWidget> SDisplayClusterLightCardList::CreateAddNewMenuContent()
+{
+	const FDisplayClusterLightCardEditorCommands& Commands = FDisplayClusterLightCardEditorCommands::Get();
+	
+	const bool bCloseAfterSelection = true;
+	FMenuBuilder MenuBuilder(bCloseAfterSelection, CommandList, Extenders);
+	MenuBuilder.AddMenuEntry(Commands.AddNewLightCard);
+	MenuBuilder.AddMenuEntry(Commands.AddExistingLightCard);
+	
+	return MenuBuilder.MakeWidget();
+}
+
+TSharedPtr<SWidget> SDisplayClusterLightCardList::CreateContextMenu()
+{
+	const FDisplayClusterLightCardEditorCommands& Commands = FDisplayClusterLightCardEditorCommands::Get();
+	
+	const bool bCloseAfterSelection = true;
+	FMenuBuilder MenuBuilder(bCloseAfterSelection, CommandList, Extenders);
+	MenuBuilder.AddMenuEntry(Commands.RemoveLightCard);
+	MenuBuilder.AddMenuEntry(FGenericCommands::Get().Delete);
+
+	return MenuBuilder.MakeWidget();
+}
+
+#undef LOCTEXT_NAMESPACE
diff --git a/Engine/Plugins/Runtime/nDisplay/Source/DisplayClusterLightCardEditor/Private/SDisplayClusterLightCardList.h b/Engine/Plugins/Runtime/nDisplay/Source/DisplayClusterLightCardEditor/Private/SDisplayClusterLightCardList.h
--- a/Engine/Plugins/Runtime/nDisplay/Source/DisplayClusterLightCardEditor/Private/SDisplayClusterLightCardList.h
+++ b/Engine/Plugins/Runtime/nDisplay/Source/DisplayClusterLightCardEditor/Private/SDisplayClusterLightCardList.h
@@ -6,7 +6,10 @@
 #include "Widgets/DeclarativeSyntaxSupport.h"
 #include "Widgets/SCompoundWidget.h"
 
+class FExtender;
+class FUICommandList;
 class ADisplayClusterRootActor;
+class ADisplayClusterLightCardActor;
 class SDisplayClusterLightCardEditor;
 class ITableRow;
 class STableViewBase;
@@ -45,17 +48,48 @@ public:
 
 	void SelectLightCards(const TArray<AActor*>& LightCardsToSelect);
 
+	/** Refreshes the list */
+	void Refresh();
+
 private:
+	void BindCommands();
+	
 	/**
-	 * Fill the LightCard list with available LightCards.
-	 * @return True if the list has been modified.
+	 * Fill the LightCard list with available LightCards
+	 * @return True if the list has been modified
 	 */
 	bool FillLightCardList();
 
+	/** Spawn and add a new Light Card */
+	void AddNewLightCard();
+	
+	/** Select an existing Light Card from a menu */
+	void AddExistingLightCard();
+	
+	/** Adds the given Light Card to the root actor */
+	void AddLightCardToActor(AActor* LightCard);
+	
+	/** If a Light Card can currently be added */
+	bool CanAddLightCard() const;
+
+	/**
+	 * Remove the light card from the actor
+	 *@param bDeleteLightCardActor Delete the actor from the level
+	 */
+	void RemoveLightCard(bool bDeleteLightCardActor);
+	
+	/** If the selected Light Card can be removed */
+	bool CanRemoveLightCard() const;
+
 	TSharedRef<ITableRow> GenerateTreeItemRow(TSharedPtr<FLightCardTreeItem> Item, const TSharedRef<STableViewBase>& OwnerTable);
 	void GetChildrenForTreeItem(TSharedPtr<FLightCardTreeItem> InItem, TArray<TSharedPtr<FLightCardTreeItem>>& OutChildren);
 	void OnTreeItemSelected(TSharedPtr<FLightCardTreeItem> InItem, ESelectInfo::Type SelectInfo);
 
+	virtual FReply OnKeyDown(const FGeometry& MyGeometry, const FKeyEvent& InKeyEvent) override;
+	
+	TSharedRef<SWidget> CreateAddNewMenuContent();
+	TSharedPtr<SWidget> CreateContextMenu();
+	
 private:
 	/** Pointer to the light card editor that owns this widget */
 	TWeakPtr<SDisplayClusterLightCardEditor> LightCardEditorPtr;
@@ -71,4 +105,10 @@ private:
 
 	/** The tree view widget displaying the light cards */
 	TSharedPtr<STreeView<TSharedPtr<FLightCardTreeItem>>> LightCardTreeView;
+	
+	/** Extenders for menus */
+	TSharedPtr<FExtender> Extenders;
+
+	/** Mapped commands for this list */
+	TSharedPtr<FUICommandList> CommandList;
 };
\ No newline at end of file
