nDisplay: Fix lightcards not refreshing the 3d preview when being added or removed. Fix issues where duplicate light cards may appear in the list of they were force deleted from a level.

diff --git a/Engine/Plugins/Runtime/nDisplay/Source/DisplayClusterLightCardEditor/Private/SDisplayClusterLightCardEditor.cpp b/Engine/Plugins/Runtime/nDisplay/Source/DisplayClusterLightCardEditor/Private/SDisplayClusterLightCardEditor.cpp
--- a/Engine/Plugins/Runtime/nDisplay/Source/DisplayClusterLightCardEditor/Private/SDisplayClusterLightCardEditor.cpp
+++ b/Engine/Plugins/Runtime/nDisplay/Source/DisplayClusterLightCardEditor/Private/SDisplayClusterLightCardEditor.cpp
@@ -148,7 +148,8 @@ void SDisplayClusterLightCardEditor::OnActiveRootActorChanged(ADisplayClusterRoo
 
 TSharedRef<SWidget> SDisplayClusterLightCardEditor::CreateLightCardListWidget()
 {
-	return SAssignNew(LightCardList, SDisplayClusterLightCardList, SharedThis(this));
+	return SAssignNew(LightCardList, SDisplayClusterLightCardList, SharedThis(this))
+	.OnLightCardListChanged(this, &SDisplayClusterLightCardEditor::OnLightCardListChanged);
 }
 
 TSharedRef<SWidget> SDisplayClusterLightCardEditor::CreateViewportWidget()
@@ -319,4 +320,9 @@ void SDisplayClusterLightCardEditor::OnObjectTransacted(UObject* Object,
 	}
 }
 
+void SDisplayClusterLightCardEditor::OnLightCardListChanged()
+{
+	RefreshPreviewActors(EDisplayClusterLightCardEditorProxyType::LightCards);
+}
+
 #undef LOCTEXT_NAMESPACE
diff --git a/Engine/Plugins/Runtime/nDisplay/Source/DisplayClusterLightCardEditor/Private/SDisplayClusterLightCardEditor.h b/Engine/Plugins/Runtime/nDisplay/Source/DisplayClusterLightCardEditor/Private/SDisplayClusterLightCardEditor.h
--- a/Engine/Plugins/Runtime/nDisplay/Source/DisplayClusterLightCardEditor/Private/SDisplayClusterLightCardEditor.h
+++ b/Engine/Plugins/Runtime/nDisplay/Source/DisplayClusterLightCardEditor/Private/SDisplayClusterLightCardEditor.h
@@ -94,6 +94,9 @@ private:
 
 	/** Raised when any object is transacted */
 	void OnObjectTransacted(UObject* Object, const FTransactionObjectEvent& TransactionObjectEvent);
+
+	/** Raised when the light card list has added or removed a card */
+	void OnLightCardListChanged();
 	
 private:
 	/** The light card list widget */
diff --git a/Engine/Plugins/Runtime/nDisplay/Source/DisplayClusterLightCardEditor/Private/SDisplayClusterLightCardList.cpp b/Engine/Plugins/Runtime/nDisplay/Source/DisplayClusterLightCardEditor/Private/SDisplayClusterLightCardList.cpp
--- a/Engine/Plugins/Runtime/nDisplay/Source/DisplayClusterLightCardEditor/Private/SDisplayClusterLightCardList.cpp
+++ b/Engine/Plugins/Runtime/nDisplay/Source/DisplayClusterLightCardEditor/Private/SDisplayClusterLightCardList.cpp
@@ -138,7 +138,8 @@ private:
 void SDisplayClusterLightCardList::Construct(const FArguments& InArgs, TSharedPtr<SDisplayClusterLightCardEditor> InLightCardEditor)
 {
 	LightCardEditorPtr = InLightCardEditor;
-
+	OnLightCardChanged = InArgs._OnLightCardListChanged;
+	
 	BindCommands();
 
 	ChildSlot
@@ -471,12 +472,31 @@ void SDisplayClusterLightCardList::AddExistingLightCard()
 
 void SDisplayClusterLightCardList::AddLightCardToActor(AActor* LightCard)
 {
+	check(LightCard);
+	
 	UDisplayClusterConfigurationData* ConfigData = RootActor->GetConfigData();
 	ConfigData->Modify();
 	FDisplayClusterConfigurationICVFX_VisibilityList& RootActorLightCards = ConfigData->StageSettings.Lightcard.ShowOnlyList;
-	RootActorLightCards.Actors.AddUnique(LightCard);
 
+	if (!RootActorLightCards.Actors.ContainsByPredicate([&](const TSoftObjectPtr<AActor>& Actor)
+	{
+		// Don't add if a loaded actor is already present.
+		return Actor.Get() == LightCard;
+	}))
+	{
+		const TSoftObjectPtr<AActor> LightCardSoftObject(LightCard);
+		
+		// Remove any exact paths to this actor. It's possible invalid actors are present if a light card
+		// was force deleted from a level.
+		RootActorLightCards.Actors.RemoveAll([&] (const TSoftObjectPtr<AActor>& Actor)
+		{
+			return Actor == LightCardSoftObject;
+		});
+	
+		RootActorLightCards.Actors.Add(LightCard);
+	}
 	Refresh();
+	OnLightCardChanged.ExecuteIfBound();
 }
 
 bool SDisplayClusterLightCardList::CanAddLightCard() const
@@ -511,7 +531,10 @@ void SDisplayClusterLightCardList::RemoveLightCard(bool bDeleteLightCardActor)
 				ConfigData->Modify();
 				
 				FDisplayClusterConfigurationICVFX_VisibilityList& RootActorLightCards = ConfigData->StageSettings.Lightcard.ShowOnlyList;
-				RootActorLightCards.Actors.Remove(Item->LightCardActor.Get());
+				RootActorLightCards.Actors.RemoveAll([&](const TSoftObjectPtr<AActor>& Actor)
+				{
+					return Actor.Get() == Item->LightCardActor.Get();
+				});
 			}
 
 			if (bDeleteLightCardActor)
@@ -533,6 +556,7 @@ void SDisplayClusterLightCardList::RemoveLightCard(bool bDeleteLightCardActor)
 	}
 
 	Refresh();
+	OnLightCardChanged.ExecuteIfBound();
 }
 
 bool SDisplayClusterLightCardList::CanRemoveLightCard() const
diff --git a/Engine/Plugins/Runtime/nDisplay/Source/DisplayClusterLightCardEditor/Private/SDisplayClusterLightCardList.h b/Engine/Plugins/Runtime/nDisplay/Source/DisplayClusterLightCardEditor/Private/SDisplayClusterLightCardList.h
--- a/Engine/Plugins/Runtime/nDisplay/Source/DisplayClusterLightCardEditor/Private/SDisplayClusterLightCardList.h
+++ b/Engine/Plugins/Runtime/nDisplay/Source/DisplayClusterLightCardEditor/Private/SDisplayClusterLightCardList.h
@@ -21,6 +21,8 @@ class STreeView;
 class SDisplayClusterLightCardList : public SCompoundWidget
 {
 public:
+	DECLARE_DELEGATE(FOnLightCardListChanged)
+	
 	struct FLightCardTreeItem
 	{
 		TWeakObjectPtr<AActor> LightCardActor;
@@ -38,6 +40,7 @@ public:
 public:
 	SLATE_BEGIN_ARGS(SDisplayClusterLightCardList)
 	{}
+	SLATE_EVENT(FOnLightCardListChanged, OnLightCardListChanged)
 	SLATE_END_ARGS()
 
 	void Construct(const FArguments& InArgs, TSharedPtr<SDisplayClusterLightCardEditor> InLightCardEditor);
@@ -111,4 +114,7 @@ private:
 
 	/** Mapped commands for this list */
 	TSharedPtr<FUICommandList> CommandList;
+
+	/** When a user adds or removes a light card */
+	FOnLightCardListChanged OnLightCardChanged;
 };
\ No newline at end of file
