nDisplay: Allow DefaultViewPoint component to be accessed directly in BP.

diff --git a/Engine/Plugins/Runtime/nDisplay/Source/DisplayCluster/Public/DisplayClusterRootActor.h b/Engine/Plugins/Runtime/nDisplay/Source/DisplayCluster/Public/DisplayClusterRootActor.h
--- a/Engine/Plugins/Runtime/nDisplay/Source/DisplayCluster/Public/DisplayClusterRootActor.h
+++ b/Engine/Plugins/Runtime/nDisplay/Source/DisplayCluster/Public/DisplayClusterRootActor.h
@@ -116,7 +116,7 @@ protected:
 	bool BuildHierarchy();
 
 public:
-	UFUNCTION(BlueprintCallable, meta = (DisplayName = "Get Default Camera"), Category = "NDisplay|Components")
+	UFUNCTION(BlueprintGetter)
 	UDisplayClusterCameraComponent* GetDefaultCamera() const;
 
 	UFUNCTION(BlueprintCallable, Category = "NDisplay|Render")
@@ -239,7 +239,7 @@ private:
 	/**
 	 * Default camera component. It's an outer camera in VP/ICVFX terminology. Always exists on a DCRA instance.
 	 */
-	UPROPERTY(VisibleAnywhere, Category = "NDisplay")
+	UPROPERTY(VisibleAnywhere, BlueprintGetter=GetDefaultCamera, Category = "NDisplay|Components")
 	UDisplayClusterCameraComponent* DefaultViewPoint;
 
 	/**
