nDisplay: Adding a new Light Card will add it to the Root Actor's level instead of the persistent level. Fix crash when loading a new map with the light card editor open and set to a DCRA in the level being destroyed.

diff --git a/Engine/Plugins/Runtime/nDisplay/Source/DisplayClusterLightCardEditor/Private/SDisplayClusterLightCardList.cpp b/Engine/Plugins/Runtime/nDisplay/Source/DisplayClusterLightCardEditor/Private/SDisplayClusterLightCardList.cpp
--- a/Engine/Plugins/Runtime/nDisplay/Source/DisplayClusterLightCardEditor/Private/SDisplayClusterLightCardList.cpp
+++ b/Engine/Plugins/Runtime/nDisplay/Source/DisplayClusterLightCardEditor/Private/SDisplayClusterLightCardList.cpp
@@ -354,11 +354,13 @@ void SDisplayClusterLightCardList::AddNewLightCard()
 	SpawnParameters.SpawnCollisionHandlingOverride = ESpawnActorCollisionHandlingMethod::AdjustIfPossibleButAlwaysSpawn;
 	SpawnParameters.Name = TEXT("LightCard");
 	SpawnParameters.NameMode = FActorSpawnParameters::ESpawnActorNameMode::Requested;
+	SpawnParameters.OverrideLevel = RootActor->GetLevel();
 	
 	ADisplayClusterLightCardActor* NewActor = CastChecked<ADisplayClusterLightCardActor>(
 		RootActor->GetWorld()->SpawnActor(ADisplayClusterLightCardActor::StaticClass(),
 		&SpawnLocation, &SpawnRotation, MoveTemp(SpawnParameters)));
 
+	ensure(NewActor->GetLevel() == RootActor->GetLevel());
 	NewActor->SetActorLabel(NewActor->GetName());
 
 	AddLightCardToActor(NewActor);
diff --git a/Engine/Plugins/Runtime/nDisplay/Source/DisplayClusterOperator/Private/SDisplayClusterOperatorToolbar.cpp b/Engine/Plugins/Runtime/nDisplay/Source/DisplayClusterOperator/Private/SDisplayClusterOperatorToolbar.cpp
--- a/Engine/Plugins/Runtime/nDisplay/Source/DisplayClusterOperator/Private/SDisplayClusterOperatorToolbar.cpp
+++ b/Engine/Plugins/Runtime/nDisplay/Source/DisplayClusterOperator/Private/SDisplayClusterOperatorToolbar.cpp
@@ -6,6 +6,7 @@
 #include "DisplayClusterRootActor.h"
 
 #include "EditorStyleSet.h"
+#include "LevelEditor.h"
 #include "Framework/MultiBox/MultiBoxBuilder.h"
 #include "Toolkits/AssetEditorToolkit.h"
 #include "Widgets/Layout/SBorder.h"
@@ -15,10 +16,16 @@
 
 SDisplayClusterOperatorToolbar::~SDisplayClusterOperatorToolbar()
 {
-	if (GEngine != nullptr)
+	if (LevelActorDeletedHandle.IsValid() && GEngine != nullptr)
 	{
 		GEngine->OnLevelActorDeleted().Remove(LevelActorDeletedHandle);
 	}
+
+	if (MapChangedHandle.IsValid())
+	{
+		FLevelEditorModule& LevelEditor = FModuleManager::LoadModuleChecked<FLevelEditorModule>("LevelEditor");
+		LevelEditor.OnMapChanged().Remove(MapChangedHandle);
+	}
 	
 	if (ActiveRootActor.IsValid())
 	{
@@ -75,6 +82,9 @@ void SDisplayClusterOperatorToolbar::Construct(const FArguments& InArgs)
 	{
 		LevelActorDeletedHandle = GEngine->OnLevelActorDeleted().AddSP(this, &SDisplayClusterOperatorToolbar::OnLevelActorDeleted);
 	}
+
+	FLevelEditorModule& LevelEditor = FModuleManager::LoadModuleChecked<FLevelEditorModule>("LevelEditor");
+	MapChangedHandle = LevelEditor.OnMapChanged().AddRaw(this, &SDisplayClusterOperatorToolbar::HandleMapChanged);
 }
 
 TSharedPtr<FString> SDisplayClusterOperatorToolbar::FillRootActorList(const FString& InitiallySelectedRootActor)
@@ -99,6 +109,22 @@ TSharedPtr<FString> SDisplayClusterOperatorToolbar::FillRootActorList(const FStr
 	return SelectedItem;
 }
 
+void SDisplayClusterOperatorToolbar::ClearSelectedRootActor()
+{
+	if (ActiveRootActor.IsValid())
+	{
+		if (UBlueprint* Blueprint = UBlueprint::GetBlueprintFromClass(ActiveRootActor->GetClass()))
+		{
+			Blueprint->OnCompiled().RemoveAll(this);
+		}
+	}
+
+	ActiveRootActorName.Reset();
+	ActiveRootActor.Reset();
+	IDisplayClusterOperator::Get().OnActiveRootActorChanged().Broadcast(nullptr);
+	RootActorComboBox->SetSelectedItem(nullptr);
+}
+
 void SDisplayClusterOperatorToolbar::OnRootActorChanged(TSharedPtr<FString> ItemSelected, ESelectInfo::Type SelectInfo)
 {
 	ActiveRootActorName = ItemSelected;
@@ -132,6 +158,7 @@ void SDisplayClusterOperatorToolbar::OnRootActorChanged(TSharedPtr<FString> Item
 	{
 		if (UBlueprint* Blueprint = UBlueprint::GetBlueprintFromClass(SelectedRootActor->GetClass()))
 		{
+			Blueprint->OnCompiled().RemoveAll(this);
 			Blueprint->OnCompiled().AddRaw(this, &SDisplayClusterOperatorToolbar::OnBlueprintCompiled);
 		}
 	}
@@ -192,11 +219,17 @@ void SDisplayClusterOperatorToolbar::OnLevelActorDeleted(AActor* Actor)
 			return;
 		}
 		
-		ActiveRootActor.Reset();
-		ActiveRootActorName.Reset();
-		IDisplayClusterOperator::Get().OnActiveRootActorChanged().Broadcast(nullptr);
-		RootActorComboBox->SetSelectedItem(nullptr);
+		ClearSelectedRootActor();
+	}
+}
+
+void SDisplayClusterOperatorToolbar::HandleMapChanged(UWorld* InWorld, EMapChangeType InMapChangeType)
+{
+	if (InMapChangeType == EMapChangeType::TearDownWorld &&
+		(!ActiveRootActor.IsValid() || ActiveRootActor->GetWorld() == InWorld))
+	{
+		ClearSelectedRootActor();
 	}
 }
 
-#undef LOCTEXT_NAMESPACE
\ No newline at end of file
+#undef LOCTEXT_NAMESPACE
diff --git a/Engine/Plugins/Runtime/nDisplay/Source/DisplayClusterOperator/Private/SDisplayClusterOperatorToolbar.h b/Engine/Plugins/Runtime/nDisplay/Source/DisplayClusterOperator/Private/SDisplayClusterOperatorToolbar.h
--- a/Engine/Plugins/Runtime/nDisplay/Source/DisplayClusterOperator/Private/SDisplayClusterOperatorToolbar.h
+++ b/Engine/Plugins/Runtime/nDisplay/Source/DisplayClusterOperator/Private/SDisplayClusterOperatorToolbar.h
@@ -7,6 +7,7 @@
 #include "Widgets/SCompoundWidget.h"
 #include "UObject/WeakObjectPtr.h"
 
+enum class EMapChangeType : uint8;
 class FUICommandList;
 class ADisplayClusterRootActor;
 
@@ -38,6 +39,9 @@ private:
 	 */
 	TSharedPtr<FString> FillRootActorList(const FString& InitiallySelectedRootActor = TEXT(""));
 
+	/** Deselect and clear the current root actor */
+	void ClearSelectedRootActor();
+	
 	/** Raised when the user selects a new root actor from the root actor picker dropdown */
 	void OnRootActorChanged(TSharedPtr<FString> ItemSelected, ESelectInfo::Type SelectInfo);
 
@@ -56,6 +60,9 @@ private:
 	/** Raised when the user deletes an actor from the level */
 	void OnLevelActorDeleted(AActor* Actor);
 
+	/** When the map changes */
+	void HandleMapChanged(UWorld* InWorld, EMapChangeType InMapChangeType);
+
 private:
 	/** The command list used by the toolbar */
 	TSharedPtr<FUICommandList> CommandList;
@@ -74,4 +81,7 @@ private:
 
 	/** The OnLevelActorDeleted delegate handle */
 	FDelegateHandle LevelActorDeletedHandle;
+
+	/** The handle to OnMapChanged */
+	FDelegateHandle MapChangedHandle;
 };
\ No newline at end of file
