Fix non-unity nopch compile errors
#rnx

diff --git a/Engine/Plugins/Developer/DumpGPUServices/Source/DumpGPUServices/Private/DumpGPUHTTPUploadServices.cpp b/Engine/Plugins/Developer/DumpGPUServices/Source/DumpGPUServices/Private/DumpGPUHTTPUploadServices.cpp
--- a/Engine/Plugins/Developer/DumpGPUServices/Source/DumpGPUServices/Private/DumpGPUHTTPUploadServices.cpp
+++ b/Engine/Plugins/Developer/DumpGPUServices/Source/DumpGPUServices/Private/DumpGPUHTTPUploadServices.cpp
@@ -2,6 +2,7 @@
 
 #include "DumpGPU.h"
 #include "DumpGPUServices.h"
+#include "Misc/App.h"
 #include "Misc/Paths.h"
 #include "Misc/FileHelper.h"
 #include "Misc/ScopedSlowTask.h"
diff --git a/Engine/Plugins/Developer/DumpGPUServices/Source/DumpGPUServices/Private/DumpGPUServices.cpp b/Engine/Plugins/Developer/DumpGPUServices/Source/DumpGPUServices/Private/DumpGPUServices.cpp
--- a/Engine/Plugins/Developer/DumpGPUServices/Source/DumpGPUServices/Private/DumpGPUServices.cpp
+++ b/Engine/Plugins/Developer/DumpGPUServices/Source/DumpGPUServices/Private/DumpGPUServices.cpp
@@ -1,8 +1,10 @@
 // Copyright Epic Games, Inc. All Rights Reserved.
 
 #include "DumpGPUServices.h"
+#include "DumpGPU.h"
 #include "CoreMinimal.h"
 #include "CoreGlobals.h"
+#include "Misc/ConfigCacheIni.h"
 #include "Modules/ModuleManager.h"
 #include "IDumpGPUServices.h"
 
diff --git a/Engine/Plugins/Enterprise/DatasmithImporter/Source/DirectLinkExtension/Private/AutoReimportManager.cpp b/Engine/Plugins/Enterprise/DatasmithImporter/Source/DirectLinkExtension/Private/AutoReimportManager.cpp
--- a/Engine/Plugins/Enterprise/DatasmithImporter/Source/DirectLinkExtension/Private/AutoReimportManager.cpp
+++ b/Engine/Plugins/Enterprise/DatasmithImporter/Source/DirectLinkExtension/Private/AutoReimportManager.cpp
@@ -2,12 +2,14 @@
 
 #include "AutoReimportManager.h"
 
+#include "AssetRegistry/AssetData.h"
 #include "Async/Async.h"
 #include "ExternalSource.h"
 #include "SourceUri.h"
 #include "EditorReimportHandler.h"
 #include "ExternalSourceModule.h"
 #include "Stats/Stats.h"
+#include "Misc/App.h"
 
 DEFINE_LOG_CATEGORY(LogAutoReimportManager);
 
diff --git a/Engine/Plugins/Enterprise/LidarPointCloud/Source/LidarPointCloudRuntime/Public/LidarPointCloudComponent.h b/Engine/Plugins/Enterprise/LidarPointCloud/Source/LidarPointCloudRuntime/Public/LidarPointCloudComponent.h
--- a/Engine/Plugins/Enterprise/LidarPointCloud/Source/LidarPointCloudRuntime/Public/LidarPointCloudComponent.h
+++ b/Engine/Plugins/Enterprise/LidarPointCloud/Source/LidarPointCloudRuntime/Public/LidarPointCloudComponent.h
@@ -8,6 +8,7 @@
 #include "LidarPointCloudComponent.generated.h"
 
 class UBodySetup;
+class FViewportClient;
 
 UENUM(BlueprintType)
 enum class ELidarPointCloudSpriteOrientation : uint8
diff --git a/Engine/Plugins/Experimental/ColorCorrectRegions/Source/ColorCorrectRegions/Public/ColorCorrectRegion.h b/Engine/Plugins/Experimental/ColorCorrectRegions/Source/ColorCorrectRegions/Public/ColorCorrectRegion.h
--- a/Engine/Plugins/Experimental/ColorCorrectRegions/Source/ColorCorrectRegions/Public/ColorCorrectRegion.h
+++ b/Engine/Plugins/Experimental/ColorCorrectRegions/Source/ColorCorrectRegions/Public/ColorCorrectRegion.h
@@ -6,6 +6,7 @@
 #include "UObject/ObjectMacros.h"
 #include "GameFramework/Actor.h"
 #include "Engine/Classes/Components/MeshComponent.h"
+#include "Engine/Scene.h"
 #include "ColorCorrectRegion.generated.h"
 
 
diff --git a/Engine/Plugins/Experimental/Interchange/Runtime/Source/Import/Public/Mesh/InterchangeStaticMeshFactory.h b/Engine/Plugins/Experimental/Interchange/Runtime/Source/Import/Public/Mesh/InterchangeStaticMeshFactory.h
--- a/Engine/Plugins/Experimental/Interchange/Runtime/Source/Import/Public/Mesh/InterchangeStaticMeshFactory.h
+++ b/Engine/Plugins/Experimental/Interchange/Runtime/Source/Import/Public/Mesh/InterchangeStaticMeshFactory.h
@@ -11,6 +11,7 @@
 #include "InterchangeStaticMeshFactory.generated.h"
 
 
+class UBodySetup;
 class UStaticMesh;
 class UInterchangeStaticMeshLodDataNode;
 struct FMeshDescription;
diff --git a/Engine/Plugins/Media/HLMedia/Source/HLMedia/Private/HLMediaPlayer.cpp b/Engine/Plugins/Media/HLMedia/Source/HLMedia/Private/HLMediaPlayer.cpp
--- a/Engine/Plugins/Media/HLMedia/Source/HLMedia/Private/HLMediaPlayer.cpp
+++ b/Engine/Plugins/Media/HLMedia/Source/HLMedia/Private/HLMediaPlayer.cpp
@@ -9,6 +9,7 @@
 #include "MediaUtils/Public/MediaSamples.h"
 #include "GenericPlatform/GenericPlatformProcess.h"
 #include "Misc/MessageDialog.h"
+#include "RenderingThread.h"
 
 #include "HLMediaLibrary.h"
 #include "HLMediaPlayerTracks.h"
diff --git a/Engine/Plugins/Runtime/Steam/SteamVR/Source/SteamVRInputDevice/Private/SteamVRInputDevice.cpp b/Engine/Plugins/Runtime/Steam/SteamVR/Source/SteamVRInputDevice/Private/SteamVRInputDevice.cpp
--- a/Engine/Plugins/Runtime/Steam/SteamVR/Source/SteamVRInputDevice/Private/SteamVRInputDevice.cpp
+++ b/Engine/Plugins/Runtime/Steam/SteamVR/Source/SteamVRInputDevice/Private/SteamVRInputDevice.cpp
@@ -43,6 +43,7 @@ POSSIBILITY OF SUCH DAMAGE.
 #include "SteamVRControllerKeys.h"
 #include "Runtime/HeadMountedDisplay/Public/IXRTrackingSystem.h"
 #include "SteamVRSkeletonDefinition.h"
+#include "Misc/EngineVersion.h"
 #include "Misc/MessageDialog.h"
 
 #if PLATFORM_WINDOWS
diff --git a/Engine/Plugins/Runtime/nDisplay/Source/DisplayCluster/Private/Render/Synchronization/DisplayClusterRenderSyncPolicyFactoryInternal.cpp b/Engine/Plugins/Runtime/nDisplay/Source/DisplayCluster/Private/Render/Synchronization/DisplayClusterRenderSyncPolicyFactoryInternal.cpp
--- a/Engine/Plugins/Runtime/nDisplay/Source/DisplayCluster/Private/Render/Synchronization/DisplayClusterRenderSyncPolicyFactoryInternal.cpp
+++ b/Engine/Plugins/Runtime/nDisplay/Source/DisplayCluster/Private/Render/Synchronization/DisplayClusterRenderSyncPolicyFactoryInternal.cpp
@@ -10,6 +10,7 @@
 #include "Misc/DisplayClusterLog.h"
 #include "Misc/DisplayClusterStrings.h"
 #include "DisplayClusterConfigurationStrings.h"
+#include "RHIDefinitions.h"
 
 
 TSharedPtr<IDisplayClusterRenderSyncPolicy> FDisplayClusterRenderSyncPolicyFactoryInternal::Create(const FString& InPolicyType, const TMap<FString, FString>& Parameters)
diff --git a/Engine/Plugins/Tests/InterchangeTests/Source/InterchangeTestEditor/Private/InterchangeTestFunctionLayout.cpp b/Engine/Plugins/Tests/InterchangeTests/Source/InterchangeTestEditor/Private/InterchangeTestFunctionLayout.cpp
--- a/Engine/Plugins/Tests/InterchangeTests/Source/InterchangeTestEditor/Private/InterchangeTestFunctionLayout.cpp
+++ b/Engine/Plugins/Tests/InterchangeTests/Source/InterchangeTestEditor/Private/InterchangeTestFunctionLayout.cpp
@@ -8,6 +8,7 @@
 #include "DetailWidgetRow.h"
 #include "DetailLayoutBuilder.h"
 #include "Editor.h"
+#include "Widgets/Input/SEditableTextBox.h"
 #include "Widgets/Input/SNumericEntryBox.h"
 #include "Widgets/Input/SComboBox.h"
 #include "ScopedTransaction.h"
diff --git a/Engine/Plugins/Tests/InterchangeTests/Source/InterchangeTests/Private/InterchangeImportTest.cpp b/Engine/Plugins/Tests/InterchangeTests/Source/InterchangeTests/Private/InterchangeImportTest.cpp
--- a/Engine/Plugins/Tests/InterchangeTests/Source/InterchangeTests/Private/InterchangeImportTest.cpp
+++ b/Engine/Plugins/Tests/InterchangeTests/Source/InterchangeTests/Private/InterchangeImportTest.cpp
@@ -16,6 +16,7 @@
 #include "UObject/StrongObjectPtr.h"
 #include "AssetRegistryModule.h"
 #include "AssetData.h"
+#include "Editor.h"
 #include "InterchangeManager.h"
 #include "UObject/SavePackage.h"
 #include "ObjectTools.h"
diff --git a/Engine/Source/Runtime/Engine/Private/EngineGlobals.cpp b/Engine/Source/Runtime/Engine/Private/EngineGlobals.cpp
--- a/Engine/Source/Runtime/Engine/Private/EngineGlobals.cpp
+++ b/Engine/Source/Runtime/Engine/Private/EngineGlobals.cpp
@@ -18,6 +18,7 @@
 #include "Modules/ModuleManager.h"
 #include "Materials/Material.h"
 #include "Materials/MaterialInstanceConstant.h"
+#include "RendererInterface.h"
 
 #if !UE_BUILD_SHIPPING && PLATFORM_DESKTOP 
 #include "ISlateReflectorModule.h"
diff --git a/Engine/Source/Runtime/UMG/Private/Slate/SWorldWidgetScreenLayer.cpp b/Engine/Source/Runtime/UMG/Private/Slate/SWorldWidgetScreenLayer.cpp
--- a/Engine/Source/Runtime/UMG/Private/Slate/SWorldWidgetScreenLayer.cpp
+++ b/Engine/Source/Runtime/UMG/Private/Slate/SWorldWidgetScreenLayer.cpp
@@ -7,6 +7,7 @@
 #include "Blueprint/SlateBlueprintLibrary.h"
 #include "Engine/GameViewportClient.h"
 #include "Widgets/SViewport.h"
+#include "SceneView.h"
 #include "Slate/SGameLayerManager.h"
 
 static int32 GSlateWorldWidgetZOrder = 1;
diff --git a/Engine/Source/Runtime/Windows/D3D11RHI/Public/D3D11ShaderResources.h b/Engine/Source/Runtime/Windows/D3D11RHI/Public/D3D11ShaderResources.h
--- a/Engine/Source/Runtime/Windows/D3D11RHI/Public/D3D11ShaderResources.h
+++ b/Engine/Source/Runtime/Windows/D3D11RHI/Public/D3D11ShaderResources.h
@@ -6,6 +6,8 @@
 
 #pragma once
 
+#include "ShaderCore.h"
+
 struct FD3D11ShaderResourceTable : public FBaseShaderResourceTable
 {
 	/** Mapping of bound Textures to their location in resource tables. */
