Non-unity fixes
#rnx

diff --git a/Engine/Plugins/Developer/Concert/ConcertUI/ConcertServerUI/Source/ConcertServerUI/Private/Widgets/SessionTabs/Live/PackageViewer/SConcertSessionPackageViewer.h b/Engine/Plugins/Developer/Concert/ConcertUI/ConcertServerUI/Source/ConcertServerUI/Private/Widgets/SessionTabs/Live/PackageViewer/SConcertSessionPackageViewer.h
--- a/Engine/Plugins/Developer/Concert/ConcertUI/ConcertServerUI/Source/ConcertServerUI/Private/Widgets/SessionTabs/Live/PackageViewer/SConcertSessionPackageViewer.h
+++ b/Engine/Plugins/Developer/Concert/ConcertUI/ConcertServerUI/Source/ConcertServerUI/Private/Widgets/SessionTabs/Live/PackageViewer/SConcertSessionPackageViewer.h
@@ -5,6 +5,7 @@
 #include "CoreMinimal.h"
 #include "Session/Activity/SConcertSessionActivities.h"
 #include "Widgets/SCompoundWidget.h"
+#include "Widgets/Input/SSearchBox.h"
 #include "Widgets/DeclarativeSyntaxSupport.h"
 
 class FConcertSessionActivitiesOptions;
diff --git a/Engine/Plugins/Runtime/nDisplay/Source/DisplayCluster/Private/Render/Synchronization/DisplayClusterRenderSyncPolicyFactoryInternal.cpp b/Engine/Plugins/Runtime/nDisplay/Source/DisplayCluster/Private/Render/Synchronization/DisplayClusterRenderSyncPolicyFactoryInternal.cpp
--- a/Engine/Plugins/Runtime/nDisplay/Source/DisplayCluster/Private/Render/Synchronization/DisplayClusterRenderSyncPolicyFactoryInternal.cpp
+++ b/Engine/Plugins/Runtime/nDisplay/Source/DisplayCluster/Private/Render/Synchronization/DisplayClusterRenderSyncPolicyFactoryInternal.cpp
@@ -10,8 +10,7 @@
 #include "Misc/DisplayClusterLog.h"
 #include "Misc/DisplayClusterStrings.h"
 #include "DisplayClusterConfigurationStrings.h"
-#include "RHIDefinitions.h"
-
+#include "RHI.h"
 
 TSharedPtr<IDisplayClusterRenderSyncPolicy> FDisplayClusterRenderSyncPolicyFactoryInternal::Create(const FString& InPolicyType, const TMap<FString, FString>& Parameters)
 {
diff --git a/Engine/Source/Runtime/Engine/Private/TriggerActors.cpp b/Engine/Source/Runtime/Engine/Private/TriggerActors.cpp
--- a/Engine/Source/Runtime/Engine/Private/TriggerActors.cpp
+++ b/Engine/Source/Runtime/Engine/Private/TriggerActors.cpp
@@ -8,6 +8,7 @@
 #include "Engine/EngineTypes.h"
 #include "GameFramework/Actor.h"
 #include "Components/BillboardComponent.h"
+#include "Engine/Texture2D.h"
 #include "Engine/TriggerBox.h"
 #include "Engine/TriggerCapsule.h"
 #include "Engine/TriggerSphere.h"
diff --git a/Engine/Source/Runtime/Experimental/ChaosSolverEngine/Private/Chaos/ChaosSolverActor.cpp b/Engine/Source/Runtime/Experimental/ChaosSolverEngine/Private/Chaos/ChaosSolverActor.cpp
--- a/Engine/Source/Runtime/Experimental/ChaosSolverEngine/Private/Chaos/ChaosSolverActor.cpp
+++ b/Engine/Source/Runtime/Experimental/ChaosSolverEngine/Private/Chaos/ChaosSolverActor.cpp
@@ -10,6 +10,7 @@
 #include "ChaosSolversModule.h"
 #include "Chaos/ChaosGameplayEventDispatcher.h"
 #include "Chaos/Framework/DebugSubstep.h"
+#include "Engine/Texture2D.h"
 #include "UObject/FortniteMainBranchObjectVersion.h"
 #include "PhysicsProxy/SingleParticlePhysicsProxy.h"
 
