nDisplay: Add support for details panel section shortcuts.

diff --git a/Engine/Plugins/Runtime/nDisplay/Source/DisplayClusterConfigurator/Private/DisplayClusterConfiguratorModule.cpp b/Engine/Plugins/Runtime/nDisplay/Source/DisplayClusterConfigurator/Private/DisplayClusterConfiguratorModule.cpp
--- a/Engine/Plugins/Runtime/nDisplay/Source/DisplayClusterConfigurator/Private/DisplayClusterConfiguratorModule.cpp
+++ b/Engine/Plugins/Runtime/nDisplay/Source/DisplayClusterConfigurator/Private/DisplayClusterConfiguratorModule.cpp
@@ -89,6 +89,7 @@ void FDisplayClusterConfiguratorModule::StartupModule()
 	RegisterCustomLayouts();
 	RegisterSettings();
 	RegisterPlacementModeItems();
+	RegisterSectionMappings();
 	
 	FDisplayClusterConfiguratorStyle::Get();
 
@@ -126,6 +127,7 @@ void FDisplayClusterConfiguratorModule::ShutdownModule()
 	UnregisterSettings();
 	UnregisterCustomLayouts();
 	UnregisterPlacementModeItems();
+	UnregisterSectionMappings();
 
 	MenuExtensibilityManager.Reset();
 	ToolBarExtensibilityManager.Reset();
@@ -222,6 +224,75 @@ void FDisplayClusterConfiguratorModule::UnregisterCustomLayouts()
 	RegisteredPropertyLayoutNames.Empty();
 }
 
+void FDisplayClusterConfiguratorModule::RegisterSectionMappings()
+{
+	FPropertyEditorModule& PropertyModule = FModuleManager::LoadModuleChecked<FPropertyEditorModule>("PropertyEditor");
+	
+	// Root Actor
+	{
+		const TSharedRef<FPropertySection> Section = PropertyModule.FindOrCreateSection(ADisplayClusterRootActor::StaticClass()->GetFName(),
+			DisplayClusterConfigurationStrings::categories::ViewportsCategory, LOCTEXT("Viewports", "Viewports"));
+		Section->AddCategory(DisplayClusterConfigurationStrings::categories::ViewportsCategory);
+	}
+	{
+		const TSharedRef<FPropertySection> Section = PropertyModule.FindOrCreateSection(ADisplayClusterRootActor::StaticClass()->GetFName(),
+			DisplayClusterConfigurationStrings::categories::ICVFXCategory, LOCTEXT("In-Camera VFX", "In-Camera VFX"));
+		Section->AddCategory(DisplayClusterConfigurationStrings::categories::ICVFXCategory);
+	}
+	{
+		const TSharedRef<FPropertySection> Section = PropertyModule.FindOrCreateSection(ADisplayClusterRootActor::StaticClass()->GetFName(),
+			DisplayClusterConfigurationStrings::categories::ColorGradingCategory, LOCTEXT("Color Grading", "Color Grading"));
+		Section->AddCategory(DisplayClusterConfigurationStrings::categories::ColorGradingCategory);
+	}
+	{
+		const TSharedRef<FPropertySection> Section = PropertyModule.FindOrCreateSection(ADisplayClusterRootActor::StaticClass()->GetFName(),
+			DisplayClusterConfigurationStrings::categories::OCIOCategory, LOCTEXT("OCIO", "OCIO"));
+		Section->AddCategory(DisplayClusterConfigurationStrings::categories::OCIOCategory);
+	}
+	{
+		const TSharedRef<FPropertySection> Section = PropertyModule.FindOrCreateSection(ADisplayClusterRootActor::StaticClass()->GetFName(),
+			DisplayClusterConfigurationStrings::categories::LightcardCategory, LOCTEXT("Light Cards", "Light Cards"));
+		Section->AddCategory(DisplayClusterConfigurationStrings::categories::LightcardCategory);
+	}
+
+	// ICVFX Component
+	{
+		const TSharedRef<FPropertySection> Section = PropertyModule.FindOrCreateSection(UDisplayClusterICVFXCameraComponent::StaticClass()->GetFName(),
+		DisplayClusterConfigurationStrings::categories::ICVFXCategory, LOCTEXT("In-Camera VFX", "In-Camera VFX"));
+		Section->AddCategory(DisplayClusterConfigurationStrings::categories::ICVFXCategory);
+	}
+	{
+		const TSharedRef<FPropertySection> Section = PropertyModule.FindOrCreateSection(UDisplayClusterICVFXCameraComponent::StaticClass()->GetFName(),
+		DisplayClusterConfigurationStrings::categories::CameraColorGradingCategory, LOCTEXT("Inner Frustum Color Grading", "Inner Frustum Color Grading"));
+		Section->AddCategory(DisplayClusterConfigurationStrings::categories::CameraColorGradingCategory);
+	}
+	{
+		const TSharedRef<FPropertySection> Section = PropertyModule.FindOrCreateSection(UDisplayClusterICVFXCameraComponent::StaticClass()->GetFName(),
+			DisplayClusterConfigurationStrings::categories::OCIOCategory, LOCTEXT("OCIO", "OCIO"));
+		Section->AddCategory(DisplayClusterConfigurationStrings::categories::OCIOCategory);
+	}
+	{
+		const TSharedRef<FPropertySection> Section = PropertyModule.FindOrCreateSection(UDisplayClusterICVFXCameraComponent::StaticClass()->GetFName(),
+			DisplayClusterConfigurationStrings::categories::ChromaKeyCategory, LOCTEXT("Chromakey", "Chromakey"));
+		Section->AddCategory(DisplayClusterConfigurationStrings::categories::ChromaKeyCategory);
+	}
+}
+
+void FDisplayClusterConfiguratorModule::UnregisterSectionMappings()
+{
+	FPropertyEditorModule& PropertyModule = FModuleManager::LoadModuleChecked<FPropertyEditorModule>("PropertyEditor");
+	PropertyModule.RemoveSection(ADisplayClusterRootActor::StaticClass()->GetFName(), DisplayClusterConfigurationStrings::categories::ViewportsCategory);
+	PropertyModule.RemoveSection(ADisplayClusterRootActor::StaticClass()->GetFName(), DisplayClusterConfigurationStrings::categories::ICVFXCategory);
+	PropertyModule.RemoveSection(ADisplayClusterRootActor::StaticClass()->GetFName(), DisplayClusterConfigurationStrings::categories::ColorGradingCategory);
+	PropertyModule.RemoveSection(ADisplayClusterRootActor::StaticClass()->GetFName(), DisplayClusterConfigurationStrings::categories::OCIOCategory);
+	PropertyModule.RemoveSection(ADisplayClusterRootActor::StaticClass()->GetFName(), DisplayClusterConfigurationStrings::categories::LightcardCategory);
+
+	PropertyModule.RemoveSection(UDisplayClusterICVFXCameraComponent::StaticClass()->GetFName(), DisplayClusterConfigurationStrings::categories::ICVFXCategory);
+	PropertyModule.RemoveSection(UDisplayClusterICVFXCameraComponent::StaticClass()->GetFName(), DisplayClusterConfigurationStrings::categories::CameraColorGradingCategory);
+	PropertyModule.RemoveSection(UDisplayClusterICVFXCameraComponent::StaticClass()->GetFName(), DisplayClusterConfigurationStrings::categories::OCIOCategory);
+	PropertyModule.RemoveSection(UDisplayClusterICVFXCameraComponent::StaticClass()->GetFName(), DisplayClusterConfigurationStrings::categories::ChromaKeyCategory);
+}
+
 void FDisplayClusterConfiguratorModule::RegisterPlacementModeItems()
 {
 	auto RegisterPlaceActors = [&]() -> void
diff --git a/Engine/Plugins/Runtime/nDisplay/Source/DisplayClusterConfigurator/Private/DisplayClusterConfiguratorModule.h b/Engine/Plugins/Runtime/nDisplay/Source/DisplayClusterConfigurator/Private/DisplayClusterConfiguratorModule.h
--- a/Engine/Plugins/Runtime/nDisplay/Source/DisplayClusterConfigurator/Private/DisplayClusterConfiguratorModule.h
+++ b/Engine/Plugins/Runtime/nDisplay/Source/DisplayClusterConfigurator/Private/DisplayClusterConfiguratorModule.h
@@ -37,6 +37,8 @@ private:
 	void UnregisterSettings();
 	void RegisterCustomLayouts();
 	void UnregisterCustomLayouts();
+	void RegisterSectionMappings();
+	void UnregisterSectionMappings();
 
 	/** Register items to show up in the Place Actors panel. */
 	void RegisterPlacementModeItems();
