nDisplay: Add Editor Preview category section.

diff --git a/Engine/Plugins/Runtime/nDisplay/Source/DisplayClusterConfigurator/Private/DisplayClusterConfiguratorModule.cpp b/Engine/Plugins/Runtime/nDisplay/Source/DisplayClusterConfigurator/Private/DisplayClusterConfiguratorModule.cpp
--- a/Engine/Plugins/Runtime/nDisplay/Source/DisplayClusterConfigurator/Private/DisplayClusterConfiguratorModule.cpp
+++ b/Engine/Plugins/Runtime/nDisplay/Source/DisplayClusterConfigurator/Private/DisplayClusterConfiguratorModule.cpp
@@ -254,6 +254,11 @@ void FDisplayClusterConfiguratorModule::RegisterSectionMappings()
 			DisplayClusterConfigurationStrings::categories::LightcardCategory, LOCTEXT("Light Cards", "Light Cards"));
 		Section->AddCategory(DisplayClusterConfigurationStrings::categories::LightcardCategory);
 	}
+	{
+		const TSharedRef<FPropertySection> Section = PropertyModule.FindOrCreateSection(ADisplayClusterRootActor::StaticClass()->GetFName(),
+			DisplayClusterConfigurationStrings::categories::PreviewCategory, LOCTEXT("Editor Preview", "Editor Preview"));
+		Section->AddCategory(DisplayClusterConfigurationStrings::categories::PreviewCategory);
+	}
 
 	// ICVFX Component
 	{
