nDisplay: disabled shadows rendering for warp meshes

diff --git a/Engine/Plugins/Runtime/nDisplay/Source/DisplayCluster/Private/Game/EngineClasses/Scene/DisplayClusterPreviewComponent.cpp b/Engine/Plugins/Runtime/nDisplay/Source/DisplayCluster/Private/Game/EngineClasses/Scene/DisplayClusterPreviewComponent.cpp
--- a/Engine/Plugins/Runtime/nDisplay/Source/DisplayCluster/Private/Game/EngineClasses/Scene/DisplayClusterPreviewComponent.cpp
+++ b/Engine/Plugins/Runtime/nDisplay/Source/DisplayCluster/Private/Game/EngineClasses/Scene/DisplayClusterPreviewComponent.cpp
@@ -192,6 +192,12 @@ bool UDisplayClusterPreviewComponent::UpdatePreviewMesh()
 					WarpMeshSavedProjectionPolicy = ViewportConfig->ProjectionPolicy;
 				}
 
+				// disable shadow rendering for preview meshes
+				if (PreviewMesh != nullptr)
+				{
+					PreviewMesh->SetCastShadow(false);
+				}
+
 				if (OriginalMaterial == nullptr)
 				{
 					// Assign preview material to mesh
