[Epic Stage App] Create and listen for changes to C++ lightcard class

diff --git a/Engine/Plugins/Runtime/nDisplay/Source/DisplayCluster/Private/Blueprints/DisplayClusterBlueprintLib.cpp b/Engine/Plugins/Runtime/nDisplay/Source/DisplayCluster/Private/Blueprints/DisplayClusterBlueprintLib.cpp
--- a/Engine/Plugins/Runtime/nDisplay/Source/DisplayCluster/Private/Blueprints/DisplayClusterBlueprintLib.cpp
+++ b/Engine/Plugins/Runtime/nDisplay/Source/DisplayCluster/Private/Blueprints/DisplayClusterBlueprintLib.cpp
@@ -4,6 +4,15 @@
 #include "Blueprints/DisplayClusterBlueprintAPIImpl.h"
 #include "UObject/Package.h"
 
+#include "DisplayClusterLightCardActor.h"
+#include "DisplayClusterRootActor.h"
+#include "DisplayClusterConfigurationTypes.h"
+
+#if WITH_EDITOR
+#include "ScopedTransaction.h"
+#endif
+
+#define LOCTEXT_NAMESPACE "UDisplayClusterBlueprintLib"
 
 UDisplayClusterBlueprintLib::UDisplayClusterBlueprintLib(class FObjectInitializer const & ObjectInitializer)
 	: Super(ObjectInitializer)
@@ -15,3 +24,49 @@ void UDisplayClusterBlueprintLib::GetAPI(TScriptInterface<IDisplayClusterBluepri
 	static UDisplayClusterBlueprintAPIImpl* Obj = NewObject<UDisplayClusterBlueprintAPIImpl>(GetTransientPackage(), NAME_None, RF_MarkAsRootSet);
 	OutAPI = Obj;
 }
+
+ADisplayClusterLightCardActor* UDisplayClusterBlueprintLib::CreateLightCard(ADisplayClusterRootActor* RootActor)
+{
+	if (!RootActor)
+	{
+		return nullptr;
+	}
+
+	// Create the light card
+#if WITH_EDITOR
+	FScopedTransaction Transaction(LOCTEXT("CreateLightCard", "Create Light Card"));
+#endif
+
+	const FVector SpawnLocation = RootActor->GetDefaultCamera()->GetComponentLocation();
+	FRotator SpawnRotation = RootActor->GetDefaultCamera()->GetComponentRotation();
+	SpawnRotation.Yaw -= 180.f;
+
+	FActorSpawnParameters SpawnParameters;
+	SpawnParameters.bNoFail = true;
+	SpawnParameters.SpawnCollisionHandlingOverride = ESpawnActorCollisionHandlingMethod::AdjustIfPossibleButAlwaysSpawn;
+	SpawnParameters.Name = TEXT("LightCard");
+	SpawnParameters.NameMode = FActorSpawnParameters::ESpawnActorNameMode::Requested;
+	SpawnParameters.OverrideLevel = RootActor->GetLevel();
+
+	ADisplayClusterLightCardActor* NewActor = CastChecked<ADisplayClusterLightCardActor>(
+		RootActor->GetWorld()->SpawnActor(ADisplayClusterLightCardActor::StaticClass(),
+			&SpawnLocation, &SpawnRotation, MoveTemp(SpawnParameters)));
+
+#if WITH_EDITOR
+	NewActor->SetActorLabel(NewActor->GetName());
+#endif
+
+	FAttachmentTransformRules AttachmentRules(EAttachmentRule::KeepWorld, false);
+	NewActor->AttachToActor(RootActor, AttachmentRules);
+
+	// Add it to the root actor
+	UDisplayClusterConfigurationData* ConfigData = RootActor->GetConfigData();
+	ConfigData->Modify();
+	FDisplayClusterConfigurationICVFX_VisibilityList& RootActorLightCards = ConfigData->StageSettings.Lightcard.ShowOnlyList;
+
+	RootActorLightCards.Actors.Add(NewActor);
+
+	return NewActor;
+}
+
+#undef LOCTEXT_NAMESPACE
\ No newline at end of file
diff --git a/Engine/Plugins/Runtime/nDisplay/Source/DisplayCluster/Public/Blueprints/DisplayClusterBlueprintLib.h b/Engine/Plugins/Runtime/nDisplay/Source/DisplayCluster/Public/Blueprints/DisplayClusterBlueprintLib.h
--- a/Engine/Plugins/Runtime/nDisplay/Source/DisplayCluster/Public/Blueprints/DisplayClusterBlueprintLib.h
+++ b/Engine/Plugins/Runtime/nDisplay/Source/DisplayCluster/Public/Blueprints/DisplayClusterBlueprintLib.h
@@ -7,6 +7,8 @@
 #include "Blueprints/IDisplayClusterBlueprintAPI.h"
 #include "DisplayClusterBlueprintLib.generated.h"
 
+class ADisplayClusterLightCardActor;
+class ADisplayClusterRootActor;
 
 /**
  * Blueprint API function library
@@ -22,5 +24,7 @@ public:
 	UFUNCTION(BlueprintPure, meta = (DisplayName = "DisplayCluster Module API"), Category = "NDisplay")
 	static void GetAPI(TScriptInterface<IDisplayClusterBlueprintAPI>& OutAPI);
 
-
+	/** Create a new light card parented to the given nDisplay root actor. */
+	UFUNCTION(BlueprintPure, meta = (DisplayName = "DisplayCluster Module API"), Category = "NDisplay")
+	static ADisplayClusterLightCardActor* CreateLightCard(ADisplayClusterRootActor* RootActor);
 };
