nDisplay: MeshProjectionShader now applies a gamma of 2.2. This works well when using in-editor preview with post process enabled.

Correction when post process is disabled is still pending.

diff --git a/Engine/Plugins/Runtime/nDisplay/Shaders/Private/MeshProjectionShaders.usf b/Engine/Plugins/Runtime/nDisplay/Shaders/Private/MeshProjectionShaders.usf
--- a/Engine/Plugins/Runtime/nDisplay/Shaders/Private/MeshProjectionShaders.usf
+++ b/Engine/Plugins/Runtime/nDisplay/Shaders/Private/MeshProjectionShaders.usf
@@ -68,4 +68,7 @@ void MainPS(FMeshProjectionVSToPS Input OPTIONAL_IsFrontFace, out float4 OutputC
 	float Opacity = GetMaterialOpacity(PixelMaterialInputs);
 	
 	OutputColor = float4(Emissive, Opacity);
+
+	float gamma = 2.2;
+	OutputColor.rgb = pow(OutputColor.rgb, gamma);
 }
\ No newline at end of file
