nDisplay: Added ability to move light card proxies within the light card editor using a axis widget that allows users to move freely, or fixed along longitude or latitude. The level instance light cards are subsequently updated to match the transformations applied to the proxies.

diff --git a/Engine/Plugins/Runtime/nDisplay/Source/DisplayCluster/Private/Game/EngineClasses/Scene/DisplayClusterLightCardActor.cpp b/Engine/Plugins/Runtime/nDisplay/Source/DisplayCluster/Private/Game/EngineClasses/Scene/DisplayClusterLightCardActor.cpp
--- a/Engine/Plugins/Runtime/nDisplay/Source/DisplayCluster/Private/Game/EngineClasses/Scene/DisplayClusterLightCardActor.cpp
+++ b/Engine/Plugins/Runtime/nDisplay/Source/DisplayCluster/Private/Game/EngineClasses/Scene/DisplayClusterLightCardActor.cpp
@@ -66,6 +66,30 @@ void ADisplayClusterLightCardActor::OnConstruction(const FTransform& Transform)
 void ADisplayClusterLightCardActor::Tick(float DeltaSeconds)
 {
 	Super::Tick(DeltaSeconds);
+
+	if (Longitude < 0 || Longitude > 360)
+	{
+		Longitude = FRotator::ClampAxis(Longitude);
+	}
+
+	if (Latitude < -90 || Latitude > 90)
+	{
+		// If latitude exceeds [-90, 90], mod it back into the appropriate range, and apply a shift of 180 degrees if
+		// needed to the longitude, to allow the latitude to be continuous (increasing latitude indefinitely should result in the LC 
+		// orbiting around a polar great circle)
+		float Parity = FMath::Fmod(FMath::Abs(Latitude) + 90, 360) - 180;
+		float DeltaLongitude = Parity > 1 ? 180.f : 0.f;
+
+		float LatMod = FMath::Fmod(Latitude + 90.f, 180.f);
+		if (LatMod < 0.f)
+		{
+			LatMod += 180.f;
+		}
+
+		Latitude = LatMod - 90;
+		Longitude = FRotator::ClampAxis(Longitude + DeltaLongitude);
+	}
+
 	UpdateLightCardTransform();
 	UpdateLightCardMaterialInstance();
 }
@@ -100,10 +124,28 @@ void ADisplayClusterLightCardActor::PostEditChangeProperty(FPropertyChangedEvent
 
 #endif
 
+FTransform ADisplayClusterLightCardActor::GetLightCardTransform(bool bIgnoreSpinYawPitch) const
+{
+	FTransform Transform;
+
+	Transform.SetLocation(LightCardComponent->GetComponentLocation());
+
+	FRotator LightCardRotation = LightCardTransformerComponent->GetComponentRotation();
+
+	if (bIgnoreSpinYawPitch)
+	{
+		LightCardRotation -= FRotator(Spin, -Yaw, -Pitch);
+	}
+
+	Transform.SetRotation(LightCardRotation.Quaternion());
+
+	return Transform;
+}
+
 void ADisplayClusterLightCardActor::UpdateLightCardTransform()
 {
 	MainSpringArmComponent->TargetArmLength = DistanceFromCenter;
-	MainSpringArmComponent->SetRelativeRotation(FRotator(-Latitude, Longitude * 180.f, 0.0));
+	MainSpringArmComponent->SetRelativeRotation(FRotator(-Latitude, Longitude, 0.0));
 
 	LightCardComponent->SetRelativeRotation(FRotator(-Spin, -90.f + Yaw, 90.f + Pitch));
 	LightCardComponent->SetRelativeScale3D(FVector(Scale, 1.f));
diff --git a/Engine/Plugins/Runtime/nDisplay/Source/DisplayCluster/Public/DisplayClusterLightCardActor.h b/Engine/Plugins/Runtime/nDisplay/Source/DisplayCluster/Public/DisplayClusterLightCardActor.h
--- a/Engine/Plugins/Runtime/nDisplay/Source/DisplayCluster/Public/DisplayClusterLightCardActor.h
+++ b/Engine/Plugins/Runtime/nDisplay/Source/DisplayCluster/Public/DisplayClusterLightCardActor.h
@@ -34,7 +34,13 @@ public:
 	virtual void PostEditChangeProperty(FPropertyChangedEvent& PropertyChangedEvent) override;
 	virtual bool ShouldTickIfViewportsOnly() const override { return true; }
 #endif
-
+	
+	/**
+	 * Gets the transform in world space of the light card component
+	 * @param bIgnoreSpinYawPitch - If the light card component's spin, yaw, and pitch should be ignored when computing the transform
+	 */
+	FTransform GetLightCardTransform(bool bIgnoreSpinYawPitch = false) const;
+	
 protected:
 	void UpdateLightCardTransform();
 	void UpdateLightCardMaterialInstance();
@@ -43,10 +49,10 @@ public:
 	UPROPERTY(EditAnywhere, BlueprintReadWrite, Category = "Orientation")
 	float DistanceFromCenter;
 
-	UPROPERTY(EditAnywhere, BlueprintReadWrite, Category = "Orientation")
+	UPROPERTY(EditAnywhere, BlueprintReadWrite, Category = "Orientation", meta = (UIMin = 0, ClampMin = 0, UIMax = 360, ClampMax = 360))
 	float Longitude;
 
-	UPROPERTY(EditAnywhere, BlueprintReadWrite, Category = "Orientation")
+	UPROPERTY(EditAnywhere, BlueprintReadWrite, Category = "Orientation", meta = (UIMin = -90, ClampMin = -90, UIMax = 90, ClampMax = 90))
 	float Latitude;
 
 	UPROPERTY(EditAnywhere, BlueprintReadWrite, Category = "Orientation")
diff --git a/Engine/Plugins/Runtime/nDisplay/Source/DisplayClusterLightCardEditor/Private/Viewport/DisplayClusterLightCardEditorViewportClient.cpp b/Engine/Plugins/Runtime/nDisplay/Source/DisplayClusterLightCardEditor/Private/Viewport/DisplayClusterLightCardEditorViewportClient.cpp
--- a/Engine/Plugins/Runtime/nDisplay/Source/DisplayClusterLightCardEditor/Private/Viewport/DisplayClusterLightCardEditorViewportClient.cpp
+++ b/Engine/Plugins/Runtime/nDisplay/Source/DisplayClusterLightCardEditor/Private/Viewport/DisplayClusterLightCardEditorViewportClient.cpp
@@ -2,9 +2,11 @@
 
 #include "DisplayClusterLightCardEditorViewportClient.h"
 
+#include "DisplayClusterLightCardEditorWidget.h"
 #include "DisplayClusterConfigurationTypes.h"
 #include "DisplayClusterProjectionStrings.h"
 #include "DisplayClusterRootActor.h"
+#include "DisplayClusterLightCardActor.h"
 #include "Components/DisplayClusterPreviewComponent.h"
 #include "Components/DisplayClusterCameraComponent.h"
 #include "Components/DisplayClusterScreenComponent.h"
@@ -52,19 +54,19 @@ FDisplayClusterLightCardEditorViewportClient::FDisplayClusterLightCardEditorView
 	
 	MeshProjectionRenderer = MakeShared<FDisplayClusterMeshProjectionRenderer>();
 	MeshProjectionRenderer->ActorSelectedDelegate = FDisplayClusterMeshProjectionRenderer::FSelection::CreateRaw(this, &FDisplayClusterLightCardEditorViewportClient::IsLightCardSelected);
+	MeshProjectionRenderer->RenderSimpleElementsDelegate = FDisplayClusterMeshProjectionRenderer::FSimpleElementPass::CreateRaw(this, &FDisplayClusterLightCardEditorViewportClient::Draw);
 
 	bDraggingActor = false;
+	DragWidgetOffset = FVector::ZeroVector;
 	ScopedTransaction = nullptr;
 	
+	EditorWidget = MakeShared<FDisplayClusterLightCardEditorWidget>();
+
 	// Setup defaults for the common draw helper.
-	DrawHelper.bDrawPivot = false;
-	DrawHelper.bDrawWorldBox = false;
-	DrawHelper.bDrawKillZ = false;
-	DrawHelper.bDrawGrid = false;
-	DrawHelper.PerspectiveGridSize = HALF_WORLD_MAX1;
+	bUsesDrawHelper = false;
 
 	EngineShowFlags.SetSelectionOutline(true);
-
+	 
 	check(Widget);
 	Widget->SetSnapEnabled(true);
 	
@@ -117,6 +119,8 @@ void FDisplayClusterLightCardEditorViewportClient::Tick(float DeltaSeconds)
 		SetViewRotation(FVector::UpVector.Rotation());
 	}
 
+	CachedEditorWidgetTransform = CalcEditorWidgetTransform();
+
 	// Tick the preview scene world.
 	if (!GIntraFrameDebuggingGameThread)
 	{
@@ -151,6 +155,16 @@ void FDisplayClusterLightCardEditorViewportClient::Tick(float DeltaSeconds)
 			}
 		}
 	}
+
+	// EditorViewportClient sets the cursor settings based on the state of the built in FWidget, which isn't being used here, so
+	// force a software cursor if we are dragging an actor so that the correct mouse cursor shows up
+	if (bDraggingActor)
+	{
+		SetRequiredCursor(false, true);
+		SetRequiredCursorOverride(true, EMouseCursor::CardinalCross);
+
+		ApplyRequiredCursorVisibility(true);
+	}
 }
 
 void FDisplayClusterLightCardEditorViewportClient::Draw(FViewport* InViewport, FCanvas* Canvas)
@@ -311,14 +325,13 @@ void FDisplayClusterLightCardEditorViewportClient::Draw(FViewport* InViewport, F
 	Viewport = ViewportBackup;
 }
 
-UE::Widget::EWidgetMode FDisplayClusterLightCardEditorViewportClient::GetWidgetMode() const
-{
-	return UE::Widget::WM_None;
-}
-
-FVector FDisplayClusterLightCardEditorViewportClient::GetWidgetLocation() const
+void FDisplayClusterLightCardEditorViewportClient::Draw(const FSceneView* View, FPrimitiveDrawInterface* PDI)
 {
-	return FEditorViewportClient::GetWidgetLocation();
+	if (SelectedLightCards.Num())
+	{
+		EditorWidget->SetTransform(CachedEditorWidgetTransform);
+		EditorWidget->Draw(View, PDI);
+	}
 }
 
 FSceneView* FDisplayClusterLightCardEditorViewportClient::CalcSceneView(FSceneViewFamily* ViewFamily, const int32 StereoViewIndex)
@@ -367,12 +380,27 @@ bool FDisplayClusterLightCardEditorViewportClient::InputKey(FViewport* InViewpor
 	return FEditorViewportClient::InputKey(InViewport, ControllerId, Key, Event, AmountDepressed, bGamepad);
 }
 
-bool FDisplayClusterLightCardEditorViewportClient::InputWidgetDelta(FViewport* InViewport, EAxisList::Type CurrentAxis, FVector& Drag,
-                                                    FRotator& Rot, FVector& Scale)
+bool FDisplayClusterLightCardEditorViewportClient::InputWidgetDelta(FViewport* InViewport,
+	EAxisList::Type CurrentAxis,
+	FVector& Drag,
+	FRotator& Rot,
+	FVector& Scale)
 {
-	// TODO: Add logic for moving a light card in "2D" space, modifying its long and lat, and making sure it is flush with the nearest screen
+	bool bHandled = false;
+	if (FEditorViewportClient::InputWidgetDelta(InViewport, CurrentAxis, Drag, Rot, Scale))
+	{
+		bHandled = true;
+	}
+	else
+	{
+		if (CurrentAxis != EAxisList::Type::None && SelectedLightCards.Num())
+		{
+			MoveSelectedLightCards(InViewport, CurrentAxis);
+			bHandled = true;
+		}
+	}
 
-	return FEditorViewportClient::InputWidgetDelta(InViewport, CurrentAxis, Drag, Rot, Scale);
+	return bHandled;
 }
 
 void FDisplayClusterLightCardEditorViewportClient::TrackingStarted(const FInputEventState& InInputState, bool bIsDraggingWidget,
@@ -388,15 +416,35 @@ void FDisplayClusterLightCardEditorViewportClient::TrackingStarted(const FInputE
 			GUnrealEd->SetPivotMovedIndependently(false);
 		}
 
-		BeginTransaction(NSLOCTEXT("LogicDriverPreview", "ModifyPreviewActor", "Modify a Preview Actor"));
+		BeginTransaction(LOCTEXT("MoveLightCard", "Move Light Card"));
 		bDraggingActor = true;
+
+		// Compute and store the delta between the widget's origin and the place the user clicked on it,
+		// in order to factor it out when transforming the selected actor
+		FIntPoint MousePos;
+		InInputState.GetViewport()->GetMousePos(MousePos);
+
+		FSceneViewFamilyContext ViewFamily(FSceneViewFamily::ConstructionValues(
+			InInputState.GetViewport(),
+			GetScene(),
+			EngineShowFlags)
+			.SetRealtimeUpdate(IsRealtime()));
+		FSceneView* View = CalcSceneView(&ViewFamily);
+
+		FVector Origin;
+		FVector Direction;
+		PixelToWorld(*View, MousePos, Origin, Direction);
+
+		DragWidgetOffset = Direction - (CachedEditorWidgetTransform.GetTranslation() - Origin).GetSafeNormal();
 	}
+
 	FEditorViewportClient::TrackingStarted(InInputState, bIsDraggingWidget, bNudge);
 }
 
 void FDisplayClusterLightCardEditorViewportClient::TrackingStopped()
 {
 	bDraggingActor = false;
+	DragWidgetOffset = FVector::ZeroVector;
 	EndTransaction();
 
 	if (SelectedLightCards.Num())
@@ -425,13 +473,13 @@ void FDisplayClusterLightCardEditorViewportClient::ProcessClick(FSceneView& View
 			{
 				if (ActorHitProxy->PrimComponent && ActorHitProxy->PrimComponent->IsA<UStaticMeshComponent>())
 				{
-					AActor* TracedLightCard = TraceScreenForLightCard(View, HitX, HitY);
+					ADisplayClusterLightCardActor* TracedLightCard = TraceScreenForLightCard(View, HitX, HitY);
 					SelectLightCard(TracedLightCard, bMultiSelect);
 				}
 			}
-			else if (LightCardProxies.Contains(ActorHitProxy->Actor))
+			else if (ActorHitProxy->Actor->IsA<ADisplayClusterLightCardActor>() && LightCardProxies.Contains(ActorHitProxy->Actor))
 			{
-				SelectLightCard(ActorHitProxy->Actor, bMultiSelect);
+				SelectLightCard(Cast<ADisplayClusterLightCardActor>(ActorHitProxy->Actor), bMultiSelect);
 			}
 			else if (!bMultiSelect)
 			{
@@ -462,8 +510,10 @@ EMouseCursor::Type FDisplayClusterLightCardEditorViewportClient::GetCursor(FView
 	{
 		MouseCursor = EMouseCursor::None;
 	}
-	else if( InViewport->IsCursorVisible() && !bWidgetAxisControlledByDrag )
+	else if (InViewport->IsCursorVisible() && !bWidgetAxisControlledByDrag)
 	{
+		EditorWidget->SetHighlightedAxis(EAxisList::Type::None);
+
 		HHitProxy* HitProxy = InViewport->GetHitProxy(X,Y);
 		if (HitProxy)
 		{
@@ -481,9 +531,9 @@ EMouseCursor::Type FDisplayClusterLightCardEditorViewportClient::GetCursor(FView
 							GetScene(),
 							EngineShowFlags)
 							.SetRealtimeUpdate(IsRealtime()));
-						FSceneView* View = CalcSceneView( &ViewFamily );
+						FSceneView* View = CalcSceneView(&ViewFamily);
 
-						AActor* TracedLightCard = TraceScreenForLightCard(*View, X, Y);
+						ADisplayClusterLightCardActor* TracedLightCard = TraceScreenForLightCard(*View, X, Y);
 						if (TracedLightCard)
 						{
 							MouseCursor = EMouseCursor::Crosshairs;
@@ -495,25 +545,22 @@ EMouseCursor::Type FDisplayClusterLightCardEditorViewportClient::GetCursor(FView
 					MouseCursor = EMouseCursor::Crosshairs;
 				}
 			}
+			else if (HitProxy->IsA(HWidgetAxis::StaticGetType()))
+			{
+				HWidgetAxis* AxisHitProxy = static_cast<HWidgetAxis*>(HitProxy);
+				
+				MouseCursor = AxisHitProxy->GetMouseCursor();
+				EditorWidget->SetHighlightedAxis(AxisHitProxy->Axis);
+			}
 		}
 	}
 
 	CachedMouseX = X;
-	CachedLastMouseY = Y;
+	CachedMouseY = Y;
 
 	return MouseCursor;
 }
 
-void FDisplayClusterLightCardEditorViewportClient::SelectActor(AActor* NewActor)
-{
-}
-
-void FDisplayClusterLightCardEditorViewportClient::ResetSelection()
-{
-	SelectActor(nullptr);
-	SetWidgetMode(UE::Widget::WM_None);
-}
-
 void FDisplayClusterLightCardEditorViewportClient::UpdatePreviewActor(ADisplayClusterRootActor* RootActor, bool bForce,
                                                                       EDisplayClusterLightCardEditorProxyType ProxyType)
 {
@@ -576,16 +623,16 @@ void FDisplayClusterLightCardEditorViewportClient::UpdatePreviewActor(ADisplayCl
 			if (ProxyType == EDisplayClusterLightCardEditorProxyType::All ||
 				ProxyType == EDisplayClusterLightCardEditorProxyType::LightCards)
 			{
-				TArray<TWeakObjectPtr<AActor>> LightCards;
+				TArray<TWeakObjectPtr<ADisplayClusterLightCardActor>> LightCards;
 				FindLightCardsForRootActor(RootActor, LightCards);
 			
-				for (const TWeakObjectPtr<AActor>& LightCard : LightCards)
+				for (const TWeakObjectPtr<ADisplayClusterLightCardActor>& LightCard : LightCards)
 				{
 					FObjectDuplicationParameters DupeActorParameters(LightCard.Get(), PreviewWorld->GetCurrentLevel());
 					DupeActorParameters.FlagMask = RF_AllFlags & ~(RF_ArchetypeObject | RF_Transactional);
 					DupeActorParameters.PortFlags = PPF_DuplicateVerbatim;
 				
-					AActor* LightCardProxy = CastChecked<AActor>(StaticDuplicateObjectEx(DupeActorParameters));
+					ADisplayClusterLightCardActor* LightCardProxy = CastChecked<ADisplayClusterLightCardActor>(StaticDuplicateObjectEx(DupeActorParameters));
 					PreviewWorld->GetCurrentLevel()->AddLoadedActor(LightCardProxy);
 				
 					LightCardProxy->SetActorLocation(LightCard->GetActorLocation() - RootActor->GetActorLocation());
@@ -685,10 +732,12 @@ void FDisplayClusterLightCardEditorViewportClient::SetProjectionMode(EDisplayClu
 	{
 		// TODO: Do we want to cache the perspective rotation and restore it when the user switches back?
 		SetViewRotation(FVector::ForwardVector.Rotation());
+		EditorWidget->SetWidgetScale(1.f);
 	}
 	else if (ProjectionMode == EDisplayClusterMeshProjectionType::Azimuthal)
 	{
 		SetViewRotation(FVector::UpVector.Rotation());
+		EditorWidget->SetWidgetScale(0.5f);
 	}
 
 	FindProjectionOriginComponent();
@@ -914,7 +963,7 @@ void FDisplayClusterLightCardEditorViewportClient::FindProjectionOriginComponent
 	}
 }
 
-void FDisplayClusterLightCardEditorViewportClient::FindLightCardsForRootActor(ADisplayClusterRootActor* RootActor, TArray<TWeakObjectPtr<AActor>>& OutLightCards)
+void FDisplayClusterLightCardEditorViewportClient::FindLightCardsForRootActor(ADisplayClusterRootActor* RootActor, TArray<TWeakObjectPtr<ADisplayClusterLightCardActor>>& OutLightCards)
 {
 	if (RootActor)
 	{
@@ -922,9 +971,9 @@ void FDisplayClusterLightCardEditorViewportClient::FindLightCardsForRootActor(AD
 
 		for (const TSoftObjectPtr<AActor>& LightCardActor : RootActorLightCards.Actors)
 		{
-			if (LightCardActor.IsValid())
+			if (LightCardActor.IsValid() && LightCardActor->IsA<ADisplayClusterLightCardActor>())
 			{
-				OutLightCards.Add(LightCardActor.Get());
+				OutLightCards.Add(Cast<ADisplayClusterLightCardActor>(LightCardActor.Get()));
 			}
 		}
 
@@ -937,13 +986,13 @@ void FDisplayClusterLightCardEditorViewportClient::FindLightCardsForRootActor(AD
 			{
 				for (const TWeakObjectPtr<AActor> WeakActor : FActorRange(World))
 				{
-					if (WeakActor.IsValid())
+					if (WeakActor.IsValid() && WeakActor->IsA<ADisplayClusterLightCardActor>())
 					{
 						for (const FActorLayer& ActorLayer : RootActorLightCards.ActorLayers)
 						{
 							if (WeakActor->Layers.Contains(ActorLayer.Name))
 							{
-								OutLightCards.Add(WeakActor);
+								OutLightCards.Add(Cast<ADisplayClusterLightCardActor>(WeakActor));
 								break;
 							}
 						}
@@ -959,13 +1008,13 @@ bool FDisplayClusterLightCardEditorViewportClient::IsLightCardSelected(const AAc
 	return SelectedLightCards.Contains(Actor);
 }
 
-void FDisplayClusterLightCardEditorViewportClient::SelectLightCard(AActor* Actor, bool bAddToSelection)
+void FDisplayClusterLightCardEditorViewportClient::SelectLightCard(ADisplayClusterLightCardActor* Actor, bool bAddToSelection)
 {
-	TArray<AActor*> UpdatedActors;
+	TArray<ADisplayClusterLightCardActor*> UpdatedActors;
 
 	if (!bAddToSelection)
 	{
-		for (const TWeakObjectPtr<AActor>& LightCard : SelectedLightCards)
+		for (const TWeakObjectPtr<ADisplayClusterLightCardActor>& LightCard : SelectedLightCards)
 		{
 			if (LightCard.IsValid())
 			{
@@ -991,7 +1040,7 @@ void FDisplayClusterLightCardEditorViewportClient::SelectLightCard(AActor* Actor
 void FDisplayClusterLightCardEditorViewportClient::PropagateLightCardSelection()
 {
 	TArray<AActor*> SelectedLevelInstances;
-	for (const TWeakObjectPtr<AActor>& SelectedLightCard : SelectedLightCards)
+	for (const TWeakObjectPtr<ADisplayClusterLightCardActor>& SelectedLightCard : SelectedLightCards)
 	{
 		if (FLightCardProxy* FoundProxy = LightCardProxies.FindByKey(SelectedLightCard.Get()))
 		{
@@ -1005,30 +1054,184 @@ void FDisplayClusterLightCardEditorViewportClient::PropagateLightCardSelection()
 	LightCardEditorPtr.Pin()->SelectLightCards(SelectedLevelInstances);
 }
 
-AActor* FDisplayClusterLightCardEditorViewportClient::TraceScreenForLightCard(const FSceneView& View, int32 HitX, int32 HitY)
+void FDisplayClusterLightCardEditorViewportClient::PropagateLightCardTransform(ADisplayClusterLightCardActor* LightCardProxy)
+{
+	FLightCardProxy* FoundProxy = LightCardProxies.FindByKey(LightCardProxy);
+	if (FoundProxy && FoundProxy->Proxy == LightCardProxy && FoundProxy->LevelInstance.IsValid())
+	{
+		ADisplayClusterLightCardActor* LevelInstance = FoundProxy->LevelInstance.Get();
+
+		LevelInstance->Modify();
+
+		LevelInstance->Longitude = LightCardProxy->Longitude;
+		LevelInstance->Latitude = LightCardProxy->Latitude;
+		LevelInstance->DistanceFromCenter = LightCardProxy->DistanceFromCenter;
+		LevelInstance->Spin = LightCardProxy->Spin;
+		LevelInstance->Pitch = LightCardProxy->Pitch;
+		LevelInstance->Yaw = LightCardProxy->Yaw;
+
+		FVector RootActorLevelInstanceLocation = RootActorLevelInstance.IsValid() ? RootActorLevelInstance->GetActorLocation() : FVector::ZeroVector;
+		LevelInstance->SetActorLocation(RootActorLevelInstanceLocation + LightCardProxy->GetActorLocation());
+	}
+}
+
+void FDisplayClusterLightCardEditorViewportClient::MoveSelectedLightCards(FViewport* InViewport, EAxisList::Type CurrentAxis)
 {
 	UWorld* PreviewWorld = PreviewScene->GetWorld();
 	check(PreviewWorld);
 
-	const FVector4 ScreenPos = View.PixelToScreen(HitX, HitY, 0);
-	const FMatrix InvProjMatrix = View.ViewMatrices.GetInvProjectionMatrix();
-	const FMatrix InvViewMatrix = View.ViewMatrices.GetInvViewMatrix();
-	FVector ViewPos = FVector(InvProjMatrix.TransformFVector4(FVector4(ScreenPos.X * GNearClippingPlane, ScreenPos.Y * GNearClippingPlane, 0.0f, GNearClippingPlane)));
+	FIntPoint MousePos;
+	InViewport->GetMousePos(MousePos);
+
+	FSceneViewFamilyContext ViewFamily(FSceneViewFamily::ConstructionValues(
+		InViewport,
+		GetScene(),
+		EngineShowFlags)
+		.SetRealtimeUpdate(IsRealtime()));
+	FSceneView* View = CalcSceneView(&ViewFamily);
+
+	FVector Origin;
+	FVector Direction;
+	PixelToWorld(*View, MousePos, Origin, Direction);
+
+	const TWeakObjectPtr<ADisplayClusterLightCardActor>& SelectedLightCard = SelectedLightCards.Last();
+	if (SelectedLightCard.IsValid())
+	{
+		const FVector CurrentLightCardLocation = SelectedLightCard->GetLightCardTransform().GetTranslation();
+
+		// Light cards should be centered on the current view origin, so set the light card position to match the current view origin. Update the light card
+		// spherical coordinates to match its current coordinates
+		if (ProjectionOriginComponent.IsValid() && ProjectionOriginComponent->GetComponentLocation() != SelectedLightCard->GetActorLocation())
+		{
+			const FVector DesiredLightCardOffset = CurrentLightCardLocation - ProjectionOriginComponent->GetComponentLocation();
+
+			SelectedLightCard->SetActorLocation(ProjectionOriginComponent->GetComponentLocation());
+
+			FSphericalCoordinates SphericalCoords(DesiredLightCardOffset);
+
+			SelectedLightCard->DistanceFromCenter = SphericalCoords.Radius;
+			SelectedLightCard->Longitude = FMath::RadiansToDegrees(SphericalCoords.Radius) - 180;
+			SelectedLightCard->Latitude = 90.f - FMath::RadiansToDegrees(SphericalCoords.Radius);
+		}
+
+		const FVector CursorRayStart = Origin;
+		const FVector CursorRayEnd = CursorRayStart + Direction * HALF_WORLD_MAX;
+
+		FCollisionQueryParams Param(SCENE_QUERY_STAT(DragDropTrace), true);
+
+		bool bHitScreen = false;
+		FHitResult ScreenHitResult;
+		if (PreviewWorld->LineTraceSingleByObjectType(ScreenHitResult, CursorRayStart, CursorRayEnd, FCollisionObjectQueryParams(FCollisionObjectQueryParams::InitType::AllObjects), Param))
+		{
+			if (AActor* HitActor = ScreenHitResult.GetActor())
+			{
+				if (RootActorProxy.Get() == HitActor && ScreenHitResult.Component.IsValid())
+				{
+					// TODO: Make the light card flush with the screen, which requires setting the LC's distance from center to just behind the screen, and setting the pitch and yaw of the card
+					// to match the screen's normal at the traced point
+				}
+			}
+		}
+
+		// If we didn't hit a screen, keep the light card's radius fixed, and simply orbit it around the view origin
+		if (!bHitScreen)
+		{
+			const FSphericalCoordinates CurrentCoords = GetLightCardCoordinates(SelectedLightCard.Get());
+			const FSphericalCoordinates DeltaCoords = GetLightCardTranslationDelta(InViewport, SelectedLightCard.Get(), CurrentAxis);
+
+			FSphericalCoordinates NewCoords;
+			NewCoords.Radius = CurrentCoords.Radius + DeltaCoords.Radius;
+			NewCoords.Azimuth = CurrentCoords.Azimuth + DeltaCoords.Azimuth;
+			NewCoords.Inclination = CurrentCoords.Inclination + DeltaCoords.Inclination;
+
+			SelectedLightCard->DistanceFromCenter = NewCoords.Radius;
+			SelectedLightCard->Longitude = FRotator::ClampAxis(FMath::RadiansToDegrees(NewCoords.Azimuth) - 180);
+			SelectedLightCard->Latitude = 90.f - FMath::RadiansToDegrees(NewCoords.Inclination);
+		}
+
+		PropagateLightCardTransform(SelectedLightCard.Get());
+	}
+}
+
+FDisplayClusterLightCardEditorViewportClient::FSphericalCoordinates FDisplayClusterLightCardEditorViewportClient::GetLightCardTranslationDelta(
+	FViewport* InViewport, 
+	ADisplayClusterLightCardActor* LightCard,
+	EAxisList::Type CurrentAxis)
+{
+	FIntPoint MousePos;
+	InViewport->GetMousePos(MousePos);
+
+	FSceneViewFamilyContext ViewFamily(FSceneViewFamily::ConstructionValues(
+		InViewport,
+		GetScene(),
+		EngineShowFlags)
+		.SetRealtimeUpdate(IsRealtime()));
+	FSceneView* View = CalcSceneView(&ViewFamily);
+
+	FVector Origin;
+	FVector Direction;
+	PixelToWorld(*View, MousePos, Origin, Direction);
+
+	Direction = (Direction - DragWidgetOffset).GetSafeNormal();
+
+	const FVector LocalDirection = LightCard->GetActorRotation().RotateVector(Direction);
+	const FVector LightCardLocation = LightCard->GetLightCardTransform().GetTranslation() - Origin;
+
+	FSphericalCoordinates LightCardCoords = GetLightCardCoordinates(LightCard);
+	FSphericalCoordinates RequestedCoords(LocalDirection * LightCardLocation.Size());
 	
-	// If the projection mode is azimuthal, we need to convert to the "true" view position by undoing the non-linear projection applied by the projection shader
-	if (ProjectionMode == EDisplayClusterMeshProjectionType::Azimuthal)
+	FSphericalCoordinates DeltaCoords;
+	DeltaCoords.Radius = RequestedCoords.Radius - LightCardCoords.Radius;
+	DeltaCoords.Azimuth = RequestedCoords.Azimuth - LightCardCoords.Azimuth;
+	DeltaCoords.Inclination = RequestedCoords.Inclination - LightCardCoords.Inclination;
+
+	if (CurrentAxis == EAxisList::Type::X)
+	{
+		DeltaCoords.Inclination = 0;
+	}
+	else if (CurrentAxis == EAxisList::Type::Y)
 	{
-		const float Rho = ViewPos.Length();
-		const FVector UnitViewPos = ViewPos.GetSafeNormal();
+		// Convert the inclination to Cartesian coordinates, project it to the x-z plane, and convert back to spherical coordinates. This ensures that the motion in the inclination
+		// plane always lines up with the mouse's projected location along that plane
+		float FixedInclination = FMath::Abs(FMath::Atan2(FMath::Cos(DeltaCoords.Azimuth) * FMath::Sin(RequestedCoords.Inclination), FMath::Cos(RequestedCoords.Inclination)));
+
+		// When translating along the inclination axis, the azimuth delta can only be intervals of pi
+		float FixedAzimuth = FMath::RoundToInt(DeltaCoords.Azimuth / PI) * PI;
 
-		const FVector PlanePos = UnitViewPos / UnitViewPos.Z;
-		const FVector2D PolarCoords = FVector2D(FMath::Sqrt(PlanePos.X * PlanePos.X + PlanePos.Y * PlanePos.Y), FMath::Atan2(PlanePos.Y, PlanePos.X));
-		ViewPos = FVector(FMath::Sin(PolarCoords.X) * FMath::Cos(PolarCoords.Y), FMath::Sin(PolarCoords.X) * FMath::Sin(PolarCoords.Y), FMath::Cos(PolarCoords.X)) * Rho;
+		DeltaCoords.Azimuth = FixedAzimuth;
+		DeltaCoords.Inclination = FixedInclination - LightCardCoords.Inclination;
 	}
 
-	const FVector Direction = InvViewMatrix.TransformVector(ViewPos).GetSafeNormal();
+	return DeltaCoords;
+}
+
+FDisplayClusterLightCardEditorViewportClient::FSphericalCoordinates FDisplayClusterLightCardEditorViewportClient::GetLightCardCoordinates(ADisplayClusterLightCardActor* LightCard) const
+{
+	const FVector LightCardLocation = LightCard->GetLightCardTransform().GetTranslation() - LightCard->GetActorLocation();
+
+	FSphericalCoordinates LightCardCoords(LightCardLocation);
+
+	// If the light card inclination is 0 or 180, the spherical coordinates will have an
+	// "undefined" azimuth value. For continuity when dragging a light card positioned there, we can manually
+	// set the azimuthal value to match the light card's configured longitude
+	if (LightCardCoords.Inclination == 0.f ||LightCardCoords.Inclination == PI)
+	{
+		LightCardCoords.Azimuth = FMath::DegreesToRadians(LightCard->Longitude + 180);
+	}
 
-	const FVector CursorRayStart = View.ViewMatrices.GetViewOrigin();
+	return LightCardCoords;
+}
+
+ADisplayClusterLightCardActor* FDisplayClusterLightCardEditorViewportClient::TraceScreenForLightCard(const FSceneView& View, int32 HitX, int32 HitY)
+{
+	UWorld* PreviewWorld = PreviewScene->GetWorld();
+	check(PreviewWorld);
+
+	FVector Origin;
+	FVector Direction;
+	PixelToWorld(View, FIntPoint(HitX, HitY), Origin, Direction);
+
+	const FVector CursorRayStart = Origin;
 	const FVector CursorRayEnd = CursorRayStart + Direction * HALF_WORLD_MAX;
 
 	FCollisionQueryParams Param(SCENE_QUERY_STAT(DragDropTrace), true);
@@ -1066,18 +1269,21 @@ AActor* FDisplayClusterLightCardEditorViewportClient::TraceScreenForLightCard(co
 						{
 							for (FHitResult& HitResult : HitResults)
 							{
-								if (LightCardProxies.Contains(HitResult.GetActor()))
+								if (ADisplayClusterLightCardActor* LightCardActor = Cast<ADisplayClusterLightCardActor>(HitResult.GetActor()))
 								{
-									return HitResult.GetActor();
+									if (LightCardProxies.Contains(LightCardActor))
+									{
+										return LightCardActor;
+									}
 								}
 							}
 						}
 					}
 				}
 			}
-			else if (LightCardProxies.Contains(HitActor))
+			else if (HitActor->IsA<ADisplayClusterLightCardActor>() && LightCardProxies.Contains(HitActor))
 			{
-				return HitActor;
+				return Cast<ADisplayClusterLightCardActor>(HitActor);
 			}
 		}
 	}
@@ -1085,4 +1291,56 @@ AActor* FDisplayClusterLightCardEditorViewportClient::TraceScreenForLightCard(co
 	return nullptr;
 }
 
+void FDisplayClusterLightCardEditorViewportClient::PixelToWorld(const FSceneView& View, const FIntPoint& PixelPos, FVector& OutOrigin, FVector& OutDirection)
+{
+	const FMatrix& InvProjMatrix = View.ViewMatrices.GetInvProjectionMatrix();
+	const FMatrix& InvViewMatrix = View.ViewMatrices.GetInvViewMatrix();
+
+	const FVector4 ScreenPos = View.PixelToScreen(PixelPos.X, PixelPos.Y, 0);
+	const FVector ViewPos = FVector(InvProjMatrix.TransformFVector4(FVector4(ScreenPos.X * GNearClippingPlane, ScreenPos.Y * GNearClippingPlane, 0.0f, GNearClippingPlane)));
+	const FVector UnprojectedViewPos = FDisplayClusterMeshProjectionRenderer::UnprojectViewPosition(ViewPos, ProjectionMode);
+
+	OutOrigin = View.ViewMatrices.GetViewOrigin();
+	OutDirection = InvViewMatrix.TransformVector(UnprojectedViewPos).GetSafeNormal();
+}
+
+FTransform FDisplayClusterLightCardEditorViewportClient::CalcEditorWidgetTransform()
+{
+	if (SelectedLightCards.Num())
+	{
+		TWeakObjectPtr<ADisplayClusterLightCardActor> LastSelected = SelectedLightCards.Last();
+		if (LastSelected.IsValid())
+		{
+			FVector LightCardPosition = LastSelected->GetLightCardTransform().GetTranslation();
+
+			FTransform WidgetTransform(FRotator::ZeroRotator, LightCardPosition, FVector::OneVector);
+
+			if (ProjectionMode != EDisplayClusterMeshProjectionType::Perspective)
+			{
+				FSceneViewInitOptions SceneVewInitOptions;
+				GetSceneViewInitOptions(SceneVewInitOptions);
+				FViewMatrices ViewMatrices(SceneVewInitOptions);
+
+				const FVector ViewPos = ViewMatrices.GetViewMatrix().TransformPosition(LightCardPosition);
+				const FVector ProjectedViewPos = FDisplayClusterMeshProjectionRenderer::ProjectViewPosition(ViewPos, ProjectionMode);
+				const FVector ProjectedPosition = ViewMatrices.GetInvViewMatrix().TransformPosition(ProjectedViewPos);
+
+				WidgetTransform.SetTranslation(ProjectedPosition);
+			}
+
+			const FVector ProjectionOrigin = ProjectionOriginComponent.IsValid() ? ProjectionOriginComponent->GetComponentLocation() : FVector::ZeroVector;
+			const FVector RadialVector = (LightCardPosition - ProjectionOrigin).GetSafeNormal();
+			const FVector AzimuthalVector = FVector::ZAxisVector ^ RadialVector;
+			const FVector InclinationVector = RadialVector ^ AzimuthalVector;
+
+			FRotator Orientation = FMatrix(AzimuthalVector, InclinationVector, RadialVector, FVector::ZeroVector).Rotator();
+			WidgetTransform.SetRotation(Orientation.Quaternion());
+
+			return WidgetTransform;
+		}
+	}
+
+	return FTransform();
+}
+
 #undef LOCTEXT_NAMESPACE
diff --git a/Engine/Plugins/Runtime/nDisplay/Source/DisplayClusterLightCardEditor/Private/Viewport/DisplayClusterLightCardEditorViewportClient.h b/Engine/Plugins/Runtime/nDisplay/Source/DisplayClusterLightCardEditor/Private/Viewport/DisplayClusterLightCardEditorViewportClient.h
--- a/Engine/Plugins/Runtime/nDisplay/Source/DisplayClusterLightCardEditor/Private/Viewport/DisplayClusterLightCardEditorViewportClient.h
+++ b/Engine/Plugins/Runtime/nDisplay/Source/DisplayClusterLightCardEditor/Private/Viewport/DisplayClusterLightCardEditorViewportClient.h
@@ -9,15 +9,49 @@
 
 #include "DisplayClusterLightCardEditorProxyType.h"
 #include "DisplayClusterMeshProjectionRenderer.h"
+#include "DisplayClusterLightCardActor.h"
 
 class ADisplayClusterRootActor;
 class SDisplayClusterLightCardEditor;
 class FScopedTransaction;
+class FDisplayClusterLightCardEditorWidget;
 class UDisplayClusterConfigurationViewport;
 
 /** Viewport Client for the preview viewport */
 class FDisplayClusterLightCardEditorViewportClient : public FEditorViewportClient, public TSharedFromThis<FDisplayClusterLightCardEditorViewportClient>
 {
+private:
+
+	struct FSphericalCoordinates
+	{
+	public:
+		FSphericalCoordinates(const FVector& CartesianPosition)
+		{
+			Radius = CartesianPosition.Size();
+
+			if (Radius > UE_SMALL_NUMBER)
+			{
+				Inclination = FMath::Acos(CartesianPosition.Z / Radius);
+			}
+			else
+			{
+				Inclination = 0.f;
+			}
+
+			Azimuth = FMath::Atan2(CartesianPosition.Y, CartesianPosition.X);
+		}
+
+		FSphericalCoordinates()
+			: Radius(0.f)
+			, Inclination(0.f)
+			, Azimuth(0.f)
+		{ }
+
+		float Radius = 0.f;
+		float Inclination = 0.f;
+		float Azimuth = 0.f;
+	};
+
 public:
 	FDisplayClusterLightCardEditorViewportClient(FAdvancedPreviewScene& InPreviewScene, const TWeakPtr<SEditorViewport>& InEditorViewportWidget,
 		TWeakPtr<SDisplayClusterLightCardEditor> InLightCardEditor);
@@ -27,10 +61,10 @@ public:
 	virtual FLinearColor GetBackgroundColor() const override;
 	virtual void Tick(float DeltaSeconds) override;
 	virtual void Draw(FViewport* InViewport, FCanvas* Canvas) override;
-	virtual bool CanSetWidgetMode(UE::Widget::EWidgetMode NewMode) const override { return true; }
-	virtual bool CanCycleWidgetMode() const override { return true; }
-	virtual UE::Widget::EWidgetMode GetWidgetMode() const override;
-	virtual FVector GetWidgetLocation() const override;
+	virtual void Draw(const FSceneView* View, FPrimitiveDrawInterface* PDI) override;
+	virtual bool CanSetWidgetMode(UE::Widget::EWidgetMode NewMode) const override { return false; }
+	virtual bool CanCycleWidgetMode() const override { return false; }
+	virtual UE::Widget::EWidgetMode GetWidgetMode() const override { return UE::Widget::WM_None; }
 	virtual FSceneView* CalcSceneView(FSceneViewFamily* ViewFamily, const int32 StereoViewIndex = INDEX_NONE) override;
 	virtual bool IsLevelEditorClient() const override { return false; }
 	virtual bool InputKey(FViewport* InViewport, int32 ControllerId, FKey Key, EInputEvent Event, float AmountDepressed, bool bGamepad) override;
@@ -42,9 +76,6 @@ public:
 	virtual ELevelViewportType GetViewportType() const override { return LVT_Perspective; }
 	// ~FEditorViewportClient
 
-	void SelectActor(AActor* NewActor);
-	void ResetSelection();
-
 	/**
 	 * Update the spawned preview actor from a root actor in the level
 	 * 
@@ -93,19 +124,37 @@ private:
 	void FindProjectionOriginComponent();
 
 	/** Gets a list of all light card actors on the level linked to the specified root actor */
-	void FindLightCardsForRootActor(ADisplayClusterRootActor* RootActor, TArray<TWeakObjectPtr<AActor>>& OutLightCards);
+	void FindLightCardsForRootActor(ADisplayClusterRootActor* RootActor, TArray<TWeakObjectPtr<ADisplayClusterLightCardActor>>& OutLightCards);
 
 	/** Callback to check if an light card actor is among the list of selected light card actors */
 	bool IsLightCardSelected(const AActor* Actor);
 
 	/** Adds the specified light card actor the the list of selected light cards */
-	void SelectLightCard(AActor* Actor, bool bAddToSelection = false);
+	void SelectLightCard(ADisplayClusterLightCardActor* Actor, bool bAddToSelection = false);
 
 	/** Notifies the light card editor of the currently selected light cards so that it may update other UI components to match */
 	void PropagateLightCardSelection();
 
+	/** Propagates the specified light card proxy's transform back to its level instance version */
+	void PropagateLightCardTransform(ADisplayClusterLightCardActor* LightCardProxy);
+
+	/** Moves the currently selected light cards */
+	void MoveSelectedLightCards(FViewport* InViewport, EAxisList::Type CurrentAxis);
+
+	/** Determines the appropriate delta in spherical coordinates needed to move the specified light card to the mouse's location */
+	FSphericalCoordinates GetLightCardTranslationDelta(FViewport* InViewport, ADisplayClusterLightCardActor* LightCard, EAxisList::Type CurrentAxis);
+
+	/** Gets the spherical coordinates of the specified light card */
+	FSphericalCoordinates GetLightCardCoordinates(ADisplayClusterLightCardActor* LightCard) const;
+
 	/** Traces to find the light card corresponding to a click on a stage screen */
-	AActor* TraceScreenForLightCard(const FSceneView& View, int32 HitX, int32 HitY);
+	ADisplayClusterLightCardActor* TraceScreenForLightCard(const FSceneView& View, int32 HitX, int32 HitY);
+
+	/** Converts a pixel coordinate into a point and direction vector in world space */
+	void PixelToWorld(const FSceneView& View, const FIntPoint& PixelPos, FVector& OutOrigin, FVector& OutDirection);
+
+	/** Calculates the world transform to render the editor widget with to align it with the selected light card */
+	FTransform CalcEditorWidgetTransform();
 
 private:
 	TWeakPtr<FSceneViewport> SceneViewportPtr;
@@ -115,10 +164,10 @@ private:
 
 	struct FLightCardProxy
 	{
-		TWeakObjectPtr<AActor> LevelInstance;
-		TWeakObjectPtr<AActor> Proxy;
+		TWeakObjectPtr<ADisplayClusterLightCardActor> LevelInstance;
+		TWeakObjectPtr<ADisplayClusterLightCardActor> Proxy;
 
-		FLightCardProxy(AActor* InLevelInstance, AActor* InProxy)
+		FLightCardProxy(ADisplayClusterLightCardActor* InLevelInstance, ADisplayClusterLightCardActor* InProxy)
 			: LevelInstance(InLevelInstance)
 			, Proxy(InProxy)
 		{ }
@@ -127,7 +176,7 @@ private:
 	};
 
 	TArray<FLightCardProxy> LightCardProxies;
-	TArray<TWeakObjectPtr<AActor>> SelectedLightCards;
+	TArray<TWeakObjectPtr<ADisplayClusterLightCardActor>> SelectedLightCards;
 	
 	/** The renderer for the viewport, which can render the meshes with a variety of projection types */
 	TSharedPtr<FDisplayClusterMeshProjectionRenderer> MeshProjectionRenderer;
@@ -135,8 +184,18 @@ private:
 	/** The current transaction for undo/redo */
 	FScopedTransaction* ScopedTransaction = nullptr;
 	
+	/** Indicates that the user is dragging an actor in the viewport */
 	bool bDraggingActor = false;
 
+	/** The LC editor widget used to manipulate light cards */
+	TSharedPtr<FDisplayClusterLightCardEditorWidget> EditorWidget;
+
+	/** The cached editor widget transform calculated for the editor widget on the last tick */
+	FTransform CachedEditorWidgetTransform;
+
+	/** The offset between the widget's origin and the place it was clicked when a drag action was started */
+	FVector DragWidgetOffset;
+
 	/** The current projection mode the 3D viewport is being displayed with */
 	EDisplayClusterMeshProjectionType ProjectionMode = EDisplayClusterMeshProjectionType::Perspective;
 
diff --git a/Engine/Plugins/Runtime/nDisplay/Source/DisplayClusterLightCardEditor/Private/Viewport/DisplayClusterLightCardEditorWidget.cpp b/Engine/Plugins/Runtime/nDisplay/Source/DisplayClusterLightCardEditor/Private/Viewport/DisplayClusterLightCardEditorWidget.cpp
new file mode 100644
--- /dev/null
+++ b/Engine/Plugins/Runtime/nDisplay/Source/DisplayClusterLightCardEditor/Private/Viewport/DisplayClusterLightCardEditorWidget.cpp
@@ -0,0 +1,78 @@
+// Copyright Epic Games, Inc. All Rights Reserved.
+
+#include "DisplayClusterLightCardEditorWidget.h"
+
+#include "SceneManagement.h"
+#include "UnrealWidget.h"
+
+void FDisplayClusterLightCardEditorWidget::Draw(const FSceneView* View, FPrimitiveDrawInterface* PDI)
+{
+	DrawAxis(View, PDI, EAxisList::Type::X);
+	DrawAxis(View, PDI, EAxisList::Type::Y);
+	DrawOrigin(View, PDI);
+}
+
+void FDisplayClusterLightCardEditorWidget::DrawAxis(const FSceneView* View, FPrimitiveDrawInterface* PDI, EAxisList::Type Axis)
+{
+	const FVector Origin = Transform.GetTranslation();
+
+	const FVector GlobalAxis = GetGlobalAxis(Axis);
+	const FVector TransformedAxis = Transform.TransformVector(GlobalAxis);
+	const FLinearColor AxisColor = HighlightedAxis == Axis ? HighlightColor : GetAxisColor(Axis);
+
+	// Compute the scalar needed to make the axis length a fixed size on the screen regardless of distance and field of view
+	const float DistanceScreenSpaceScalar = WidgetScale * FMath::Max(FVector::Dist(Origin, View->ViewMatrices.GetViewOrigin()), 1.f) / View->ViewMatrices.GetScreenScale();
+	const float ThicknessScreenSpaceScalar = 100.f / View->ViewMatrices.GetScreenScale();
+
+
+	PDI->SetHitProxy(new HWidgetAxis(Axis));
+	PDI->DrawLine(Origin, Origin + TransformedAxis * AxisLength * DistanceScreenSpaceScalar, AxisColor, ESceneDepthPriorityGroup::SDPG_Foreground, AxisThickness * ThicknessScreenSpaceScalar, 0.0, true);
+	PDI->SetHitProxy(nullptr);
+}
+
+void FDisplayClusterLightCardEditorWidget::DrawOrigin(const FSceneView* View, FPrimitiveDrawInterface* PDI)
+{
+	const FVector Origin = Transform.GetTranslation();
+	const FLinearColor Color = HighlightedAxis == EAxisList::Type::XY ? HighlightColor : FLinearColor::Black;
+
+	// Compute the scalar needed to make the axis length a fixed size on the screen regardless of distance and field of view
+	const float ScreenSpaceScalar = 100.f / View->ViewMatrices.GetScreenScale();
+
+	PDI->SetHitProxy(new HWidgetAxis(EAxisList::Type::XY));
+	PDI->DrawPoint(Origin, Color, OriginSize * ScreenSpaceScalar, ESceneDepthPriorityGroup::SDPG_Foreground);
+	PDI->SetHitProxy(nullptr);
+}
+
+FVector FDisplayClusterLightCardEditorWidget::GetGlobalAxis(EAxisList::Type Axis) const
+{
+	switch (Axis)
+	{
+	case EAxisList::Type::X:
+		return FVector::UnitX();
+
+	case EAxisList::Type::Y:
+		return FVector::UnitY();
+		
+	case EAxisList::Type::Z:
+		return FVector::UnitZ();
+	}
+
+	return FVector::ZeroVector;
+}
+
+FLinearColor FDisplayClusterLightCardEditorWidget::GetAxisColor(EAxisList::Type Axis) const
+{
+	switch (Axis)
+	{
+	case EAxisList::Type::X:
+		return AxisColorX;
+
+	case EAxisList::Type::Y:
+		return AxisColorY;
+		
+	case EAxisList::Type::Z:
+		return AxisColorZ;
+	}
+
+	return FLinearColor::Transparent;
+}
\ No newline at end of file
diff --git a/Engine/Plugins/Runtime/nDisplay/Source/DisplayClusterLightCardEditor/Private/Viewport/DisplayClusterLightCardEditorWidget.h b/Engine/Plugins/Runtime/nDisplay/Source/DisplayClusterLightCardEditor/Private/Viewport/DisplayClusterLightCardEditorWidget.h
new file mode 100644
--- /dev/null
+++ b/Engine/Plugins/Runtime/nDisplay/Source/DisplayClusterLightCardEditor/Private/Viewport/DisplayClusterLightCardEditorWidget.h
@@ -0,0 +1,45 @@
+// Copyright Epic Games, Inc. All Rights Reserved.
+
+#pragma once
+
+#include "CoreMinimal.h"
+#include "Math/Axis.h"
+
+class FSceneView;
+class FPrimitiveDrawInterface;
+
+/** Widget used to manipulate entities in the light card editor's viewport */
+class FDisplayClusterLightCardEditorWidget
+{
+public:
+	constexpr static float AxisLength = 100.f;
+	constexpr static float AxisThickness = 35.f;
+	constexpr static float OriginSize = 200.f;
+	constexpr static FLinearColor AxisColorX = FLinearColor(0.594f, 0.0197f, 0.0f);
+	constexpr static FLinearColor AxisColorY = FLinearColor(0.1349f, 0.3959f, 0.0f);
+	constexpr static FLinearColor AxisColorZ = FLinearColor(0.0251f, 0.207f, 0.85f);
+	constexpr static FLinearColor HighlightColor = FLinearColor(1.0f, 1.0f, 0.0f);
+
+	void Draw(const FSceneView* View, FPrimitiveDrawInterface* PDI);
+
+	void SetTransform(const FTransform& NewTransform) { Transform = FTransform(NewTransform); }
+	void SetHighlightedAxis(EAxisList::Type InAxis) { HighlightedAxis = InAxis; }
+	void SetWidgetScale(float InWidgetScale) { WidgetScale = InWidgetScale; }
+
+private:
+	void DrawAxis(const FSceneView* View, FPrimitiveDrawInterface* PDI, EAxisList::Type Axis);
+	void DrawOrigin(const FSceneView* View, FPrimitiveDrawInterface* PDI);
+
+	FVector GetGlobalAxis(EAxisList::Type Axis) const;
+	FLinearColor GetAxisColor(EAxisList::Type Axis) const;
+
+private:
+	/** The world transform to apply to the widget when rendering */
+	FTransform Transform;
+
+	/** The axis on the widget which should be highlighted */
+	EAxisList::Type HighlightedAxis = EAxisList::Type::None;
+
+	/** A scale factor to scale the widget's rendered size */
+	float WidgetScale = 1.f;
+};
\ No newline at end of file
diff --git a/Engine/Plugins/Runtime/nDisplay/Source/DisplayClusterLightCardEditorShaders/Private/DisplayClusterMeshProjectionRenderer.cpp b/Engine/Plugins/Runtime/nDisplay/Source/DisplayClusterLightCardEditorShaders/Private/DisplayClusterMeshProjectionRenderer.cpp
--- a/Engine/Plugins/Runtime/nDisplay/Source/DisplayClusterLightCardEditorShaders/Private/DisplayClusterMeshProjectionRenderer.cpp
+++ b/Engine/Plugins/Runtime/nDisplay/Source/DisplayClusterLightCardEditorShaders/Private/DisplayClusterMeshProjectionRenderer.cpp
@@ -10,6 +10,7 @@
 #include "Shader.h"
 #include "CanvasTypes.h"
 #include "EngineModule.h"
+#include "SceneManagement.h"
 #include "SceneViewExtension.h"
 #include "ScreenPass.h"
 #include "Components/PrimitiveComponent.h"
@@ -85,14 +86,15 @@ public:
 		, DrawRenderState(*InView)
 		, bTranslucencyPass(bIsTranslucencyPass)
 	{
-		DrawRenderState.SetDepthStencilState(TStaticDepthStencilState<true, CF_DepthNearOrEqual>::GetRHI());
 
 		if (bTranslucencyPass)
 		{
+			DrawRenderState.SetDepthStencilState(TStaticDepthStencilState<false, CF_DepthNearOrEqual>::GetRHI());
 			DrawRenderState.SetBlendState(TStaticBlendState<CW_RGBA, BO_Add, BF_SourceAlpha, BF_InverseSourceAlpha, BO_Add, BF_Zero, BF_One>::GetRHI());
 		}
 		else
 		{
+			DrawRenderState.SetDepthStencilState(TStaticDepthStencilState<true, CF_DepthNearOrEqual>::GetRHI());
 			DrawRenderState.SetBlendState(TStaticBlendStateWriteMask<CW_RGBA, CW_RGBA, CW_RGBA, CW_RGBA>::GetRHI());
 		}
 	}
@@ -510,6 +512,63 @@ namespace
 //////////////////////////////////////////////////////////////////////////
 // FDisplayClusterMeshProjectionRenderer
 
+/** A version of FSimpleElementCollector that has a hit proxy consumer, allowing support for adding hit proxies to the rendered elements */
+class FMeshProjectionElementCollector : public FSimpleElementCollector
+{
+public:
+	FMeshProjectionElementCollector(FHitProxyConsumer* InHitProxyConsumer)
+		: FSimpleElementCollector()
+		, HitProxyConsumer(InHitProxyConsumer)
+	{ }
+
+	virtual void SetHitProxy(HHitProxy* HitProxy) override
+	{
+		FSimpleElementCollector::SetHitProxy(HitProxy);
+
+		if (HitProxyConsumer && HitProxy)
+		{
+			HitProxyConsumer->AddHitProxy(HitProxy);
+		}
+	}
+
+private:
+	FHitProxyConsumer* HitProxyConsumer;
+};
+
+FVector FDisplayClusterMeshProjectionRenderer::ProjectViewPosition(const FVector& ViewPosition, EDisplayClusterMeshProjectionType  ProjectionType)
+{
+	FVector ProjectedViewPosition(ViewPosition);
+
+	if (ProjectionType == EDisplayClusterMeshProjectionType::Azimuthal)
+	{
+		const float Rho = ViewPosition.Length();
+		const FVector UnitViewPos = ViewPosition.GetSafeNormal();
+		const FVector2D PolarCoords = FVector2D(FMath::Acos(UnitViewPos.Z), FMath::Atan2(UnitViewPos.Y, UnitViewPos.X));
+		const FVector PlanePos = FVector(PolarCoords.X * FMath::Cos(PolarCoords.Y), PolarCoords.X * FMath::Sin(PolarCoords.Y), 1);
+
+		ProjectedViewPosition = PlanePos.GetSafeNormal() * Rho;
+	}
+
+	return ProjectedViewPosition;
+}
+
+FVector FDisplayClusterMeshProjectionRenderer::UnprojectViewPosition(const FVector& ProjectedViewPosition, EDisplayClusterMeshProjectionType  ProjectionType)
+{
+	FVector ViewPosition(ProjectedViewPosition);
+
+	if (ProjectionType == EDisplayClusterMeshProjectionType::Azimuthal)
+	{
+		const float Rho = ProjectedViewPosition.Length();
+		const FVector UnitViewPos = ProjectedViewPosition.GetSafeNormal();
+
+		const FVector PlanePos = UnitViewPos / UnitViewPos.Z;
+		const FVector2D PolarCoords = FVector2D(FMath::Sqrt(PlanePos.X * PlanePos.X + PlanePos.Y * PlanePos.Y), FMath::Atan2(PlanePos.Y, PlanePos.X));
+		ViewPosition = FVector(FMath::Sin(PolarCoords.X) * FMath::Cos(PolarCoords.Y), FMath::Sin(PolarCoords.X) * FMath::Sin(PolarCoords.Y), FMath::Cos(PolarCoords.X)) * Rho;
+	}
+
+	return ViewPosition;
+}
+
 void FDisplayClusterMeshProjectionRenderer::AddActor(AActor* Actor)
 {
 	AddActor(Actor, [](const UPrimitiveComponent* PrimitiveComponent)
@@ -570,9 +629,10 @@ void FDisplayClusterMeshProjectionRenderer::Render(FCanvas* Canvas, FSceneInterf
 	Canvas->Flush_GameThread();
 	FRenderTarget* RenderTarget = Canvas->GetRenderTarget();
 	const bool bIsHitTesting = Canvas->IsHitTesting();
+	FHitProxyConsumer* HitProxyConsumer = Canvas->GetHitProxyConsumer();
 
 	ENQUEUE_RENDER_COMMAND(FDrawProjectedMeshes)(
-		[RenderTarget, Scene, ViewInitOptions, EngineShowFlags, bIsHitTesting, ProjectionType, this](FRHICommandListImmediate& RHICmdList)
+		[RenderTarget, Scene, ViewInitOptions, EngineShowFlags, bIsHitTesting, HitProxyConsumer, ProjectionType, this](FRHICommandListImmediate& RHICmdList)
 		{
 			FMemMark Mark(FMemStack::Get());
 			FRDGBuilder GraphBuilder(RHICmdList);
@@ -602,7 +662,7 @@ void FDisplayClusterMeshProjectionRenderer::Render(FCanvas* Canvas, FSceneInterf
 			NewInitOptions.ViewFamily = &ViewFamily;
 
 			GetRendererModule().CreateAndInitSingleView(RHICmdList, &ViewFamily, &NewInitOptions);
-			const FViewInfo* View = (FViewInfo*)ViewFamily.Views[0];
+			FViewInfo* View = (FViewInfo*)ViewFamily.Views[0];
 
 			FRDGTextureRef OutputTexture = GraphBuilder.RegisterExternalTexture(CreateRenderTarget(RenderTarget->GetRenderTargetTexture(), TEXT("ViewRenderTarget")));
 			FRenderTargetBinding OutputRenderTargetBinding(OutputTexture, ERenderTargetLoadAction::ELoad);
@@ -638,7 +698,7 @@ void FDisplayClusterMeshProjectionRenderer::Render(FCanvas* Canvas, FSceneInterf
 						RDG_EVENT_NAME("MeshProjectionRenderer::CopyHitProxyTexture"),
 						ScreenPassParameters,
 						ERDGPassFlags::Raster,
-						[View, ScreenPassVS, CopyPixelShader, &RegionViewport, ScreenPassParameters, DefaultBlendState](FRHICommandList& RHICmdList)
+						[View, ScreenPassVS, CopyPixelShader, RegionViewport, ScreenPassParameters, DefaultBlendState](FRHICommandList& RHICmdList)
 					{
 						DrawScreenPass(
 							RHICmdList,
@@ -696,7 +756,7 @@ void FDisplayClusterMeshProjectionRenderer::Render(FCanvas* Canvas, FSceneInterf
 						RDG_EVENT_NAME("MeshProjectionRenderer::CopyColorTexture"),
 						ScreenPassParameters,
 						ERDGPassFlags::Raster,
-						[View, ScreenPassVS, CopyPixelShader, &RegionViewport, ScreenPassParameters, DefaultBlendState](FRHICommandList& RHICmdList)
+						[View, ScreenPassVS, CopyPixelShader, RegionViewport, ScreenPassParameters, DefaultBlendState](FRHICommandList& RHICmdList)
 					{
 						DrawScreenPass(
 							RHICmdList,
@@ -712,6 +772,13 @@ void FDisplayClusterMeshProjectionRenderer::Render(FCanvas* Canvas, FSceneInterf
 				}
 #endif
 			}
+			
+			FMeshProjectionElementCollector ElementCollector(HitProxyConsumer);
+			if (RenderSimpleElementsDelegate.IsBound())
+			{
+				RenderSimpleElementsDelegate.Execute(View, &ElementCollector);
+				AddSimpleElementPass(GraphBuilder, View, OutputRenderTargetBinding, ElementCollector);
+			}
 
 			GraphBuilder.Execute();
 		});
@@ -845,7 +912,7 @@ void FDisplayClusterMeshProjectionRenderer::AddSelectionOutlineScreenPass(FRDGBu
 		RDG_EVENT_NAME("MeshProjectionRenderer::SelectionScreen"),
 		ScreenPassParameters,
 		ERDGPassFlags::Raster,
-		[View, ScreenPassVS, SelectionOutlinePS, &OutputViewport, ScreenPassParameters, DefaultBlendState](FRHICommandList& RHICmdList)
+		[View, ScreenPassVS, SelectionOutlinePS, OutputViewport, ScreenPassParameters, DefaultBlendState](FRHICommandList& RHICmdList)
 	{
 		DrawScreenPass(
 			RHICmdList,
@@ -861,6 +928,36 @@ void FDisplayClusterMeshProjectionRenderer::AddSelectionOutlineScreenPass(FRDGBu
 }
 #endif
 
+void FDisplayClusterMeshProjectionRenderer::AddSimpleElementPass(FRDGBuilder& GraphBuilder,
+	const FViewInfo* View,
+	FRenderTargetBinding& OutputRenderTargetBinding,
+	FSimpleElementCollector& ElementCollector)
+{
+	FMeshProjectionPassParameters* PassParameters = GraphBuilder.AllocParameters<FMeshProjectionPassParameters>();
+	PassParameters->View = View->ViewUniformBuffer;
+	PassParameters->InstanceCulling = FInstanceCullingContext::CreateDummyInstanceCullingUniformBuffer(GraphBuilder);
+	PassParameters->RenderTargets[0] = OutputRenderTargetBinding;
+
+	GraphBuilder.AddPass(
+		RDG_EVENT_NAME("MeshProjectionRenderer::SimpleElements"),
+		PassParameters,
+		ERDGPassFlags::Raster,
+		[this, View,  &ElementCollector](FRHICommandListImmediate& RHICmdList)
+		{
+			FIntRect ViewRect = View->UnscaledViewRect;
+			RHICmdList.SetViewport(ViewRect.Min.X, ViewRect.Min.Y, 0.0f, ViewRect.Max.X, ViewRect.Max.Y, 1.0f);
+
+			FMeshPassProcessorRenderState DrawRenderState(*View);
+			DrawRenderState.SetDepthStencilAccess(FExclusiveDepthStencil::DepthRead_StencilWrite);
+			DrawRenderState.SetBlendState(TStaticBlendStateWriteMask<CW_RGBA, CW_RGBA, CW_RGBA, CW_RGBA>::GetRHI());
+			DrawRenderState.SetDepthStencilState(TStaticDepthStencilState<false, CF_DepthNearOrEqual>::GetRHI());
+
+			ElementCollector.DrawBatchedElements(RHICmdList, DrawRenderState, *View, EBlendModeFilter::OpaqueAndMasked, ESceneDepthPriorityGroup::SDPG_World);
+			ElementCollector.DrawBatchedElements(RHICmdList, DrawRenderState, *View, EBlendModeFilter::OpaqueAndMasked, ESceneDepthPriorityGroup::SDPG_Foreground);
+		}
+	);
+}
+
 template<EDisplayClusterMeshProjectionType ProjectionType>
 void FDisplayClusterMeshProjectionRenderer::RenderPrimitives_RenderThread(const FSceneView* View, FRHICommandList& RHICmdList, bool bTranslucencyPass)
 {
@@ -881,6 +978,8 @@ void FDisplayClusterMeshProjectionRenderer::RenderPrimitives_RenderThread(const
 		{
 			if (const FMeshBatch* MeshBatch = PrimitiveProxy->GetPrimitiveSceneInfo()->GetMeshBatch(PrimitiveProxy->GetPrimitiveSceneInfo()->StaticMeshes.Num() - 1))
 			{
+				MeshBatch->MaterialRenderProxy->UpdateUniformExpressionCacheIfNeeded(View->GetFeatureLevel());
+
 				const uint64 BatchElementMask = ~0ull;
 				MeshProcessor.AddMeshBatch(*MeshBatch, BatchElementMask, PrimitiveProxy);
 			}
@@ -908,6 +1007,8 @@ void FDisplayClusterMeshProjectionRenderer::RenderHitProxies_RenderThread(const
 		{
 			if (const FMeshBatch* MeshBatch = PrimitiveProxy->GetPrimitiveSceneInfo()->GetMeshBatch(PrimitiveProxy->GetPrimitiveSceneInfo()->StaticMeshes.Num() - 1))
 			{
+				MeshBatch->MaterialRenderProxy->UpdateUniformExpressionCacheIfNeeded(View->GetFeatureLevel());
+
 				const uint64 BatchElementMask = ~0ull;
 				MeshProcessor.AddMeshBatch(*MeshBatch, BatchElementMask, PrimitiveProxy);
 			}
@@ -936,6 +1037,8 @@ void FDisplayClusterMeshProjectionRenderer::RenderSelection_RenderThread(const F
 		{
 			if (const FMeshBatch* MeshBatch = PrimitiveProxy->GetPrimitiveSceneInfo()->GetMeshBatch(PrimitiveProxy->GetPrimitiveSceneInfo()->StaticMeshes.Num() - 1))
 			{
+				MeshBatch->MaterialRenderProxy->UpdateUniformExpressionCacheIfNeeded(View->GetFeatureLevel());
+
 				const uint64 BatchElementMask = ~0ull;
 				MeshProcessor.AddMeshBatch(*MeshBatch, BatchElementMask, PrimitiveProxy);
 			}
diff --git a/Engine/Plugins/Runtime/nDisplay/Source/DisplayClusterLightCardEditorShaders/Public/DisplayClusterMeshProjectionRenderer.h b/Engine/Plugins/Runtime/nDisplay/Source/DisplayClusterLightCardEditorShaders/Public/DisplayClusterMeshProjectionRenderer.h
--- a/Engine/Plugins/Runtime/nDisplay/Source/DisplayClusterLightCardEditorShaders/Public/DisplayClusterMeshProjectionRenderer.h
+++ b/Engine/Plugins/Runtime/nDisplay/Source/DisplayClusterLightCardEditorShaders/Public/DisplayClusterMeshProjectionRenderer.h
@@ -9,6 +9,9 @@ class AActor;
 class UPrimitiveComponent;
 class FSceneView;
 class FViewInfo;
+class FPrimitiveDrawInterface;
+class FSimpleElementCollector;
+struct FSceneViewInitOptions;
 struct FEngineShowFlags;
 
 /** Indicates which kind of projection is used by the renderer */
@@ -25,6 +28,12 @@ enum DISPLAYCLUSTERLIGHTCARDEDITORSHADERS_API EDisplayClusterMeshProjectionType
 class DISPLAYCLUSTERLIGHTCARDEDITORSHADERS_API FDisplayClusterMeshProjectionRenderer
 {
 public:
+	/** Projects a position in view coordinates into the projected view space of the specified projection type */
+	static FVector ProjectViewPosition(const FVector& ViewPosition, EDisplayClusterMeshProjectionType  ProjectionType);
+
+	/** Projects a position in the projected view space of the specified projection type to ordinary view coordinates */
+	static FVector UnprojectViewPosition(const FVector& ProjectedViewPosition, EDisplayClusterMeshProjectionType  ProjectionType);
+
 	/** Adds an actor's primitive components to the list of primitives to render */
 	void AddActor(AActor* Actor);
 
@@ -78,6 +87,12 @@ private:
 		FRDGTexture* SelectionDepth);
 #endif
 
+	/** Adds a pass that renders any elements added to the PDI through the renderer's RenderSimpleElements callback */
+	void AddSimpleElementPass(FRDGBuilder& GraphBuilder,
+		const FViewInfo* View,
+		FRenderTargetBinding& OutputRenderTargetBinding,
+		FSimpleElementCollector& ElementCollector);
+
 	/** Renders the list of primitive components using the appropriate mesh pass processor given by the template parameter */
 	template<EDisplayClusterMeshProjectionType ProjectionType>
 	void RenderPrimitives_RenderThread(const FSceneView* View, FRHICommandList& RHICmdList, bool bTranslucencyPass);
@@ -100,6 +115,10 @@ public:
 	DECLARE_DELEGATE_RetVal_OneParam(bool, FSelection, const AActor*);
 	FSelection ActorSelectedDelegate;
 
+	/** Delegate raised during a render pass allowing simple elements ro be rendered to the viewport after meshes are projected to it */
+	DECLARE_DELEGATE_TwoParams(FSimpleElementPass, const FSceneView*, FPrimitiveDrawInterface*);
+	FSimpleElementPass RenderSimpleElementsDelegate;
+
 private:
 	/** The list of primitive components to render */
 	TArray<TWeakObjectPtr<UPrimitiveComponent>> PrimitiveComponents;
