nDisplay: Allow the DCRA to render with and without post process to preview components. Prevent LC DCRA proxy from rendering its preview internally.

diff --git a/Engine/Plugins/Runtime/nDisplay/Source/DisplayCluster/Private/Game/EngineClasses/Scene/DisplayClusterPreviewComponent.cpp b/Engine/Plugins/Runtime/nDisplay/Source/DisplayCluster/Private/Game/EngineClasses/Scene/DisplayClusterPreviewComponent.cpp
--- a/Engine/Plugins/Runtime/nDisplay/Source/DisplayCluster/Private/Game/EngineClasses/Scene/DisplayClusterPreviewComponent.cpp
+++ b/Engine/Plugins/Runtime/nDisplay/Source/DisplayCluster/Private/Game/EngineClasses/Scene/DisplayClusterPreviewComponent.cpp
@@ -88,7 +88,8 @@ IDisplayClusterViewport* UDisplayClusterPreviewComponent::GetCurrentViewport() c
 	return nullptr;
 }
 
-bool UDisplayClusterPreviewComponent::InitializePreviewComponent(ADisplayClusterRootActor* InRootActor, const FString& InClusterNodeId, const FString& InViewportId, UDisplayClusterConfigurationViewport* InViewportConfig)
+bool UDisplayClusterPreviewComponent::InitializePreviewComponent(ADisplayClusterRootActor* InRootActor, const FString& InClusterNodeId,
+	const FString& InViewportId, UDisplayClusterConfigurationViewport* InViewportConfig)
 {
 	RootActor = InRootActor;
 	ViewportId = InViewportId;
@@ -167,7 +168,8 @@ bool UDisplayClusterPreviewComponent::UpdatePreviewMesh()
 
 		// And search for new mesh reference
 		IDisplayClusterViewport* Viewport = GetCurrentViewport();
-		const bool bIsViewportEnabled = RenderTarget != nullptr && Viewport != nullptr && Viewport->GetRenderSettings().bEnable && Viewport->GetProjectionPolicy().IsValid();
+		const bool bIsViewportEnabled = (RenderTarget != nullptr || RenderTargetPostProcess != nullptr)
+		&& Viewport != nullptr && Viewport->GetRenderSettings().bEnable && Viewport->GetProjectionPolicy().IsValid();
 		if (bIsViewportEnabled)
 		{
 			// Handle preview mesh:
@@ -250,7 +252,8 @@ void UDisplayClusterPreviewComponent::UpdatePreviewMaterial()
 		}
 		else
 		{
-			PreviewMaterialInstance->SetTextureParameterValue(TEXT("Preview"), RenderTarget);
+			PreviewMaterialInstance->SetTextureParameterValue(TEXT("Preview"),
+				RootActor && RootActor->bPreviewEnablePostProcess ? RenderTargetPostProcess : RenderTarget);
 		}
 	}
 }
@@ -274,13 +277,32 @@ void UDisplayClusterPreviewComponent::ReleasePreviewMaterial()
 
 void UDisplayClusterPreviewComponent::ReleasePreviewRenderTarget()
 {
-	if (RenderTarget != nullptr)
+	ReleaseRenderTargetImpl(&RenderTarget);
+	ReleaseRenderTargetImpl(&RenderTargetPostProcess);
+}
+
+void UDisplayClusterPreviewComponent::UpdatePreviewRenderTarget()
+{
+	if (RootActor)
 	{
-		RenderTarget = nullptr;
+		if (RootActor->ShouldThisFrameOutputPreviewToPostProcessRenderTarget())
+		{
+			UpdateRenderTargetImpl(&RenderTargetPostProcess);
+		}
+		else
+		{
+			UpdateRenderTargetImpl(&RenderTarget);
+		}
 	}
 }
 
-void UDisplayClusterPreviewComponent::UpdatePreviewRenderTarget()
+void UDisplayClusterPreviewComponent::ReleaseRenderTargetImpl(UTextureRenderTarget2D** InOutRenderTarget)
+{
+	checkSlow(InOutRenderTarget);
+	*InOutRenderTarget = nullptr;
+}
+
+void UDisplayClusterPreviewComponent::UpdateRenderTargetImpl(UTextureRenderTarget2D** InOutRenderTarget)
 {
 	DECLARE_SCOPE_CYCLE_COUNTER(TEXT("UDisplayClusterPreviewComponent::UpdatePreviewRenderTarget"), STAT_UpdatePreviewRenderTarget, STATGROUP_NDisplay);
 	
@@ -289,50 +311,53 @@ void UDisplayClusterPreviewComponent::UpdatePreviewRenderTarget()
 	float        TextureGamma = 1.f;
 	bool         bTextureSRGB = false;
 
+	checkSlow(InOutRenderTarget);
+	UTextureRenderTarget2D*& RenderTargetPtr = *InOutRenderTarget;
+	
 	if (GetPreviewTextureSettings(TextureSize, TextureFormat, TextureGamma, bTextureSRGB))
 	{
-		if (RenderTarget != nullptr)
+		if (RenderTargetPtr != nullptr)
 		{
 			// Re-create RTT when format changed
-			if (RenderTarget->GetFormat() != TextureFormat)
+			if (RenderTargetPtr->GetFormat() != TextureFormat)
 			{
-				ReleasePreviewRenderTarget();
+				ReleaseRenderTargetImpl(&RenderTargetPtr);
 			}
 		}
 
-		if (RenderTarget != nullptr)
+		if (RenderTargetPtr != nullptr)
 		{
 			// Update an existing RTT resource only when settings change
-			if (RenderTarget->TargetGamma != TextureGamma
-				|| RenderTarget->SRGB != bTextureSRGB
-				|| RenderTarget->GetSurfaceWidth() != TextureSize.X
-				|| RenderTarget->GetSurfaceHeight() != TextureSize.Y)
+			if (RenderTargetPtr->TargetGamma != TextureGamma
+				|| RenderTargetPtr->SRGB != bTextureSRGB
+				|| RenderTargetPtr->GetSurfaceWidth() != TextureSize.X
+				|| RenderTargetPtr->GetSurfaceHeight() != TextureSize.Y)
 			{
-				RenderTarget->TargetGamma = TextureGamma;
-				RenderTarget->SRGB = bTextureSRGB;
+				RenderTargetPtr->TargetGamma = TextureGamma;
+				RenderTargetPtr->SRGB = bTextureSRGB;
 
-				RenderTarget->ResizeTarget(TextureSize.X, TextureSize.Y);
+				RenderTargetPtr->ResizeTarget(TextureSize.X, TextureSize.Y);
 			}
 		}
 		else
 		{
 			// Create new RTT
-			RenderTarget = NewObject<UTextureRenderTarget2D>(this);
-			RenderTarget->ClearColor = FLinearColor::Black;
-			RenderTarget->TargetGamma = TextureGamma;
-			RenderTarget->SRGB = bTextureSRGB;
+			RenderTargetPtr = NewObject<UTextureRenderTarget2D>(this);
+			RenderTargetPtr->ClearColor = FLinearColor::Black;
+			RenderTargetPtr->TargetGamma = TextureGamma;
+			RenderTargetPtr->SRGB = bTextureSRGB;
 
-			RenderTarget->InitCustomFormat(TextureSize.X, TextureSize.Y, TextureFormat, false);
+			RenderTargetPtr->InitCustomFormat(TextureSize.X, TextureSize.Y, TextureFormat, false);
 			UpdatePreviewMaterial();
 		}
 	}
 	else
 	{
 		//@todo: disable this viewport
-		if (RenderTarget)
+		if (RenderTargetPtr)
 		{
 			// clear preview RTT to black in this case
-			FTextureRenderTarget2DResource* TexResource = (FTextureRenderTarget2DResource*)RenderTarget->GetResource();
+			FTextureRenderTarget2DResource* TexResource = (FTextureRenderTarget2DResource*)RenderTargetPtr->GetResource();
 			if (TexResource)
 			{
 				FCanvas Canvas(TexResource, NULL, FGameTime(), GMaxRHIFeatureLevel);
@@ -340,6 +365,8 @@ void UDisplayClusterPreviewComponent::UpdatePreviewRenderTarget()
 			}
 		}
 	}
+
+	*InOutRenderTarget = RenderTargetPtr;
 }
 
 bool UDisplayClusterPreviewComponent::GetPreviewTextureSettings(FIntPoint& OutSize, EPixelFormat& OutTextureFormat, float& OutGamma, bool& bOutSRGB) const
diff --git a/Engine/Plugins/Runtime/nDisplay/Source/DisplayCluster/Private/Game/EngineClasses/Scene/DisplayClusterRootActorEditor.cpp b/Engine/Plugins/Runtime/nDisplay/Source/DisplayCluster/Private/Game/EngineClasses/Scene/DisplayClusterRootActorEditor.cpp
--- a/Engine/Plugins/Runtime/nDisplay/Source/DisplayCluster/Private/Game/EngineClasses/Scene/DisplayClusterRootActorEditor.cpp
+++ b/Engine/Plugins/Runtime/nDisplay/Source/DisplayCluster/Private/Game/EngineClasses/Scene/DisplayClusterRootActorEditor.cpp
@@ -126,7 +126,23 @@ void ADisplayClusterRootActor::Tick_Editor(float DeltaSeconds)
 			TickPerFrameCounter = 0;
 
 			// Render viewport for preview material RTTs
-			ImplRenderPreview_Editor();
+			if (ViewportManager.IsValid())
+			{
+				// Either render to the preview components' standard RenderTarget or the RenderTargetPostProcess.
+				// This is needed in case preview post process is disabled but this preview should also be displayed
+				// in another editor window that requires post process completed.
+				
+				// TODO: Optimize the render routine to cache the pre PP preview then apply PostProcess.
+				
+				if (!PreviewRenderFrame.IsValid()) // Frame has finished.
+				{
+					bOutputFrameToPostProcessRenderTarget = bOutputFrameToPostProcessRenderTarget ?
+						bPreviewEnablePostProcess :
+						bPreviewEnablePostProcess || DoObserversNeedPostProcessRenderTarget();
+				}
+
+				ImplRenderPreview_Editor();
+			}
 		}
 
 		// preview frustums on each tick
@@ -170,6 +186,11 @@ void ADisplayClusterRootActor::RerunConstructionScripts_Editor()
 	ResetPreviewComponents_Editor(false);
 }
 
+void ADisplayClusterRootActor::EnableEditorRender(bool bValue)
+{
+	bEnableEditorRender = bValue;
+}
+
 bool ADisplayClusterRootActor::IsPreviewEnabled() const
 {
 	//@todo: (GUI) Scene preview can be disabled when the configuration window with internal preview is open.
@@ -199,7 +220,8 @@ void ADisplayClusterRootActor::GetPreviewRenderTargetableTextures(const TArray<F
 			if (OutTextureIndex != INDEX_NONE)
 			{
 				// Add scope for func GetRenderTargetTexture()
-				if (UTextureRenderTarget2D* RenderTarget2D = PreviewComponentIt.Value->GetRenderTargetTexture())
+				if (UTextureRenderTarget2D* RenderTarget2D = ShouldThisFrameOutputPreviewToPostProcessRenderTarget() ?
+					PreviewComponentIt.Value->GetRenderTargetTexturePostProcess() : PreviewComponentIt.Value->GetRenderTargetTexture())
 				{
 					FTextureRenderTargetResource* DstRenderTarget = RenderTarget2D->GameThread_GetRenderTargetResource();
 					if (DstRenderTarget != nullptr)
@@ -305,7 +327,7 @@ bool ADisplayClusterRootActor::ImplUpdatePreviewConfiguration_Editor(const FStri
 		PreviewSettings.PreviewRenderTargetRatioMult = PreviewRenderTargetRatioMult;
 
 		PreviewSettings.bFreezePreviewRender = bFreezePreviewRender;
-		PreviewSettings.bPreviewEnablePostProcess = bPreviewEnablePostProcess;
+		PreviewSettings.bPreviewEnablePostProcess = ShouldThisFrameOutputPreviewToPostProcessRenderTarget();
 
 		PreviewSettings.PreviewMaxTextureDimension = PreviewMaxTextureDimension;
 		
@@ -714,7 +736,6 @@ void ADisplayClusterRootActor::PostEditChangeProperty(FPropertyChangedEvent& Pro
 	// Cluster node ID has been changed
 	if (PropertyName == GET_MEMBER_NAME_CHECKED(ADisplayClusterRootActor, bFreezePreviewRender)
 	|| PropertyName == GET_MEMBER_NAME_CHECKED(ADisplayClusterRootActor, PreviewRenderTargetRatioMult)
-	|| PropertyName == GET_MEMBER_NAME_CHECKED(ADisplayClusterRootActor, bPreviewEnablePostProcess)
 	|| PropertyName == GET_MEMBER_NAME_CHECKED(ADisplayClusterRootActor, bPreviewICVFXFrustums)
 	|| PropertyName == GET_MEMBER_NAME_CHECKED(ADisplayClusterRootActor, PreviewICVFXFrustumsFarDistance)
 	|| PropertyName == GET_MEMBER_NAME_CHECKED(ADisplayClusterRootActor, PreviewMaxTextureDimension)
@@ -734,7 +755,12 @@ void ADisplayClusterRootActor::PostEditChangeProperty(FPropertyChangedEvent& Pro
 	{
 		ResetInnerFrustumPriority();
 	}
-
+	else if (PropertyName == GET_MEMBER_NAME_CHECKED(ADisplayClusterRootActor, bPreviewEnablePostProcess))
+	{
+		ResetPreviewComponents_Editor(false);
+		PreviewRenderFrame.Reset();
+		bReinitializeActor = false;
+	}
 	if (bReinitializeActor)
 	{
 		InitializeRootActor();
@@ -848,12 +874,36 @@ void ADisplayClusterRootActor::ReleasePreviewComponents()
 	OnPreviewDestroyed.ExecuteIfBound();
 }
 
+int32 ADisplayClusterRootActor::SubscribeToPostProcessRenderTarget(const uint8* Object)
+{
+	check(Object);
+	PostProcessRenderTargetObservers.Add(Object);
+	return PostProcessRenderTargetObservers.Num();
+}
+
+int32 ADisplayClusterRootActor::UnsubscribeFromPostProcessRenderTarget(const uint8* Object)
+{
+	check(Object);
+	PostProcessRenderTargetObservers.Remove(Object);
+	return PostProcessRenderTargetObservers.Num();
+}
+
+bool ADisplayClusterRootActor::DoObserversNeedPostProcessRenderTarget() const
+{
+	return PostProcessRenderTargetObservers.Num() > 0;
+}
+
+bool ADisplayClusterRootActor::ShouldThisFrameOutputPreviewToPostProcessRenderTarget() const
+{
+	return bOutputFrameToPostProcessRenderTarget;
+}
+
 FString ADisplayClusterRootActor::GeneratePreviewComponentName_Editor(const FString& NodeId, const FString& ViewportId) const
 {
 	return FString::Printf(TEXT("%s_%s"), *NodeId, *ViewportId);
 }
 
-UDisplayClusterPreviewComponent* ADisplayClusterRootActor::GetPreviewComponent(const FString& NodeId, const FString& ViewportId)
+UDisplayClusterPreviewComponent* ADisplayClusterRootActor::GetPreviewComponent(const FString& NodeId, const FString& ViewportId) const
 {
 	const FString PreviewCompId = GeneratePreviewComponentName_Editor(NodeId, ViewportId);
 	if (PreviewComponents.Contains(PreviewCompId))
diff --git a/Engine/Plugins/Runtime/nDisplay/Source/DisplayCluster/Private/Render/Viewport/DisplayClusterViewportManager_EditorPreview.cpp b/Engine/Plugins/Runtime/nDisplay/Source/DisplayCluster/Private/Render/Viewport/DisplayClusterViewportManager_EditorPreview.cpp
--- a/Engine/Plugins/Runtime/nDisplay/Source/DisplayCluster/Private/Render/Viewport/DisplayClusterViewportManager_EditorPreview.cpp
+++ b/Engine/Plugins/Runtime/nDisplay/Source/DisplayCluster/Private/Render/Viewport/DisplayClusterViewportManager_EditorPreview.cpp
@@ -110,6 +110,13 @@ bool FDisplayClusterViewportManager::RenderInEditor(class FDisplayClusterRenderF
 		return false;
 	}
 
+	const ADisplayClusterRootActor* RootActor = GetRootActor();
+	if (RootActor && !RootActor->IsEditorRenderEnabled())
+	{
+		bOutFrameRendered = true;
+		return true;
+	}
+	
 	FSceneInterface* PreviewScene = CurrentWorld->Scene;
 	FEngineShowFlags EngineShowFlags = FEngineShowFlags(EShowFlagInitMode::ESFIM_Game);
 
@@ -153,6 +160,18 @@ bool FDisplayClusterViewportManager::RenderInEditor(class FDisplayClusterRenderF
 
 				ConfigureViewFamily(RenderTargetIt, ViewFamiliesIt, ViewFamily);
 
+				if (RenderTargetIt.CaptureMode == EDisplayClusterViewportCaptureMode::Default
+					&& RootActor
+					&& !RootActor->bPreviewEnablePostProcess
+					&& RootActor->DoObserversNeedPostProcessRenderTarget())
+				{
+					if (ViewFamily.EngineShowFlags.TemporalAA)
+					{
+						ViewFamily.EngineShowFlags.SetTemporalAA(false);
+						ViewFamily.EngineShowFlags.SetAntiAliasing(true);
+					}
+				}
+				
 				for (FDisplayClusterRenderFrame::FFrameView& ViewIt : ViewFamiliesIt.Views)
 				{
 					bool bViewportAlreadyRendered = ViewportIndex < (int32)InFirstViewportNum;
diff --git a/Engine/Plugins/Runtime/nDisplay/Source/DisplayCluster/Public/Components/DisplayClusterPreviewComponent.h b/Engine/Plugins/Runtime/nDisplay/Source/DisplayCluster/Public/Components/DisplayClusterPreviewComponent.h
--- a/Engine/Plugins/Runtime/nDisplay/Source/DisplayCluster/Public/Components/DisplayClusterPreviewComponent.h
+++ b/Engine/Plugins/Runtime/nDisplay/Source/DisplayCluster/Public/Components/DisplayClusterPreviewComponent.h
@@ -62,6 +62,11 @@ public:
 		return RenderTarget;
 	}
 
+	UTextureRenderTarget2D* GetRenderTargetTexturePostProcess() const
+	{
+		return RenderTargetPostProcess;
+	}
+
 	void UpdatePreviewResources();
 
 	UMeshComponent* GetPreviewMesh()
@@ -88,6 +93,11 @@ protected:
 	void UpdatePreviewRenderTarget();
 	void ReleasePreviewRenderTarget();
 
+private:
+	void UpdateRenderTargetImpl(UTextureRenderTarget2D** InOutRenderTarget);
+	void ReleaseRenderTargetImpl(UTextureRenderTarget2D** InOutRenderTarget);
+
+protected:
 	bool UpdatePreviewMesh();
 	void ReleasePreviewMesh();
 	void UpdatePreviewMeshReference();
@@ -106,6 +116,10 @@ protected:
 	UPROPERTY(VisibleAnywhere, BlueprintReadOnly, Category = "Preview", meta = (DisplayName = "Render Target"))
 	UTextureRenderTarget2D* RenderTarget;
 
+	// Texture when DCRA has post process disabled but is requesting a post process render target.
+	UPROPERTY(VisibleAnywhere, BlueprintReadOnly, Category = "Preview")
+	UTextureRenderTarget2D* RenderTargetPostProcess;
+	
 private:
 	// Saved mesh policy params
 	UPROPERTY()
diff --git a/Engine/Plugins/Runtime/nDisplay/Source/DisplayCluster/Public/DisplayClusterRootActor.h b/Engine/Plugins/Runtime/nDisplay/Source/DisplayCluster/Public/DisplayClusterRootActor.h
--- a/Engine/Plugins/Runtime/nDisplay/Source/DisplayCluster/Public/DisplayClusterRootActor.h
+++ b/Engine/Plugins/Runtime/nDisplay/Source/DisplayCluster/Public/DisplayClusterRootActor.h
@@ -343,6 +343,17 @@ private:
 #endif
 
 #if WITH_EDITOR
+public:
+	/** Enable or disable editor render. Preview components may need to be on for texture overrides, but capture and rendering disabled. */
+	void EnableEditorRender(bool bValue);
+
+	/** If editor rendering is enabled. */
+	bool IsEditorRenderEnabled() const { return bEnableEditorRender; }
+	
+private:
+	/** Is editor rendering enabled? This can be false and the preview still enabled. */
+	bool bEnableEditorRender = true;
+	
 public:
 	DECLARE_DELEGATE(FOnPreviewUpdated);
 
@@ -368,11 +379,35 @@ public:
 	// Preview components free referenced meshes and materials
 	void ResetPreviewComponents_Editor(bool bInRestoreSceneMaterial);
 
-	UDisplayClusterPreviewComponent* GetPreviewComponent(const FString& NodeId, const FString& ViewportId);
+	UDisplayClusterPreviewComponent* GetPreviewComponent(const FString& NodeId, const FString& ViewportId) const;
 
 	void UpdatePreviewComponents();
 	void ReleasePreviewComponents();
 
+	/**
+	 * Enable the use of a post process render target when bPreviewEnablePostProcess is disabled on the actor. The root actor
+	 * will still display a pre post processed preview. This may increase editor overhead.
+	 *
+	 * Retrieve the post process texture from the preview component with GetRenderTargetTexturePostProcess().
+	 *
+	 * @param Object The object subscribing to updates
+	 * @return The number of subscribers to use post process.
+	 */
+	int32 SubscribeToPostProcessRenderTarget(const uint8* Object);
+
+	/**
+	 * Unsubscribe a registered object from requiring post process render target updates.
+	 *
+	 * @param Object The object subscribing to updates. When the counter is zero post process render targets will not be used.
+	 * @return The number of subscribers to use post process.
+	 */
+	int32 UnsubscribeFromPostProcessRenderTarget(const uint8* Object);
+
+	/** If one or more observers are subscribed to receive post process preview targets. */
+	bool DoObserversNeedPostProcessRenderTarget() const;
+	
+	/** When rendering the preview determine which render target should be used for the current frame. */
+	bool ShouldThisFrameOutputPreviewToPostProcessRenderTarget() const;
 
 	float GetPreviewRenderTargetRatioMult() const
 	{
@@ -415,7 +450,12 @@ protected:
 private:
 	bool bIsSelectedInEditor = false;
 	
-private:
+	/** When the preview render should be directed to the post process render target. */
+	bool bOutputFrameToPostProcessRenderTarget;
+
+	/** Enables preview components to output to the post process render target when bPreviewEnablePostProcess is disabled. Contains all subscribed objects. */
+	TSet<const uint8*> PostProcessRenderTargetObservers;
+	
 	TWeakPtr<IDisplayClusterConfiguratorBlueprintEditor> ToolkitPtr;
 
 	int32 TickPerFrameCounter = 0;
diff --git a/Engine/Plugins/Runtime/nDisplay/Source/DisplayClusterLightCardEditor/Private/Viewport/DisplayClusterLightCardEditorViewportClient.cpp b/Engine/Plugins/Runtime/nDisplay/Source/DisplayClusterLightCardEditor/Private/Viewport/DisplayClusterLightCardEditorViewportClient.cpp
--- a/Engine/Plugins/Runtime/nDisplay/Source/DisplayClusterLightCardEditor/Private/Viewport/DisplayClusterLightCardEditorViewportClient.cpp
+++ b/Engine/Plugins/Runtime/nDisplay/Source/DisplayClusterLightCardEditor/Private/Viewport/DisplayClusterLightCardEditorViewportClient.cpp
@@ -93,6 +93,10 @@ FDisplayClusterLightCardEditorViewportClient::FDisplayClusterLightCardEditorView
 FDisplayClusterLightCardEditorViewportClient::~FDisplayClusterLightCardEditorViewportClient()
 {
 	EndTransaction();
+	if (RootActorLevelInstance.IsValid())
+	{
+		RootActorLevelInstance->UnsubscribeFromPostProcessRenderTarget(reinterpret_cast<uint8*>(this));
+	}
 }
 
 FLinearColor FDisplayClusterLightCardEditorViewportClient::GetBackgroundColor() const
@@ -150,7 +154,7 @@ void FDisplayClusterLightCardEditorViewportClient::Tick(float DeltaSeconds)
 
 				if (PreviewComp && LevelInstancePreviewComp)
 				{
-					PreviewComp->SetOverrideTexture(LevelInstancePreviewComp->GetRenderTargetTexture());
+					PreviewComp->SetOverrideTexture(LevelInstancePreviewComp->GetRenderTargetTexturePostProcess());
 				}
 			}
 		}
@@ -590,6 +594,7 @@ void FDisplayClusterLightCardEditorViewportClient::UpdatePreviewActor(ADisplayCl
 		PreviewWorld->GetTimerManager().SetTimerForNextTick([=]()
 		{
 			DestroyProxies(ProxyType);
+			RootActor->SubscribeToPostProcessRenderTarget(reinterpret_cast<uint8*>(this));
 			RootActorLevelInstance = RootActor;
 			
 			if (ProxyType == EDisplayClusterLightCardEditorProxyType::All ||
@@ -612,6 +617,7 @@ void FDisplayClusterLightCardEditorViewportClient::UpdatePreviewActor(ADisplayCl
 				FindProjectionOriginComponent();
 
 				RootActorProxy->UpdatePreviewComponents();
+				RootActorProxy->EnableEditorRender(false);
 			}
 
 			// Filter out any primitives hidden in game except screen components
@@ -690,8 +696,12 @@ EDisplayClusterLightCardEditorProxyType ProxyType)
 			PreviewWorld->EditorDestroyActor(RootActorProxy.Get(), false);
 			RootActorProxy.Reset();
 		}
-	
-		RootActorLevelInstance.Reset();
+
+		if (RootActorLevelInstance.IsValid())
+		{
+			RootActorLevelInstance->UnsubscribeFromPostProcessRenderTarget(reinterpret_cast<uint8*>(this));
+			RootActorLevelInstance.Reset();
+		}
 	}
 	
 	if (ProxyType == EDisplayClusterLightCardEditorProxyType::All ||
