nDisplay: Removed invisible characters.

diff --git a/Engine/Plugins/Runtime/nDisplay/Source/DisplayCluster/Private/Game/EngineClasses/Scene/DisplayClusterRootActorEditor.cpp b/Engine/Plugins/Runtime/nDisplay/Source/DisplayCluster/Private/Game/EngineClasses/Scene/DisplayClusterRootActorEditor.cpp
--- a/Engine/Plugins/Runtime/nDisplay/Source/DisplayCluster/Private/Game/EngineClasses/Scene/DisplayClusterRootActorEditor.cpp
+++ b/Engine/Plugins/Runtime/nDisplay/Source/DisplayCluster/Private/Game/EngineClasses/Scene/DisplayClusterRootActorEditor.cpp
@@ -550,8 +550,8 @@ void ADisplayClusterRootActor::ImplRenderPreviewFrustums_Editor()
 	{
 		const TArray<FDisplayClusterViewport_Context>& Contexts = ViewportIt->GetContexts();
 
-		// Due to rendering optimizations for DCRA preview (node ​​rendered over multiple frames). 
-		// As a result, viewport math values ​​(stored in contexts) may not be ready at the moment.
+		// Due to rendering optimizations for DCRA preview (node rendered over multiple frames).
+		// As a result, viewport math values (stored in contexts) may not be ready at the moment.
 		FFrustumPreviewViewportContextCache FrustumPreviewViewportContext;
 		bool bIsValidViewportContext = false;
 
