nDisplay: Fix lightcard proxies being deselected when a property is changed.

diff --git a/Engine/Plugins/Runtime/nDisplay/Source/DisplayClusterLightCardEditor/Private/Viewport/DisplayClusterLightCardEditorViewportClient.cpp b/Engine/Plugins/Runtime/nDisplay/Source/DisplayClusterLightCardEditor/Private/Viewport/DisplayClusterLightCardEditorViewportClient.cpp
--- a/Engine/Plugins/Runtime/nDisplay/Source/DisplayClusterLightCardEditor/Private/Viewport/DisplayClusterLightCardEditorViewportClient.cpp
+++ b/Engine/Plugins/Runtime/nDisplay/Source/DisplayClusterLightCardEditor/Private/Viewport/DisplayClusterLightCardEditorViewportClient.cpp
@@ -593,6 +593,21 @@ void FDisplayClusterLightCardEditorViewportClient::UpdatePreviewActor(ADisplayCl
 		// regeneration & compiles. nDisplay's custom propagation may have issues if the archetype isn't correct.
 		PreviewWorld->GetTimerManager().SetTimerForNextTick([=]()
 		{
+			TSet<AActor*> LastSelectedLightCardLevelInstances;
+			for (TWeakObjectPtr<ADisplayClusterLightCardActor>& SelectedLightCard : SelectedLightCards)
+			{
+				if (SelectedLightCard.IsValid())
+				{
+					if (FLightCardProxy* FoundProxy = LightCardProxies.FindByKey(SelectedLightCard.Get()))
+					{
+						if (FoundProxy->LevelInstance.IsValid())
+						{
+							LastSelectedLightCardLevelInstances.Add(FoundProxy->LevelInstance.Get());
+						}
+					}
+				}
+			}
+			
 			DestroyProxies(ProxyType);
 			RootActor->SubscribeToPostProcessRenderTarget(reinterpret_cast<uint8*>(this));
 			RootActorLevelInstance = RootActor;
@@ -631,7 +646,9 @@ void FDisplayClusterLightCardEditorViewportClient::UpdatePreviewActor(ADisplayCl
 			{
 				TArray<TWeakObjectPtr<ADisplayClusterLightCardActor>> LightCards;
 				FindLightCardsForRootActor(RootActor, LightCards);
-			
+
+				SelectLightCard(nullptr);
+				
 				for (const TWeakObjectPtr<ADisplayClusterLightCardActor>& LightCard : LightCards)
 				{
 					FObjectDuplicationParameters DupeActorParameters(LightCard.Get(), PreviewWorld->GetCurrentLevel());
@@ -645,6 +662,11 @@ void FDisplayClusterLightCardEditorViewportClient::UpdatePreviewActor(ADisplayCl
 					LightCardProxy->SetActorRotation(LightCard->GetActorRotation() - RootActor->GetActorRotation());
 
 					LightCardProxies.Add(FLightCardProxy(LightCard.Get(), LightCardProxy));
+					
+					if (LastSelectedLightCardLevelInstances.Contains(LightCard.Get()))
+					{
+						SelectLightCard(LightCardProxy, true);
+					}
 				}
 			}
 
