nDisplay: Fix LightCard tab spawner registering its own menu item. Prevent LightCard editor tab from being closed as part of the operator panel.

diff --git a/Engine/Plugins/Runtime/nDisplay/Source/DisplayClusterLightCardEditor/Private/SDisplayClusterLightCardEditor.cpp b/Engine/Plugins/Runtime/nDisplay/Source/DisplayClusterLightCardEditor/Private/SDisplayClusterLightCardEditor.cpp
--- a/Engine/Plugins/Runtime/nDisplay/Source/DisplayClusterLightCardEditor/Private/SDisplayClusterLightCardEditor.cpp
+++ b/Engine/Plugins/Runtime/nDisplay/Source/DisplayClusterLightCardEditor/Private/SDisplayClusterLightCardEditor.cpp
@@ -28,8 +28,9 @@ void SDisplayClusterLightCardEditor::RegisterTabSpawner()
 
 	FGlobalTabmanager::Get()->RegisterNomadTabSpawner(TabName, FOnSpawnTab::CreateStatic(&SDisplayClusterLightCardEditor::SpawnInTab))
 		.SetDisplayName(LOCTEXT("TabDisplayName", "Light Cards Editor"))
-		.SetTooltipText(LOCTEXT("TabTooltip", "Editing tools for nDisplay light cards."));
-
+		.SetTooltipText(LOCTEXT("TabTooltip", "Editing tools for nDisplay light cards."))
+		.SetMenuType(ETabSpawnerMenuType::Hidden);
+	
 	TSharedPtr<FExtender> ToolbarExtender = MakeShareable(new FExtender);
 	ToolbarExtender->AddToolBarExtension("General", EExtensionHook::After, nullptr, FToolBarExtensionDelegate::CreateStatic(&SDisplayClusterLightCardEditor::ExtendToolbar));
 	IDisplayClusterOperator::Get().GetOperatorToolBarExtensibilityManager()->AddExtender(ToolbarExtender);
@@ -49,8 +50,10 @@ void SDisplayClusterLightCardEditor::RegisterLayoutExtension(FLayoutExtender& In
 TSharedRef<SDockTab> SDisplayClusterLightCardEditor::SpawnInTab(const FSpawnTabArgs& SpawnTabArgs)
 {
 	const TSharedRef<SDockTab> MajorTab = SNew(SDockTab)
-		.TabRole(ETabRole::NomadTab);
-
+		.TabRole(ETabRole::NomadTab)
+		// Prevent close until we can add a menu item in the operator panel to spawn this tab.
+		.OnCanCloseTab(SDockTab::FCanCloseTab::CreateLambda([]() -> bool { return false; }));
+	
 	MajorTab->SetContent(SNew(SDisplayClusterLightCardEditor, MajorTab, SpawnTabArgs.GetOwnerWindow()));
 
 	return MajorTab;
