nDisplay: added support of r.setnearclipplane cvar

diff --git a/Engine/Plugins/Runtime/nDisplay/Source/DisplayCluster/Private/Render/Device/DisplayClusterDeviceBase.cpp b/Engine/Plugins/Runtime/nDisplay/Source/DisplayCluster/Private/Render/Device/DisplayClusterDeviceBase.cpp
--- a/Engine/Plugins/Runtime/nDisplay/Source/DisplayCluster/Private/Render/Device/DisplayClusterDeviceBase.cpp
+++ b/Engine/Plugins/Runtime/nDisplay/Source/DisplayCluster/Private/Render/Device/DisplayClusterDeviceBase.cpp
@@ -302,7 +302,7 @@ void FDisplayClusterDeviceBase::CalculateStereoViewOffset(const int32 ViewIndex,
 	// Get the actual camera settings
 	const float CfgEyeDist = ViewCamera->GetInterpupillaryDistance();
 	const bool  CfgEyeSwap = ViewCamera->GetSwapEyes();
-	const float CfgNCP     = 1.f;
+	const float CfgNCP     = GNearClippingPlane;
 	const EDisplayClusterEyeStereoOffset CfgEyeOffset = ViewCamera->GetStereoOffset();
 
 	// Calculate eye offset considering the world scale
diff --git a/Engine/Plugins/Runtime/nDisplay/Source/DisplayCluster/Private/Render/Viewport/DisplayClusterViewport_EditorPreview.cpp b/Engine/Plugins/Runtime/nDisplay/Source/DisplayCluster/Private/Render/Viewport/DisplayClusterViewport_EditorPreview.cpp
--- a/Engine/Plugins/Runtime/nDisplay/Source/DisplayCluster/Private/Render/Viewport/DisplayClusterViewport_EditorPreview.cpp
+++ b/Engine/Plugins/Runtime/nDisplay/Source/DisplayCluster/Private/Render/Viewport/DisplayClusterViewport_EditorPreview.cpp
@@ -237,7 +237,7 @@ bool FDisplayClusterViewport::ImplPreview_CalculateStereoViewOffset(const uint32
 	// Get the actual camera settings
 	const float CfgEyeDist  = ViewCamera ? ViewCamera->GetInterpupillaryDistance() : 6.4f;
 	const bool  bCfgEyeSwap = ViewCamera ? ViewCamera->GetSwapEyes() : false;
-	const float CfgNCP = 1.f;
+	const float CfgNCP = GNearClippingPlane;
 
 	const EDisplayClusterEyeStereoOffset CfgEyeOffset = ViewCamera ? ViewCamera->GetStereoOffset() : EDisplayClusterEyeStereoOffset::None;
 
