nDisplay: added near plane value into projection clip computation

diff --git a/Engine/Plugins/Runtime/nDisplay/Source/DisplayClusterShaders/Private/WarpBlend/Math/DisplayClusterWarpBlendMath_Frustum.h b/Engine/Plugins/Runtime/nDisplay/Source/DisplayClusterShaders/Private/WarpBlend/Math/DisplayClusterWarpBlendMath_Frustum.h
--- a/Engine/Plugins/Runtime/nDisplay/Source/DisplayClusterShaders/Private/WarpBlend/Math/DisplayClusterWarpBlendMath_Frustum.h
+++ b/Engine/Plugins/Runtime/nDisplay/Source/DisplayClusterShaders/Private/WarpBlend/Math/DisplayClusterWarpBlendMath_Frustum.h
@@ -317,7 +317,7 @@ private:
 				return false;
 			}
 
-			const float Scale = 1.0f / ProjectedVertice.X;
+			const float Scale = Eye.ZNear / ProjectedVertice.X;
 
 			ProjectedVertice.Y *= Scale;
 			ProjectedVertice.Z *= Scale;
