nDisplay: fixed frustum calculations for curved surfaces by adding weight based on triangle area

diff --git a/Engine/Plugins/Runtime/nDisplay/Source/DisplayClusterShaders/Private/WarpBlend/Math/DisplayClusterWarpBlendMath_WarpMesh.h b/Engine/Plugins/Runtime/nDisplay/Source/DisplayClusterShaders/Private/WarpBlend/Math/DisplayClusterWarpBlendMath_WarpMesh.h
--- a/Engine/Plugins/Runtime/nDisplay/Source/DisplayClusterShaders/Private/WarpBlend/Math/DisplayClusterWarpBlendMath_WarpMesh.h
+++ b/Engine/Plugins/Runtime/nDisplay/Source/DisplayClusterShaders/Private/WarpBlend/Math/DisplayClusterWarpBlendMath_WarpMesh.h
@@ -47,9 +47,9 @@ public:
 			return;
 		}
 
-		double Nxyz[3] = { 0, 0, 0 };
-
+		FVector SurfaceNormalViewSum = FVector::ZeroVector;
 		const FPositionVertexBuffer& PositionBuffer = MeshLODResources.VertexBuffers.PositionVertexBuffer;
+
 		for (int32 TriIdx = 0; TriIdx < TriNum; ++TriIdx)
 		{
 			const int32 Index0 = MeshLODResources.IndexBuffer.GetIndex(TriIdx * 3 + 0);
@@ -62,23 +62,17 @@ public:
 
 			const FVector N1 = Pts1 - Pts0;
 			const FVector N2 = Pts2 - Pts0;
-			const FVector N = FVector::CrossProduct(N2, N1).GetSafeNormal();
 
-			for (int32 j = 0; j < 3; j++)
-			{
-				Nxyz[j] += N[j];
-			}
-		}
+			const FVector Normal = FVector::CrossProduct(N2, N1);
+			const double Magnitude = Normal.Length();
 
-		double Scale = double(1) / TriNum;
-		for (int32 AxisIndex = 0; AxisIndex < 3; AxisIndex++)
-		{
-			Nxyz[AxisIndex] *= Scale;
+			SurfaceNormalViewSum += (Normal * Magnitude);
 		}
 
-		OutSurfaceViewNormal = FVector(Nxyz[0], Nxyz[1], Nxyz[2]).GetSafeNormal();
+		OutSurfaceViewNormal = SurfaceNormalViewSum.GetSafeNormal();
 
 		//@todo: MeshSurfaceViewPlane not implemented, use MeshSurfaceViewNormal
+		// this is for the case when normals are set in the side direction to have same sampling quality over the multiple heavily curved surfaces
 		OutSurfaceViewPlane = OutSurfaceViewNormal;
 	}
 
