nDisplay: Dragging a LightCard in the LightCard editor through the widget will send proper snapshots to MU and update its location while dragging.

diff --git a/Engine/Plugins/Runtime/nDisplay/Source/DisplayClusterLightCardEditor/Private/Viewport/DisplayClusterLightCardEditorViewportClient.cpp b/Engine/Plugins/Runtime/nDisplay/Source/DisplayClusterLightCardEditor/Private/Viewport/DisplayClusterLightCardEditorViewportClient.cpp
--- a/Engine/Plugins/Runtime/nDisplay/Source/DisplayClusterLightCardEditor/Private/Viewport/DisplayClusterLightCardEditorViewportClient.cpp
+++ b/Engine/Plugins/Runtime/nDisplay/Source/DisplayClusterLightCardEditor/Private/Viewport/DisplayClusterLightCardEditorViewportClient.cpp
@@ -58,7 +58,6 @@ FDisplayClusterLightCardEditorViewportClient::FDisplayClusterLightCardEditorView
 
 	bDraggingActor = false;
 	DragWidgetOffset = FVector::ZeroVector;
-	ScopedTransaction = nullptr;
 	
 	EditorWidget = MakeShared<FDisplayClusterLightCardEditorWidget>();
 
@@ -806,19 +805,12 @@ void FDisplayClusterLightCardEditorViewportClient::SetProjectionModeFOV(EDisplay
 
 void FDisplayClusterLightCardEditorViewportClient::BeginTransaction(const FText& Description)
 {
-	if (!ScopedTransaction)
-	{
-		ScopedTransaction = new FScopedTransaction(Description);
-	}
+	GEditor->BeginTransaction(Description);
 }
 
 void FDisplayClusterLightCardEditorViewportClient::EndTransaction()
 {
-	if (ScopedTransaction)
-	{
-		delete ScopedTransaction;
-		ScopedTransaction = nullptr;
-	}
+	GEditor->EndTransaction();
 }
 
 void FDisplayClusterLightCardEditorViewportClient::GetScenePrimitiveComponents(TArray<UPrimitiveComponent*>& OutPrimitiveComponents)
@@ -1088,22 +1080,47 @@ void FDisplayClusterLightCardEditorViewportClient::PropagateLightCardSelection()
 
 void FDisplayClusterLightCardEditorViewportClient::PropagateLightCardTransform(ADisplayClusterLightCardActor* LightCardProxy)
 {
-	FLightCardProxy* FoundProxy = LightCardProxies.FindByKey(LightCardProxy);
+	const FLightCardProxy* FoundProxy = LightCardProxies.FindByKey(LightCardProxy);
 	if (FoundProxy && FoundProxy->Proxy == LightCardProxy && FoundProxy->LevelInstance.IsValid())
 	{
 		ADisplayClusterLightCardActor* LevelInstance = FoundProxy->LevelInstance.Get();
 
 		LevelInstance->Modify();
 
-		LevelInstance->Longitude = LightCardProxy->Longitude;
-		LevelInstance->Latitude = LightCardProxy->Latitude;
-		LevelInstance->DistanceFromCenter = LightCardProxy->DistanceFromCenter;
-		LevelInstance->Spin = LightCardProxy->Spin;
-		LevelInstance->Pitch = LightCardProxy->Pitch;
-		LevelInstance->Yaw = LightCardProxy->Yaw;
+		TArray<const FProperty*> ChangedProperties;
+		
+		// Set the level instance property value to our proxy property value.
+		auto TryChangeProperty = [&](FName InPropertyName) -> void
+		{
+			const FProperty* Property = FindFProperty<FProperty>(LevelInstance->GetClass(), InPropertyName);
+			check(Property);
+			
+			// Only change if values are different.
+			if (!Property->Identical_InContainer(LightCardProxy, LevelInstance))
+			{
+				void* NewValue = nullptr;
+				Property->GetValue_InContainer(LightCardProxy, &NewValue);
+				Property->SetValue_InContainer(LevelInstance, &NewValue);
 
-		FVector RootActorLevelInstanceLocation = RootActorLevelInstance.IsValid() ? RootActorLevelInstance->GetActorLocation() : FVector::ZeroVector;
+				ChangedProperties.Add(Property);
+			}
+		};
+		
+		TryChangeProperty(GET_MEMBER_NAME_CHECKED(ADisplayClusterLightCardActor, Longitude));
+		TryChangeProperty(GET_MEMBER_NAME_CHECKED(ADisplayClusterLightCardActor, Latitude));
+		TryChangeProperty(GET_MEMBER_NAME_CHECKED(ADisplayClusterLightCardActor, DistanceFromCenter));
+		TryChangeProperty(GET_MEMBER_NAME_CHECKED(ADisplayClusterLightCardActor, Spin));
+		TryChangeProperty(GET_MEMBER_NAME_CHECKED(ADisplayClusterLightCardActor, Pitch));
+		TryChangeProperty(GET_MEMBER_NAME_CHECKED(ADisplayClusterLightCardActor, Yaw));
+		
+		const FVector RootActorLevelInstanceLocation = RootActorLevelInstance.IsValid() ? RootActorLevelInstance->GetActorLocation() : FVector::ZeroVector;
 		LevelInstance->SetActorLocation(RootActorLevelInstanceLocation + LightCardProxy->GetActorLocation());
+
+		// Snapshot the changed properties so multi-user can update while dragging.
+		if (ChangedProperties.Num() > 0)
+		{
+			SnapshotTransactionBuffer(LevelInstance, MakeArrayView(ChangedProperties));
+		}
 	}
 }
 
diff --git a/Engine/Plugins/Runtime/nDisplay/Source/DisplayClusterLightCardEditor/Private/Viewport/DisplayClusterLightCardEditorViewportClient.h b/Engine/Plugins/Runtime/nDisplay/Source/DisplayClusterLightCardEditor/Private/Viewport/DisplayClusterLightCardEditorViewportClient.h
--- a/Engine/Plugins/Runtime/nDisplay/Source/DisplayClusterLightCardEditor/Private/Viewport/DisplayClusterLightCardEditorViewportClient.h
+++ b/Engine/Plugins/Runtime/nDisplay/Source/DisplayClusterLightCardEditor/Private/Viewport/DisplayClusterLightCardEditorViewportClient.h
@@ -180,9 +180,6 @@ private:
 	
 	/** The renderer for the viewport, which can render the meshes with a variety of projection types */
 	TSharedPtr<FDisplayClusterMeshProjectionRenderer> MeshProjectionRenderer;
-
-	/** The current transaction for undo/redo */
-	FScopedTransaction* ScopedTransaction = nullptr;
 	
 	/** Indicates that the user is dragging an actor in the viewport */
 	bool bDraggingActor = false;
