nDisplay: Unlocked azimuthal equidistant view rotation.

* Its shader now tries to hide triangles with vertices near the south pole that w/o a geometry shader end up in front of the normal view.
* Made preview material translucent instead of opaque so that alpha works on them in order to support the above item.
* Azimuthal is now the default projection.
* Proxy light cards are now forced to have low opacity because with full opacity they are distracting to see when they don't line up well with their projection on the walls. With this you can still see them when they are in outer space (i.e. they are not projected onto any wall).

diff --git a/Engine/Plugins/Runtime/nDisplay/Shaders/Private/MeshProjectionShaders.usf b/Engine/Plugins/Runtime/nDisplay/Shaders/Private/MeshProjectionShaders.usf
--- a/Engine/Plugins/Runtime/nDisplay/Shaders/Private/MeshProjectionShaders.usf
+++ b/Engine/Plugins/Runtime/nDisplay/Shaders/Private/MeshProjectionShaders.usf
@@ -10,33 +10,39 @@ struct FMeshProjectionVSToPS
 	FVertexFactoryInterpolantsVSToPS FactoryInterpolants;
 	float4 Position : SV_POSITION;
 	float4 SavedWorldPosition : POSITION1;
+	float Azimuth : CUSTOM_AZIMUTH;
 };
 
 void AzimuthalVS(FVertexFactoryInput Input, out FMeshProjectionVSToPS Output)
 {
 	ResolvedView = ResolveView();
 
-	FVertexFactoryIntermediates VFIntermediates = GetVertexFactoryIntermediates(Input);
+	const FVertexFactoryIntermediates VFIntermediates = GetVertexFactoryIntermediates(Input);
 	Output.SavedWorldPosition = VertexFactoryGetWorldPosition(Input, VFIntermediates);
-	float3x3 TangentToLocal = VertexFactoryGetTangentToLocal(Input, VFIntermediates);
+	const float3x3 TangentToLocal = VertexFactoryGetTangentToLocal(Input, VFIntermediates);
 
-	FMaterialVertexParameters VertexParameters = GetMaterialVertexParameters(Input, VFIntermediates, Output.SavedWorldPosition.xyz, TangentToLocal);
+	const FMaterialVertexParameters VertexParameters = GetMaterialVertexParameters(Input, VFIntermediates, Output.SavedWorldPosition.xyz, TangentToLocal);
 	Output.SavedWorldPosition.xyz += GetMaterialWorldPositionOffset(VertexParameters);
 
 	Output.FactoryInterpolants = VertexFactoryGetInterpolantsVSToPS(Input, VFIntermediates, VertexParameters);
 
 	// Project the vertex view position onto a unit sphere
-	float4 ViewPosition = mul(Output.SavedWorldPosition, ResolvedView.TranslatedWorldToView);
-	float rho = length(ViewPosition.xyz);
-	float3 UnitPosition = ViewPosition.xyz / rho;
+	const float4 ViewPosition = mul(Output.SavedWorldPosition, ResolvedView.TranslatedWorldToView);
+	const float Rho = length(ViewPosition.xyz);
+	const float3 UnitPosition = ViewPosition.xyz / Rho;
 	
 	// Compute the position in view space of the point on the projection plane using azimthual equidistance projection
-	float2 PolarCoords = float2(acos(UnitPosition.z), atan2(UnitPosition.y, UnitPosition.x));
-	float3 PlanePosition = float3(PolarCoords.x * cos(PolarCoords.y), PolarCoords.x * sin(PolarCoords.y), 1);
 	
+	const float Azimuth = acos(UnitPosition.z); // 0 <= azimuth <= pi
+	const float Theta = atan2(UnitPosition.y, UnitPosition.x); // -pi <= theta <= pi
+		
+	const float2 PolarCoords = float2(Azimuth, Theta);
+	const float3 PlanePosition = float3(PolarCoords.x * cos(PolarCoords.y), PolarCoords.x * sin(PolarCoords.y), 1);
+		
 	// Project the plane position into view space
-	float4 ProjectedPosition = float4(PlanePosition * rho / length(PlanePosition), 1);
+	const float4 ProjectedPosition = float4(PlanePosition * Rho / length(PlanePosition), 1);
 	
+	Output.Azimuth = Azimuth;
 	Output.Position = mul(ProjectedPosition, ResolvedView.ViewToClip);
 }
 
@@ -44,15 +50,16 @@ void PerspectiveVS(FVertexFactoryInput Input, out FMeshProjectionVSToPS Output)
 {
 	ResolvedView = ResolveView();
 
-	FVertexFactoryIntermediates VFIntermediates = GetVertexFactoryIntermediates(Input);
+	const FVertexFactoryIntermediates VFIntermediates = GetVertexFactoryIntermediates(Input);
 	Output.SavedWorldPosition = VertexFactoryGetWorldPosition(Input, VFIntermediates);
-	float3x3 TangentToLocal = VertexFactoryGetTangentToLocal(Input, VFIntermediates);
+	const float3x3 TangentToLocal = VertexFactoryGetTangentToLocal(Input, VFIntermediates);
 
-	FMaterialVertexParameters VertexParameters = GetMaterialVertexParameters(Input, VFIntermediates, Output.SavedWorldPosition.xyz, TangentToLocal);
+	const FMaterialVertexParameters VertexParameters = GetMaterialVertexParameters(Input, VFIntermediates, Output.SavedWorldPosition.xyz, TangentToLocal);
 	Output.SavedWorldPosition.xyz += GetMaterialWorldPositionOffset(VertexParameters);
 
 	Output.FactoryInterpolants = VertexFactoryGetInterpolantsVSToPS(Input, VFIntermediates, VertexParameters);
 	
+	Output.Azimuth = 0;
 	Output.Position = mul(Output.SavedWorldPosition, ResolvedView.TranslatedWorldToClip);
 }
 
@@ -64,11 +71,17 @@ void MainPS(FMeshProjectionVSToPS Input OPTIONAL_IsFrontFace, out float4 OutputC
 
 	GetMaterialCoverageAndClipping(MaterialParameters, PixelMaterialInputs);
 	
-	float3 Emissive = GetMaterialEmissive(PixelMaterialInputs);
-	float Opacity = GetMaterialOpacity(PixelMaterialInputs);
-	
+	const float3 Emissive = GetMaterialEmissive(PixelMaterialInputs);
+	const float Opacity = GetMaterialOpacity(PixelMaterialInputs);
+		
 	OutputColor = float4(Emissive, Opacity);
 
-	float gamma = 2.2;
-	OutputColor.rgb = pow(OutputColor.rgb, gamma);
+	static const float Gamma = 2.2;
+	OutputColor.rgb = pow(OutputColor.rgb, Gamma);
+	
+	static const float AzimuthFeatherStart = radians(180 - 30);
+	static const float AzimuthFeatherEnd   = radians(180 - 15);
+	
+	const float AlphaBlendWeight = saturate((Input.Azimuth - AzimuthFeatherStart) / (AzimuthFeatherEnd - AzimuthFeatherStart));
+	OutputColor.a = lerp(OutputColor.a, 0, AlphaBlendWeight);
 }
\ No newline at end of file
diff --git a/Engine/Plugins/Runtime/nDisplay/Source/DisplayCluster/Private/Game/EngineClasses/Scene/DisplayClusterLightCardActor.cpp b/Engine/Plugins/Runtime/nDisplay/Source/DisplayCluster/Private/Game/EngineClasses/Scene/DisplayClusterLightCardActor.cpp
--- a/Engine/Plugins/Runtime/nDisplay/Source/DisplayCluster/Private/Game/EngineClasses/Scene/DisplayClusterLightCardActor.cpp
+++ b/Engine/Plugins/Runtime/nDisplay/Source/DisplayCluster/Private/Game/EngineClasses/Scene/DisplayClusterLightCardActor.cpp
@@ -155,6 +155,9 @@ void ADisplayClusterLightCardActor::UpdateLightCardMaterialInstance()
 {
 	if (UMaterialInstanceDynamic* LightCardMaterialInstance = Cast<UMaterialInstanceDynamic>(LightCardComponent->GetMaterial(0)))
 	{
+		// Showing proxy with low opacity to make it less distracting when it doesn't line up well with its projection in the Light Card Editor.
+		constexpr float ProxyOpacity = 0.25;
+
 		LightCardMaterialInstance->ClearParameterValues();
 
 		LightCardMaterialInstance->SetTextureParameterValue(TEXT("Texture"), Texture);
@@ -163,7 +166,7 @@ void ADisplayClusterLightCardActor::UpdateLightCardMaterialInstance()
 		LightCardMaterialInstance->SetVectorParameterValue(TEXT("CardColor"), Color);
 		LightCardMaterialInstance->SetScalarParameterValue(TEXT("Exposure"), Exposure);
 		LightCardMaterialInstance->SetScalarParameterValue(TEXT("EmissiveStrength"), Exposure);
-		LightCardMaterialInstance->SetScalarParameterValue(TEXT("Opacity"), Opacity);
+		LightCardMaterialInstance->SetScalarParameterValue(TEXT("Opacity"), bIsProxy ? ProxyOpacity : Opacity);
 		LightCardMaterialInstance->SetScalarParameterValue(TEXT("Feather"), Feathering);
 	}
 }
\ No newline at end of file
diff --git a/Engine/Plugins/Runtime/nDisplay/Source/DisplayCluster/Public/DisplayClusterLightCardActor.h b/Engine/Plugins/Runtime/nDisplay/Source/DisplayCluster/Public/DisplayClusterLightCardActor.h
--- a/Engine/Plugins/Runtime/nDisplay/Source/DisplayCluster/Public/DisplayClusterLightCardActor.h
+++ b/Engine/Plugins/Runtime/nDisplay/Source/DisplayCluster/Public/DisplayClusterLightCardActor.h
@@ -85,6 +85,10 @@ public:
 	UPROPERTY(EditAnywhere, BlueprintReadWrite, Category = "Appearance", meta = (UIMin = 0, ClampMin = 0))
 	float Feathering;
 
+	/** Used to flag this light card as a proxy of a "real" light card. Used by the LightCard Editor */
+	UPROPERTY(Transient)
+	bool bIsProxy = false;
+
 protected:
 	UPROPERTY(EditDefaultsOnly, BlueprintReadOnly, Category = "Default")
 	TObjectPtr<USceneComponent> DefaultSceneRootComponent;
diff --git a/Engine/Plugins/Runtime/nDisplay/Source/DisplayClusterLightCardEditor/Private/Viewport/DisplayClusterLightCardEditorViewportClient.cpp b/Engine/Plugins/Runtime/nDisplay/Source/DisplayClusterLightCardEditor/Private/Viewport/DisplayClusterLightCardEditorViewportClient.cpp
--- a/Engine/Plugins/Runtime/nDisplay/Source/DisplayClusterLightCardEditor/Private/Viewport/DisplayClusterLightCardEditorViewportClient.cpp
+++ b/Engine/Plugins/Runtime/nDisplay/Source/DisplayClusterLightCardEditor/Private/Viewport/DisplayClusterLightCardEditorViewportClient.cpp
@@ -87,6 +87,8 @@ FDisplayClusterLightCardEditorViewportClient::FDisplayClusterLightCardEditorView
 	SetIsSimulateInEditorViewport(true);
 	
 	UpdatePreviewActor(LightCardEditorPtr.Pin()->GetActiveRootActor().Get());
+
+	SetProjectionMode(EDisplayClusterMeshProjectionType::Azimuthal);
 }
 
 FDisplayClusterLightCardEditorViewportClient::~FDisplayClusterLightCardEditorViewportClient()
@@ -116,12 +118,6 @@ void FDisplayClusterLightCardEditorViewportClient::Tick(float DeltaSeconds)
 
 	SetViewLocation(Location);
 
-	// View rotation is also locked for the azimuthal projection
-	if (ProjectionMode != EDisplayClusterMeshProjectionType::Perspective)
-	{
-		SetViewRotation(FVector::UpVector.Rotation());
-	}
-
 	CachedEditorWidgetTransform = CalcEditorWidgetTransform();
 
 	// Tick the preview scene world.
@@ -659,6 +655,7 @@ void FDisplayClusterLightCardEditorViewportClient::UpdatePreviewActor(ADisplayCl
 				
 					LightCardProxy->SetActorLocation(LightCard->GetActorLocation() - RootActor->GetActorLocation());
 					LightCardProxy->SetActorRotation(LightCard->GetActorRotation() - RootActor->GetActorRotation());
+					LightCardProxy->bIsProxy = true;
 
 					LightCardProxies.Add(FLightCardProxy(LightCard.Get(), LightCardProxy));
 					
@@ -773,7 +770,11 @@ void FDisplayClusterLightCardEditorViewportClient::SetProjectionMode(EDisplayClu
 
 	FindProjectionOriginComponent();
 
-	Viewport->InvalidateHitProxy();
+	if (Viewport)
+	{
+		Viewport->InvalidateHitProxy();
+	}
+
 	bShouldCheckHitProxy = true;
 }
 
diff --git a/Engine/Plugins/Runtime/nDisplay/Source/DisplayClusterLightCardEditor/Private/Viewport/DisplayClusterLightCardEditorViewportClient.h b/Engine/Plugins/Runtime/nDisplay/Source/DisplayClusterLightCardEditor/Private/Viewport/DisplayClusterLightCardEditorViewportClient.h
--- a/Engine/Plugins/Runtime/nDisplay/Source/DisplayClusterLightCardEditor/Private/Viewport/DisplayClusterLightCardEditorViewportClient.h
+++ b/Engine/Plugins/Runtime/nDisplay/Source/DisplayClusterLightCardEditor/Private/Viewport/DisplayClusterLightCardEditorViewportClient.h
@@ -194,7 +194,7 @@ private:
 	FVector DragWidgetOffset;
 
 	/** The current projection mode the 3D viewport is being displayed with */
-	EDisplayClusterMeshProjectionType ProjectionMode = EDisplayClusterMeshProjectionType::Perspective;
+	EDisplayClusterMeshProjectionType ProjectionMode = EDisplayClusterMeshProjectionType::Azimuthal;
 
 	/** The component of the root actor that is acting as the projection origin. Can be either the root component (stage origin) or a view origin component */
 	TWeakObjectPtr<USceneComponent> ProjectionOriginComponent;
