nDisplay: Fix incorrect offset compensation when dragging lightcards. Also adjusted invisibility range of south pole geometry.

diff --git a/Engine/Plugins/Runtime/nDisplay/Shaders/Private/MeshProjectionShaders.usf b/Engine/Plugins/Runtime/nDisplay/Shaders/Private/MeshProjectionShaders.usf
--- a/Engine/Plugins/Runtime/nDisplay/Shaders/Private/MeshProjectionShaders.usf
+++ b/Engine/Plugins/Runtime/nDisplay/Shaders/Private/MeshProjectionShaders.usf
@@ -79,8 +79,8 @@ void MainPS(FMeshProjectionVSToPS Input OPTIONAL_IsFrontFace, out float4 OutputC
 	static const float Gamma = 2.2;
 	OutputColor.rgb = pow(OutputColor.rgb, Gamma);
 	
-	static const float AzimuthFeatherStart = radians(180 - 30);
-	static const float AzimuthFeatherEnd   = radians(180 - 15);
+	static const float AzimuthFeatherStart = radians(180 - 45);
+	static const float AzimuthFeatherEnd   = radians(180 - 30);
 	
 	const float AlphaBlendWeight = saturate((Input.Azimuth - AzimuthFeatherStart) / (AzimuthFeatherEnd - AzimuthFeatherStart));
 	OutputColor.a = lerp(OutputColor.a, 0, AlphaBlendWeight);
diff --git a/Engine/Plugins/Runtime/nDisplay/Source/DisplayClusterLightCardEditor/Private/Viewport/DisplayClusterLightCardEditorViewportClient.cpp b/Engine/Plugins/Runtime/nDisplay/Source/DisplayClusterLightCardEditor/Private/Viewport/DisplayClusterLightCardEditorViewportClient.cpp
--- a/Engine/Plugins/Runtime/nDisplay/Source/DisplayClusterLightCardEditor/Private/Viewport/DisplayClusterLightCardEditorViewportClient.cpp
+++ b/Engine/Plugins/Runtime/nDisplay/Source/DisplayClusterLightCardEditor/Private/Viewport/DisplayClusterLightCardEditorViewportClient.cpp
@@ -118,7 +118,7 @@ void FDisplayClusterLightCardEditorViewportClient::Tick(float DeltaSeconds)
 
 	SetViewLocation(Location);
 
-	CachedEditorWidgetTransform = CalcEditorWidgetTransform();
+	CalcEditorWidgetTransform(CachedEditorWidgetTransformBeforeMapProjection, CachedEditorWidgetTransformAfterMapProjection);
 
 	// Tick the preview scene world.
 	if (!GIntraFrameDebuggingGameThread)
@@ -328,7 +328,7 @@ void FDisplayClusterLightCardEditorViewportClient::Draw(const FSceneView* View,
 {
 	if (SelectedLightCards.Num())
 	{
-		EditorWidget->SetTransform(CachedEditorWidgetTransform);
+		EditorWidget->SetTransform(CachedEditorWidgetTransformAfterMapProjection);
 		EditorWidget->Draw(View, PDI);
 	}
 }
@@ -434,7 +434,7 @@ void FDisplayClusterLightCardEditorViewportClient::TrackingStarted(const FInputE
 		FVector Direction;
 		PixelToWorld(*View, MousePos, Origin, Direction);
 
-		DragWidgetOffset = Direction - (CachedEditorWidgetTransform.GetTranslation() - Origin).GetSafeNormal();
+		DragWidgetOffset = Direction - (CachedEditorWidgetTransformBeforeMapProjection.GetTranslation() - Origin).GetSafeNormal();
 	}
 
 	FEditorViewportClient::TrackingStarted(InInputState, bIsDraggingWidget, bNudge);
@@ -1354,43 +1354,48 @@ void FDisplayClusterLightCardEditorViewportClient::PixelToWorld(const FSceneView
 	OutDirection = InvViewMatrix.TransformVector(UnprojectedViewPos).GetSafeNormal();
 }
 
-FTransform FDisplayClusterLightCardEditorViewportClient::CalcEditorWidgetTransform()
+bool FDisplayClusterLightCardEditorViewportClient::CalcEditorWidgetTransform(FTransform& WidgetTransformBeforeMapProjection, FTransform& WidgetTransformAfterMapProjection)
 {
-	if (SelectedLightCards.Num())
+	if (!SelectedLightCards.Num())
 	{
-		TWeakObjectPtr<ADisplayClusterLightCardActor> LastSelected = SelectedLightCards.Last();
-		if (LastSelected.IsValid())
-		{
-			FVector LightCardPosition = LastSelected->GetLightCardTransform().GetTranslation();
+		return false;
+	}
 
-			FTransform WidgetTransform(FRotator::ZeroRotator, LightCardPosition, FVector::OneVector);
+	TWeakObjectPtr<ADisplayClusterLightCardActor> LastSelected = SelectedLightCards.Last();
 
-			if (ProjectionMode != EDisplayClusterMeshProjectionType::Perspective)
-			{
-				FSceneViewInitOptions SceneVewInitOptions;
-				GetSceneViewInitOptions(SceneVewInitOptions);
-				FViewMatrices ViewMatrices(SceneVewInitOptions);
+	if (!LastSelected.IsValid())
+	{
+		return false;
+	}
 
-				const FVector ViewPos = ViewMatrices.GetViewMatrix().TransformPosition(LightCardPosition);
-				const FVector ProjectedViewPos = FDisplayClusterMeshProjectionRenderer::ProjectViewPosition(ViewPos, ProjectionMode);
-				const FVector ProjectedPosition = ViewMatrices.GetInvViewMatrix().TransformPosition(ProjectedViewPos);
+	FVector LightCardPosition = LastSelected->GetLightCardTransform().GetTranslation();
 
-				WidgetTransform.SetTranslation(ProjectedPosition);
-			}
+	WidgetTransformAfterMapProjection = WidgetTransformBeforeMapProjection = FTransform(FRotator::ZeroRotator, LightCardPosition, FVector::OneVector);
 
-			const FVector ProjectionOrigin = ProjectionOriginComponent.IsValid() ? ProjectionOriginComponent->GetComponentLocation() : FVector::ZeroVector;
-			const FVector RadialVector = (LightCardPosition - ProjectionOrigin).GetSafeNormal();
-			const FVector AzimuthalVector = FVector::ZAxisVector ^ RadialVector;
-			const FVector InclinationVector = RadialVector ^ AzimuthalVector;
+	if (ProjectionMode != EDisplayClusterMeshProjectionType::Perspective)
+	{
+		FSceneViewInitOptions SceneVewInitOptions;
+		GetSceneViewInitOptions(SceneVewInitOptions);
+		FViewMatrices ViewMatrices(SceneVewInitOptions);
 
-			FRotator Orientation = FMatrix(AzimuthalVector, InclinationVector, RadialVector, FVector::ZeroVector).Rotator();
-			WidgetTransform.SetRotation(Orientation.Quaternion());
+		const FVector ViewPos = ViewMatrices.GetViewMatrix().TransformPosition(LightCardPosition);
+		const FVector ProjectedViewPos = FDisplayClusterMeshProjectionRenderer::ProjectViewPosition(ViewPos, ProjectionMode);
+		const FVector ProjectedPosition = ViewMatrices.GetInvViewMatrix().TransformPosition(ProjectedViewPos);
 
-			return WidgetTransform;
-		}
+		WidgetTransformAfterMapProjection.SetTranslation(ProjectedPosition);
 	}
 
-	return FTransform();
+	const FVector ProjectionOrigin = ProjectionOriginComponent.IsValid() ? ProjectionOriginComponent->GetComponentLocation() : FVector::ZeroVector;
+	const FVector RadialVector = (LightCardPosition - ProjectionOrigin).GetSafeNormal();
+	const FVector AzimuthalVector = FVector::ZAxisVector ^ RadialVector;
+	const FVector InclinationVector = RadialVector ^ AzimuthalVector;
+
+	FRotator Orientation = FMatrix(AzimuthalVector, InclinationVector, RadialVector, FVector::ZeroVector).Rotator();
+
+	WidgetTransformBeforeMapProjection.SetRotation(Orientation.Quaternion());
+	WidgetTransformAfterMapProjection.SetRotation(Orientation.Quaternion());
+
+	return true;
 }
 
 #undef LOCTEXT_NAMESPACE
diff --git a/Engine/Plugins/Runtime/nDisplay/Source/DisplayClusterLightCardEditor/Private/Viewport/DisplayClusterLightCardEditorViewportClient.h b/Engine/Plugins/Runtime/nDisplay/Source/DisplayClusterLightCardEditor/Private/Viewport/DisplayClusterLightCardEditorViewportClient.h
--- a/Engine/Plugins/Runtime/nDisplay/Source/DisplayClusterLightCardEditor/Private/Viewport/DisplayClusterLightCardEditorViewportClient.h
+++ b/Engine/Plugins/Runtime/nDisplay/Source/DisplayClusterLightCardEditor/Private/Viewport/DisplayClusterLightCardEditorViewportClient.h
@@ -154,7 +154,7 @@ private:
 	void PixelToWorld(const FSceneView& View, const FIntPoint& PixelPos, FVector& OutOrigin, FVector& OutDirection);
 
 	/** Calculates the world transform to render the editor widget with to align it with the selected light card */
-	FTransform CalcEditorWidgetTransform();
+	bool CalcEditorWidgetTransform(FTransform& WidgetTransformBeforeMapProjection, FTransform& WidgetTransformAfterMapProjection);
 
 private:
 	TWeakPtr<FSceneViewport> SceneViewportPtr;
@@ -187,8 +187,11 @@ private:
 	/** The LC editor widget used to manipulate light cards */
 	TSharedPtr<FDisplayClusterLightCardEditorWidget> EditorWidget;
 
-	/** The cached editor widget transform calculated for the editor widget on the last tick */
-	FTransform CachedEditorWidgetTransform;
+	/** The cached editor widget transform after map projection calculated for the editor widget on the last tick */
+	FTransform CachedEditorWidgetTransformAfterMapProjection;
+
+	/** The cached editor widget transform before map projection calculated for the editor widget on the last tick */
+	FTransform CachedEditorWidgetTransformBeforeMapProjection;
 
 	/** The offset between the widget's origin and the place it was clicked when a drag action was started */
 	FVector DragWidgetOffset;
