nDisplay: Fix viewport transform behaving incorrectly after a compile.

diff --git a/Engine/Plugins/Runtime/nDisplay/Source/DisplayClusterConfigurator/Private/Views/OutputMapping/SDisplayClusterConfiguratorGraphEditor.cpp b/Engine/Plugins/Runtime/nDisplay/Source/DisplayClusterConfigurator/Private/Views/OutputMapping/SDisplayClusterConfiguratorGraphEditor.cpp
--- a/Engine/Plugins/Runtime/nDisplay/Source/DisplayClusterConfigurator/Private/Views/OutputMapping/SDisplayClusterConfiguratorGraphEditor.cpp
+++ b/Engine/Plugins/Runtime/nDisplay/Source/DisplayClusterConfigurator/Private/Views/OutputMapping/SDisplayClusterConfiguratorGraphEditor.cpp
@@ -57,6 +57,8 @@ void SDisplayClusterConfiguratorGraphEditor::FindAndSelectObjects(const TArray<U
 	check(ConfiguratorGraph != nullptr);
 
 	bSelectionSetDirectly = true;
+	// Clear the current selection or stale data might persist after a compile.
+	ClearSelectionSet();
 	ConfiguratorGraph->ForEachGraphNode([this, &ObjectsToSelect](UDisplayClusterConfiguratorBaseNode* Node)
 	{
 		if (ObjectsToSelect.Contains(Node->GetObject()))
