nDisplay: some cleaning & coding style + minor update to the IDisplayClusterProjectionPolicy

diff --git a/Engine/Plugins/Runtime/nDisplay/Source/DisplayCluster/Private/Render/Viewport/Configuration/DisplayClusterViewportConfigurationProjectionPolicy.cpp b/Engine/Plugins/Runtime/nDisplay/Source/DisplayCluster/Private/Render/Viewport/Configuration/DisplayClusterViewportConfigurationProjectionPolicy.cpp
--- a/Engine/Plugins/Runtime/nDisplay/Source/DisplayCluster/Private/Render/Viewport/Configuration/DisplayClusterViewportConfigurationProjectionPolicy.cpp
+++ b/Engine/Plugins/Runtime/nDisplay/Source/DisplayCluster/Private/Render/Viewport/Configuration/DisplayClusterViewportConfigurationProjectionPolicy.cpp
@@ -34,7 +34,7 @@ void FDisplayClusterViewportConfigurationProjectionPolicy::Update()
 		if (ViewportIt && (ViewportIt->GetRenderSettingsICVFX().RuntimeFlags & ViewportRuntime_InternalResource) == 0)
 		{
 			// Support advanced logic for 'camera' projection policy
-			if (ViewportIt->ProjectionPolicy.IsValid() && ViewportIt->ProjectionPolicy->GetTypeId().Compare(DisplayClusterProjectionStrings::projection::Camera) == 0)
+			if (ViewportIt->ProjectionPolicy.IsValid() && ViewportIt->ProjectionPolicy->GetType().Compare(DisplayClusterProjectionStrings::projection::Camera) == 0)
 			{
 				UpdateCameraPolicy(*ViewportIt);
 			}
diff --git a/Engine/Plugins/Runtime/nDisplay/Source/DisplayCluster/Public/Render/Projection/IDisplayClusterProjectionPolicy.h b/Engine/Plugins/Runtime/nDisplay/Source/DisplayCluster/Public/Render/Projection/IDisplayClusterProjectionPolicy.h
--- a/Engine/Plugins/Runtime/nDisplay/Source/DisplayCluster/Public/Render/Projection/IDisplayClusterProjectionPolicy.h
+++ b/Engine/Plugins/Runtime/nDisplay/Source/DisplayCluster/Public/Render/Projection/IDisplayClusterProjectionPolicy.h
@@ -4,21 +4,13 @@
 
 #include "CoreMinimal.h"
 
-#include "RHI.h"
-#include "RHIDefinitions.h"
-#include "RHIResources.h"
-#include "RHICommandList.h"
-
 class FRHICommandListImmediate;
+class IDisplayClusterViewport;
+class IDisplayClusterViewportProxy;
+class IDisplayClusterWarpBlend;
+class UMeshComponent;
 struct FDisplayClusterConfigurationProjection;
 
-#if WITH_EDITORONLY_DATA
-class UDisplayClusterConfigurationViewport;
-class FTextureRenderTargetResource;
-class IDetailCategoryBuilder;
-class IDetailLayoutBuilder;
-class UMeshComponent;
-#endif
 
 /**
  * nDisplay projection policy
@@ -30,14 +22,23 @@ public:
 
 public:
 	/**
-	* Return projection policy name
+	* Return projection policy instance ID
 	*/
 	virtual const FString& GetId() const = 0;
 
 	/**
 	* Return projection policy type
 	*/
-	virtual const FString GetTypeId() const = 0;
+	virtual const FString& GetType() const = 0;
+
+	/**
+	* Return projection policy type
+	*/
+	UE_DEPRECATED(5.1, "This function has beend deprecated. Please use 'GetType'.")
+	virtual const FString GetTypeId() const
+	{
+		return GetType();
+	}
 
 	/**
 	* Return projection policy configuration
@@ -48,7 +49,7 @@ public:
 	* Send projection policy game thread data to render thread proxy
 	* called once per frame from FDisplayClusterViewportManager::FinalizeNewFrame
 	*/
-	virtual void UpdateProxyData(class IDisplayClusterViewport* InViewport)
+	virtual void UpdateProxyData(IDisplayClusterViewport* InViewport)
 	{ }
 
 	/**
@@ -56,14 +57,18 @@ public:
 	*
 	* @param InViewport - a owner viewport
 	*/
-	virtual bool HandleStartScene(class IDisplayClusterViewport* InViewport) = 0;
+	virtual bool HandleStartScene(IDisplayClusterViewport* InViewport)
+	{
+		return true;
+	}
 
 	/**
 	* Called when current level is going to be closed (i.e. before loading a new map)
 	*
 	* @param InViewport - a owner viewport
 	*/
-	virtual void HandleEndScene(class IDisplayClusterViewport* InViewport) = 0;
+	virtual void HandleEndScene(IDisplayClusterViewport* InViewport)
+	{ }
 
 	// Handle request for additional render targetable resource inside viewport api for projection policy
 	virtual bool ShouldUseAdditionalTargetableResource() const
@@ -115,7 +120,7 @@ public:
 	*
 	* @return - True if success
 	*/
-	virtual bool CalculateView(class IDisplayClusterViewport* InViewport, const uint32 InContextNum, FVector& InOutViewLocation, FRotator& InOutViewRotation, const FVector& ViewOffset, const float WorldToMeters, const float NCP, const float FCP) = 0;
+	virtual bool CalculateView(IDisplayClusterViewport* InViewport, const uint32 InContextNum, FVector& InOutViewLocation, FRotator& InOutViewRotation, const FVector& ViewOffset, const float WorldToMeters, const float NCP, const float FCP) = 0;
 
 	/**
 	* @param ViewIdx      - Index of view that is being processed for this viewport
@@ -123,7 +128,7 @@ public:
 	*
 	* @return - True if success
 	*/
-	virtual bool GetProjectionMatrix(class IDisplayClusterViewport* InViewport, const uint32 InContextNum, FMatrix& OutPrjMatrix) = 0;
+	virtual bool GetProjectionMatrix(IDisplayClusterViewport* InViewport, const uint32 InContextNum, FMatrix& OutPrjMatrix) = 0;
 
 	/**
 	* Returns if a policy provides warp&blend feature
@@ -142,7 +147,7 @@ public:
 	* @param InViewportProxy - viewport proxy
 	*
 	*/
-	virtual void BeginWarpBlend_RenderThread(FRHICommandListImmediate& RHICmdList, const class IDisplayClusterViewportProxy* InViewportProxy)
+	virtual void BeginWarpBlend_RenderThread(FRHICommandListImmediate& RHICmdList, const IDisplayClusterViewportProxy* InViewportProxy)
 	{ }
 
 	/**
@@ -152,7 +157,7 @@ public:
 	* @param InViewportProxy - viewport proxy
 	*
 	*/
-	virtual void ApplyWarpBlend_RenderThread(FRHICommandListImmediate& RHICmdList, const class IDisplayClusterViewportProxy* InViewportProxy)
+	virtual void ApplyWarpBlend_RenderThread(FRHICommandListImmediate& RHICmdList, const IDisplayClusterViewportProxy* InViewportProxy)
 	{ }
 
 	/**
@@ -164,7 +169,7 @@ public:
 	* @param ViewportRect - Region of the SrcTexture to perform warp&blend operations
 	*
 	*/
-	virtual void EndWarpBlend_RenderThread(FRHICommandListImmediate& RHICmdList, const class IDisplayClusterViewportProxy* InViewportProxy)
+	virtual void EndWarpBlend_RenderThread(FRHICommandListImmediate& RHICmdList, const IDisplayClusterViewportProxy* InViewportProxy)
 	{ }
 
 	/**
@@ -174,7 +179,7 @@ public:
 	* 
 	* @return - true if output interface valid
 	*/
-	virtual bool GetWarpBlendInterface(TSharedPtr<class IDisplayClusterWarpBlend, ESPMode::ThreadSafe>& OutWarpBlendInterface) const
+	virtual bool GetWarpBlendInterface(TSharedPtr<IDisplayClusterWarpBlend, ESPMode::ThreadSafe>& OutWarpBlendInterface) const
 	{
 		return false;
 	}
@@ -186,7 +191,7 @@ public:
 	*
 	* @return - true if output interface valid
 	*/
-	virtual bool GetWarpBlendInterface_RenderThread(TSharedPtr<class IDisplayClusterWarpBlend, ESPMode::ThreadSafe>& OutWarpBlendInterfaceProxy) const
+	virtual bool GetWarpBlendInterface_RenderThread(TSharedPtr<IDisplayClusterWarpBlend, ESPMode::ThreadSafe>& OutWarpBlendInterfaceProxy) const
 	{
 		return false;
 	}
@@ -208,7 +213,7 @@ public:
 	* @param InViewport - Projection specific parameters.
 	* @param bOutIsRootActorComponent - return true, if used custom root actor component. return false, if created unique temporary component
 	*/
-	virtual class UMeshComponent* GetOrCreatePreviewMeshComponent(class IDisplayClusterViewport* InViewport, bool& bOutIsRootActorComponent)
+	virtual UMeshComponent* GetOrCreatePreviewMeshComponent(IDisplayClusterViewport* InViewport, bool& bOutIsRootActorComponent)
 	{
 		return nullptr;
 	}
diff --git a/Engine/Plugins/Runtime/nDisplay/Source/DisplayClusterProjection/Private/DisplayClusterProjectionLog.cpp b/Engine/Plugins/Runtime/nDisplay/Source/DisplayClusterProjection/Private/DisplayClusterProjectionLog.cpp
--- a/Engine/Plugins/Runtime/nDisplay/Source/DisplayClusterProjection/Private/DisplayClusterProjectionLog.cpp
+++ b/Engine/Plugins/Runtime/nDisplay/Source/DisplayClusterProjection/Private/DisplayClusterProjectionLog.cpp
@@ -5,10 +5,10 @@
 
 DEFINE_LOG_CATEGORY(LogDisplayClusterProjection);
 DEFINE_LOG_CATEGORY(LogDisplayClusterProjectionCamera);
-DEFINE_LOG_CATEGORY(LogDisplayClusterProjectionSimple);
-DEFINE_LOG_CATEGORY(LogDisplayClusterProjectionMPCDI);
-DEFINE_LOG_CATEGORY(LogDisplayClusterProjectionMesh);
+DEFINE_LOG_CATEGORY(LogDisplayClusterProjectionDomeprojection);
 DEFINE_LOG_CATEGORY(LogDisplayClusterProjectionEasyBlend);
-DEFINE_LOG_CATEGORY(LogDisplayClusterProjectionVIOSO);
 DEFINE_LOG_CATEGORY(LogDisplayClusterProjectionManual);
-DEFINE_LOG_CATEGORY(LogDisplayClusterProjectionDomeprojection);
+DEFINE_LOG_CATEGORY(LogDisplayClusterProjectionMesh);
+DEFINE_LOG_CATEGORY(LogDisplayClusterProjectionMPCDI);
+DEFINE_LOG_CATEGORY(LogDisplayClusterProjectionSimple);
+DEFINE_LOG_CATEGORY(LogDisplayClusterProjectionVIOSO);
diff --git a/Engine/Plugins/Runtime/nDisplay/Source/DisplayClusterProjection/Private/DisplayClusterProjectionLog.h b/Engine/Plugins/Runtime/nDisplay/Source/DisplayClusterProjection/Private/DisplayClusterProjectionLog.h
--- a/Engine/Plugins/Runtime/nDisplay/Source/DisplayClusterProjection/Private/DisplayClusterProjectionLog.h
+++ b/Engine/Plugins/Runtime/nDisplay/Source/DisplayClusterProjection/Private/DisplayClusterProjectionLog.h
@@ -6,10 +6,10 @@
 
 DECLARE_LOG_CATEGORY_EXTERN(LogDisplayClusterProjection,               Log, All);
 DECLARE_LOG_CATEGORY_EXTERN(LogDisplayClusterProjectionCamera,         Log, All);
-DECLARE_LOG_CATEGORY_EXTERN(LogDisplayClusterProjectionSimple,         Log, All);
-DECLARE_LOG_CATEGORY_EXTERN(LogDisplayClusterProjectionMPCDI,          Log, All);
-DECLARE_LOG_CATEGORY_EXTERN(LogDisplayClusterProjectionMesh,           Log, All);
+DECLARE_LOG_CATEGORY_EXTERN(LogDisplayClusterProjectionDomeprojection, Log, All);
 DECLARE_LOG_CATEGORY_EXTERN(LogDisplayClusterProjectionEasyBlend,      Log, All);
-DECLARE_LOG_CATEGORY_EXTERN(LogDisplayClusterProjectionVIOSO,          Log, All);
 DECLARE_LOG_CATEGORY_EXTERN(LogDisplayClusterProjectionManual,         Log, All);
-DECLARE_LOG_CATEGORY_EXTERN(LogDisplayClusterProjectionDomeprojection, Log, All);
+DECLARE_LOG_CATEGORY_EXTERN(LogDisplayClusterProjectionMesh,           Log, All);
+DECLARE_LOG_CATEGORY_EXTERN(LogDisplayClusterProjectionMPCDI,          Log, All);
+DECLARE_LOG_CATEGORY_EXTERN(LogDisplayClusterProjectionSimple,         Log, All);
+DECLARE_LOG_CATEGORY_EXTERN(LogDisplayClusterProjectionVIOSO,          Log, All);
diff --git a/Engine/Plugins/Runtime/nDisplay/Source/DisplayClusterProjection/Private/DisplayClusterProjectionModule.cpp b/Engine/Plugins/Runtime/nDisplay/Source/DisplayClusterProjection/Private/DisplayClusterProjectionModule.cpp
--- a/Engine/Plugins/Runtime/nDisplay/Source/DisplayClusterProjection/Private/DisplayClusterProjectionModule.cpp
+++ b/Engine/Plugins/Runtime/nDisplay/Source/DisplayClusterProjection/Private/DisplayClusterProjectionModule.cpp
@@ -7,12 +7,12 @@
 
 #include "Policy/Camera/DisplayClusterProjectionCameraPolicyFactory.h"
 #include "Policy/EasyBlend/DisplayClusterProjectionEasyBlendPolicyFactory.h"
-#include "Policy/VIOSO/DisplayClusterProjectionVIOSOPolicyFactory.h"
-#include "Policy/Simple/DisplayClusterProjectionSimplePolicyFactory.h"
-#include "Policy/MPCDI/DisplayClusterProjectionMPCDIPolicyFactory.h"
-#include "Policy/Manual/DisplayClusterProjectionManualPolicyFactory.h"
 #include "Policy/Domeprojection/DisplayClusterProjectionDomeprojectionPolicyFactory.h"
 #include "Policy/Link/DisplayClusterProjectionLinkPolicyFactory.h"
+#include "Policy/Manual/DisplayClusterProjectionManualPolicyFactory.h"
+#include "Policy/MPCDI/DisplayClusterProjectionMPCDIPolicyFactory.h"
+#include "Policy/Simple/DisplayClusterProjectionSimplePolicyFactory.h"
+#include "Policy/VIOSO/DisplayClusterProjectionVIOSOPolicyFactory.h"
 
 #include "Policy/Mesh/DisplayClusterProjectionMeshPolicy.h"
 
@@ -38,35 +38,34 @@ FDisplayClusterProjectionModule::FDisplayClusterProjectionModule()
 	Factory = MakeShared<FDisplayClusterProjectionCameraPolicyFactory>();
 	ProjectionPolicyFactories.Emplace(DisplayClusterProjectionStrings::projection::Camera, Factory);
 
-	// Simple projection
-	Factory = MakeShared<FDisplayClusterProjectionSimplePolicyFactory>();
-	ProjectionPolicyFactories.Emplace(DisplayClusterProjectionStrings::projection::Simple, Factory);
-
-	// MPCDI and Mesh projection
-	Factory = MakeShared<FDisplayClusterProjectionMPCDIPolicyFactory>();
-	ProjectionPolicyFactories.Emplace(DisplayClusterProjectionStrings::projection::MPCDI, Factory);
-	ProjectionPolicyFactories.Emplace(DisplayClusterProjectionStrings::projection::Mesh,  Factory);
+	// Domeprojection projection
+	Factory = MakeShared<FDisplayClusterProjectionDomeprojectionPolicyFactory>();
+	ProjectionPolicyFactories.Emplace(DisplayClusterProjectionStrings::projection::Domeprojection, Factory);
 
 	// EasyBlend projection
 	Factory = MakeShared<FDisplayClusterProjectionEasyBlendPolicyFactory>();
 	ProjectionPolicyFactories.Emplace(DisplayClusterProjectionStrings::projection::EasyBlend, Factory);
 
-	// VIOSO projection
-	Factory = MakeShared<FDisplayClusterProjectionVIOSOPolicyFactory>();
-	ProjectionPolicyFactories.Emplace(DisplayClusterProjectionStrings::projection::VIOSO, Factory);
+	// Link projection
+	Factory = MakeShared<FDisplayClusterProjectionLinkPolicyFactory>();
+	ProjectionPolicyFactories.Emplace(DisplayClusterProjectionStrings::projection::Link, Factory);
 
 	// Manual projection
 	Factory = MakeShared<FDisplayClusterProjectionManualPolicyFactory>();
 	ProjectionPolicyFactories.Emplace(DisplayClusterProjectionStrings::projection::Manual, Factory);
 
-	// Domeprojection projection
-	Factory = MakeShared<FDisplayClusterProjectionDomeprojectionPolicyFactory>();
-	ProjectionPolicyFactories.Emplace(DisplayClusterProjectionStrings::projection::Domeprojection, Factory);
+	// MPCDI and Mesh projection
+	Factory = MakeShared<FDisplayClusterProjectionMPCDIPolicyFactory>();
+	ProjectionPolicyFactories.Emplace(DisplayClusterProjectionStrings::projection::MPCDI, Factory);
+	ProjectionPolicyFactories.Emplace(DisplayClusterProjectionStrings::projection::Mesh, Factory);
 
-	// Linkprojection
-	Factory = MakeShared<FDisplayClusterProjectionLinkPolicyFactory>();
-	ProjectionPolicyFactories.Emplace(DisplayClusterProjectionStrings::projection::Link, Factory);
+	// Simple projection
+	Factory = MakeShared<FDisplayClusterProjectionSimplePolicyFactory>();
+	ProjectionPolicyFactories.Emplace(DisplayClusterProjectionStrings::projection::Simple, Factory);
 
+	// VIOSO projection
+	Factory = MakeShared<FDisplayClusterProjectionVIOSOPolicyFactory>();
+	ProjectionPolicyFactories.Emplace(DisplayClusterProjectionStrings::projection::VIOSO, Factory);
 
 	UE_LOG(LogDisplayClusterProjection, Log, TEXT("Projection module has been instantiated"));
 }
diff --git a/Engine/Plugins/Runtime/nDisplay/Source/DisplayClusterProjection/Private/Policy/Camera/DisplayClusterProjectionCameraPolicy.cpp b/Engine/Plugins/Runtime/nDisplay/Source/DisplayClusterProjection/Private/Policy/Camera/DisplayClusterProjectionCameraPolicy.cpp
--- a/Engine/Plugins/Runtime/nDisplay/Source/DisplayClusterProjection/Private/Policy/Camera/DisplayClusterProjectionCameraPolicy.cpp
+++ b/Engine/Plugins/Runtime/nDisplay/Source/DisplayClusterProjection/Private/Policy/Camera/DisplayClusterProjectionCameraPolicy.cpp
@@ -15,26 +15,29 @@
 #include "Render/Viewport/IDisplayClusterViewport.h"
 #include "Render/Viewport/IDisplayClusterViewportManager.h"
 
-FDisplayClusterProjectionCameraPolicy::FDisplayClusterProjectionCameraPolicy(const FString& ProjectionPolicyId, const struct FDisplayClusterConfigurationProjection* InConfigurationProjectionPolicy)
-	: FDisplayClusterProjectionPolicyBase(ProjectionPolicyId, InConfigurationProjectionPolicy)
-{
-}
 
-FDisplayClusterProjectionCameraPolicy::~FDisplayClusterProjectionCameraPolicy()
+FDisplayClusterProjectionCameraPolicy::FDisplayClusterProjectionCameraPolicy(const FString& ProjectionPolicyId, const FDisplayClusterConfigurationProjection* InConfigurationProjectionPolicy)
+	: FDisplayClusterProjectionPolicyBase(ProjectionPolicyId, InConfigurationProjectionPolicy)
 {
 }
 
 //////////////////////////////////////////////////////////////////////////////////////////////
 // IDisplayClusterProjectionPolicy
 //////////////////////////////////////////////////////////////////////////////////////////////
-bool FDisplayClusterProjectionCameraPolicy::HandleStartScene(class IDisplayClusterViewport* InViewport)
+const FString& FDisplayClusterProjectionCameraPolicy::GetType() const
+{
+	static const FString Type(DisplayClusterProjectionStrings::projection::Camera);
+	return Type;
+}
+
+bool FDisplayClusterProjectionCameraPolicy::HandleStartScene(IDisplayClusterViewport* InViewport)
 {
 	check(IsInGameThread());
 
 	return true;
 }
 
-void FDisplayClusterProjectionCameraPolicy::HandleEndScene(class IDisplayClusterViewport* InViewport)
+void FDisplayClusterProjectionCameraPolicy::HandleEndScene(IDisplayClusterViewport* InViewport)
 {
 	check(IsInGameThread());
 
@@ -74,7 +77,7 @@ UCameraComponent* FDisplayClusterProjectionCameraPolicy::GetCameraComponent()
 	return  nullptr;
 }
 
-bool FDisplayClusterProjectionCameraPolicy::CalculateView(class IDisplayClusterViewport* InViewport, const uint32 InContextNum, FVector& InOutViewLocation, FRotator& InOutViewRotation, const FVector& ViewOffset, const float WorldToMeters, const float NCP, const float FCP)
+bool FDisplayClusterProjectionCameraPolicy::CalculateView(IDisplayClusterViewport* InViewport, const uint32 InContextNum, FVector& InOutViewLocation, FRotator& InOutViewRotation, const FVector& ViewOffset, const float WorldToMeters, const float NCP, const float FCP)
 {
 	check(IsInGameThread());
 
@@ -159,6 +162,7 @@ bool FDisplayClusterProjectionCameraPolicy::GetProjectionMatrix(IDisplayClusterV
 	return false;
 }
 
+
 //////////////////////////////////////////////////////////////////////////////////////////////
 // FDisplayClusterProjectionCameraPolicy
 //////////////////////////////////////////////////////////////////////////////////////////////
diff --git a/Engine/Plugins/Runtime/nDisplay/Source/DisplayClusterProjection/Private/Policy/Camera/DisplayClusterProjectionCameraPolicy.h b/Engine/Plugins/Runtime/nDisplay/Source/DisplayClusterProjection/Private/Policy/Camera/DisplayClusterProjectionCameraPolicy.h
--- a/Engine/Plugins/Runtime/nDisplay/Source/DisplayClusterProjection/Private/Policy/Camera/DisplayClusterProjectionCameraPolicy.h
+++ b/Engine/Plugins/Runtime/nDisplay/Source/DisplayClusterProjection/Private/Policy/Camera/DisplayClusterProjectionCameraPolicy.h
@@ -6,38 +6,34 @@
 #include "Containers/DisplayClusterProjectionCameraPolicySettings.h"
 #include "Misc/DisplayClusterObjectRef.h"
 
-class IDisplayClusterViewport;
 class UCameraComponent;
-class UWorld;
+
 
 /**
- * Implements math behind the native camera projection policy (use symmetric frustum of a camera)
+ * Camera projection policy implementation
  */
 class FDisplayClusterProjectionCameraPolicy
 	: public FDisplayClusterProjectionPolicyBase
 {
 public:
-	FDisplayClusterProjectionCameraPolicy(const FString& ProjectionPolicyId, const struct FDisplayClusterConfigurationProjection* InConfigurationProjectionPolicy);
-	virtual ~FDisplayClusterProjectionCameraPolicy();
-
-	virtual const FString GetTypeId() const
-	{ return DisplayClusterProjectionStrings::projection::Camera; }
+	FDisplayClusterProjectionCameraPolicy(const FString& ProjectionPolicyId, const FDisplayClusterConfigurationProjection* InConfigurationProjectionPolicy);
 
 public:
 	//////////////////////////////////////////////////////////////////////////////////////////////
 	// IDisplayClusterProjectionPolicy
 	//////////////////////////////////////////////////////////////////////////////////////////////
+	virtual const FString& GetType() const override;
+
 	virtual bool HandleStartScene(IDisplayClusterViewport* InViewport) override;
 	virtual void HandleEndScene(IDisplayClusterViewport* InViewport) override;
 
 	virtual bool CalculateView(IDisplayClusterViewport* InViewport, const uint32 InContextNum, FVector& InOutViewLocation, FRotator& InOutViewRotation, const FVector& ViewOffset, const float WorldToMeters, const float NCP, const float FCP) override;
 	virtual bool GetProjectionMatrix(IDisplayClusterViewport* InViewport, const uint32 InContextNum, FMatrix& OutPrjMatrix) override;
 
-	virtual bool IsWarpBlendSupported() override
-	{ return false; }
-
 	virtual bool ShouldUseSourceTextureWithMips() const override
-	{ return true; }
+	{
+		return true;
+	}
 
 public:
 	void SetCamera(UCameraComponent* const NewCamera, const FDisplayClusterProjectionCameraPolicySettings& InCameraSettings);
@@ -52,6 +48,6 @@ private:
 	// Camera to use for rendering
 	FDisplayClusterSceneComponentRef CameraRef;
 	FDisplayClusterProjectionCameraPolicySettings CameraSettings;
-	float ZNear = 1;
-	float ZFar = 1;
+	float ZNear = 1.f;
+	float ZFar = 1.f;
 };
diff --git a/Engine/Plugins/Runtime/nDisplay/Source/DisplayClusterProjection/Private/Policy/Camera/DisplayClusterProjectionCameraPolicyFactory.cpp b/Engine/Plugins/Runtime/nDisplay/Source/DisplayClusterProjection/Private/Policy/Camera/DisplayClusterProjectionCameraPolicyFactory.cpp
--- a/Engine/Plugins/Runtime/nDisplay/Source/DisplayClusterProjection/Private/Policy/Camera/DisplayClusterProjectionCameraPolicyFactory.cpp
+++ b/Engine/Plugins/Runtime/nDisplay/Source/DisplayClusterProjection/Private/Policy/Camera/DisplayClusterProjectionCameraPolicyFactory.cpp
@@ -7,10 +7,11 @@
 
 #include "DisplayClusterConfigurationTypes.h"
 
+
 //////////////////////////////////////////////////////////////////////////////////////////////
 // IDisplayClusterProjectionPolicyFactory
 //////////////////////////////////////////////////////////////////////////////////////////////
-TSharedPtr<IDisplayClusterProjectionPolicy, ESPMode::ThreadSafe> FDisplayClusterProjectionCameraPolicyFactory::Create(const FString& ProjectionPolicyId, const struct FDisplayClusterConfigurationProjection* InConfigurationProjectionPolicy)
+TSharedPtr<IDisplayClusterProjectionPolicy, ESPMode::ThreadSafe> FDisplayClusterProjectionCameraPolicyFactory::Create(const FString& ProjectionPolicyId, const FDisplayClusterConfigurationProjection* InConfigurationProjectionPolicy)
 {
 	check(InConfigurationProjectionPolicy != nullptr);
 
diff --git a/Engine/Plugins/Runtime/nDisplay/Source/DisplayClusterProjection/Private/Policy/Camera/DisplayClusterProjectionCameraPolicyFactory.h b/Engine/Plugins/Runtime/nDisplay/Source/DisplayClusterProjection/Private/Policy/Camera/DisplayClusterProjectionCameraPolicyFactory.h
--- a/Engine/Plugins/Runtime/nDisplay/Source/DisplayClusterProjection/Private/Policy/Camera/DisplayClusterProjectionCameraPolicyFactory.h
+++ b/Engine/Plugins/Runtime/nDisplay/Source/DisplayClusterProjection/Private/Policy/Camera/DisplayClusterProjectionCameraPolicyFactory.h
@@ -15,5 +15,5 @@ public:
 	//////////////////////////////////////////////////////////////////////////////////////////////
 	// IDisplayClusterProjectionPolicyFactory
 	//////////////////////////////////////////////////////////////////////////////////////////////
-	virtual TSharedPtr<IDisplayClusterProjectionPolicy, ESPMode::ThreadSafe> Create(const FString& ProjectionPolicyId, const struct FDisplayClusterConfigurationProjection* InConfigurationProjectionPolicy) override;
+	virtual TSharedPtr<IDisplayClusterProjectionPolicy, ESPMode::ThreadSafe> Create(const FString& ProjectionPolicyId, const FDisplayClusterConfigurationProjection* InConfigurationProjectionPolicy) override;
 };
diff --git a/Engine/Plugins/Runtime/nDisplay/Source/DisplayClusterProjection/Private/Policy/DisplayClusterProjectionPolicyBase.cpp b/Engine/Plugins/Runtime/nDisplay/Source/DisplayClusterProjection/Private/Policy/DisplayClusterProjectionPolicyBase.cpp
--- a/Engine/Plugins/Runtime/nDisplay/Source/DisplayClusterProjection/Private/Policy/DisplayClusterProjectionPolicyBase.cpp
+++ b/Engine/Plugins/Runtime/nDisplay/Source/DisplayClusterProjection/Private/Policy/DisplayClusterProjectionPolicyBase.cpp
@@ -12,18 +12,16 @@
 #include "Render/Viewport/IDisplayClusterViewportManager.h"
 #include "Render/Viewport/IDisplayClusterViewport.h"
 
+
 //////////////////////////////////////////////////////////////////////////////////////////////
 // FDisplayClusterProjectionPolicyBase
 //////////////////////////////////////////////////////////////////////////////////////////////
-FDisplayClusterProjectionPolicyBase::FDisplayClusterProjectionPolicyBase(const FString& InProjectionPolicyId, const struct FDisplayClusterConfigurationProjection* InConfigurationProjectionPolicy)
-	: ProjectionPolicyId(InProjectionPolicyId)
+FDisplayClusterProjectionPolicyBase::FDisplayClusterProjectionPolicyBase(const FString& InProjectionPolicyId, const FDisplayClusterConfigurationProjection* InConfigurationProjectionPolicy)
+	: PolicyInstanceId(InProjectionPolicyId)
 {
 	Parameters.Append(InConfigurationProjectionPolicy->Parameters);
 }
 
-FDisplayClusterProjectionPolicyBase::~FDisplayClusterProjectionPolicyBase()
-{
-}
 
 bool FDisplayClusterProjectionPolicyBase::IsEditorOperationMode_RenderThread(const IDisplayClusterViewportProxy* InViewportProxy)
 {
@@ -37,7 +35,7 @@ bool FDisplayClusterProjectionPolicyBase::IsEditorOperationMode_RenderThread(con
 	return true;
 }
 
-bool FDisplayClusterProjectionPolicyBase::IsEditorOperationMode(class IDisplayClusterViewport* InViewport)
+bool FDisplayClusterProjectionPolicyBase::IsEditorOperationMode(IDisplayClusterViewport* InViewport)
 {
 	check(IsInGameThread());
 
@@ -70,7 +68,7 @@ bool FDisplayClusterProjectionPolicyBase::IsEditorOperationMode(class IDisplayCl
 	return false;
 }
 
-bool FDisplayClusterProjectionPolicyBase::IsConfigurationChanged(const struct FDisplayClusterConfigurationProjection* InConfigurationProjectionPolicy) const
+bool FDisplayClusterProjectionPolicyBase::IsConfigurationChanged(const FDisplayClusterConfigurationProjection* InConfigurationProjectionPolicy) const
 {
 	if (InConfigurationProjectionPolicy->Parameters.Num() != Parameters.Num()) {
 
diff --git a/Engine/Plugins/Runtime/nDisplay/Source/DisplayClusterProjection/Private/Policy/DisplayClusterProjectionPolicyBase.h b/Engine/Plugins/Runtime/nDisplay/Source/DisplayClusterProjection/Private/Policy/DisplayClusterProjectionPolicyBase.h
--- a/Engine/Plugins/Runtime/nDisplay/Source/DisplayClusterProjection/Private/Policy/DisplayClusterProjectionPolicyBase.h
+++ b/Engine/Plugins/Runtime/nDisplay/Source/DisplayClusterProjection/Private/Policy/DisplayClusterProjectionPolicyBase.h
@@ -9,6 +9,7 @@
 
 class USceneComponent;
 
+
 /**
  * Base projection policy
  */
@@ -16,12 +17,12 @@ class FDisplayClusterProjectionPolicyBase
 	: public IDisplayClusterProjectionPolicy
 {
 public:
-	FDisplayClusterProjectionPolicyBase(const FString& ProjectionPolicyId, const struct FDisplayClusterConfigurationProjection* InConfigurationProjectionPolicy);
-	virtual ~FDisplayClusterProjectionPolicyBase() = 0;
+	FDisplayClusterProjectionPolicyBase(const FString& ProjectionPolicyId, const FDisplayClusterConfigurationProjection* InConfigurationProjectionPolicy);
 
+public:
 	virtual const FString& GetId() const override
 	{
-		return ProjectionPolicyId;
+		return PolicyInstanceId;
 	}
 
 	const USceneComponent* const GetOriginComp() const
@@ -44,18 +45,17 @@ public:
 		return Parameters;
 	}
 
-	virtual bool IsConfigurationChanged(const struct FDisplayClusterConfigurationProjection* InConfigurationProjectionPolicy) const override;
+	virtual bool IsConfigurationChanged(const FDisplayClusterConfigurationProjection* InConfigurationProjectionPolicy) const override;
 
-	static bool IsEditorOperationMode(class IDisplayClusterViewport* InViewport);
+	static bool IsEditorOperationMode(IDisplayClusterViewport* InViewport);
 	static bool IsEditorOperationMode_RenderThread(const IDisplayClusterViewportProxy* InViewportProxy);
 
 protected:
-	void InitializeOriginComponent(class IDisplayClusterViewport* InViewport, const FString& OriginCopmId);
+	void InitializeOriginComponent(IDisplayClusterViewport* InViewport, const FString& OriginCopmId);
 	void ReleaseOriginComponent();
 
 private:
-	// Added 'Policy' prefix to avoid "... hides class name ..." warnings in child classes
-	FString ProjectionPolicyId;
+	const FString PolicyInstanceId;
 
 	FString PolicyOriginCompId;
 	TMap<FString, FString> Parameters;
diff --git a/Engine/Plugins/Runtime/nDisplay/Source/DisplayClusterProjection/Private/Policy/Domeprojection/DisplayClusterProjectionDomeprojectionPolicyBase.cpp b/Engine/Plugins/Runtime/nDisplay/Source/DisplayClusterProjection/Private/Policy/Domeprojection/DisplayClusterProjectionDomeprojectionPolicyBase.cpp
--- a/Engine/Plugins/Runtime/nDisplay/Source/DisplayClusterProjection/Private/Policy/Domeprojection/DisplayClusterProjectionDomeprojectionPolicyBase.cpp
+++ b/Engine/Plugins/Runtime/nDisplay/Source/DisplayClusterProjection/Private/Policy/Domeprojection/DisplayClusterProjectionDomeprojectionPolicyBase.cpp
@@ -14,15 +14,23 @@
 #include "Render/Viewport/IDisplayClusterViewport.h"
 #include "Render/Viewport/IDisplayClusterViewportProxy.h"
 
-FDisplayClusterProjectionDomeprojectionPolicyBase::FDisplayClusterProjectionDomeprojectionPolicyBase(const FString& ProjectionPolicyId, const struct FDisplayClusterConfigurationProjection* InConfigurationProjectionPolicy)
+
+FDisplayClusterProjectionDomeprojectionPolicyBase::FDisplayClusterProjectionDomeprojectionPolicyBase(const FString& ProjectionPolicyId, const FDisplayClusterConfigurationProjection* InConfigurationProjectionPolicy)
 	: FDisplayClusterProjectionPolicyBase(ProjectionPolicyId, InConfigurationProjectionPolicy)
 {
 }
 
+
 //////////////////////////////////////////////////////////////////////////////////////////////
 // IDisplayClusterProjectionPolicy
 //////////////////////////////////////////////////////////////////////////////////////////////
-bool FDisplayClusterProjectionDomeprojectionPolicyBase::HandleStartScene(class IDisplayClusterViewport* InViewport)
+const FString& FDisplayClusterProjectionDomeprojectionPolicyBase::GetType() const
+{
+	static const FString& Type(DisplayClusterProjectionStrings::projection::Domeprojection);
+	return Type;
+}
+
+bool FDisplayClusterProjectionDomeprojectionPolicyBase::HandleStartScene(IDisplayClusterViewport* InViewport)
 {
 	check(IsInGameThread());
 
@@ -79,7 +87,7 @@ bool FDisplayClusterProjectionDomeprojectionPolicyBase::HandleStartScene(class I
 	return true;
 }
 
-void FDisplayClusterProjectionDomeprojectionPolicyBase::HandleEndScene(class IDisplayClusterViewport* InViewport)
+void FDisplayClusterProjectionDomeprojectionPolicyBase::HandleEndScene(IDisplayClusterViewport* InViewport)
 {
 	check(IsInGameThread());
 
@@ -87,7 +95,7 @@ void FDisplayClusterProjectionDomeprojectionPolicyBase::HandleEndScene(class IDi
 }
 
 
-bool FDisplayClusterProjectionDomeprojectionPolicyBase::CalculateView(class IDisplayClusterViewport* InViewport, const uint32 InContextNum, FVector& InOutViewLocation, FRotator& InOutViewRotation, const FVector& ViewOffset, const float WorldToMeters, const float NCP, const float FCP)
+bool FDisplayClusterProjectionDomeprojectionPolicyBase::CalculateView(IDisplayClusterViewport* InViewport, const uint32 InContextNum, FVector& InOutViewLocation, FRotator& InOutViewRotation, const FVector& ViewOffset, const float WorldToMeters, const float NCP, const float FCP)
 {
 	check(IsInGameThread());
 
@@ -127,7 +135,7 @@ bool FDisplayClusterProjectionDomeprojectionPolicyBase::CalculateView(class IDis
 	return true;
 }
 
-bool FDisplayClusterProjectionDomeprojectionPolicyBase::GetProjectionMatrix(class IDisplayClusterViewport* InViewport, const uint32 InContextNum, FMatrix& OutPrjMatrix)
+bool FDisplayClusterProjectionDomeprojectionPolicyBase::GetProjectionMatrix(IDisplayClusterViewport* InViewport, const uint32 InContextNum, FMatrix& OutPrjMatrix)
 {
 	check(IsInGameThread());
 
@@ -145,7 +153,7 @@ bool FDisplayClusterProjectionDomeprojectionPolicyBase::IsWarpBlendSupported()
 	return true;
 }
 
-void FDisplayClusterProjectionDomeprojectionPolicyBase::ApplyWarpBlend_RenderThread(FRHICommandListImmediate& RHICmdList, const class IDisplayClusterViewportProxy* InViewportProxy)
+void FDisplayClusterProjectionDomeprojectionPolicyBase::ApplyWarpBlend_RenderThread(FRHICommandListImmediate& RHICmdList, const IDisplayClusterViewportProxy* InViewportProxy)
 {
 	check(IsInRenderingThread());
 
@@ -158,7 +166,7 @@ void FDisplayClusterProjectionDomeprojectionPolicyBase::ApplyWarpBlend_RenderThr
 //////////////////////////////////////////////////////////////////////////////////////////////
 // FDisplayClusterProjectionDomeprojectionPolicy
 //////////////////////////////////////////////////////////////////////////////////////////////
-bool FDisplayClusterProjectionDomeprojectionPolicyBase::ReadConfigData(class IDisplayClusterViewport* InViewport, FString& OutFile, FString& OutOrigin, uint32& OutChannel)
+bool FDisplayClusterProjectionDomeprojectionPolicyBase::ReadConfigData(IDisplayClusterViewport* InViewport, FString& OutFile, FString& OutOrigin, uint32& OutChannel)
 {
 	check(InViewport);
 
diff --git a/Engine/Plugins/Runtime/nDisplay/Source/DisplayClusterProjection/Private/Policy/Domeprojection/DisplayClusterProjectionDomeprojectionPolicyBase.h b/Engine/Plugins/Runtime/nDisplay/Source/DisplayClusterProjection/Private/Policy/Domeprojection/DisplayClusterProjectionDomeprojectionPolicyBase.h
--- a/Engine/Plugins/Runtime/nDisplay/Source/DisplayClusterProjection/Private/Policy/Domeprojection/DisplayClusterProjectionDomeprojectionPolicyBase.h
+++ b/Engine/Plugins/Runtime/nDisplay/Source/DisplayClusterProjection/Private/Policy/Domeprojection/DisplayClusterProjectionDomeprojectionPolicyBase.h
@@ -7,31 +7,28 @@
 
 
 /**
- * Domeprojection projection policy
+ * Domeprojection projection policy base class
  */
 class FDisplayClusterProjectionDomeprojectionPolicyBase
 	: public FDisplayClusterProjectionPolicyBase
 {
 public:
-	FDisplayClusterProjectionDomeprojectionPolicyBase(const FString& ProjectionPolicyId, const struct FDisplayClusterConfigurationProjection* InConfigurationProjectionPolicy);
-
-	virtual const FString GetTypeId() const
-	{
-		return DisplayClusterProjectionStrings::projection::Domeprojection;
-	}
+	FDisplayClusterProjectionDomeprojectionPolicyBase(const FString& ProjectionPolicyId, const FDisplayClusterConfigurationProjection* InConfigurationProjectionPolicy);
 
 public:
 	//////////////////////////////////////////////////////////////////////////////////////////////
 	// IDisplayClusterProjectionPolicy
 	//////////////////////////////////////////////////////////////////////////////////////////////
-	virtual bool HandleStartScene(class IDisplayClusterViewport* InViewport) override;
-	virtual void HandleEndScene(class IDisplayClusterViewport* InViewport) override;
+	virtual const FString& GetType() const override;
+
+	virtual bool HandleStartScene(IDisplayClusterViewport* InViewport) override;
+	virtual void HandleEndScene(IDisplayClusterViewport* InViewport) override;
 
-	virtual bool CalculateView(class IDisplayClusterViewport* InViewport, const uint32 InContextNum, FVector& InOutViewLocation, FRotator& InOutViewRotation, const FVector& ViewOffset, const float WorldToMeters, const float NCP, const float FCP) override;
-	virtual bool GetProjectionMatrix(class IDisplayClusterViewport* InViewport, const uint32 InContextNum, FMatrix& OutPrjMatrix) override;
+	virtual bool CalculateView(IDisplayClusterViewport* InViewport, const uint32 InContextNum, FVector& InOutViewLocation, FRotator& InOutViewRotation, const FVector& ViewOffset, const float WorldToMeters, const float NCP, const float FCP) override;
+	virtual bool GetProjectionMatrix(IDisplayClusterViewport* InViewport, const uint32 InContextNum, FMatrix& OutPrjMatrix) override;
 
 	virtual bool IsWarpBlendSupported() override;
-	virtual void ApplyWarpBlend_RenderThread(FRHICommandListImmediate& RHICmdList, const class IDisplayClusterViewportProxy* InViewportProxy) override;
+	virtual void ApplyWarpBlend_RenderThread(FRHICommandListImmediate& RHICmdList, const IDisplayClusterViewportProxy* InViewportProxy) override;
 	
 	// Request additional targetable resources for domeprojection  external warpblend
 	virtual bool ShouldUseAdditionalTargetableResource() const override
@@ -45,7 +42,7 @@ protected:
 
 private:
 	// Parse Domeprojection related data from the nDisplay config file
-	bool ReadConfigData(class IDisplayClusterViewport* InViewport, FString& OutFile, FString& OutOrigin, uint32& OutChannel);
+	bool ReadConfigData(IDisplayClusterViewport* InViewport, FString& OutFile, FString& OutOrigin, uint32& OutChannel);
 
 private:
 	FString OriginCompId;
diff --git a/Engine/Plugins/Runtime/nDisplay/Source/DisplayClusterProjection/Private/Policy/Domeprojection/DisplayClusterProjectionDomeprojectionPolicyFactory.cpp b/Engine/Plugins/Runtime/nDisplay/Source/DisplayClusterProjection/Private/Policy/Domeprojection/DisplayClusterProjectionDomeprojectionPolicyFactory.cpp
--- a/Engine/Plugins/Runtime/nDisplay/Source/DisplayClusterProjection/Private/Policy/Domeprojection/DisplayClusterProjectionDomeprojectionPolicyFactory.cpp
+++ b/Engine/Plugins/Runtime/nDisplay/Source/DisplayClusterProjection/Private/Policy/Domeprojection/DisplayClusterProjectionDomeprojectionPolicyFactory.cpp
@@ -15,7 +15,7 @@
 //////////////////////////////////////////////////////////////////////////////////////////////
 // IDisplayClusterProjectionPolicyFactory
 //////////////////////////////////////////////////////////////////////////////////////////////
-TSharedPtr<IDisplayClusterProjectionPolicy, ESPMode::ThreadSafe> FDisplayClusterProjectionDomeprojectionPolicyFactory::Create(const FString& ProjectionPolicyId, const struct FDisplayClusterConfigurationProjection* InConfigurationProjectionPolicy)
+TSharedPtr<IDisplayClusterProjectionPolicy, ESPMode::ThreadSafe> FDisplayClusterProjectionDomeprojectionPolicyFactory::Create(const FString& ProjectionPolicyId, const FDisplayClusterConfigurationProjection* InConfigurationProjectionPolicy)
 {
 	const ERHIInterfaceType RHIType = RHIGetInterfaceType();
 
diff --git a/Engine/Plugins/Runtime/nDisplay/Source/DisplayClusterProjection/Private/Policy/Domeprojection/DisplayClusterProjectionDomeprojectionPolicyFactory.h b/Engine/Plugins/Runtime/nDisplay/Source/DisplayClusterProjection/Private/Policy/Domeprojection/DisplayClusterProjectionDomeprojectionPolicyFactory.h
--- a/Engine/Plugins/Runtime/nDisplay/Source/DisplayClusterProjection/Private/Policy/Domeprojection/DisplayClusterProjectionDomeprojectionPolicyFactory.h
+++ b/Engine/Plugins/Runtime/nDisplay/Source/DisplayClusterProjection/Private/Policy/Domeprojection/DisplayClusterProjectionDomeprojectionPolicyFactory.h
@@ -15,5 +15,5 @@ public:
 	//////////////////////////////////////////////////////////////////////////////////////////////
 	// IDisplayClusterProjectionPolicyFactory
 	//////////////////////////////////////////////////////////////////////////////////////////////
-	virtual TSharedPtr<IDisplayClusterProjectionPolicy, ESPMode::ThreadSafe> Create(const FString& ProjectionPolicyId, const struct FDisplayClusterConfigurationProjection* InConfigurationProjectionPolicy) override;
+	virtual TSharedPtr<IDisplayClusterProjectionPolicy, ESPMode::ThreadSafe> Create(const FString& ProjectionPolicyId, const FDisplayClusterConfigurationProjection* InConfigurationProjectionPolicy) override;
 };
diff --git a/Engine/Plugins/Runtime/nDisplay/Source/DisplayClusterProjection/Private/Policy/Domeprojection/DisplayClusterProjectionDomeprojectionViewAdapterBase.h b/Engine/Plugins/Runtime/nDisplay/Source/DisplayClusterProjection/Private/Policy/Domeprojection/DisplayClusterProjectionDomeprojectionViewAdapterBase.h
--- a/Engine/Plugins/Runtime/nDisplay/Source/DisplayClusterProjection/Private/Policy/Domeprojection/DisplayClusterProjectionDomeprojectionViewAdapterBase.h
+++ b/Engine/Plugins/Runtime/nDisplay/Source/DisplayClusterProjection/Private/Policy/Domeprojection/DisplayClusterProjectionDomeprojectionViewAdapterBase.h
@@ -2,9 +2,9 @@
 
 #pragma once
 
-#include "RHI.h"
-#include "RHIResources.h"
-#include "RHIUtilities.h"
+class IDisplayClusterViewport;
+class IDisplayClusterViewportProxy;
+class FRHICommandListImmediate;
 
 
 /**
@@ -26,7 +26,7 @@ public:
 	virtual ~FDisplayClusterProjectionDomeprojectionViewAdapterBase() = default;
 
 public:
-	virtual bool Initialize(class IDisplayClusterViewport* InViewport, const FString& File) = 0;
+	virtual bool Initialize(IDisplayClusterViewport* InViewport, const FString& File) = 0;
 	
 	virtual void Release()
 	{ }
@@ -38,9 +38,9 @@ public:
 	}
 
 public:
-	virtual bool CalculateView(class IDisplayClusterViewport* InViewport, const uint32 InContextNum, const uint32 Channel, FVector& InOutViewLocation, FRotator& InOutViewRotation, const FVector& ViewOffset, const float WorldToMeters, const float NCP, const float FCP) = 0;
-	virtual bool GetProjectionMatrix(class IDisplayClusterViewport* InViewport, const uint32 InContextNum, const uint32 Channel, FMatrix& OutPrjMatrix) = 0;
-	virtual bool ApplyWarpBlend_RenderThread(FRHICommandListImmediate& RHICmdList, const class IDisplayClusterViewportProxy* InViewportProxy, const uint32 Channel) = 0;
+	virtual bool CalculateView(IDisplayClusterViewport* InViewport, const uint32 InContextNum, const uint32 Channel, FVector& InOutViewLocation, FRotator& InOutViewRotation, const FVector& ViewOffset, const float WorldToMeters, const float NCP, const float FCP) = 0;
+	virtual bool GetProjectionMatrix(IDisplayClusterViewport* InViewport, const uint32 InContextNum, const uint32 Channel, FMatrix& OutPrjMatrix) = 0;
+	virtual bool ApplyWarpBlend_RenderThread(FRHICommandListImmediate& RHICmdList, const IDisplayClusterViewportProxy* InViewportProxy, const uint32 Channel) = 0;
 
 private:
 	const FInitParams InitParams;
diff --git a/Engine/Plugins/Runtime/nDisplay/Source/DisplayClusterProjection/Private/Policy/Domeprojection/Windows/DX11/DisplayClusterProjectionDomeprojectionPolicyDX11.h b/Engine/Plugins/Runtime/nDisplay/Source/DisplayClusterProjection/Private/Policy/Domeprojection/Windows/DX11/DisplayClusterProjectionDomeprojectionPolicyDX11.h
--- a/Engine/Plugins/Runtime/nDisplay/Source/DisplayClusterProjection/Private/Policy/Domeprojection/Windows/DX11/DisplayClusterProjectionDomeprojectionPolicyDX11.h
+++ b/Engine/Plugins/Runtime/nDisplay/Source/DisplayClusterProjection/Private/Policy/Domeprojection/Windows/DX11/DisplayClusterProjectionDomeprojectionPolicyDX11.h
@@ -7,6 +7,7 @@
 
 #include "DisplayClusterProjectionLog.h"
 
+
 /**
  * Domeprojection projection policy (DX11 implementation)
  */
@@ -14,7 +15,7 @@ class FDisplayClusterProjectionDomeprojectionPolicyDX11
 	: public FDisplayClusterProjectionDomeprojectionPolicyBase
 {
 public:
-	FDisplayClusterProjectionDomeprojectionPolicyDX11(const FString& ProjectionPolicyId, const struct FDisplayClusterConfigurationProjection* InConfigurationProjectionPolicy)
+	FDisplayClusterProjectionDomeprojectionPolicyDX11(const FString& ProjectionPolicyId, const FDisplayClusterConfigurationProjection* InConfigurationProjectionPolicy)
 		: FDisplayClusterProjectionDomeprojectionPolicyBase(ProjectionPolicyId, InConfigurationProjectionPolicy)
 	{ }
 
diff --git a/Engine/Plugins/Runtime/nDisplay/Source/DisplayClusterProjection/Private/Policy/Domeprojection/Windows/DX11/DisplayClusterProjectionDomeprojectionViewAdapterDX11.cpp b/Engine/Plugins/Runtime/nDisplay/Source/DisplayClusterProjection/Private/Policy/Domeprojection/Windows/DX11/DisplayClusterProjectionDomeprojectionViewAdapterDX11.cpp
--- a/Engine/Plugins/Runtime/nDisplay/Source/DisplayClusterProjection/Private/Policy/Domeprojection/Windows/DX11/DisplayClusterProjectionDomeprojectionViewAdapterDX11.cpp
+++ b/Engine/Plugins/Runtime/nDisplay/Source/DisplayClusterProjection/Private/Policy/Domeprojection/Windows/DX11/DisplayClusterProjectionDomeprojectionViewAdapterDX11.cpp
@@ -38,7 +38,7 @@ FDisplayClusterProjectionDomeprojectionViewAdapterDX11::~FDisplayClusterProjecti
 	Views.Empty();
 }
 
-bool FDisplayClusterProjectionDomeprojectionViewAdapterDX11::Initialize(class IDisplayClusterViewport* InViewport, const FString& InFile)
+bool FDisplayClusterProjectionDomeprojectionViewAdapterDX11::Initialize(IDisplayClusterViewport* InViewport, const FString& InFile)
 {
 	bool bResult = true;
 
@@ -106,7 +106,7 @@ bool FDisplayClusterProjectionDomeprojectionViewAdapterDX11::CalculateView(IDisp
 	return true;
 }
 
-bool FDisplayClusterProjectionDomeprojectionViewAdapterDX11::GetProjectionMatrix(class IDisplayClusterViewport* InViewport, const uint32 InContextNum, const uint32 Channel, FMatrix& OutPrjMatrix)
+bool FDisplayClusterProjectionDomeprojectionViewAdapterDX11::GetProjectionMatrix(IDisplayClusterViewport* InViewport, const uint32 InContextNum, const uint32 Channel, FMatrix& OutPrjMatrix)
 {
 	check(Views.Num() > (int32)InContextNum);
 
@@ -177,7 +177,7 @@ bool FDisplayClusterProjectionDomeprojectionViewAdapterDX11::ImplApplyWarpBlend_
 	return false;
 }
 
-bool FDisplayClusterProjectionDomeprojectionViewAdapterDX11::ApplyWarpBlend_RenderThread(FRHICommandListImmediate& RHICmdList, const class IDisplayClusterViewportProxy* InViewportProxy, const uint32 Channel)
+bool FDisplayClusterProjectionDomeprojectionViewAdapterDX11::ApplyWarpBlend_RenderThread(FRHICommandListImmediate& RHICmdList, const IDisplayClusterViewportProxy* InViewportProxy, const uint32 Channel)
 {
 	check(IsInRenderingThread());
 
@@ -229,7 +229,7 @@ void FDisplayClusterProjectionDomeprojectionViewAdapterDX11::FViewData::Release(
 	}
 }
 
-bool FDisplayClusterProjectionDomeprojectionViewAdapterDX11::FViewData::Initialize(class IDisplayClusterViewport* InViewport, const FString& InFile, FCriticalSection& DllAccessCS)
+bool FDisplayClusterProjectionDomeprojectionViewAdapterDX11::FViewData::Initialize(IDisplayClusterViewport* InViewport, const FString& InFile, FCriticalSection& DllAccessCS)
 {
 	// Initialize Domeprojection DLL API
 	FScopeLock lock(&DllAccessCS);
diff --git a/Engine/Plugins/Runtime/nDisplay/Source/DisplayClusterProjection/Private/Policy/Domeprojection/Windows/DX11/DisplayClusterProjectionDomeprojectionViewAdapterDX11.h b/Engine/Plugins/Runtime/nDisplay/Source/DisplayClusterProjection/Private/Policy/Domeprojection/Windows/DX11/DisplayClusterProjectionDomeprojectionViewAdapterDX11.h
--- a/Engine/Plugins/Runtime/nDisplay/Source/DisplayClusterProjection/Private/Policy/Domeprojection/Windows/DX11/DisplayClusterProjectionDomeprojectionViewAdapterDX11.h
+++ b/Engine/Plugins/Runtime/nDisplay/Source/DisplayClusterProjection/Private/Policy/Domeprojection/Windows/DX11/DisplayClusterProjectionDomeprojectionViewAdapterDX11.h
@@ -9,6 +9,9 @@ THIRD_PARTY_INCLUDES_START
 THIRD_PARTY_INCLUDES_END
 
 
+/**
+ * Domeprojection projection policy (DX11 adapter)
+ */
 class FDisplayClusterProjectionDomeprojectionViewAdapterDX11
 	: public FDisplayClusterProjectionDomeprojectionViewAdapterBase
 {
@@ -17,12 +20,15 @@ public:
 	virtual ~FDisplayClusterProjectionDomeprojectionViewAdapterDX11();
 
 public:
-	virtual bool Initialize(class IDisplayClusterViewport* InViewport, const FString& File) override;
+	virtual bool Initialize(IDisplayClusterViewport* InViewport, const FString& File) override;
 
 public:
-	virtual bool CalculateView(class IDisplayClusterViewport* InViewport, const uint32 InContextNum, const uint32 Channel, FVector& InOutViewLocation, FRotator& InOutViewRotation, const FVector& ViewOffset, const float WorldToMeters, const float NCP, const float FCP) override;
-	virtual bool GetProjectionMatrix(class IDisplayClusterViewport* InViewport, const uint32 InContextNum, const uint32 Channel, FMatrix& OutPrjMatrix) override;
-	virtual bool ApplyWarpBlend_RenderThread(FRHICommandListImmediate& RHICmdList, const class IDisplayClusterViewportProxy* InViewportProxy, const uint32 Channel) override;
+	//////////////////////////////////////////////////////////////////////////////////////////////
+	// IDisplayClusterProjectionPolicy
+	//////////////////////////////////////////////////////////////////////////////////////////////
+	virtual bool CalculateView(IDisplayClusterViewport* InViewport, const uint32 InContextNum, const uint32 Channel, FVector& InOutViewLocation, FRotator& InOutViewRotation, const FVector& ViewOffset, const float WorldToMeters, const float NCP, const float FCP) override;
+	virtual bool GetProjectionMatrix(IDisplayClusterViewport* InViewport, const uint32 InContextNum, const uint32 Channel, FMatrix& OutPrjMatrix) override;
+	virtual bool ApplyWarpBlend_RenderThread(FRHICommandListImmediate& RHICmdList, const IDisplayClusterViewportProxy* InViewportProxy, const uint32 Channel) override;
 
 protected:
 	bool ImplApplyWarpBlend_RenderThread(FRHICommandListImmediate& RHICmdList, uint32 ContextNum, const uint32 Channel, FRHITexture2D* InputTextures, FRHITexture2D* OutputTextures);
@@ -34,7 +40,7 @@ private:
 	class FViewData
 	{
 	public:
-		bool Initialize(class IDisplayClusterViewport* InViewport, const FString& InFile, FCriticalSection& DllAccessCS);
+		bool Initialize(IDisplayClusterViewport* InViewport, const FString& InFile, FCriticalSection& DllAccessCS);
 		void Release(FCriticalSection& DllAccessCS);
 
 	public:
diff --git a/Engine/Plugins/Runtime/nDisplay/Source/DisplayClusterProjection/Private/Policy/EasyBlend/DisplayClusterProjectionEasyBlendPolicyBase.cpp b/Engine/Plugins/Runtime/nDisplay/Source/DisplayClusterProjection/Private/Policy/EasyBlend/DisplayClusterProjectionEasyBlendPolicyBase.cpp
--- a/Engine/Plugins/Runtime/nDisplay/Source/DisplayClusterProjection/Private/Policy/EasyBlend/DisplayClusterProjectionEasyBlendPolicyBase.cpp
+++ b/Engine/Plugins/Runtime/nDisplay/Source/DisplayClusterProjection/Private/Policy/EasyBlend/DisplayClusterProjectionEasyBlendPolicyBase.cpp
@@ -15,7 +15,8 @@
 #include "Render/Viewport/IDisplayClusterViewport.h"
 #include "Render/Viewport/IDisplayClusterViewportProxy.h"
 
-FDisplayClusterProjectionEasyBlendPolicyBase::FDisplayClusterProjectionEasyBlendPolicyBase(const FString& ProjectionPolicyId, const struct FDisplayClusterConfigurationProjection* InConfigurationProjectionPolicy)
+
+FDisplayClusterProjectionEasyBlendPolicyBase::FDisplayClusterProjectionEasyBlendPolicyBase(const FString& ProjectionPolicyId, const FDisplayClusterConfigurationProjection* InConfigurationProjectionPolicy)
 	: FDisplayClusterProjectionPolicyBase(ProjectionPolicyId, InConfigurationProjectionPolicy)
 {
 }
@@ -24,7 +25,13 @@ FDisplayClusterProjectionEasyBlendPolicyBase::FDisplayClusterProjectionEasyBlend
 //////////////////////////////////////////////////////////////////////////////////////////////
 // IDisplayClusterProjectionPolicy
 //////////////////////////////////////////////////////////////////////////////////////////////
-bool FDisplayClusterProjectionEasyBlendPolicyBase::HandleStartScene(class IDisplayClusterViewport* InViewport)
+const FString& FDisplayClusterProjectionEasyBlendPolicyBase::GetType() const
+{
+	static const FString Type(DisplayClusterProjectionStrings::projection::EasyBlend);
+	return Type;
+}
+
+bool FDisplayClusterProjectionEasyBlendPolicyBase::HandleStartScene(IDisplayClusterViewport* InViewport)
 {
 	check(IsInGameThread());
 
@@ -95,7 +102,7 @@ bool FDisplayClusterProjectionEasyBlendPolicyBase::HandleStartScene(class IDispl
 	return true;
 }
 
-void FDisplayClusterProjectionEasyBlendPolicyBase::HandleEndScene(class IDisplayClusterViewport* InViewport)
+void FDisplayClusterProjectionEasyBlendPolicyBase::HandleEndScene(IDisplayClusterViewport* InViewport)
 {
 	check(IsInGameThread());
 
@@ -173,7 +180,7 @@ bool FDisplayClusterProjectionEasyBlendPolicyBase::GetProjectionMatrix(IDisplayC
 //////////////////////////////////////////////////////////////////////////////////////////////
 // FDisplayClusterProjectionEasyBlendPolicy
 //////////////////////////////////////////////////////////////////////////////////////////////
-bool FDisplayClusterProjectionEasyBlendPolicyBase::ReadConfigData(class IDisplayClusterViewport* InViewport, FString& OutFile, FString& OutOrigin, float& OutGeometryScale)
+bool FDisplayClusterProjectionEasyBlendPolicyBase::ReadConfigData(IDisplayClusterViewport* InViewport, FString& OutFile, FString& OutOrigin, float& OutGeometryScale)
 {
 	check(InViewport);
 
diff --git a/Engine/Plugins/Runtime/nDisplay/Source/DisplayClusterProjection/Private/Policy/EasyBlend/DisplayClusterProjectionEasyBlendPolicyBase.h b/Engine/Plugins/Runtime/nDisplay/Source/DisplayClusterProjection/Private/Policy/EasyBlend/DisplayClusterProjectionEasyBlendPolicyBase.h
--- a/Engine/Plugins/Runtime/nDisplay/Source/DisplayClusterProjection/Private/Policy/EasyBlend/DisplayClusterProjectionEasyBlendPolicyBase.h
+++ b/Engine/Plugins/Runtime/nDisplay/Source/DisplayClusterProjection/Private/Policy/EasyBlend/DisplayClusterProjectionEasyBlendPolicyBase.h
@@ -7,35 +7,40 @@
 
 
 /**
- * EasyBlend projection policy
+ * EasyBlend projection policy base class
  */
 class FDisplayClusterProjectionEasyBlendPolicyBase
 	: public FDisplayClusterProjectionPolicyBase
 {
 public:
-	FDisplayClusterProjectionEasyBlendPolicyBase(const FString& ProjectionPolicyId, const struct FDisplayClusterConfigurationProjection* InConfigurationProjectionPolicy);
-
-	virtual const FString GetTypeId() const
-	{ return DisplayClusterProjectionStrings::projection::EasyBlend; }
+	FDisplayClusterProjectionEasyBlendPolicyBase(const FString& ProjectionPolicyId, const FDisplayClusterConfigurationProjection* InConfigurationProjectionPolicy);
 
 public:
 	//////////////////////////////////////////////////////////////////////////////////////////////
 	// IDisplayClusterProjectionPolicy
 	//////////////////////////////////////////////////////////////////////////////////////////////
-	virtual bool HandleStartScene(class IDisplayClusterViewport* InViewport) override;
-	virtual void HandleEndScene(class IDisplayClusterViewport* InViewport) override;
+	virtual const FString& GetType() const override;
+
+	virtual bool HandleStartScene(IDisplayClusterViewport* InViewport) override;
+	virtual void HandleEndScene(IDisplayClusterViewport* InViewport) override;
 
-	virtual bool CalculateView(class IDisplayClusterViewport* InViewport, const uint32 InContextNum, FVector& InOutViewLocation, FRotator& InOutViewRotation, const FVector& ViewOffset, const float WorldToMeters, const float NCP, const float FCP) override;
-	virtual bool GetProjectionMatrix(class IDisplayClusterViewport* InViewport, const uint32 InContextNum, FMatrix& OutPrjMatrix) override;
+	virtual bool CalculateView(IDisplayClusterViewport* InViewport, const uint32 InContextNum, FVector& InOutViewLocation, FRotator& InOutViewRotation, const FVector& ViewOffset, const float WorldToMeters, const float NCP, const float FCP) override;
+	virtual bool GetProjectionMatrix(IDisplayClusterViewport* InViewport, const uint32 InContextNum, FMatrix& OutPrjMatrix) override;
 
 	virtual bool IsWarpBlendSupported() override
-	{ return true; }
-	virtual void ApplyWarpBlend_RenderThread(FRHICommandListImmediate& RHICmdList, const class IDisplayClusterViewportProxy* InViewportProxy) override;
+	{
+		return true;
+	}
 
 	// Request additional targetable resources for easyblend external warpblend
 	virtual bool ShouldUseAdditionalTargetableResource() const override
-	{ return true; }
+	{
+		return true;
+	}
 
+	virtual void ApplyWarpBlend_RenderThread(FRHICommandListImmediate& RHICmdList, const IDisplayClusterViewportProxy* InViewportProxy) override;
+
+public:
 	virtual bool IsEasyBlendRenderingEnabled() = 0;
 
 protected:
@@ -44,7 +49,7 @@ protected:
 
 private:
 	// Parse EasyBlend related data from the nDisplay config file
-	bool ReadConfigData(class IDisplayClusterViewport* InViewport, FString& OutFile, FString& OutOrigin, float& OutGeometryScale);
+	bool ReadConfigData(IDisplayClusterViewport* InViewport, FString& OutFile, FString& OutOrigin, float& OutGeometryScale);
 
 private:
 	FString OriginCompId;
diff --git a/Engine/Plugins/Runtime/nDisplay/Source/DisplayClusterProjection/Private/Policy/EasyBlend/DisplayClusterProjectionEasyBlendPolicyFactory.cpp b/Engine/Plugins/Runtime/nDisplay/Source/DisplayClusterProjection/Private/Policy/EasyBlend/DisplayClusterProjectionEasyBlendPolicyFactory.cpp
--- a/Engine/Plugins/Runtime/nDisplay/Source/DisplayClusterProjection/Private/Policy/EasyBlend/DisplayClusterProjectionEasyBlendPolicyFactory.cpp
+++ b/Engine/Plugins/Runtime/nDisplay/Source/DisplayClusterProjection/Private/Policy/EasyBlend/DisplayClusterProjectionEasyBlendPolicyFactory.cpp
@@ -11,10 +11,11 @@
 #include "Policy/EasyBlend/Windows/DX11/DisplayClusterProjectionEasyBlendPolicyDX11.h"
 #endif
 
+
 //////////////////////////////////////////////////////////////////////////////////////////////
 // IDisplayClusterProjectionPolicyFactory
 //////////////////////////////////////////////////////////////////////////////////////////////
-TSharedPtr<IDisplayClusterProjectionPolicy, ESPMode::ThreadSafe> FDisplayClusterProjectionEasyBlendPolicyFactory::Create(const FString& ProjectionPolicyId, const struct FDisplayClusterConfigurationProjection* InConfigurationProjectionPolicy)
+TSharedPtr<IDisplayClusterProjectionPolicy, ESPMode::ThreadSafe> FDisplayClusterProjectionEasyBlendPolicyFactory::Create(const FString& ProjectionPolicyId, const FDisplayClusterConfigurationProjection* InConfigurationProjectionPolicy)
 {
 	check(InConfigurationProjectionPolicy != nullptr);
 
diff --git a/Engine/Plugins/Runtime/nDisplay/Source/DisplayClusterProjection/Private/Policy/EasyBlend/DisplayClusterProjectionEasyBlendPolicyFactory.h b/Engine/Plugins/Runtime/nDisplay/Source/DisplayClusterProjection/Private/Policy/EasyBlend/DisplayClusterProjectionEasyBlendPolicyFactory.h
--- a/Engine/Plugins/Runtime/nDisplay/Source/DisplayClusterProjection/Private/Policy/EasyBlend/DisplayClusterProjectionEasyBlendPolicyFactory.h
+++ b/Engine/Plugins/Runtime/nDisplay/Source/DisplayClusterProjection/Private/Policy/EasyBlend/DisplayClusterProjectionEasyBlendPolicyFactory.h
@@ -15,5 +15,5 @@ public:
 	//////////////////////////////////////////////////////////////////////////////////////////////
 	// IDisplayClusterProjectionPolicyFactory
 	//////////////////////////////////////////////////////////////////////////////////////////////
-	virtual TSharedPtr<IDisplayClusterProjectionPolicy, ESPMode::ThreadSafe> Create(const FString& ProjectionPolicyId, const struct FDisplayClusterConfigurationProjection* InConfigurationProjectionPolicy) override;
+	virtual TSharedPtr<IDisplayClusterProjectionPolicy, ESPMode::ThreadSafe> Create(const FString& ProjectionPolicyId, const FDisplayClusterConfigurationProjection* InConfigurationProjectionPolicy) override;
 };
diff --git a/Engine/Plugins/Runtime/nDisplay/Source/DisplayClusterProjection/Private/Policy/EasyBlend/DisplayClusterProjectionEasyBlendViewAdapterBase.h b/Engine/Plugins/Runtime/nDisplay/Source/DisplayClusterProjection/Private/Policy/EasyBlend/DisplayClusterProjectionEasyBlendViewAdapterBase.h
--- a/Engine/Plugins/Runtime/nDisplay/Source/DisplayClusterProjection/Private/Policy/EasyBlend/DisplayClusterProjectionEasyBlendViewAdapterBase.h
+++ b/Engine/Plugins/Runtime/nDisplay/Source/DisplayClusterProjection/Private/Policy/EasyBlend/DisplayClusterProjectionEasyBlendViewAdapterBase.h
@@ -2,9 +2,9 @@
 
 #pragma once
 
-#include "RHI.h"
-#include "RHIResources.h"
-#include "RHIUtilities.h"
+class IDisplayClusterViewport;
+class IDisplayClusterViewportProxy;
+class FRHICommandListImmediate;
 
 
 /**
@@ -26,7 +26,7 @@ public:
 	virtual ~FDisplayClusterProjectionEasyBlendViewAdapterBase() = default;
 
 public:
-	virtual bool Initialize(class IDisplayClusterViewport* InViewport, const FString& File) = 0;
+	virtual bool Initialize(IDisplayClusterViewport* InViewport, const FString& File) = 0;
 	
 	virtual void Release()
 	{ }
@@ -38,9 +38,9 @@ public:
 	}
 
 public:
-	virtual bool CalculateView(class IDisplayClusterViewport* InViewport, const uint32 InContextNum, FVector& InOutViewLocation, FRotator& InOutViewRotation, const FVector& ViewOffset, const float WorldToMeters, const float NCP, const float FCP) = 0;
-	virtual bool GetProjectionMatrix(class IDisplayClusterViewport* InViewport, const uint32 InContextNum, FMatrix& OutPrjMatrix) = 0;
-	virtual bool ApplyWarpBlend_RenderThread(FRHICommandListImmediate& RHICmdList, const class IDisplayClusterViewportProxy* InViewportProxy) = 0;
+	virtual bool CalculateView(IDisplayClusterViewport* InViewport, const uint32 InContextNum, FVector& InOutViewLocation, FRotator& InOutViewRotation, const FVector& ViewOffset, const float WorldToMeters, const float NCP, const float FCP) = 0;
+	virtual bool GetProjectionMatrix(IDisplayClusterViewport* InViewport, const uint32 InContextNum, FMatrix& OutPrjMatrix) = 0;
+	virtual bool ApplyWarpBlend_RenderThread(FRHICommandListImmediate& RHICmdList, const IDisplayClusterViewportProxy* InViewportProxy) = 0;
 
 private:
 	const FInitParams InitParams;
diff --git a/Engine/Plugins/Runtime/nDisplay/Source/DisplayClusterProjection/Private/Policy/EasyBlend/Windows/DX11/DisplayClusterProjectionEasyBlendPolicyDX11.h b/Engine/Plugins/Runtime/nDisplay/Source/DisplayClusterProjection/Private/Policy/EasyBlend/Windows/DX11/DisplayClusterProjectionEasyBlendPolicyDX11.h
--- a/Engine/Plugins/Runtime/nDisplay/Source/DisplayClusterProjection/Private/Policy/EasyBlend/Windows/DX11/DisplayClusterProjectionEasyBlendPolicyDX11.h
+++ b/Engine/Plugins/Runtime/nDisplay/Source/DisplayClusterProjection/Private/Policy/EasyBlend/Windows/DX11/DisplayClusterProjectionEasyBlendPolicyDX11.h
@@ -7,6 +7,7 @@
 
 #include "DisplayClusterProjectionLog.h"
 
+
 /**
  * EasyBlend projection policy (DX11 implementation)
  */
@@ -14,7 +15,7 @@ class FDisplayClusterProjectionEasyBlendPolicyDX11
 	: public FDisplayClusterProjectionEasyBlendPolicyBase
 {
 public:
-	FDisplayClusterProjectionEasyBlendPolicyDX11(const FString& ProjectionPolicyId, const struct FDisplayClusterConfigurationProjection* InConfigurationProjectionPolicy)
+	FDisplayClusterProjectionEasyBlendPolicyDX11(const FString& ProjectionPolicyId, const FDisplayClusterConfigurationProjection* InConfigurationProjectionPolicy)
 		: FDisplayClusterProjectionEasyBlendPolicyBase(ProjectionPolicyId, InConfigurationProjectionPolicy)
 	{ }
 
diff --git a/Engine/Plugins/Runtime/nDisplay/Source/DisplayClusterProjection/Private/Policy/EasyBlend/Windows/DX11/DisplayClusterProjectionEasyBlendViewAdapterDX11.cpp b/Engine/Plugins/Runtime/nDisplay/Source/DisplayClusterProjection/Private/Policy/EasyBlend/Windows/DX11/DisplayClusterProjectionEasyBlendViewAdapterDX11.cpp
--- a/Engine/Plugins/Runtime/nDisplay/Source/DisplayClusterProjection/Private/Policy/EasyBlend/Windows/DX11/DisplayClusterProjectionEasyBlendViewAdapterDX11.cpp
+++ b/Engine/Plugins/Runtime/nDisplay/Source/DisplayClusterProjection/Private/Policy/EasyBlend/Windows/DX11/DisplayClusterProjectionEasyBlendViewAdapterDX11.cpp
@@ -54,7 +54,7 @@ FDisplayClusterProjectionEasyBlendViewAdapterDX11::~FDisplayClusterProjectionEas
 	}
 }
 
-bool FDisplayClusterProjectionEasyBlendViewAdapterDX11::Initialize(class IDisplayClusterViewport* InViewport, const FString& File)
+bool FDisplayClusterProjectionEasyBlendViewAdapterDX11::Initialize(IDisplayClusterViewport* InViewport, const FString& File)
 {
 	if (!IsEasyBlendRenderingEnabled())
 	{
diff --git a/Engine/Plugins/Runtime/nDisplay/Source/DisplayClusterProjection/Private/Policy/EasyBlend/Windows/DX11/DisplayClusterProjectionEasyBlendViewAdapterDX11.h b/Engine/Plugins/Runtime/nDisplay/Source/DisplayClusterProjection/Private/Policy/EasyBlend/Windows/DX11/DisplayClusterProjectionEasyBlendViewAdapterDX11.h
--- a/Engine/Plugins/Runtime/nDisplay/Source/DisplayClusterProjection/Private/Policy/EasyBlend/Windows/DX11/DisplayClusterProjectionEasyBlendViewAdapterDX11.h
+++ b/Engine/Plugins/Runtime/nDisplay/Source/DisplayClusterProjection/Private/Policy/EasyBlend/Windows/DX11/DisplayClusterProjectionEasyBlendViewAdapterDX11.h
@@ -18,12 +18,12 @@ public:
 	virtual ~FDisplayClusterProjectionEasyBlendViewAdapterDX11();
 
 public:
-	virtual bool Initialize(class IDisplayClusterViewport* InViewport, const FString& File) override;
+	virtual bool Initialize(IDisplayClusterViewport* InViewport, const FString& File) override;
 
 public:
-	virtual bool CalculateView(class IDisplayClusterViewport* InViewport, const uint32 InContextNum, FVector& InOutViewLocation, FRotator& InOutViewRotation, const FVector& ViewOffset, const float WorldToMeters, const float NCP, const float FCP) override;
-	virtual bool GetProjectionMatrix(class IDisplayClusterViewport* InViewport, const uint32 InContextNum, FMatrix& OutPrjMatrix) override;
-	virtual bool ApplyWarpBlend_RenderThread(FRHICommandListImmediate& RHICmdList, const class IDisplayClusterViewportProxy* InViewportProxy) override;
+	virtual bool CalculateView(IDisplayClusterViewport* InViewport, const uint32 InContextNum, FVector& InOutViewLocation, FRotator& InOutViewRotation, const FVector& ViewOffset, const float WorldToMeters, const float NCP, const float FCP) override;
+	virtual bool GetProjectionMatrix(IDisplayClusterViewport* InViewport, const uint32 InContextNum, FMatrix& OutPrjMatrix) override;
+	virtual bool ApplyWarpBlend_RenderThread(FRHICommandListImmediate& RHICmdList, const IDisplayClusterViewportProxy* InViewportProxy) override;
 
 	static bool IsEasyBlendRenderingEnabled();
 
diff --git a/Engine/Plugins/Runtime/nDisplay/Source/DisplayClusterProjection/Private/Policy/Link/DisplayClusterProjectionLinkPolicy.cpp b/Engine/Plugins/Runtime/nDisplay/Source/DisplayClusterProjection/Private/Policy/Link/DisplayClusterProjectionLinkPolicy.cpp
--- a/Engine/Plugins/Runtime/nDisplay/Source/DisplayClusterProjection/Private/Policy/Link/DisplayClusterProjectionLinkPolicy.cpp
+++ b/Engine/Plugins/Runtime/nDisplay/Source/DisplayClusterProjection/Private/Policy/Link/DisplayClusterProjectionLinkPolicy.cpp
@@ -12,19 +12,20 @@
 #include "Render/Viewport/IDisplayClusterViewportProxy.h"
 
 
-FDisplayClusterProjectionLinkPolicy::FDisplayClusterProjectionLinkPolicy(const FString& ProjectionPolicyId, const struct FDisplayClusterConfigurationProjection* InConfigurationProjectionPolicy)
+FDisplayClusterProjectionLinkPolicy::FDisplayClusterProjectionLinkPolicy(const FString& ProjectionPolicyId, const FDisplayClusterConfigurationProjection* InConfigurationProjectionPolicy)
 	: FDisplayClusterProjectionPolicyBase(ProjectionPolicyId, InConfigurationProjectionPolicy)
 {
 }
 
-FDisplayClusterProjectionLinkPolicy::~FDisplayClusterProjectionLinkPolicy()
-{
-}
-
 
 //////////////////////////////////////////////////////////////////////////////////////////////
 // IDisplayClusterProjectionPolicy
 //////////////////////////////////////////////////////////////////////////////////////////////
+const FString& FDisplayClusterProjectionLinkPolicy::GetType() const
+{
+	static const FString Type(DisplayClusterProjectionStrings::projection::Link);
+	return Type;
+}
 
 bool FDisplayClusterProjectionLinkPolicy::CalculateView(IDisplayClusterViewport* InViewport, const uint32 InContextNum, FVector& InOutViewLocation, FRotator& InOutViewRotation, const FVector& ViewOffset, const float WorldToMeters, const float NCP, const float FCP)
 {
@@ -55,7 +56,7 @@ bool FDisplayClusterProjectionLinkPolicy::CalculateView(IDisplayClusterViewport*
 	return false;
 }
 
-bool FDisplayClusterProjectionLinkPolicy::GetProjectionMatrix(class IDisplayClusterViewport* InViewport, const uint32 InContextNum, FMatrix& OutPrjMatrix)
+bool FDisplayClusterProjectionLinkPolicy::GetProjectionMatrix(IDisplayClusterViewport* InViewport, const uint32 InContextNum, FMatrix& OutPrjMatrix)
 {
 	check(IsInGameThread());
 	check(InViewport);
diff --git a/Engine/Plugins/Runtime/nDisplay/Source/DisplayClusterProjection/Private/Policy/Link/DisplayClusterProjectionLinkPolicy.h b/Engine/Plugins/Runtime/nDisplay/Source/DisplayClusterProjection/Private/Policy/Link/DisplayClusterProjectionLinkPolicy.h
--- a/Engine/Plugins/Runtime/nDisplay/Source/DisplayClusterProjection/Private/Policy/Link/DisplayClusterProjectionLinkPolicy.h
+++ b/Engine/Plugins/Runtime/nDisplay/Source/DisplayClusterProjection/Private/Policy/Link/DisplayClusterProjectionLinkPolicy.h
@@ -6,39 +6,29 @@
 
 #include "Policy/DisplayClusterProjectionPolicyBase.h"
 
+
+/**
+ * Link projection policy (for internal usage)
+ */
 class FDisplayClusterProjectionLinkPolicy
 	: public FDisplayClusterProjectionPolicyBase
 {
 public:
-	FDisplayClusterProjectionLinkPolicy(const FString& ProjectionPolicyId, const struct FDisplayClusterConfigurationProjection* InConfigurationProjectionPolicy);
-	virtual ~FDisplayClusterProjectionLinkPolicy();
-
-	virtual const FString GetTypeId() const
-	{ return DisplayClusterProjectionStrings::projection::Link; }
+	FDisplayClusterProjectionLinkPolicy(const FString& ProjectionPolicyId, const FDisplayClusterConfigurationProjection* InConfigurationProjectionPolicy);
 
 public:
 	//////////////////////////////////////////////////////////////////////////////////////////////
 	// IDisplayClusterProjectionPolicy
 	//////////////////////////////////////////////////////////////////////////////////////////////
-	virtual bool HandleStartScene(class IDisplayClusterViewport* InViewport) override
-	{
-		// no resources
-		return true;
-	}
-
-	virtual void HandleEndScene(class IDisplayClusterViewport* InViewport) override
-	{ }
+	virtual const FString& GetType() const override;
 
 	// Return values from linked parent viewport
 	// Important note: before doing this, the parent viewports must be updated.
-	virtual bool CalculateView(class IDisplayClusterViewport* InViewport, const uint32 InContextNum, FVector& InOutViewLocation, FRotator& InOutViewRotation, const FVector& ViewOffset, const float WorldToMeters, const float NCP, const float FCP) override;
-	virtual bool GetProjectionMatrix(class IDisplayClusterViewport* InViewport, const uint32 InContextNum, FMatrix& OutPrjMatrix) override;
+	virtual bool CalculateView(IDisplayClusterViewport* InViewport, const uint32 InContextNum, FVector& InOutViewLocation, FRotator& InOutViewRotation, const FVector& ViewOffset, const float WorldToMeters, const float NCP, const float FCP) override;
+	virtual bool GetProjectionMatrix(IDisplayClusterViewport* InViewport, const uint32 InContextNum, FMatrix& OutPrjMatrix) override;
 
-	virtual bool IsWarpBlendSupported() override
+	virtual bool ShouldUseSourceTextureWithMips() const override
 	{
-		return false;
+		return true;
 	}
-
-	virtual bool ShouldUseSourceTextureWithMips() const override
-	{ return true; }
 };
diff --git a/Engine/Plugins/Runtime/nDisplay/Source/DisplayClusterProjection/Private/Policy/Link/DisplayClusterProjectionLinkPolicyFactory.cpp b/Engine/Plugins/Runtime/nDisplay/Source/DisplayClusterProjection/Private/Policy/Link/DisplayClusterProjectionLinkPolicyFactory.cpp
--- a/Engine/Plugins/Runtime/nDisplay/Source/DisplayClusterProjection/Private/Policy/Link/DisplayClusterProjectionLinkPolicyFactory.cpp
+++ b/Engine/Plugins/Runtime/nDisplay/Source/DisplayClusterProjection/Private/Policy/Link/DisplayClusterProjectionLinkPolicyFactory.cpp
@@ -7,10 +7,11 @@
 
 #include "DisplayClusterConfigurationTypes.h"
 
+
 //////////////////////////////////////////////////////////////////////////////////////////////
 // IDisplayClusterProjectionPolicyFactory
 //////////////////////////////////////////////////////////////////////////////////////////////
-TSharedPtr<IDisplayClusterProjectionPolicy, ESPMode::ThreadSafe> FDisplayClusterProjectionLinkPolicyFactory::Create(const FString& ProjectionPolicyId, const struct FDisplayClusterConfigurationProjection* InConfigurationProjectionPolicy)
+TSharedPtr<IDisplayClusterProjectionPolicy, ESPMode::ThreadSafe> FDisplayClusterProjectionLinkPolicyFactory::Create(const FString& ProjectionPolicyId, const FDisplayClusterConfigurationProjection* InConfigurationProjectionPolicy)
 {
 	check(InConfigurationProjectionPolicy != nullptr);
 
diff --git a/Engine/Plugins/Runtime/nDisplay/Source/DisplayClusterProjection/Private/Policy/Link/DisplayClusterProjectionLinkPolicyFactory.h b/Engine/Plugins/Runtime/nDisplay/Source/DisplayClusterProjection/Private/Policy/Link/DisplayClusterProjectionLinkPolicyFactory.h
--- a/Engine/Plugins/Runtime/nDisplay/Source/DisplayClusterProjection/Private/Policy/Link/DisplayClusterProjectionLinkPolicyFactory.h
+++ b/Engine/Plugins/Runtime/nDisplay/Source/DisplayClusterProjection/Private/Policy/Link/DisplayClusterProjectionLinkPolicyFactory.h
@@ -15,5 +15,5 @@ public:
 	//////////////////////////////////////////////////////////////////////////////////////////////
 	// IDisplayClusterProjectionPolicyFactory
 	//////////////////////////////////////////////////////////////////////////////////////////////
-	virtual TSharedPtr<IDisplayClusterProjectionPolicy, ESPMode::ThreadSafe> Create(const FString& ProjectionPolicyId, const struct FDisplayClusterConfigurationProjection* InConfigurationProjectionPolicy) override;
+	virtual TSharedPtr<IDisplayClusterProjectionPolicy, ESPMode::ThreadSafe> Create(const FString& ProjectionPolicyId, const FDisplayClusterConfigurationProjection* InConfigurationProjectionPolicy) override;
 };
diff --git a/Engine/Plugins/Runtime/nDisplay/Source/DisplayClusterProjection/Private/Policy/MPCDI/DisplayClusterProjectionMPCDIPolicy.cpp b/Engine/Plugins/Runtime/nDisplay/Source/DisplayClusterProjection/Private/Policy/MPCDI/DisplayClusterProjectionMPCDIPolicy.cpp
--- a/Engine/Plugins/Runtime/nDisplay/Source/DisplayClusterProjection/Private/Policy/MPCDI/DisplayClusterProjectionMPCDIPolicy.cpp
+++ b/Engine/Plugins/Runtime/nDisplay/Source/DisplayClusterProjection/Private/Policy/MPCDI/DisplayClusterProjectionMPCDIPolicy.cpp
@@ -31,7 +31,8 @@
 #include "Blueprints/MPCDIGeometryData.h"
 #include "DisplayClusterRootActor.h"
 
-FDisplayClusterProjectionMPCDIPolicy::FDisplayClusterProjectionMPCDIPolicy(const FString& ProjectionPolicyId, const struct FDisplayClusterConfigurationProjection* InConfigurationProjectionPolicy)
+
+FDisplayClusterProjectionMPCDIPolicy::FDisplayClusterProjectionMPCDIPolicy(const FString& ProjectionPolicyId, const FDisplayClusterConfigurationProjection* InConfigurationProjectionPolicy)
 	: FDisplayClusterProjectionPolicyBase(ProjectionPolicyId, InConfigurationProjectionPolicy)
 {
 }
@@ -67,6 +68,12 @@ void FDisplayClusterProjectionMPCDIPolicy::UpdateProxyData(IDisplayClusterViewpo
 //////////////////////////////////////////////////////////////////////////////////////////////
 // IDisplayClusterProjectionPolicy
 //////////////////////////////////////////////////////////////////////////////////////////////
+const FString& FDisplayClusterProjectionMPCDIPolicy::GetType() const
+{
+	static const FString Type(DisplayClusterProjectionStrings::projection::MPCDI);
+	return Type;
+}
+
 bool FDisplayClusterProjectionMPCDIPolicy::HandleStartScene(IDisplayClusterViewport* InViewport)
 {
 	if (bInvalidConfiguration)
diff --git a/Engine/Plugins/Runtime/nDisplay/Source/DisplayClusterProjection/Private/Policy/MPCDI/DisplayClusterProjectionMPCDIPolicy.h b/Engine/Plugins/Runtime/nDisplay/Source/DisplayClusterProjection/Private/Policy/MPCDI/DisplayClusterProjectionMPCDIPolicy.h
--- a/Engine/Plugins/Runtime/nDisplay/Source/DisplayClusterProjection/Private/Policy/MPCDI/DisplayClusterProjectionMPCDIPolicy.h
+++ b/Engine/Plugins/Runtime/nDisplay/Source/DisplayClusterProjection/Private/Policy/MPCDI/DisplayClusterProjectionMPCDIPolicy.h
@@ -7,7 +7,6 @@
 #include "WarpBlend/IDisplayClusterWarpBlend.h"
 #include "WarpBlend/DisplayClusterWarpContext.h"
 
-class IDisplayClusterViewport;
 
 /**
  * MPCDI projection policy
@@ -24,12 +23,9 @@ public:
 	};
 
 public:
-	FDisplayClusterProjectionMPCDIPolicy(const FString& ProjectionPolicyId, const struct FDisplayClusterConfigurationProjection* InConfigurationProjectionPolicy);
+	FDisplayClusterProjectionMPCDIPolicy(const FString& ProjectionPolicyId, const FDisplayClusterConfigurationProjection* InConfigurationProjectionPolicy);
 	virtual ~FDisplayClusterProjectionMPCDIPolicy();
 
-	virtual const FString GetTypeId() const
-	{ return DisplayClusterProjectionStrings::projection::MPCDI; }
-
 public:
 	virtual EWarpType GetWarpType() const
 	{
@@ -46,6 +42,8 @@ public:
 	//////////////////////////////////////////////////////////////////////////////////////////////
 	// IDisplayClusterProjectionPolicy
 	//////////////////////////////////////////////////////////////////////////////////////////////
+	virtual const FString& GetType() const override;
+
 	virtual bool HandleStartScene(IDisplayClusterViewport* InViewport) override;
 	virtual void HandleEndScene(IDisplayClusterViewport* InViewport) override;
 
@@ -53,7 +51,7 @@ public:
 	virtual bool GetProjectionMatrix(IDisplayClusterViewport* InViewport, const uint32 InContextNum, FMatrix& OutPrjMatrix) override;
 
 	virtual bool IsWarpBlendSupported() override;
-	virtual void ApplyWarpBlend_RenderThread(FRHICommandListImmediate& RHICmdList, const class IDisplayClusterViewportProxy* InViewportProxy) override;
+	virtual void ApplyWarpBlend_RenderThread(FRHICommandListImmediate& RHICmdList, const IDisplayClusterViewportProxy* InViewportProxy) override;
 
 	virtual bool GetWarpBlendInterface(TSharedPtr<class IDisplayClusterWarpBlend, ESPMode::ThreadSafe>& OutWarpBlendInterface) const override;
 	virtual bool GetWarpBlendInterface_RenderThread(TSharedPtr<class IDisplayClusterWarpBlend, ESPMode::ThreadSafe>& OutWarpBlendInterfaceProxy) const override;
@@ -64,6 +62,16 @@ public:
 		return true;
 	}
 
+#if WITH_EDITOR
+	virtual bool HasPreviewMesh() override
+	{
+		return true;
+	}
+
+	virtual class UMeshComponent* GetOrCreatePreviewMeshComponent(IDisplayClusterViewport* InViewport, bool& bOutIsRootActorComponent) override;
+#endif
+
+public:
 	virtual void UpdateProxyData(IDisplayClusterViewport* InViewport) override;
 
 protected:
@@ -83,15 +91,6 @@ private:
 
 #if WITH_EDITOR
 protected:
-	//////////////////////////////////////////////////////////////////////////////////////////////
-	// IDisplayClusterProjectionPolicyPreview
-	//////////////////////////////////////////////////////////////////////////////////////////////
-	virtual bool HasPreviewMesh() override
-	{
-		return true;
-	}
-	virtual class UMeshComponent* GetOrCreatePreviewMeshComponent(IDisplayClusterViewport* InViewport, bool& bOutIsRootActorComponent) override;
-
 	void ReleasePreviewMeshComponent();
 
 private:
diff --git a/Engine/Plugins/Runtime/nDisplay/Source/DisplayClusterProjection/Private/Policy/MPCDI/DisplayClusterProjectionMPCDIPolicyFactory.cpp b/Engine/Plugins/Runtime/nDisplay/Source/DisplayClusterProjection/Private/Policy/MPCDI/DisplayClusterProjectionMPCDIPolicyFactory.cpp
--- a/Engine/Plugins/Runtime/nDisplay/Source/DisplayClusterProjection/Private/Policy/MPCDI/DisplayClusterProjectionMPCDIPolicyFactory.cpp
+++ b/Engine/Plugins/Runtime/nDisplay/Source/DisplayClusterProjection/Private/Policy/MPCDI/DisplayClusterProjectionMPCDIPolicyFactory.cpp
@@ -9,10 +9,11 @@
 
 #include "DisplayClusterConfigurationTypes.h"
 
+
 //////////////////////////////////////////////////////////////////////////////////////////////
 // IDisplayClusterProjectionPolicyFactory
 //////////////////////////////////////////////////////////////////////////////////////////////
-TSharedPtr<IDisplayClusterProjectionPolicy, ESPMode::ThreadSafe> FDisplayClusterProjectionMPCDIPolicyFactory::Create(const FString& ProjectionPolicyId, const struct FDisplayClusterConfigurationProjection* InConfigurationProjectionPolicy)
+TSharedPtr<IDisplayClusterProjectionPolicy, ESPMode::ThreadSafe> FDisplayClusterProjectionMPCDIPolicyFactory::Create(const FString& ProjectionPolicyId, const FDisplayClusterConfigurationProjection* InConfigurationProjectionPolicy)
 {
 	check(InConfigurationProjectionPolicy != nullptr);
 
diff --git a/Engine/Plugins/Runtime/nDisplay/Source/DisplayClusterProjection/Private/Policy/MPCDI/DisplayClusterProjectionMPCDIPolicyFactory.h b/Engine/Plugins/Runtime/nDisplay/Source/DisplayClusterProjection/Private/Policy/MPCDI/DisplayClusterProjectionMPCDIPolicyFactory.h
--- a/Engine/Plugins/Runtime/nDisplay/Source/DisplayClusterProjection/Private/Policy/MPCDI/DisplayClusterProjectionMPCDIPolicyFactory.h
+++ b/Engine/Plugins/Runtime/nDisplay/Source/DisplayClusterProjection/Private/Policy/MPCDI/DisplayClusterProjectionMPCDIPolicyFactory.h
@@ -16,5 +16,5 @@ public:
 	//////////////////////////////////////////////////////////////////////////////////////////////
 	// IDisplayClusterProjectionPolicyFactory
 	//////////////////////////////////////////////////////////////////////////////////////////////
-	virtual TSharedPtr<IDisplayClusterProjectionPolicy, ESPMode::ThreadSafe> Create(const FString& ProjectionPolicyId, const struct FDisplayClusterConfigurationProjection* InConfigurationProjectionPolicy) override;
+	virtual TSharedPtr<IDisplayClusterProjectionPolicy, ESPMode::ThreadSafe> Create(const FString& ProjectionPolicyId, const FDisplayClusterConfigurationProjection* InConfigurationProjectionPolicy) override;
 };
diff --git a/Engine/Plugins/Runtime/nDisplay/Source/DisplayClusterProjection/Private/Policy/MPCDI/DisplayClusterProjectionMPCDIPolicy_Config.h b/Engine/Plugins/Runtime/nDisplay/Source/DisplayClusterProjection/Private/Policy/MPCDI/DisplayClusterProjectionMPCDIPolicy_Config.h
--- a/Engine/Plugins/Runtime/nDisplay/Source/DisplayClusterProjection/Private/Policy/MPCDI/DisplayClusterProjectionMPCDIPolicy_Config.h
+++ b/Engine/Plugins/Runtime/nDisplay/Source/DisplayClusterProjection/Private/Policy/MPCDI/DisplayClusterProjectionMPCDIPolicy_Config.h
@@ -37,7 +37,7 @@ struct FConfigParser
 
 	bool     bEnablePreview = false;
 
-	inline bool ImplLoadConfig(class IDisplayClusterViewport* InViewport, const TMap<FString, FString>& InConfigParameters)
+	inline bool ImplLoadConfig(IDisplayClusterViewport* InViewport, const TMap<FString, FString>& InConfigParameters)
 	{
 		FString MPCDITypeKey;
 		if (DisplayClusterHelpers::map::template ExtractValue(InConfigParameters, DisplayClusterProjectionStrings::cfg::mpcdi::MPCDITypeKey, MPCDITypeKey))
@@ -69,7 +69,7 @@ struct FConfigParser
 	}
 
 private:
-	inline bool ImplLoadMPCDIConfig(class IDisplayClusterViewport* InViewport, const TMap<FString, FString>& InConfigParameters)
+	inline bool ImplLoadMPCDIConfig(IDisplayClusterViewport* InViewport, const TMap<FString, FString>& InConfigParameters)
 	{
 		// Filename
 		FString LocalMPCDIFileName;
@@ -119,7 +119,7 @@ private:
 		return true;
 	}
 
-	inline bool ImplLoadPFMConfig(class IDisplayClusterViewport* InViewport, const TMap<FString, FString>& InConfigParameters)
+	inline bool ImplLoadPFMConfig(IDisplayClusterViewport* InViewport, const TMap<FString, FString>& InConfigParameters)
 	{
 		// PFM file
 		FString LocalPFMFile;
@@ -199,7 +199,7 @@ private:
 		return true;
 	}
 
-	inline bool ImplLoadBase(class IDisplayClusterViewport* InViewport, const TMap<FString, FString>& InConfigParameters)
+	inline bool ImplLoadBase(IDisplayClusterViewport* InViewport, const TMap<FString, FString>& InConfigParameters)
 	{
 		// Origin node (optional)
 		if (DisplayClusterHelpers::map::template ExtractValue(InConfigParameters, DisplayClusterProjectionStrings::cfg::mpcdi::Origin, OriginType))
diff --git a/Engine/Plugins/Runtime/nDisplay/Source/DisplayClusterProjection/Private/Policy/Manual/DisplayClusterProjectionManualPolicy.cpp b/Engine/Plugins/Runtime/nDisplay/Source/DisplayClusterProjection/Private/Policy/Manual/DisplayClusterProjectionManualPolicy.cpp
--- a/Engine/Plugins/Runtime/nDisplay/Source/DisplayClusterProjection/Private/Policy/Manual/DisplayClusterProjectionManualPolicy.cpp
+++ b/Engine/Plugins/Runtime/nDisplay/Source/DisplayClusterProjection/Private/Policy/Manual/DisplayClusterProjectionManualPolicy.cpp
@@ -15,22 +15,24 @@
 #include "Render/Viewport/IDisplayClusterViewport.h"
 
 
-FDisplayClusterProjectionManualPolicy::FDisplayClusterProjectionManualPolicy(const FString& ProjectionPolicyId, const struct FDisplayClusterConfigurationProjection* InConfigurationProjectionPolicy)
+FDisplayClusterProjectionManualPolicy::FDisplayClusterProjectionManualPolicy(const FString& ProjectionPolicyId, const FDisplayClusterConfigurationProjection* InConfigurationProjectionPolicy)
 	: FDisplayClusterProjectionPolicyBase(ProjectionPolicyId, InConfigurationProjectionPolicy)
 {
 	ProjectionMatrix[0] = FMatrix(FPlane(1.0f, 0.0f, 0.0f, 0.0f), FPlane(0.0f, 1.0f, 0.0f, 0.0f), FPlane(0.0f, 0.0f, 0.0f, 1.0f), FPlane(0.0f, 0.0f, 1.0f, 0.0f));
 	ProjectionMatrix[1] = ProjectionMatrix[0];
 }
 
-FDisplayClusterProjectionManualPolicy::~FDisplayClusterProjectionManualPolicy()
-{
-}
-
 
 //////////////////////////////////////////////////////////////////////////////////////////////
 // IDisplayClusterProjectionPolicy
 //////////////////////////////////////////////////////////////////////////////////////////////
-bool FDisplayClusterProjectionManualPolicy::HandleStartScene(class IDisplayClusterViewport* InViewport)
+const FString& FDisplayClusterProjectionManualPolicy::GetType() const
+{
+	static const FString Type(DisplayClusterProjectionStrings::projection::Manual);
+	return Type;
+}
+
+bool FDisplayClusterProjectionManualPolicy::HandleStartScene(IDisplayClusterViewport* InViewport)
 {
 	check(IsInGameThread());
 	UE_LOG(LogDisplayClusterProjectionManual, Verbose, TEXT("Initializing internals for the viewport '%s'"), *InViewport->GetId());
@@ -124,12 +126,7 @@ bool FDisplayClusterProjectionManualPolicy::HandleStartScene(class IDisplayClust
 	return true;
 }
 
-void FDisplayClusterProjectionManualPolicy::HandleEndScene(class IDisplayClusterViewport* InViewport)
-{
-	
-}
-
-bool FDisplayClusterProjectionManualPolicy::CalculateView(class IDisplayClusterViewport* InViewport, const uint32 InContextNum, FVector& InOutViewLocation, FRotator& InOutViewRotation, const FVector& InViewOffset, const float InWorldToMeters, const float InNCP, const float InFCP)
+bool FDisplayClusterProjectionManualPolicy::CalculateView(IDisplayClusterViewport* InViewport, const uint32 InContextNum, FVector& InOutViewLocation, FRotator& InOutViewRotation, const FVector& InViewOffset, const float InWorldToMeters, const float InNCP, const float InFCP)
 {
 	check(IsInGameThread());
 	check(InContextNum < 2);
@@ -144,7 +141,7 @@ bool FDisplayClusterProjectionManualPolicy::CalculateView(class IDisplayClusterV
 	return true;
 }
 
-bool FDisplayClusterProjectionManualPolicy::GetProjectionMatrix(class IDisplayClusterViewport* InViewport, const uint32 InContextNum, FMatrix& OutPrjMatrix)
+bool FDisplayClusterProjectionManualPolicy::GetProjectionMatrix(IDisplayClusterViewport* InViewport, const uint32 InContextNum, FMatrix& OutPrjMatrix)
 {
 	check(IsInGameThread());
 	check(InContextNum < 2);
@@ -206,7 +203,7 @@ bool FDisplayClusterProjectionManualPolicy::ExtractAngles(const FString& InAngle
 	return true;
 }
 
-FDisplayClusterProjectionManualPolicy::EManualDataType FDisplayClusterProjectionManualPolicy::DataTypeFromString(class IDisplayClusterViewport* InViewport, const FString& DataTypeInString) const
+FDisplayClusterProjectionManualPolicy::EManualDataType FDisplayClusterProjectionManualPolicy::DataTypeFromString(IDisplayClusterViewport* InViewport, const FString& DataTypeInString) const
 {
 	if (DataTypeInString == DisplayClusterProjectionStrings::cfg::manual::FrustumType::Matrix)
 	{
diff --git a/Engine/Plugins/Runtime/nDisplay/Source/DisplayClusterProjection/Private/Policy/Manual/DisplayClusterProjectionManualPolicy.h b/Engine/Plugins/Runtime/nDisplay/Source/DisplayClusterProjection/Private/Policy/Manual/DisplayClusterProjectionManualPolicy.h
--- a/Engine/Plugins/Runtime/nDisplay/Source/DisplayClusterProjection/Private/Policy/Manual/DisplayClusterProjectionManualPolicy.h
+++ b/Engine/Plugins/Runtime/nDisplay/Source/DisplayClusterProjection/Private/Policy/Manual/DisplayClusterProjectionManualPolicy.h
@@ -6,17 +6,13 @@
 
 
 /**
- * Implements math behind the native (manual) quad based projections
+ * Manual projection policy
  */
 class FDisplayClusterProjectionManualPolicy
 	: public FDisplayClusterProjectionPolicyBase
 {
 public:
-	FDisplayClusterProjectionManualPolicy(const FString& ProjectionPolicyId, const struct FDisplayClusterConfigurationProjection* InConfigurationProjectionPolicy);
-	virtual ~FDisplayClusterProjectionManualPolicy();
-
-	virtual const FString GetTypeId() const
-	{ return DisplayClusterProjectionStrings::projection::Manual; }
+	FDisplayClusterProjectionManualPolicy(const FString& ProjectionPolicyId, const FDisplayClusterConfigurationProjection* InConfigurationProjectionPolicy);
 
 public:
 	enum class EManualDataType
@@ -26,20 +22,18 @@ public:
 	};
 
 	EManualDataType GetDataType() const
-	{ return DataType; }
+	{
+		return DataType;
+	}
 
 public:
 	//////////////////////////////////////////////////////////////////////////////////////////////
 	// IDisplayClusterProjectionPolicy
 	//////////////////////////////////////////////////////////////////////////////////////////////
-	virtual bool HandleStartScene(class IDisplayClusterViewport* InViewport) override;
-	virtual void HandleEndScene(class IDisplayClusterViewport* InViewport) override;
-
-	virtual bool CalculateView(class IDisplayClusterViewport* InViewport, const uint32 InContextNum, FVector& InOutViewLocation, FRotator& InOutViewRotation, const FVector& InViewOffset, const float InWorldToMeters, const float InNCP, const float InFCP) override;
-	virtual bool GetProjectionMatrix(class IDisplayClusterViewport* InViewport, const uint32 InContextNum, FMatrix& OutPrjMatrix) override;
-
-	virtual bool IsWarpBlendSupported() override
-	{ return false; }
+	virtual const FString& GetType() const override;
+	virtual bool HandleStartScene(IDisplayClusterViewport* InViewport) override;
+	virtual bool CalculateView(IDisplayClusterViewport* InViewport, const uint32 InContextNum, FVector& InOutViewLocation, FRotator& InOutViewRotation, const FVector& InViewOffset, const float InWorldToMeters, const float InNCP, const float InFCP) override;
+	virtual bool GetProjectionMatrix(IDisplayClusterViewport* InViewport, const uint32 InContextNum, FMatrix& OutPrjMatrix) override;
 
 protected:
 	struct FFrustumAngles
@@ -51,8 +45,10 @@ protected:
 	};
 
 	virtual bool ExtractAngles(const FString& InAngles, FFrustumAngles& OutAngles);
+
 private:
-	EManualDataType DataTypeFromString(class IDisplayClusterViewport* InViewport, const FString& DataTypeInString) const;
+	EManualDataType DataTypeFromString(IDisplayClusterViewport* InViewport, const FString& DataTypeInString) const;
+
 private:
 	// Current data type (matrix, frustum angle, ...)
 	EManualDataType DataType = EManualDataType::Matrix;
diff --git a/Engine/Plugins/Runtime/nDisplay/Source/DisplayClusterProjection/Private/Policy/Manual/DisplayClusterProjectionManualPolicyFactory.cpp b/Engine/Plugins/Runtime/nDisplay/Source/DisplayClusterProjection/Private/Policy/Manual/DisplayClusterProjectionManualPolicyFactory.cpp
--- a/Engine/Plugins/Runtime/nDisplay/Source/DisplayClusterProjection/Private/Policy/Manual/DisplayClusterProjectionManualPolicyFactory.cpp
+++ b/Engine/Plugins/Runtime/nDisplay/Source/DisplayClusterProjection/Private/Policy/Manual/DisplayClusterProjectionManualPolicyFactory.cpp
@@ -7,10 +7,11 @@
 
 #include "DisplayClusterConfigurationTypes.h"
 
+
 //////////////////////////////////////////////////////////////////////////////////////////////
 // IDisplayClusterProjectionPolicyFactory
 //////////////////////////////////////////////////////////////////////////////////////////////
-TSharedPtr<IDisplayClusterProjectionPolicy, ESPMode::ThreadSafe> FDisplayClusterProjectionManualPolicyFactory::Create(const FString& ProjectionPolicyId, const struct FDisplayClusterConfigurationProjection* InConfigurationProjectionPolicy)
+TSharedPtr<IDisplayClusterProjectionPolicy, ESPMode::ThreadSafe> FDisplayClusterProjectionManualPolicyFactory::Create(const FString& ProjectionPolicyId, const FDisplayClusterConfigurationProjection* InConfigurationProjectionPolicy)
 {
 	check(InConfigurationProjectionPolicy != nullptr);
 
diff --git a/Engine/Plugins/Runtime/nDisplay/Source/DisplayClusterProjection/Private/Policy/Manual/DisplayClusterProjectionManualPolicyFactory.h b/Engine/Plugins/Runtime/nDisplay/Source/DisplayClusterProjection/Private/Policy/Manual/DisplayClusterProjectionManualPolicyFactory.h
--- a/Engine/Plugins/Runtime/nDisplay/Source/DisplayClusterProjection/Private/Policy/Manual/DisplayClusterProjectionManualPolicyFactory.h
+++ b/Engine/Plugins/Runtime/nDisplay/Source/DisplayClusterProjection/Private/Policy/Manual/DisplayClusterProjectionManualPolicyFactory.h
@@ -15,5 +15,5 @@ public:
 	//////////////////////////////////////////////////////////////////////////////////////////////
 	// IDisplayClusterProjectionPolicyFactory
 	//////////////////////////////////////////////////////////////////////////////////////////////
-	virtual TSharedPtr<IDisplayClusterProjectionPolicy, ESPMode::ThreadSafe> Create(const FString& ProjectionPolicyId, const struct FDisplayClusterConfigurationProjection* InConfigurationProjectionPolicy) override;
+	virtual TSharedPtr<IDisplayClusterProjectionPolicy, ESPMode::ThreadSafe> Create(const FString& ProjectionPolicyId, const FDisplayClusterConfigurationProjection* InConfigurationProjectionPolicy) override;
 };
diff --git a/Engine/Plugins/Runtime/nDisplay/Source/DisplayClusterProjection/Private/Policy/Mesh/DisplayClusterProjectionMeshPolicy.cpp b/Engine/Plugins/Runtime/nDisplay/Source/DisplayClusterProjection/Private/Policy/Mesh/DisplayClusterProjectionMeshPolicy.cpp
--- a/Engine/Plugins/Runtime/nDisplay/Source/DisplayClusterProjection/Private/Policy/Mesh/DisplayClusterProjectionMeshPolicy.cpp
+++ b/Engine/Plugins/Runtime/nDisplay/Source/DisplayClusterProjection/Private/Policy/Mesh/DisplayClusterProjectionMeshPolicy.cpp
@@ -17,7 +17,8 @@
 #include "Components/StaticMeshComponent.h"
 #include "ProceduralMeshComponent.h"
 
-FDisplayClusterProjectionMeshPolicy::FDisplayClusterProjectionMeshPolicy(const FString& ProjectionPolicyId, const struct FDisplayClusterConfigurationProjection* InConfigurationProjectionPolicy)
+
+FDisplayClusterProjectionMeshPolicy::FDisplayClusterProjectionMeshPolicy(const FString& ProjectionPolicyId, const FDisplayClusterConfigurationProjection* InConfigurationProjectionPolicy)
 	: FDisplayClusterProjectionMPCDIPolicy(ProjectionPolicyId, InConfigurationProjectionPolicy)
 {
 }
@@ -84,6 +85,12 @@ bool FDisplayClusterProjectionMeshPolicy::CreateWarpMeshInterface(IDisplayCluste
 	return true;
 }
 
+const FString& FDisplayClusterProjectionMeshPolicy::GetType() const
+{
+	static const FString Type(DisplayClusterProjectionStrings::projection::Mesh);
+	return Type;
+}
+
 bool FDisplayClusterProjectionMeshPolicy::HandleStartScene(IDisplayClusterViewport* InViewport)
 {
 	check(IsInGameThread());
@@ -224,9 +231,6 @@ bool FDisplayClusterProjectionMeshPolicy::GetWarpMeshConfiguration(IDisplayClust
 }
 
 #if WITH_EDITOR
-//////////////////////////////////////////////////////////////////////////////////////////////
-// IDisplayClusterProjectionPolicyPreview
-//////////////////////////////////////////////////////////////////////////////////////////////
 UMeshComponent* FDisplayClusterProjectionMeshPolicy::GetOrCreatePreviewMeshComponent(IDisplayClusterViewport* InViewport, bool& bOutIsRootActorComponent)
 {
 	bOutIsRootActorComponent = true;
diff --git a/Engine/Plugins/Runtime/nDisplay/Source/DisplayClusterProjection/Private/Policy/Mesh/DisplayClusterProjectionMeshPolicy.h b/Engine/Plugins/Runtime/nDisplay/Source/DisplayClusterProjection/Private/Policy/Mesh/DisplayClusterProjectionMeshPolicy.h
--- a/Engine/Plugins/Runtime/nDisplay/Source/DisplayClusterProjection/Private/Policy/Mesh/DisplayClusterProjectionMeshPolicy.h
+++ b/Engine/Plugins/Runtime/nDisplay/Source/DisplayClusterProjection/Private/Policy/Mesh/DisplayClusterProjectionMeshPolicy.h
@@ -4,10 +4,9 @@
 
 #include "Policy/MPCDI/DisplayClusterProjectionMPCDIPolicy.h"
 
-class IDisplayClusterViewport;
 class UStaticMeshComponent;
 class UProceduralMeshComponent;
-class USceneComponent;
+
 
 /*
  * Mesh projection policy
@@ -17,19 +16,29 @@ class FDisplayClusterProjectionMeshPolicy
 	: public FDisplayClusterProjectionMPCDIPolicy
 {
 public:
-	FDisplayClusterProjectionMeshPolicy(const FString& ProjectionPolicyId, const struct FDisplayClusterConfigurationProjection* InConfigurationProjectionPolicy);
-	virtual ~FDisplayClusterProjectionMeshPolicy() = default;
-
-	virtual const FString GetTypeId() const
-	{ return DisplayClusterProjectionStrings::projection::Mesh; }
+	FDisplayClusterProjectionMeshPolicy(const FString& ProjectionPolicyId, const FDisplayClusterConfigurationProjection* InConfigurationProjectionPolicy);
 
+public:
 	//////////////////////////////////////////////////////////////////////////////////////////////
 	// IDisplayClusterProjectionPolicy
 	//////////////////////////////////////////////////////////////////////////////////////////////
+	virtual const FString& GetType() const override;
 	virtual bool HandleStartScene(IDisplayClusterViewport* InViewport) override;
-	
+
+#if WITH_EDITOR
+	virtual bool HasPreviewMesh() override
+	{
+		return true;
+	}
+
+	virtual UMeshComponent* GetOrCreatePreviewMeshComponent(IDisplayClusterViewport* InViewport, bool& bOutIsRootActorComponent) override;
+#endif
+
+public:	
 	virtual EWarpType GetWarpType() const override
-	{ return EWarpType::mesh; }
+	{
+		return EWarpType::mesh;
+	}
 
 	/** Parse the config data for a mesh id and try to retrieve it from the root actor. */
 	bool CreateWarpMeshInterface(IDisplayClusterViewport* InViewport);
@@ -53,17 +62,4 @@ private:
 	};
 
 	bool GetWarpMeshConfiguration(IDisplayClusterViewport* InViewport, FWarpMeshConfiguration& OutWarpCfg);
-
-#if WITH_EDITOR
-protected:
-	//////////////////////////////////////////////////////////////////////////////////////////////
-	// IDisplayClusterProjectionPolicyPreview
-	//////////////////////////////////////////////////////////////////////////////////////////////
-	virtual bool HasPreviewMesh() override
-	{
-		return true;
-	}
-
-	virtual class UMeshComponent* GetOrCreatePreviewMeshComponent(IDisplayClusterViewport* InViewport, bool& bOutIsRootActorComponent) override;
-#endif
 };
diff --git a/Engine/Plugins/Runtime/nDisplay/Source/DisplayClusterProjection/Private/Policy/Simple/DisplayClusterProjectionSimplePolicy.cpp b/Engine/Plugins/Runtime/nDisplay/Source/DisplayClusterProjection/Private/Policy/Simple/DisplayClusterProjectionSimplePolicy.cpp
--- a/Engine/Plugins/Runtime/nDisplay/Source/DisplayClusterProjection/Private/Policy/Simple/DisplayClusterProjectionSimplePolicy.cpp
+++ b/Engine/Plugins/Runtime/nDisplay/Source/DisplayClusterProjection/Private/Policy/Simple/DisplayClusterProjectionSimplePolicy.cpp
@@ -19,19 +19,21 @@
 #include "Render/Viewport/IDisplayClusterViewportProxy.h"
 
 
-FDisplayClusterProjectionSimplePolicy::FDisplayClusterProjectionSimplePolicy(const FString& ProjectionPolicyId, const struct FDisplayClusterConfigurationProjection* InConfigurationProjectionPolicy)
+FDisplayClusterProjectionSimplePolicy::FDisplayClusterProjectionSimplePolicy(const FString& ProjectionPolicyId, const FDisplayClusterConfigurationProjection* InConfigurationProjectionPolicy)
 	: FDisplayClusterProjectionPolicyBase(ProjectionPolicyId, InConfigurationProjectionPolicy)
 {
 }
 
-FDisplayClusterProjectionSimplePolicy::~FDisplayClusterProjectionSimplePolicy()
-{
-}
-
 
 //////////////////////////////////////////////////////////////////////////////////////////////
 // IDisplayClusterProjectionPolicy
 //////////////////////////////////////////////////////////////////////////////////////////////
+const FString& FDisplayClusterProjectionSimplePolicy::GetType() const
+{
+	static const FString Type(DisplayClusterProjectionStrings::projection::Simple);
+	return Type;
+}
+
 bool FDisplayClusterProjectionSimplePolicy::HandleStartScene(IDisplayClusterViewport* InViewport)
 {
 	check(IsInGameThread());
@@ -213,9 +215,6 @@ void FDisplayClusterProjectionSimplePolicy::ReleaseMeshData(IDisplayClusterViewp
 }
 
 #if WITH_EDITOR
-//////////////////////////////////////////////////////////////////////////////////////////////
-// IDisplayClusterProjectionPolicyPreview
-//////////////////////////////////////////////////////////////////////////////////////////////
 UMeshComponent* FDisplayClusterProjectionSimplePolicy::GetOrCreatePreviewMeshComponent(IDisplayClusterViewport* InViewport, bool& bOutIsRootActorComponent)
 {
 	bOutIsRootActorComponent = true;
diff --git a/Engine/Plugins/Runtime/nDisplay/Source/DisplayClusterProjection/Private/Policy/Simple/DisplayClusterProjectionSimplePolicy.h b/Engine/Plugins/Runtime/nDisplay/Source/DisplayClusterProjection/Private/Policy/Simple/DisplayClusterProjectionSimplePolicy.h
--- a/Engine/Plugins/Runtime/nDisplay/Source/DisplayClusterProjection/Private/Policy/Simple/DisplayClusterProjectionSimplePolicy.h
+++ b/Engine/Plugins/Runtime/nDisplay/Source/DisplayClusterProjection/Private/Policy/Simple/DisplayClusterProjectionSimplePolicy.h
@@ -7,43 +7,40 @@
 #include "Misc/DisplayClusterObjectRef.h"
 #include "Policy/DisplayClusterProjectionPolicyBase.h"
 
-class ADisplayClusterRootActor;
-class USceneComponent;
-class UDisplayClusterConfigurationViewport;
-class UDisplayClusterCameraComponent;
-class UDisplayClusterScreenComponent;
 
 /**
- * Implements math behind the native (simple) quad based projections
+ * Simple projection policy
  */
 class FDisplayClusterProjectionSimplePolicy
 	: public FDisplayClusterProjectionPolicyBase
 {
 public:
-	FDisplayClusterProjectionSimplePolicy(const FString& ProjectionPolicyId, const struct FDisplayClusterConfigurationProjection* InConfigurationProjectionPolicy);
-	virtual ~FDisplayClusterProjectionSimplePolicy();
-
-	virtual const FString GetTypeId() const
-	{ return DisplayClusterProjectionStrings::projection::Simple; }
+	FDisplayClusterProjectionSimplePolicy(const FString& ProjectionPolicyId, const FDisplayClusterConfigurationProjection* InConfigurationProjectionPolicy);
 
 public:
 	//////////////////////////////////////////////////////////////////////////////////////////////
 	// IDisplayClusterProjectionPolicy
 	//////////////////////////////////////////////////////////////////////////////////////////////
-	virtual bool HandleStartScene(class IDisplayClusterViewport* InViewport) override;
-	virtual void HandleEndScene(class IDisplayClusterViewport* InViewport) override;
+	virtual const FString& GetType() const override;
+
+	virtual bool HandleStartScene(IDisplayClusterViewport* InViewport) override;
+	virtual void HandleEndScene(IDisplayClusterViewport* InViewport) override;
 
-	virtual bool CalculateView(class IDisplayClusterViewport* InViewport, const uint32 InContextNum, FVector& InOutViewLocation, FRotator& InOutViewRotation, const FVector& ViewOffset, const float WorldToMeters, const float NCP, const float FCP) override;
-	virtual bool GetProjectionMatrix(class IDisplayClusterViewport* InViewport, const uint32 InContextNum, FMatrix& OutPrjMatrix) override;
+	virtual bool CalculateView(IDisplayClusterViewport* InViewport, const uint32 InContextNum, FVector& InOutViewLocation, FRotator& InOutViewRotation, const FVector& ViewOffset, const float WorldToMeters, const float NCP, const float FCP) override;
+	virtual bool GetProjectionMatrix(IDisplayClusterViewport* InViewport, const uint32 InContextNum, FMatrix& OutPrjMatrix) override;
 
-	virtual bool IsWarpBlendSupported() override
+#if WITH_EDITOR
+	virtual bool HasPreviewMesh() override
 	{
-		return false;
+		return true;
 	}
 
+	virtual UMeshComponent* GetOrCreatePreviewMeshComponent(IDisplayClusterViewport* InViewport, bool& bOutIsRootActorComponent) override;
+#endif
+
 protected:
-	bool InitializeMeshData(class IDisplayClusterViewport* InViewport);
-	void ReleaseMeshData(class IDisplayClusterViewport* InViewport);
+	bool InitializeMeshData(IDisplayClusterViewport* InViewport);
+	void ReleaseMeshData(IDisplayClusterViewport* InViewport);
 
 protected:
 	// Custom projection screen geometry (hw - half-width, hh - half-height of projection screen)
@@ -81,18 +78,4 @@ private:
 	};
 
 	TArray<FViewData> ViewData;
-
-
-#if WITH_EDITOR
-protected:
-	//////////////////////////////////////////////////////////////////////////////////////////////
-	// IDisplayClusterProjectionPolicyPreview
-	//////////////////////////////////////////////////////////////////////////////////////////////
-	virtual bool HasPreviewMesh() override
-	{
-		return true;
-	}
-
-	virtual class UMeshComponent* GetOrCreatePreviewMeshComponent(class IDisplayClusterViewport* InViewport, bool& bOutIsRootActorComponent) override;
-#endif
 };
diff --git a/Engine/Plugins/Runtime/nDisplay/Source/DisplayClusterProjection/Private/Policy/Simple/DisplayClusterProjectionSimplePolicyFactory.cpp b/Engine/Plugins/Runtime/nDisplay/Source/DisplayClusterProjection/Private/Policy/Simple/DisplayClusterProjectionSimplePolicyFactory.cpp
--- a/Engine/Plugins/Runtime/nDisplay/Source/DisplayClusterProjection/Private/Policy/Simple/DisplayClusterProjectionSimplePolicyFactory.cpp
+++ b/Engine/Plugins/Runtime/nDisplay/Source/DisplayClusterProjection/Private/Policy/Simple/DisplayClusterProjectionSimplePolicyFactory.cpp
@@ -7,10 +7,11 @@
 
 #include "DisplayClusterConfigurationTypes.h"
 
+
 //////////////////////////////////////////////////////////////////////////////////////////////
 // IDisplayClusterProjectionPolicyFactory
 //////////////////////////////////////////////////////////////////////////////////////////////
-TSharedPtr<IDisplayClusterProjectionPolicy, ESPMode::ThreadSafe> FDisplayClusterProjectionSimplePolicyFactory::Create(const FString& ProjectionPolicyId, const struct FDisplayClusterConfigurationProjection* InConfigurationProjectionPolicy)
+TSharedPtr<IDisplayClusterProjectionPolicy, ESPMode::ThreadSafe> FDisplayClusterProjectionSimplePolicyFactory::Create(const FString& ProjectionPolicyId, const FDisplayClusterConfigurationProjection* InConfigurationProjectionPolicy)
 {
 	check(InConfigurationProjectionPolicy != nullptr);
 
diff --git a/Engine/Plugins/Runtime/nDisplay/Source/DisplayClusterProjection/Private/Policy/Simple/DisplayClusterProjectionSimplePolicyFactory.h b/Engine/Plugins/Runtime/nDisplay/Source/DisplayClusterProjection/Private/Policy/Simple/DisplayClusterProjectionSimplePolicyFactory.h
--- a/Engine/Plugins/Runtime/nDisplay/Source/DisplayClusterProjection/Private/Policy/Simple/DisplayClusterProjectionSimplePolicyFactory.h
+++ b/Engine/Plugins/Runtime/nDisplay/Source/DisplayClusterProjection/Private/Policy/Simple/DisplayClusterProjectionSimplePolicyFactory.h
@@ -15,5 +15,5 @@ public:
 	//////////////////////////////////////////////////////////////////////////////////////////////
 	// IDisplayClusterProjectionPolicyFactory
 	//////////////////////////////////////////////////////////////////////////////////////////////
-	virtual TSharedPtr<IDisplayClusterProjectionPolicy, ESPMode::ThreadSafe> Create(const FString& ProjectionPolicyId, const struct FDisplayClusterConfigurationProjection* InConfigurationProjectionPolicy) override;
+	virtual TSharedPtr<IDisplayClusterProjectionPolicy, ESPMode::ThreadSafe> Create(const FString& ProjectionPolicyId, const FDisplayClusterConfigurationProjection* InConfigurationProjectionPolicy) override;
 };
diff --git a/Engine/Plugins/Runtime/nDisplay/Source/DisplayClusterProjection/Private/Policy/VIOSO/DisplayClusterProjectionVIOSOPolicyFactory.cpp b/Engine/Plugins/Runtime/nDisplay/Source/DisplayClusterProjection/Private/Policy/VIOSO/DisplayClusterProjectionVIOSOPolicyFactory.cpp
--- a/Engine/Plugins/Runtime/nDisplay/Source/DisplayClusterProjection/Private/Policy/VIOSO/DisplayClusterProjectionVIOSOPolicyFactory.cpp
+++ b/Engine/Plugins/Runtime/nDisplay/Source/DisplayClusterProjection/Private/Policy/VIOSO/DisplayClusterProjectionVIOSOPolicyFactory.cpp
@@ -15,7 +15,7 @@
 //////////////////////////////////////////////////////////////////////////////////////////////
 // IDisplayClusterProjectionPolicyFactory
 //////////////////////////////////////////////////////////////////////////////////////////////
-TSharedPtr<IDisplayClusterProjectionPolicy, ESPMode::ThreadSafe> FDisplayClusterProjectionVIOSOPolicyFactory::Create(const FString& ProjectionPolicyId, const struct FDisplayClusterConfigurationProjection* InConfigurationProjectionPolicy)
+TSharedPtr<IDisplayClusterProjectionPolicy, ESPMode::ThreadSafe> FDisplayClusterProjectionVIOSOPolicyFactory::Create(const FString& ProjectionPolicyId, const FDisplayClusterConfigurationProjection* InConfigurationProjectionPolicy)
 {
 	check(InConfigurationProjectionPolicy != nullptr);
 
diff --git a/Engine/Plugins/Runtime/nDisplay/Source/DisplayClusterProjection/Private/Policy/VIOSO/DisplayClusterProjectionVIOSOPolicyFactory.h b/Engine/Plugins/Runtime/nDisplay/Source/DisplayClusterProjection/Private/Policy/VIOSO/DisplayClusterProjectionVIOSOPolicyFactory.h
--- a/Engine/Plugins/Runtime/nDisplay/Source/DisplayClusterProjection/Private/Policy/VIOSO/DisplayClusterProjectionVIOSOPolicyFactory.h
+++ b/Engine/Plugins/Runtime/nDisplay/Source/DisplayClusterProjection/Private/Policy/VIOSO/DisplayClusterProjectionVIOSOPolicyFactory.h
@@ -5,6 +5,7 @@
 #include "Render/Projection/IDisplayClusterProjectionPolicyFactory.h"
 #include "Policy/DisplayClusterProjectionPolicyBase.h"
 
+
 /**
  * Implements projection policy factory for the 'VIOSO' policy
  */
@@ -19,5 +20,5 @@ public:
 	//////////////////////////////////////////////////////////////////////////////////////////////
 	// IDisplayClusterProjectionPolicyFactory
 	//////////////////////////////////////////////////////////////////////////////////////////////
-	virtual TSharedPtr<IDisplayClusterProjectionPolicy, ESPMode::ThreadSafe> Create(const FString& ProjectionPolicyId, const struct FDisplayClusterConfigurationProjection* InConfigurationProjectionPolicy) override;
+	virtual TSharedPtr<IDisplayClusterProjectionPolicy, ESPMode::ThreadSafe> Create(const FString& ProjectionPolicyId, const FDisplayClusterConfigurationProjection* InConfigurationProjectionPolicy) override;
 };
diff --git a/Engine/Plugins/Runtime/nDisplay/Source/DisplayClusterProjection/Private/Policy/VIOSO/ViosoPolicyConfiguration.cpp b/Engine/Plugins/Runtime/nDisplay/Source/DisplayClusterProjection/Private/Policy/VIOSO/ViosoPolicyConfiguration.cpp
--- a/Engine/Plugins/Runtime/nDisplay/Source/DisplayClusterProjection/Private/Policy/VIOSO/ViosoPolicyConfiguration.cpp
+++ b/Engine/Plugins/Runtime/nDisplay/Source/DisplayClusterProjection/Private/Policy/VIOSO/ViosoPolicyConfiguration.cpp
@@ -8,6 +8,7 @@
 #include "Policy/DisplayClusterProjectionPolicyBase.h"
 #include "Render/Viewport/IDisplayClusterViewport.h"
 
+
 //////////////////////////////////////////////////////////////////////////////////////////////
 // FViosoPolicyConfiguration
 //////////////////////////////////////////////////////////////////////////////////////////////
@@ -35,7 +36,7 @@ FString FViosoPolicyConfiguration::ToString() const
 	}
 }
 
-bool FViosoPolicyConfiguration::Initialize(const TMap<FString, FString>& InParameters, class IDisplayClusterViewport* InViewport)
+bool FViosoPolicyConfiguration::Initialize(const TMap<FString, FString>& InParameters, IDisplayClusterViewport* InViewport)
 {
 	check(InViewport);
 
diff --git a/Engine/Plugins/Runtime/nDisplay/Source/DisplayClusterProjection/Private/Policy/VIOSO/ViosoPolicyConfiguration.h b/Engine/Plugins/Runtime/nDisplay/Source/DisplayClusterProjection/Private/Policy/VIOSO/ViosoPolicyConfiguration.h
--- a/Engine/Plugins/Runtime/nDisplay/Source/DisplayClusterProjection/Private/Policy/VIOSO/ViosoPolicyConfiguration.h
+++ b/Engine/Plugins/Runtime/nDisplay/Source/DisplayClusterProjection/Private/Policy/VIOSO/ViosoPolicyConfiguration.h
@@ -4,6 +4,9 @@
 
 #include "CoreMinimal.h"
 
+class IDisplayClusterViewport;
+
+
 struct FViosoPolicyConfiguration
 {
 	FString OriginCompId;
@@ -20,7 +23,7 @@ struct FViosoPolicyConfiguration
 	//@todo add more vioso options, if required
 	float Gamma = 1.f;
 
-	bool Initialize(const TMap<FString, FString>& InParameters, class IDisplayClusterViewport* InViewport);
+	bool Initialize(const TMap<FString, FString>& InParameters, IDisplayClusterViewport* InViewport);
 
 	FString ToString() const;
 };
diff --git a/Engine/Plugins/Runtime/nDisplay/Source/DisplayClusterProjection/Private/Policy/VIOSO/Windows/DisplayClusterProjectionVIOSOPolicy.cpp b/Engine/Plugins/Runtime/nDisplay/Source/DisplayClusterProjection/Private/Policy/VIOSO/Windows/DisplayClusterProjectionVIOSOPolicy.cpp
--- a/Engine/Plugins/Runtime/nDisplay/Source/DisplayClusterProjection/Private/Policy/VIOSO/Windows/DisplayClusterProjectionVIOSOPolicy.cpp
+++ b/Engine/Plugins/Runtime/nDisplay/Source/DisplayClusterProjection/Private/Policy/VIOSO/Windows/DisplayClusterProjectionVIOSOPolicy.cpp
@@ -36,6 +36,7 @@
 #include "ID3D12DynamicRHI.h"
 #endif // VIOSO_USE_GRAPHICS_API_D3D12
 
+
 #if VIOSO_USE_GRAPHICS_API_D3D11
 /**
  * Class for caching D3D11 changed values on scope
@@ -134,7 +135,7 @@ private:
 //////////////////////////////////////////////////////////////////////////////////////////////
 // FDisplayClusterProjectionVIOSOPolicy
 //////////////////////////////////////////////////////////////////////////////////////////////
-FDisplayClusterProjectionVIOSOPolicy::FDisplayClusterProjectionVIOSOPolicy(const FString& ProjectionPolicyId, const struct FDisplayClusterConfigurationProjection* InConfigurationProjectionPolicy)
+FDisplayClusterProjectionVIOSOPolicy::FDisplayClusterProjectionVIOSOPolicy(const FString& ProjectionPolicyId, const FDisplayClusterConfigurationProjection* InConfigurationProjectionPolicy)
 	: FDisplayClusterProjectionPolicyBase(ProjectionPolicyId, InConfigurationProjectionPolicy)
 {
 	const ERHIInterfaceType RHIType = RHIGetInterfaceType();
@@ -159,7 +160,14 @@ FDisplayClusterProjectionVIOSOPolicy::~FDisplayClusterProjectionVIOSOPolicy()
 	ImplRelease();
 }
 
-bool FDisplayClusterProjectionVIOSOPolicy::HandleStartScene(class IDisplayClusterViewport* InViewport)
+
+const FString& FDisplayClusterProjectionVIOSOPolicy::GetType() const
+{
+	static const FString Type(DisplayClusterProjectionStrings::projection::VIOSO);
+	return Type;
+}
+
+bool FDisplayClusterProjectionVIOSOPolicy::HandleStartScene(IDisplayClusterViewport* InViewport)
 {
 	check(IsInGameThread());	
 
@@ -196,7 +204,7 @@ bool FDisplayClusterProjectionVIOSOPolicy::HandleStartScene(class IDisplayCluste
 	return true;
 }
 
-void FDisplayClusterProjectionVIOSOPolicy::HandleEndScene(class IDisplayClusterViewport* InViewport)
+void FDisplayClusterProjectionVIOSOPolicy::HandleEndScene(IDisplayClusterViewport* InViewport)
 {
 	check(IsInGameThread());
 
@@ -248,7 +256,7 @@ bool FDisplayClusterProjectionVIOSOPolicy::CalculateView(IDisplayClusterViewport
 	return true;
 }
 
-bool FDisplayClusterProjectionVIOSOPolicy::GetProjectionMatrix(class IDisplayClusterViewport* InViewport, const uint32 InContextNum, FMatrix& OutPrjMatrix)
+bool FDisplayClusterProjectionVIOSOPolicy::GetProjectionMatrix(IDisplayClusterViewport* InViewport, const uint32 InContextNum, FMatrix& OutPrjMatrix)
 {
 	check(IsInGameThread());
 
@@ -262,7 +270,7 @@ bool FDisplayClusterProjectionVIOSOPolicy::IsWarpBlendSupported()
 	return true;
 }
 
-void FDisplayClusterProjectionVIOSOPolicy::ApplyWarpBlend_RenderThread(FRHICommandListImmediate& RHICmdList, const class IDisplayClusterViewportProxy* InViewportProxy)
+void FDisplayClusterProjectionVIOSOPolicy::ApplyWarpBlend_RenderThread(FRHICommandListImmediate& RHICmdList, const IDisplayClusterViewportProxy* InViewportProxy)
 {
 	check(IsInRenderingThread());
 
@@ -273,7 +281,7 @@ void FDisplayClusterProjectionVIOSOPolicy::ApplyWarpBlend_RenderThread(FRHIComma
 	}
 }
 
-bool FDisplayClusterProjectionVIOSOPolicy::ImplApplyWarpBlend_RenderThread(FRHICommandListImmediate& RHICmdList, const class IDisplayClusterViewportProxy* InViewportProxy)
+bool FDisplayClusterProjectionVIOSOPolicy::ImplApplyWarpBlend_RenderThread(FRHICommandListImmediate& RHICmdList, const IDisplayClusterViewportProxy* InViewportProxy)
 {
 	check(IsInRenderingThread());
 
diff --git a/Engine/Plugins/Runtime/nDisplay/Source/DisplayClusterProjection/Private/Policy/VIOSO/Windows/DisplayClusterProjectionVIOSOPolicy.h b/Engine/Plugins/Runtime/nDisplay/Source/DisplayClusterProjection/Private/Policy/VIOSO/Windows/DisplayClusterProjectionVIOSOPolicy.h
--- a/Engine/Plugins/Runtime/nDisplay/Source/DisplayClusterProjection/Private/Policy/VIOSO/Windows/DisplayClusterProjectionVIOSOPolicy.h
+++ b/Engine/Plugins/Runtime/nDisplay/Source/DisplayClusterProjection/Private/Policy/VIOSO/Windows/DisplayClusterProjectionVIOSOPolicy.h
@@ -9,6 +9,7 @@
 #include "RHICommandList.h"
 #include "RHIResources.h"
 
+
 /**
  * VIOSO projection policy
  */
@@ -23,34 +24,37 @@ class FDisplayClusterProjectionVIOSOPolicy
 	};
 
 public:
-	FDisplayClusterProjectionVIOSOPolicy(const FString& ProjectionPolicyId, const struct FDisplayClusterConfigurationProjection* InConfigurationProjectionPolicy);
+	FDisplayClusterProjectionVIOSOPolicy(const FString& ProjectionPolicyId, const FDisplayClusterConfigurationProjection* InConfigurationProjectionPolicy);
 	virtual ~FDisplayClusterProjectionVIOSOPolicy();
 
-	virtual const FString GetTypeId() const
-	{ return DisplayClusterProjectionStrings::projection::VIOSO; }
-
 public:
 	//////////////////////////////////////////////////////////////////////////////////////////////
 	// IDisplayClusterProjectionPolicy
 	//////////////////////////////////////////////////////////////////////////////////////////////
-	virtual bool HandleStartScene(class IDisplayClusterViewport* InViewport) override;
-	virtual void HandleEndScene(class IDisplayClusterViewport* InViewport) override;
+	virtual const FString& GetType() const override;
+
+	virtual bool HandleStartScene(IDisplayClusterViewport* InViewport) override;
+	virtual void HandleEndScene(IDisplayClusterViewport* InViewport) override;
 
-	virtual bool CalculateView(class IDisplayClusterViewport* InViewport, const uint32 InContextNum, FVector& InOutViewLocation, FRotator& InOutViewRotation, const FVector& ViewOffset, const float WorldToMeters, const float NCP, const float FCP) override;
-	virtual bool GetProjectionMatrix(class IDisplayClusterViewport* InViewport, const uint32 InContextNum, FMatrix& OutPrjMatrix) override;
+	virtual bool CalculateView(IDisplayClusterViewport* InViewport, const uint32 InContextNum, FVector& InOutViewLocation, FRotator& InOutViewRotation, const FVector& ViewOffset, const float WorldToMeters, const float NCP, const float FCP) override;
+	virtual bool GetProjectionMatrix(IDisplayClusterViewport* InViewport, const uint32 InContextNum, FMatrix& OutPrjMatrix) override;
 
 	virtual bool IsWarpBlendSupported() override;
-	virtual void ApplyWarpBlend_RenderThread(FRHICommandListImmediate& RHICmdList, const class IDisplayClusterViewportProxy* InViewportProxy) override;
+	virtual void ApplyWarpBlend_RenderThread(FRHICommandListImmediate& RHICmdList, const IDisplayClusterViewportProxy* InViewportProxy) override;
 
 	// Request additional targetable resources for domeprojection  external warpblend
 	virtual bool ShouldUseAdditionalTargetableResource() const override
-	{ return true; }
+	{
+		return true;
+	}
 
 	virtual bool ShouldUseSourceTextureWithMips() const override
-	{ return true; }
+	{
+		return true;
+	}
 
 protected:
-	bool ImplApplyWarpBlend_RenderThread(FRHICommandListImmediate& RHICmdList, const class IDisplayClusterViewportProxy* InViewportProxy);
+	bool ImplApplyWarpBlend_RenderThread(FRHICommandListImmediate& RHICmdList, const IDisplayClusterViewportProxy* InViewportProxy);
 	void ImplRelease();
 
 protected:
@@ -66,7 +70,7 @@ protected:
 		bool Initialize(ERenderDevice RenderDevice, const FViosoPolicyConfiguration& InConfigData);
 		void DestroyVIOSO();
 
-		bool UpdateVIOSO(class IDisplayClusterViewport* InViewport, const uint32 InContextNum, const FVector& LocalLocation, const FRotator& LocalRotator, const float WorldToMeters, const float NCP, const float FCP);
+		bool UpdateVIOSO(IDisplayClusterViewport* InViewport, const uint32 InContextNum, const FVector& LocalLocation, const FRotator& LocalRotator, const float WorldToMeters, const float NCP, const float FCP);
 		bool RenderVIOSO_RenderThread(FRHICommandListImmediate& RHICmdList, FRHITexture2D* ShaderResourceTexture, FRHITexture2D* RenderTargetTexture, const FViosoPolicyConfiguration& InConfigData);
 
 	protected:
diff --git a/Engine/Plugins/Runtime/nDisplay/Source/DisplayClusterProjection/Private/Policy/VIOSO/Windows/DisplayClusterProjectionVIOSOWarper.h b/Engine/Plugins/Runtime/nDisplay/Source/DisplayClusterProjection/Private/Policy/VIOSO/Windows/DisplayClusterProjectionVIOSOWarper.h
--- a/Engine/Plugins/Runtime/nDisplay/Source/DisplayClusterProjection/Private/Policy/VIOSO/Windows/DisplayClusterProjectionVIOSOWarper.h
+++ b/Engine/Plugins/Runtime/nDisplay/Source/DisplayClusterProjection/Private/Policy/VIOSO/Windows/DisplayClusterProjectionVIOSOWarper.h
@@ -28,7 +28,7 @@ public:
 	bool Initialize(void* pDxDevice, const FViosoPolicyConfiguration &InConfigData);
 	void Release();
 
-	bool CalculateViewProjection(class IDisplayClusterViewport* InViewport, const uint32 InContextNum, FVector& InOutViewLocation, FRotator& InOutViewRotation, FMatrix& OutProjMatrix, const float WorldToMeters, const float NCP, const float FCP);
+	bool CalculateViewProjection(IDisplayClusterViewport* InViewport, const uint32 InContextNum, FVector& InOutViewLocation, FRotator& InOutViewRotation, FMatrix& OutProjMatrix, const float WorldToMeters, const float NCP, const float FCP);
 	bool Render(VWB_param RenderParam, VWB_uint StateMask);
 
 private:
@@ -42,7 +42,7 @@ private:
 	FVector GetViosoViewRotationEulerR_LHT() const;
 	FVector GetViosoViewRotationEulerR_RH() const;
 	FVector GetViosoViewLocation() const;
-	FMatrix GetProjMatrix(class IDisplayClusterViewport* InViewport, const uint32 InContextNum, const float NCP, const float FCP) const;
+	FMatrix GetProjMatrix(IDisplayClusterViewport* InViewport, const uint32 InContextNum, const float NCP, const float FCP) const;
 
 protected:
 	// Internal VIOSO data
diff --git a/Engine/Plugins/Runtime/nDisplay/Source/DisplayClusterProjection/Public/DisplayClusterProjectionStrings.h b/Engine/Plugins/Runtime/nDisplay/Source/DisplayClusterProjection/Public/DisplayClusterProjectionStrings.h
--- a/Engine/Plugins/Runtime/nDisplay/Source/DisplayClusterProjection/Public/DisplayClusterProjectionStrings.h
+++ b/Engine/Plugins/Runtime/nDisplay/Source/DisplayClusterProjection/Public/DisplayClusterProjectionStrings.h
@@ -25,16 +25,16 @@ namespace DisplayClusterProjectionStrings
 		{
 			static constexpr const TCHAR* Origin = TEXT("origin");
 
-			static constexpr const TCHAR* INIFile     = TEXT("inifile");
-			static constexpr const TCHAR* ChannelName = TEXT("channel");
+			static constexpr const TCHAR* INIFile      = TEXT("inifile");
+			static constexpr const TCHAR* ChannelName  = TEXT("channel");
 
 			static constexpr const TCHAR* File         = TEXT("file");
 			static constexpr const TCHAR* CalibIndex   = TEXT("index");
 			static constexpr const TCHAR* CalibAdapter = TEXT("adapter");
 
-			static constexpr const TCHAR* Gamma = TEXT("gamma");
+			static constexpr const TCHAR* Gamma        = TEXT("gamma");
 
-			static constexpr const TCHAR* BaseMatrix = TEXT("base");
+			static constexpr const TCHAR* BaseMatrix   = TEXT("base");
 		}
 
 		namespace manual
@@ -51,10 +51,10 @@ namespace DisplayClusterProjectionStrings
 			static constexpr const TCHAR* FrustumLeft  = TEXT("frustum_left");
 			static constexpr const TCHAR* FrustumRight = TEXT("frustum_right");
 
-			static constexpr const TCHAR* AngleL      = TEXT("l");
-			static constexpr const TCHAR* AngleR      = TEXT("r");
-			static constexpr const TCHAR* AngleT      = TEXT("t");
-			static constexpr const TCHAR* AngleB      = TEXT("b");
+			static constexpr const TCHAR* AngleL       = TEXT("l");
+			static constexpr const TCHAR* AngleR       = TEXT("r");
+			static constexpr const TCHAR* AngleT       = TEXT("t");
+			static constexpr const TCHAR* AngleB       = TEXT("b");
 			namespace FrustumType
 			{
 				static constexpr const TCHAR* Matrix = TEXT("Matrix");
@@ -62,8 +62,8 @@ namespace DisplayClusterProjectionStrings
 			}
 			namespace RenderingType
 			{
-				static constexpr const TCHAR* Mono = TEXT("Mono");
-				static constexpr const TCHAR* Stereo = TEXT("Stereo");
+				static constexpr const TCHAR* Mono       = TEXT("Mono");
+				static constexpr const TCHAR* Stereo     = TEXT("Stereo");
 				static constexpr const TCHAR* MonoStereo = TEXT("Mono & Stereo");
 			}
 		}
@@ -71,9 +71,9 @@ namespace DisplayClusterProjectionStrings
 		namespace camera
 		{
 			static constexpr const TCHAR* Component = TEXT("camera_component");
-			static constexpr const TCHAR* Native = TEXT("native");
+			static constexpr const TCHAR* Native    = TEXT("native");
 		}
-		
+
 		namespace mesh
 		{
 			static constexpr const TCHAR* Component        = TEXT("mesh_component");
@@ -120,13 +120,13 @@ namespace DisplayClusterProjectionStrings
 	namespace projection
 	{
 		static constexpr const TCHAR* Camera         = TEXT("camera");
-		static constexpr const TCHAR* Simple         = TEXT("simple");
-		static constexpr const TCHAR* MPCDI          = TEXT("mpcdi");
-		static constexpr const TCHAR* Mesh           = TEXT("mesh");
-		static constexpr const TCHAR* EasyBlend      = TEXT("easyblend");
-		static constexpr const TCHAR* VIOSO          = TEXT("vioso");
-		static constexpr const TCHAR* Manual         = TEXT("manual");
 		static constexpr const TCHAR* Domeprojection = TEXT("domeprojection");
+		static constexpr const TCHAR* EasyBlend      = TEXT("easyblend");
 		static constexpr const TCHAR* Link           = TEXT("link");
+		static constexpr const TCHAR* Manual         = TEXT("manual");
+		static constexpr const TCHAR* Mesh           = TEXT("mesh");
+		static constexpr const TCHAR* MPCDI          = TEXT("mpcdi");
+		static constexpr const TCHAR* Simple         = TEXT("simple");
+		static constexpr const TCHAR* VIOSO          = TEXT("vioso");
 	}
-};
\ No newline at end of file
+};
