nDisplay: fixed NonUnity build (CIS)

diff --git a/Engine/Plugins/Runtime/nDisplay/Source/DisplayClusterProjection/Private/Policy/Domeprojection/DisplayClusterProjectionDomeprojectionViewAdapterBase.h b/Engine/Plugins/Runtime/nDisplay/Source/DisplayClusterProjection/Private/Policy/Domeprojection/DisplayClusterProjectionDomeprojectionViewAdapterBase.h
--- a/Engine/Plugins/Runtime/nDisplay/Source/DisplayClusterProjection/Private/Policy/Domeprojection/DisplayClusterProjectionDomeprojectionViewAdapterBase.h
+++ b/Engine/Plugins/Runtime/nDisplay/Source/DisplayClusterProjection/Private/Policy/Domeprojection/DisplayClusterProjectionDomeprojectionViewAdapterBase.h
@@ -2,9 +2,10 @@
 
 #pragma once
 
+#include "RHI.h"
+
 class IDisplayClusterViewport;
 class IDisplayClusterViewportProxy;
-class FRHICommandListImmediate;
 
 
 /**
diff --git a/Engine/Plugins/Runtime/nDisplay/Source/DisplayClusterProjection/Private/Policy/EasyBlend/DisplayClusterProjectionEasyBlendViewAdapterBase.h b/Engine/Plugins/Runtime/nDisplay/Source/DisplayClusterProjection/Private/Policy/EasyBlend/DisplayClusterProjectionEasyBlendViewAdapterBase.h
--- a/Engine/Plugins/Runtime/nDisplay/Source/DisplayClusterProjection/Private/Policy/EasyBlend/DisplayClusterProjectionEasyBlendViewAdapterBase.h
+++ b/Engine/Plugins/Runtime/nDisplay/Source/DisplayClusterProjection/Private/Policy/EasyBlend/DisplayClusterProjectionEasyBlendViewAdapterBase.h
@@ -2,9 +2,10 @@
 
 #pragma once
 
+#include "RHI.h"
+
 class IDisplayClusterViewport;
 class IDisplayClusterViewportProxy;
-class FRHICommandListImmediate;
 
 
 /**
