nDisplay: Added the ability for light cards being moved in the LC editor to be flush with the stage's screens; the normals and depths used are blended to ensure a continuous transition between screens and empty space. Also added the ability to drag multiple light cards at once

diff --git a/Engine/Plugins/Runtime/nDisplay/Shaders/Private/MeshProjectionNormalSmoothing.usf b/Engine/Plugins/Runtime/nDisplay/Shaders/Private/MeshProjectionNormalSmoothing.usf
new file mode 100644
--- /dev/null
+++ b/Engine/Plugins/Runtime/nDisplay/Shaders/Private/MeshProjectionNormalSmoothing.usf
@@ -0,0 +1,240 @@
+// Copyright Epic Games, Inc. All Rights Reserved.
+
+#include "/Engine/Private/Common.ush"
+#include "/Engine/Private/ScreenPass.ush"
+
+SCREEN_PASS_TEXTURE_VIEWPORT(Input)
+
+Texture2D SceneColor;
+Texture2D SceneDepth;
+Texture2D<uint2> SceneStencil;
+
+RWTexture2D<float4> RWColor;
+RWTexture2D<float> RWDepth;
+RWTexture2D<uint> RWStencil;
+
+#define FILTER_KERNEL_SIZE 25
+
+DECLARE_SCALAR_ARRAY(float, SpatialKernel, FILTER_KERNEL_SIZE);
+DECLARE_SCALAR_ARRAY(float, InteriorSpatialKernel, FILTER_KERNEL_SIZE);
+
+float2 SampleDirection;
+float2 SampleOffsetScale;
+float InvSigma;
+
+/** Wraps the specified pixel position around the input viewport's extent */
+uint2 WrapPixelPos(uint2 PixelPos)
+{
+	return uint2(PixelPos.x % (uint)Input_Extent.x, PixelPos.y % (uint)Input_Extent.y);
+}
+
+/** Converts a vector in normal space to color space */
+float3 NormalToColorSpace(float3 Normal)
+{
+	return Normal * 0.5 + 0.5;
+}
+
+/** Converts a color from color space to normal space */
+float3 NormalFromColorSpace(float3 Color)
+{
+	return 2.0 * Color - 1.0;
+}
+
+/** Transforms the specified world vector into the radial space of the specified vector,
+  * oriented so that the x axis points in the direction of the radial vector */
+float3 TransformToRadialBasis(float3 Vector, float3 RadialVector)
+{
+	const float KINDA_SMALL_NUMBER = 0.0001;
+	float3 Up = abs(RadialVector.z) < (1.0 - KINDA_SMALL_NUMBER) ? float3(0, 0, 1) : float3(1, 0, 0);
+	
+	float3 AzimuthalVector = normalize(cross(Up, RadialVector));
+	float3 InclinationVector = cross(RadialVector, AzimuthalVector);
+	
+	return float3(dot(Vector, RadialVector), dot(Vector, AzimuthalVector), dot(Vector, InclinationVector));
+}
+
+/** A compute pass that copies the scene's color, depth, and stencil buffers into RW textures, as well as converting all normals from world space to radial space */
+[numthreads(8, 8, 1)]
+void CreateRWTexturesCS(uint2 DispatchThreadId : SV_DispatchThreadID)
+{
+	// In order to compute the radial space vector to convert the normals to radial space with, compute the world position of the current pixel
+	// This assumes the scene was rendered using an azimuthal projection
+	const uint2 PixelPos = DispatchThreadId;
+	const float2 UV = PixelPos * View.ViewSizeAndInvSize.zw;
+	const float2 ScreenPos = ViewportUVToScreenPos(UV);
+	
+	float3 ProjectedViewPos = mul(float4(ScreenPos.x * View.NearPlane, ScreenPos.y * View.NearPlane, 0.0, View.NearPlane), View.ClipToView).xyz;
+	
+	float Rho = length(ProjectedViewPos);
+	float3 UnitViewPos = normalize(ProjectedViewPos);
+	float3 PlanePos = UnitViewPos / UnitViewPos.z;
+	float2 PolarCoords = float2(sqrt(PlanePos.x * PlanePos.x + PlanePos.y * PlanePos.y), atan2(PlanePos.y, PlanePos.x));
+	
+	float3 ViewPos = float3(sin(PolarCoords.x) * cos(PolarCoords.y), sin(PolarCoords.x) * sin(PolarCoords.y), cos(PolarCoords.x)) * Rho;
+	float3 WorldPos = mul(float4(ViewPos, 0), View.ViewToTranslatedWorld).xyz;
+	
+	float3 RadialVector = normalize(WorldPos);
+	
+	float4 Color = SceneColor.Load(uint3(PixelPos, 0));
+	float Depth = SceneDepth.Load(uint3(PixelPos, 0)).r;
+	uint Stencil = SceneStencil.Load(uint3(PixelPos, 0)) STENCIL_COMPONENT_SWIZZLE;
+	
+	// If the current pixel is geometry, as indicated by the stencil, use the scene's color as the normal; otherwise,
+	// use the negative radial vector for the empty space normal
+	float3 Normal = lerp(-RadialVector, NormalFromColorSpace(Color.rgb), (float)Stencil);
+	float3 LocalNormal = TransformToRadialBasis(Normal, RadialVector);
+	
+	RWColor[PixelPos] = float4(NormalToColorSpace(LocalNormal), 1.0);
+	RWDepth[PixelPos] = Depth;
+	RWStencil[PixelPos] = Stencil;
+}
+
+/** A compute pass that dilates the depth buffer of overlapping geometry, allowing nearer geometry depth to dilate into further geometry depth */
+[numthreads(8, 8, 1)]
+void DepthDilationCS(uint2 DispatchThreadId : SV_DispatchThreadID)
+{
+	const uint2 PixelPos = DispatchThreadId;
+	
+	const int kRadius = (FILTER_KERNEL_SIZE - 1) / 2;
+
+	float InitialDepth = RWDepth[PixelPos];
+	uint Stencil = RWStencil[PixelPos];
+
+	float FinalDepth = 0.0;
+	float W = 0.0;
+	
+	// Only dilate onto geometry; if the current pixel's stencil is zero, the dilation can be skipped
+	if (Stencil > 0)
+	{
+		UNROLL
+		for (int kIndex = 0; kIndex < FILTER_KERNEL_SIZE; ++kIndex)
+		{
+			int2 PixelOffset = int2(kIndex - kRadius, kIndex - kRadius) * SampleOffsetScale * SampleDirection;
+			uint2 SamplePos = WrapPixelPos(PixelPos + PixelOffset);
+			
+			float Depth = RWDepth[SamplePos];
+			uint Stencil = RWStencil[SamplePos];
+			
+			// Use a spatial factor here to ensure the depth closest to the edge of the geometry is dilated outwards
+			float SpatialFactor = GET_SCALAR_ARRAY_ELEMENT(InteriorSpatialKernel, kIndex);
+			
+			// Only dilate overlapping geometry outwards if it is nearer to the view origin than the current pixel's depth
+			float DepthFactor = max(Depth - InitialDepth, 0) * Stencil * SpatialFactor;
+			
+			FinalDepth += Depth * DepthFactor;
+			W += DepthFactor;
+		}
+		
+		if (W > 0.0)
+		{
+			RWDepth[PixelPos] = FinalDepth / W;
+		}
+	}
+}
+
+/** A compute pass that dilates the normal, depth, and stencil buffer of geometry into empty space */
+[numthreads(8, 8, 1)]
+void DilationCS(uint2 DispatchThreadId : SV_DispatchThreadID)
+{
+	const uint2 PixelPos = DispatchThreadId;
+	
+	const int kRadius = (FILTER_KERNEL_SIZE - 1) / 2;
+	
+	float4 InitialColor = RWColor[PixelPos];
+	float InitialDepth = RWDepth[PixelPos];
+	uint InitialStencil = RWStencil[PixelPos];
+	
+	float3 FinalColor = float3(0.0, 0.0, 0.0);
+	float FinalDepth = 0.0;
+	uint FinalStencil = 0;
+	
+	float W = 0.0;
+	
+	// Only dilate into empty space; if the current pixel's stencil is non-zero, the dilation can be skipped
+	if (InitialStencil < 1)
+	{
+		UNROLL
+		for (int kIndex = 0; kIndex < FILTER_KERNEL_SIZE; ++kIndex)
+		{
+			int2 PixelOffset = int2(kIndex - kRadius, kIndex - kRadius) * SampleOffsetScale * SampleDirection;
+			uint2 SamplePos = WrapPixelPos(PixelPos + PixelOffset);
+		
+			float3 Color = RWColor[SamplePos].rgb;
+			float Depth = RWDepth[SamplePos];
+			uint Stencil = RWStencil[SamplePos];
+			
+			// Use a spatial factor here to ensure the depth closest to the edge of the geometry is dilated outwards
+			float SpatialFactor = GET_SCALAR_ARRAY_ELEMENT(SpatialKernel, kIndex);
+			
+			// Only dilate geometry into empty space
+			float StencilFactor = Stencil * SpatialFactor;
+
+			FinalColor += Color * StencilFactor;
+			FinalDepth += Depth * StencilFactor;
+			FinalStencil += Stencil;
+			
+			W += StencilFactor;
+		}
+	
+		if (W > 0.0)
+		{
+			RWColor[PixelPos] = float4(FinalColor / W, 1.0);
+			RWDepth[PixelPos] = FinalDepth / W;
+		}
+
+		RWStencil[PixelPos] = FinalStencil;
+	}
+}
+
+/** A compute pass that performs a standard Gaussian blur of the normal and depth buffers. Uses a different spatial deviation for blurring empty space or geometry */
+[numthreads(8, 8, 1)]
+void BlurCS(uint2 DispatchThreadId : SV_DispatchThreadID)
+{
+	const uint2 PixelPos = DispatchThreadId;
+	
+	const int kRadius = (FILTER_KERNEL_SIZE - 1) / 2;
+	
+	float3 InitialNormal = NormalFromColorSpace(RWColor[PixelPos].rgb);
+	float InitialDepth = RWDepth[PixelPos];
+	uint Stencil = SceneStencil.Load(int3(PixelPos, 0)) STENCIL_COMPONENT_SWIZZLE;
+	
+	float3 FinalNormal = float3(0.0, 0.0, 0.0);
+	float FinalDepth = 0.0;
+	float W = 0.0;
+	
+	UNROLL
+	for (int kIndex = 0; kIndex < FILTER_KERNEL_SIZE; ++kIndex)
+	{
+		int2 PixelOffset = int2(kIndex - kRadius, kIndex - kRadius) * SampleOffsetScale * SampleDirection;
+		uint2 SamplePos = WrapPixelPos(PixelPos + PixelOffset);
+		
+		float3 Normal = NormalFromColorSpace(RWColor[SamplePos].rgb);
+		float Depth = RWDepth[SamplePos];
+
+		// Get the correct spatial kernel depending on if we are blurring a screen or empty space
+		float Factor = lerp(GET_SCALAR_ARRAY_ELEMENT(SpatialKernel, kIndex), GET_SCALAR_ARRAY_ELEMENT(InteriorSpatialKernel, kIndex), Stencil);
+
+		FinalNormal += Normal * Factor;
+		FinalDepth += Depth * Factor;
+		W += Factor;
+	}
+	
+	FinalNormal /= W;
+	FinalDepth /= W;
+	
+	RWColor[PixelPos] = float4(NormalToColorSpace(FinalNormal), 1.0);
+	RWDepth[PixelPos] = FinalDepth;
+}
+
+/** A pixel shader that renders the results of any compute shader passes to an output render target, packing the normal and depth values into the correct RGBA channels */
+void OutputNormalMapPS(noperspective float4 UVAndScreenPos : TEXCOORD0, float4 SvPosition : SV_POSITION, out float4 OutColor : SV_Target0)
+{
+	const float2 UV = UVAndScreenPos.xy;
+	const uint2 PixelPos = int2(UV * Input_Extent);
+	
+	float3 NormalColor = RWColor[PixelPos].rgb;
+	float Depth = RWDepth[PixelPos];
+	
+	// Normal vector is placed into the RGB channel, while depth is placed into the A channel
+	OutColor = float4(NormalColor, Depth);
+}
\ No newline at end of file
diff --git a/Engine/Plugins/Runtime/nDisplay/Shaders/Private/MeshProjectionShaders.usf b/Engine/Plugins/Runtime/nDisplay/Shaders/Private/MeshProjectionShaders.usf
--- a/Engine/Plugins/Runtime/nDisplay/Shaders/Private/MeshProjectionShaders.usf
+++ b/Engine/Plugins/Runtime/nDisplay/Shaders/Private/MeshProjectionShaders.usf
@@ -63,6 +63,28 @@ void PerspectiveVS(FVertexFactoryInput Input, out FMeshProjectionVSToPS Output)
 	Output.Position = mul(Output.SavedWorldPosition, ResolvedView.TranslatedWorldToClip);
 }
 
+/** Gets a scalar value that indicates how much the current pixel diverges from where it should be from the projected vertices */
+float GetProjectionDivergence(FMeshProjectionVSToPS Input)
+{
+	// Compute the "actual" azimuth of the pixel based on its screen space position
+	// This will be subtracted from the expected azimuth interpolated from the vertices of the primitive
+	// to determine the divergence magnitude
+	
+	float4 ScreenPos = SvPositionToResolvedScreenPosition(Input.Position);
+	float3 ViewPos = mul(ScreenPos, ResolvedView.ClipToView).xyz;
+	float3 UnitPos = normalize(ViewPos);
+	float ActualAzimuth = length(UnitPos.xy) / UnitPos.z;
+	
+	return abs(ActualAzimuth - Input.Azimuth);
+}
+
+/** Clips any pixel whose projection divergence exceeds the specified threshold */
+void ClipDivergentPixel(FMeshProjectionVSToPS Input, float DivergenceThreshold)
+{
+	float Divergence = GetProjectionDivergence(Input);
+	clip(DivergenceThreshold - Divergence);
+}
+
 void MainPS(FMeshProjectionVSToPS Input OPTIONAL_IsFrontFace, out float4 OutputColor : SV_Target0)
 {
 	FMaterialPixelParameters MaterialParameters = GetMaterialPixelParameters(Input.FactoryInterpolants, Input.Position);
@@ -73,9 +95,9 @@ void MainPS(FMeshProjectionVSToPS Input OPTIONAL_IsFrontFace, out float4 OutputC
 	
 	const float3 Emissive = GetMaterialEmissive(PixelMaterialInputs);
 	const float Opacity = GetMaterialOpacity(PixelMaterialInputs);
-		
+	
 	OutputColor = float4(Emissive, Opacity);
-
+	
 	static const float Gamma = 2.2;
 	OutputColor.rgb = pow(OutputColor.rgb, Gamma);
 	
@@ -84,4 +106,22 @@ void MainPS(FMeshProjectionVSToPS Input OPTIONAL_IsFrontFace, out float4 OutputC
 	
 	const float AlphaBlendWeight = saturate((Input.Azimuth - AzimuthFeatherStart) / (AzimuthFeatherEnd - AzimuthFeatherStart));
 	OutputColor.a = lerp(OutputColor.a, 0, AlphaBlendWeight);
+}
+
+void NormalPS(FMeshProjectionVSToPS Input OPTIONAL_IsFrontFace, out float4 OutputColor : SV_Target0)
+{
+	ResolvedView = ResolveView();
+	
+	const float ClipThreshold = 0.25;
+	ClipDivergentPixel(Input, ClipThreshold);
+	
+	FMaterialPixelParameters MaterialParameters = GetMaterialPixelParameters(Input.FactoryInterpolants, Input.Position);
+	FPixelMaterialInputs PixelMaterialInputs;
+	CalcMaterialParameters(MaterialParameters, PixelMaterialInputs, Input.Position, bIsFrontFace);
+
+	GetMaterialCoverageAndClipping(MaterialParameters, PixelMaterialInputs);
+
+	float3 Normal = TransformTangentNormalToWorld(MaterialParameters.TangentToWorld, float3(0, 0, 1));
+	
+	OutputColor = float4(Normal * 0.5 + 0.5, 1.0f);
 }
\ No newline at end of file
diff --git a/Engine/Plugins/Runtime/nDisplay/Source/DisplayClusterLightCardEditor/Private/Viewport/DisplayClusterLightCardEditorViewportClient.cpp b/Engine/Plugins/Runtime/nDisplay/Source/DisplayClusterLightCardEditor/Private/Viewport/DisplayClusterLightCardEditorViewportClient.cpp
--- a/Engine/Plugins/Runtime/nDisplay/Source/DisplayClusterLightCardEditor/Private/Viewport/DisplayClusterLightCardEditorViewportClient.cpp
+++ b/Engine/Plugins/Runtime/nDisplay/Source/DisplayClusterLightCardEditor/Private/Viewport/DisplayClusterLightCardEditorViewportClient.cpp
@@ -37,12 +37,138 @@
 #include "Engine/Canvas.h"
 #include "Engine/TextureRenderTarget2D.h"
 #include "CameraController.h"
-
+#include "ImageUtils.h"
 
 #include "Renderer/Private/SceneRendering.h"
 
 #define LOCTEXT_NAMESPACE "DisplayClusterLightCardEditorViewportClient"
 
+
+//////////////////////////////////////////////////////////////////////////
+// FNormalMap
+
+const int32 FDisplayClusterLightCardEditorViewportClient::FNormalMap::NormalMapSize = 512;
+const float FDisplayClusterLightCardEditorViewportClient::FNormalMap::NormalMapFOV = 2.0f * FMath::RadiansToDegrees(FMath::Atan(0.55 * PI)); // Equation for FOV from desired angle from north pole;
+
+void FDisplayClusterLightCardEditorViewportClient::FNormalMap::Init(const FSceneViewInitOptions& InSceneViewInitOptions)
+{
+	SizeX = InSceneViewInitOptions.GetViewRect().Width();
+	SizeY = InSceneViewInitOptions.GetViewRect().Height();
+
+	ViewMatrices = FViewMatrices(InSceneViewInitOptions);
+
+	if (NormalMapTexture.IsValid())
+	{
+		NormalMapTexture->MarkAsGarbage();
+		NormalMapTexture = nullptr;
+	}
+
+	ENQUEUE_RENDER_COMMAND(InitRHIResourcesCommand)([this](FRHICommandListImmediate& RHICmdList)
+		{
+			const FRHITextureCreateDesc Desc =
+				FRHITextureCreateDesc::Create2D(TEXT("NormalMapTexture"))
+				.SetExtent(SizeX, SizeY)
+				.SetFormat(PF_FloatRGBA)
+				.SetClearValue(FClearValueBinding::Black);
+
+			RHICreateTargetableShaderResource(Desc, ETextureCreateFlags::RenderTargetable, RenderTargetTextureRHI);
+		});
+}
+
+void FDisplayClusterLightCardEditorViewportClient::FNormalMap::Release()
+{
+	ENQUEUE_RENDER_COMMAND(ReleaseRHIResourcesCommand)([this](FRHICommandListImmediate& RHICmdList)
+		{
+			RenderTargetTextureRHI.SafeRelease();
+		});
+}
+
+bool FDisplayClusterLightCardEditorViewportClient::FNormalMap::GetNormalAndDistanceAtPosition(FVector Position, FVector& OutNormal, float& OutDistance)
+{
+	auto GetPixel = [this](uint32 InX, uint32 InY)
+	{
+		uint32 WrappedX = InX % SizeX;
+		uint32 WrappedY = InY % SizeY;
+
+		return CachedNormalData[WrappedY * SizeX + WrappedX].GetFloats();
+	};
+
+	const FVector ViewPos = FVector(ViewMatrices.GetViewMatrix().TransformFVector4(FVector4(Position, 1)));
+	const FVector ProjectedViewPos = FDisplayClusterMeshProjectionRenderer::ProjectViewPosition(ViewPos, EDisplayClusterMeshProjectionType::Azimuthal);
+
+	const FVector4 ScreenPos = ViewMatrices.GetProjectionMatrix().TransformFVector4(FVector4(ProjectedViewPos, 1));
+
+	if (ScreenPos.W != 0.0)
+	{
+		const float InvW = (ScreenPos.W > 0.0f ? 1.0f : -1.0f) / ScreenPos.W;
+		const float Y = (GProjectionSignY > 0.0f) ? ScreenPos.Y : 1.0f - ScreenPos.Y;
+		const FVector2D PixelPos = FVector2D((0.5f + ScreenPos.X * 0.5f * InvW) * SizeX, (0.5f - Y * 0.5f * InvW) * SizeY);
+
+		// Perform a bilinear interpolation on the computed pixel position to ensure a continuous normal regardless of the resolution of the normal map
+		const uint32 PixelX = FMath::Floor(PixelPos.X - 0.5f);
+		const uint32 PixelY = FMath::Floor(PixelPos.Y - 0.5f);
+		const float PixelXFrac = FMath::Frac(PixelPos.X);
+		const float PixelYFrac = FMath::Frac(PixelPos.X);
+
+		FLinearColor NormalData;
+		NormalData = FMath::Lerp(
+			FMath::Lerp(GetPixel(PixelX, PixelY), GetPixel(PixelX + 1, PixelY), PixelXFrac),
+			FMath::Lerp(GetPixel(PixelX, PixelY + 1), GetPixel(PixelX + 1, PixelY + 1), PixelXFrac),
+			PixelYFrac);
+
+		const FVector NormalVector = 2.f * FVector(NormalData.R, NormalData.G, NormalData.B) - 1.f;
+		OutNormal = NormalVector.GetSafeNormal();
+
+		// Make sure the depth value is not 0, as that will cause a divide by zero when transformed, resulting in an NaN distance being returned
+		const float Depth = FMath::Max(0.001f, NormalData.A);
+
+		FVector4 DepthPos = ViewMatrices.GetInvProjectionMatrix().TransformFVector4(FVector4(ScreenPos.X * InvW, ScreenPos.Y * InvW, Depth, 1.0f));
+		DepthPos /= DepthPos.W;
+
+		const FVector UnprojectedDepthPos = FDisplayClusterMeshProjectionRenderer::UnprojectViewPosition(DepthPos, EDisplayClusterMeshProjectionType::Azimuthal);
+		OutDistance = UnprojectedDepthPos.Length();
+
+		return true;
+	}
+	else
+	{
+		OutNormal = FVector::ZeroVector;
+		OutDistance = 0.0;
+		return false;
+	}
+}
+
+UTexture2D* FDisplayClusterLightCardEditorViewportClient::FNormalMap::GenerateNormalMapTexture(const FString& TextureName)
+{
+	if (NormalMapTexture.IsValid())
+	{
+		NormalMapTexture->MarkAsGarbage();
+		NormalMapTexture = nullptr;
+	}
+
+	if (CachedNormalData.Num())
+	{
+		FCreateTexture2DParameters Params;
+		Params.bDeferCompression = true;
+
+		TArray<FColor> Bitmap;
+		Bitmap.AddZeroed(CachedNormalData.Num());
+
+		for (int32 Index = 0; Index < CachedNormalData.Num(); ++Index)
+		{
+			Bitmap[Index] = CachedNormalData[Index].GetFloats().ToFColor(false);
+		}
+
+		NormalMapTexture = FImageUtils::CreateTexture2D(SizeX, SizeY, Bitmap, GetTransientPackage(), TextureName, RF_Transient, Params);
+	}
+
+	return GetNormalMapTexture();
+}
+
+
+//////////////////////////////////////////////////////////////////////////
+// FDisplayClusterLightCardEditorViewportClient
+
 FDisplayClusterLightCardEditorViewportClient::FDisplayClusterLightCardEditorViewportClient(FAdvancedPreviewScene& InPreviewScene,
                                                            const TWeakPtr<SEditorViewport>& InEditorViewportWidget,
                                                            TWeakPtr<SDisplayClusterLightCardEditor> InLightCardEditor) :
@@ -168,6 +294,13 @@ void FDisplayClusterLightCardEditorViewportClient::Tick(float DeltaSeconds)
 
 void FDisplayClusterLightCardEditorViewportClient::Draw(FViewport* InViewport, FCanvas* Canvas)
 {
+	if (bNormalMapInvalid)
+	{
+		RenderNormalMap(NorthNormalMap, FVector::UpVector);
+		RenderNormalMap(SouthNormalMap, -FVector::UpVector);
+		bNormalMapInvalid = false;
+	}
+
 	FViewport* ViewportBackup = Viewport;
 	Viewport = InViewport ? InViewport : Viewport;
 
@@ -274,6 +407,27 @@ void FDisplayClusterLightCardEditorViewportClient::Draw(FViewport* InViewport, F
 		DrawCanvas(*Viewport, *View, *Canvas);
 	}
 
+	if (bDisplayNormalMapVisualization)
+	{
+		auto DrawNormalMap = [Canvas](FNormalMap& NormalMap, const FString& TextureName, FVector2D Position)
+		{
+			UTexture2D* NormalMapTexture = NormalMap.GetNormalMapTexture();
+
+			if (!NormalMapTexture)
+			{
+				NormalMapTexture = NormalMap.GenerateNormalMapTexture(TextureName);
+			}
+
+			if (NormalMapTexture)
+			{
+				Canvas->DrawTile(Position.X, Position.Y, 512, 512, 0, 0, 1, 1, FLinearColor::White, NormalMapTexture->GetResource());
+			}
+		};
+
+		DrawNormalMap(NorthNormalMap, TEXT("DisplayClusterLightCardEditor.NorthNormalMap"), FVector2D(0.0f, 0.0f));
+		DrawNormalMap(SouthNormalMap, TEXT("DisplayClusterLightCardEditor.SouthNormalMap"), FVector2D(0.0f, 512.0f));
+	}
+
 	// Remove temporary debug lines.
 	// Possibly a hack. Lines may get added without the scene being rendered etc.
 	if (World->LineBatcher != NULL && (World->LineBatcher->BatchedLines.Num() || World->LineBatcher->BatchedPoints.Num() || World->LineBatcher->BatchedMeshes.Num() ) )
@@ -571,7 +725,8 @@ void FDisplayClusterLightCardEditorViewportClient::UpdatePreviewActor(ADisplayCl
 	auto Finalize = [this]()
 	{
 		Viewport->InvalidateHitProxy();
-		bShouldCheckHitProxy = true;	
+		bShouldCheckHitProxy = true;
+		InvalidateNormalMap();
 	};
 	
 	if (RootActor == nullptr)
@@ -628,6 +783,8 @@ void FDisplayClusterLightCardEditorViewportClient::UpdatePreviewActor(ADisplayCl
 
 				RootActorProxy->UpdatePreviewComponents();
 				RootActorProxy->EnableEditorRender(false);
+
+				RootActorBoundingRadius = 0.5f * RootActorProxy->GetComponentsBoundingBox().GetSize().Length();
 			}
 
 			// Filter out any primitives hidden in game except screen components
@@ -802,6 +959,9 @@ void FDisplayClusterLightCardEditorViewportClient::SetProjectionModeFOV(EDisplay
 	{
 		ViewFOV = NewFOV;
 	}
+
+	Viewport->InvalidateHitProxy();
+	bShouldCheckHitProxy = true;
 }
 
 void FDisplayClusterLightCardEditorViewportClient::BeginTransaction(const FText& Description)
@@ -929,6 +1089,73 @@ void FDisplayClusterLightCardEditorViewportClient::GetSceneViewInitOptions(FScen
 	OutViewInitOptions = ViewInitOptions;
 }
 
+void FDisplayClusterLightCardEditorViewportClient::GetNormalMapSceneViewInitOptions(FIntPoint NormalMapSize, float NormalMapFOV, const FVector& ViewDirection, FSceneViewInitOptions& OutViewInitOptions)
+{
+	FViewportCameraTransform& ViewTransform = GetViewTransform();
+
+	OutViewInitOptions.ViewLocation = ProjectionOriginComponent.IsValid() ? ProjectionOriginComponent->GetComponentLocation() : FVector::ZeroVector;
+	OutViewInitOptions.ViewRotation = ViewDirection.Rotation();
+	OutViewInitOptions.ViewOrigin = OutViewInitOptions.ViewLocation;
+
+	OutViewInitOptions.SetViewRectangle(FIntRect(0, 0, NormalMapSize.X, NormalMapSize.Y));
+
+	AWorldSettings* WorldSettings = nullptr;
+	if (GetScene() != nullptr && GetScene()->GetWorld() != nullptr)
+	{
+		WorldSettings = GetScene()->GetWorld()->GetWorldSettings();
+	}
+
+	if (WorldSettings != nullptr)
+	{
+		OutViewInitOptions.WorldToMetersScale = WorldSettings->WorldToMeters;
+	}
+
+	// Rotate view 90 degrees
+	OutViewInitOptions.ViewRotationMatrix = FInverseRotationMatrix(OutViewInitOptions.ViewRotation) * FMatrix(
+		FPlane(0, 0, 1, 0),
+		FPlane(1, 0, 0, 0),
+		FPlane(0, 1, 0, 0),
+		FPlane(0, 0, 0, 1));
+
+	const float MinZ = GetNearClipPlane();
+	const float MaxZ = FMath::Max(RootActorBoundingRadius, MinZ);
+
+	// Avoid zero ViewFOV's which cause divide by zero's in projection matrix
+	const float MatrixFOV = FMath::Max(0.001f, NormalMapFOV) * (float)PI / 360.0f;
+
+	const float XAxisMultiplier = 1.0f;
+	const float YAxisMultiplier = 1.0f;
+
+	if ((bool)ERHIZBuffer::IsInverted)
+	{
+		OutViewInitOptions.ProjectionMatrix = FReversedZPerspectiveMatrix(
+			MatrixFOV,
+			MatrixFOV,
+			XAxisMultiplier,
+			YAxisMultiplier,
+			MinZ,
+			MaxZ);
+	}
+	else
+	{
+		OutViewInitOptions.ProjectionMatrix = FPerspectiveMatrix(
+			MatrixFOV,
+			MatrixFOV,
+			XAxisMultiplier,
+			YAxisMultiplier,
+			MinZ,
+			MaxZ);
+	}
+
+	OutViewInitOptions.SceneViewStateInterface = ViewState.GetReference();
+	OutViewInitOptions.ViewElementDrawer = this;
+
+	OutViewInitOptions.EditorViewBitflag = (uint64)1 << ViewIndex, // send the bit for this view - each actor will check it's visibility bits against this
+
+	OutViewInitOptions.FOV = NormalMapFOV;
+	OutViewInitOptions.OverrideFarClippingPlaneDistance = GetFarClipPlaneOverride();
+}
+
 UDisplayClusterConfigurationViewport* FDisplayClusterLightCardEditorViewportClient::FindViewportForPrimitiveComponent(UPrimitiveComponent* PrimitiveComponent)
 {
 	if (RootActorProxy.IsValid())
@@ -1144,62 +1371,72 @@ void FDisplayClusterLightCardEditorViewportClient::MoveSelectedLightCards(FViewp
 	FVector Direction;
 	PixelToWorld(*View, MousePos, Origin, Direction);
 
-	const TWeakObjectPtr<ADisplayClusterLightCardActor>& SelectedLightCard = SelectedLightCards.Last();
-	if (SelectedLightCard.IsValid())
+	const FVector CursorRayStart = Origin;
+	const FVector CursorRayEnd = CursorRayStart + Direction * HALF_WORLD_MAX;
+
+	const TWeakObjectPtr<ADisplayClusterLightCardActor>& LastSelectedLightCard = SelectedLightCards.Last();
+	if (LastSelectedLightCard.IsValid())
 	{
-		const FVector CurrentLightCardLocation = SelectedLightCard->GetLightCardTransform().GetTranslation();
+		const FSphericalCoordinates DeltaCoords = GetLightCardTranslationDelta(InViewport, LastSelectedLightCard.Get(), CurrentAxis);
 
-		// Light cards should be centered on the current view origin, so set the light card position to match the current view origin. Update the light card
-		// spherical coordinates to match its current coordinates
-		if (ProjectionOriginComponent.IsValid() && ProjectionOriginComponent->GetComponentLocation() != SelectedLightCard->GetActorLocation())
+		for (const TWeakObjectPtr<ADisplayClusterLightCardActor>& LightCard : SelectedLightCards)
 		{
-			const FVector DesiredLightCardOffset = CurrentLightCardLocation - ProjectionOriginComponent->GetComponentLocation();
-
-			SelectedLightCard->SetActorLocation(ProjectionOriginComponent->GetComponentLocation());
-
-			FSphericalCoordinates SphericalCoords(DesiredLightCardOffset);
+			if (!LightCard.IsValid())
+			{
+				continue;
+			}
 
-			SelectedLightCard->DistanceFromCenter = SphericalCoords.Radius;
-			SelectedLightCard->Longitude = FMath::RadiansToDegrees(SphericalCoords.Radius) - 180;
-			SelectedLightCard->Latitude = 90.f - FMath::RadiansToDegrees(SphericalCoords.Radius);
-		}
+			// Light cards should be centered on the current view origin, so set the light card position to match the current view origin. Update the light card
+			// spherical coordinates to match its current coordinates
+			if (ProjectionOriginComponent.IsValid() && ProjectionOriginComponent->GetComponentLocation() != LightCard->GetActorLocation())
+			{
+				const FVector DesiredLightCardOffset = LightCard->GetLightCardTransform().GetTranslation() - ProjectionOriginComponent->GetComponentLocation();
 
-		const FVector CursorRayStart = Origin;
-		const FVector CursorRayEnd = CursorRayStart + Direction * HALF_WORLD_MAX;
+				LightCard->SetActorLocation(ProjectionOriginComponent->GetComponentLocation());
 
-		FCollisionQueryParams Param(SCENE_QUERY_STAT(DragDropTrace), true);
+				FSphericalCoordinates SphericalCoords(DesiredLightCardOffset);
 
-		bool bHitScreen = false;
-		FHitResult ScreenHitResult;
-		if (PreviewWorld->LineTraceSingleByObjectType(ScreenHitResult, CursorRayStart, CursorRayEnd, FCollisionObjectQueryParams(FCollisionObjectQueryParams::InitType::AllObjects), Param))
-		{
-			if (AActor* HitActor = ScreenHitResult.GetActor())
-			{
-				if (RootActorProxy.Get() == HitActor && ScreenHitResult.Component.IsValid())
-				{
-					// TODO: Make the light card flush with the screen, which requires setting the LC's distance from center to just behind the screen, and setting the pitch and yaw of the card
-					// to match the screen's normal at the traced point
-				}
+				LightCard->DistanceFromCenter = SphericalCoords.Radius;
+				LightCard->Longitude = FMath::RadiansToDegrees(SphericalCoords.Radius) - 180;
+				LightCard->Latitude = 90.f - FMath::RadiansToDegrees(SphericalCoords.Radius);
 			}
-		}
 
-		// If we didn't hit a screen, keep the light card's radius fixed, and simply orbit it around the view origin
-		if (!bHitScreen)
-		{
-			const FSphericalCoordinates CurrentCoords = GetLightCardCoordinates(SelectedLightCard.Get());
-			const FSphericalCoordinates DeltaCoords = GetLightCardTranslationDelta(InViewport, SelectedLightCard.Get(), CurrentAxis);
+			const FSphericalCoordinates CurrentCoords = GetLightCardCoordinates(LightCard.Get());
 
 			FSphericalCoordinates NewCoords;
 			NewCoords.Radius = CurrentCoords.Radius + DeltaCoords.Radius;
 			NewCoords.Azimuth = CurrentCoords.Azimuth + DeltaCoords.Azimuth;
 			NewCoords.Inclination = CurrentCoords.Inclination + DeltaCoords.Inclination;
 
-			SelectedLightCard->DistanceFromCenter = NewCoords.Radius;
-			SelectedLightCard->Longitude = FRotator::ClampAxis(FMath::RadiansToDegrees(NewCoords.Azimuth) - 180);
-			SelectedLightCard->Latitude = 90.f - FMath::RadiansToDegrees(NewCoords.Inclination);
-		}
+			LightCard->DistanceFromCenter = NewCoords.Radius;
+			LightCard->Longitude = FRotator::ClampAxis(FMath::RadiansToDegrees(NewCoords.Azimuth) - 180);
+			LightCard->Latitude = 90.f - FMath::RadiansToDegrees(NewCoords.Inclination);
+
+			{
+				const FVector LightCardPosition = LightCard->GetLightCardTransform().GetTranslation();
+
+				FVector DesiredNormal;
+				float DesiredDistance;
 
-		PropagateLightCardTransform(SelectedLightCard.Get());
+				// If the light card is in the southern hemisphere of the view origin, use the southern normal map; otherwise, use the north normal map
+				if (LightCardPosition.Z < Origin.Z)
+				{
+					SouthNormalMap.GetNormalAndDistanceAtPosition(LightCardPosition, DesiredNormal, DesiredDistance);
+				}
+				else
+				{
+					NorthNormalMap.GetNormalAndDistanceAtPosition(LightCardPosition, DesiredNormal, DesiredDistance);
+				}
+
+				const FRotator Rotation = FRotationMatrix::MakeFromX(-DesiredNormal).Rotator();
+
+				LightCard->Pitch = Rotation.Pitch;
+				LightCard->Yaw = Rotation.Yaw;
+				LightCard->DistanceFromCenter = FMath::Min(DesiredDistance, RootActorBoundingRadius) + LightCardFlushOffset;
+			}
+
+			PropagateLightCardTransform(LightCard.Get());
+		}
 	}
 }
 
@@ -1227,8 +1464,21 @@ FDisplayClusterLightCardEditorViewportClient::FSphericalCoordinates FDisplayClus
 	const FVector LocalDirection = LightCard->GetActorRotation().RotateVector(Direction);
 	const FVector LightCardLocation = LightCard->GetLightCardTransform().GetTranslation() - Origin;
 
+	FVector Normal;
+	float Distance;
+
+	// If the light card is in the southern hemisphere of the view origin, use the southern normal map; otherwise, use the north normal map
+	if (LightCardLocation.Z < 0.0f)
+	{
+		SouthNormalMap.GetNormalAndDistanceAtPosition(LightCard->GetLightCardTransform().GetTranslation(), Normal, Distance);
+	}
+	else
+	{
+		NorthNormalMap.GetNormalAndDistanceAtPosition(LightCard->GetLightCardTransform().GetTranslation(), Normal, Distance);
+	}
+
 	FSphericalCoordinates LightCardCoords = GetLightCardCoordinates(LightCard);
-	FSphericalCoordinates RequestedCoords(LocalDirection * LightCardLocation.Size());
+	FSphericalCoordinates RequestedCoords(LocalDirection * Distance);
 	
 	FSphericalCoordinates DeltaCoords;
 	DeltaCoords.Radius = RequestedCoords.Radius - LightCardCoords.Radius;
@@ -1387,7 +1637,7 @@ bool FDisplayClusterLightCardEditorViewportClient::CalcEditorWidgetTransform(FTr
 
 	const FVector ProjectionOrigin = ProjectionOriginComponent.IsValid() ? ProjectionOriginComponent->GetComponentLocation() : FVector::ZeroVector;
 	const FVector RadialVector = (LightCardPosition - ProjectionOrigin).GetSafeNormal();
-	const FVector AzimuthalVector = FVector::ZAxisVector ^ RadialVector;
+	const FVector AzimuthalVector = (FVector::ZAxisVector ^ RadialVector).GetSafeNormal();
 	const FVector InclinationVector = RadialVector ^ AzimuthalVector;
 
 	FRotator Orientation = FMatrix(AzimuthalVector, InclinationVector, RadialVector, FVector::ZeroVector).Rotator();
@@ -1398,4 +1648,37 @@ bool FDisplayClusterLightCardEditorViewportClient::CalcEditorWidgetTransform(FTr
 	return true;
 }
 
+void FDisplayClusterLightCardEditorViewportClient::RenderNormalMap(FNormalMap& NormalMap, const FVector& NormalMapDirection)
+{
+	FSceneViewInitOptions ViewInitOptions;
+	GetNormalMapSceneViewInitOptions(FIntPoint(FNormalMap::NormalMapSize), FNormalMap::NormalMapFOV, NormalMapDirection, ViewInitOptions);
+
+	NormalMap.Init(ViewInitOptions);
+
+	// Only render primitive components from the stage actor for the normal map
+	FDisplayClusterMeshProjectionPrimitiveFilter PrimitiveFilter;
+	PrimitiveFilter.PrimitiveFilterDelegate = FDisplayClusterMeshProjectionPrimitiveFilter::FPrimitiveFilter::CreateLambda([this](const UPrimitiveComponent* PrimitiveComponent)
+	{
+		return PrimitiveComponent->GetOwner() == RootActorProxy;
+	});
+
+	FCanvas Canvas(&NormalMap, nullptr, GetWorld(), GetScene()->GetFeatureLevel(), FCanvas::CDM_DeferDrawing, 1.0f);
+	{
+		Canvas.Clear(FLinearColor::Black);
+
+		MeshProjectionRenderer->RenderNormals(&Canvas, GetScene(), ViewInitOptions, EngineShowFlags, EDisplayClusterMeshProjectionType::Azimuthal, &PrimitiveFilter);
+	}
+	Canvas.Flush_GameThread();
+
+	NormalMap.ReadFloat16Pixels(NormalMap.GetCachedNormalData());
+	NormalMap.Release();
+
+	FlushRenderingCommands();
+}
+
+void FDisplayClusterLightCardEditorViewportClient::InvalidateNormalMap()
+{
+	bNormalMapInvalid = true;
+}
+
 #undef LOCTEXT_NAMESPACE
diff --git a/Engine/Plugins/Runtime/nDisplay/Source/DisplayClusterLightCardEditor/Private/Viewport/DisplayClusterLightCardEditorViewportClient.h b/Engine/Plugins/Runtime/nDisplay/Source/DisplayClusterLightCardEditor/Private/Viewport/DisplayClusterLightCardEditorViewportClient.h
--- a/Engine/Plugins/Runtime/nDisplay/Source/DisplayClusterLightCardEditor/Private/Viewport/DisplayClusterLightCardEditorViewportClient.h
+++ b/Engine/Plugins/Runtime/nDisplay/Source/DisplayClusterLightCardEditor/Private/Viewport/DisplayClusterLightCardEditorViewportClient.h
@@ -21,7 +21,6 @@ class UDisplayClusterConfigurationViewport;
 class FDisplayClusterLightCardEditorViewportClient : public FEditorViewportClient, public TSharedFromThis<FDisplayClusterLightCardEditorViewportClient>
 {
 private:
-
 	struct FSphericalCoordinates
 	{
 	public:
@@ -52,6 +51,58 @@ private:
 		float Azimuth = 0.f;
 	};
 
+	/** Custom render target that stores the normal data for the stage */
+	class FNormalMap : public FRenderTarget
+	{
+	public:
+		/** The size of the normal map */
+		static const int32 NormalMapSize;
+
+		/** The field of view to render the normal map with. Using the azimuthal projection,
+		 * this is set so that the entire 360 degree scene is rendered */
+		static const float NormalMapFOV;
+
+		/** Initializes the normal map render target using the specified scene view options */
+		void Init(const FSceneViewInitOptions& InSceneViewInitOptions);
+
+		/** Releases the normal map render target's resources */
+		void Release();
+
+		/** Gets the size of the render target */
+		virtual FIntPoint GetSizeXY() const override { return FIntPoint(SizeX, SizeY); }
+
+		/** Gets a reference to the normal map data array, which stores the normal vector in the RGB components (color = 0.5 * Normal + 0.5) and the depth in the A component */
+		TArray<FFloat16Color>& GetCachedNormalData() { return CachedNormalData; }
+
+		/** Gets the normal vector and distance at the specified world location. The normal and distance are bilinearly interpolated from the nearest pixels in the normal map */
+		bool GetNormalAndDistanceAtPosition(FVector Position, FVector& OutNormal, float& OutDistance);
+
+		/** Generates a texture object that can be used to visualize the normal map */
+		UTexture2D* GenerateNormalMapTexture(const FString& TextureName);
+
+		/** Gets the normal map visualization texture, or null if it hasn't been generated */
+		UTexture2D* GetNormalMapTexture() const { return NormalMapTexture.IsValid() ? NormalMapTexture.Get() : nullptr; }
+
+	private:
+		/** The view matrices used when the normal map was last rendered */
+		FViewMatrices ViewMatrices;
+
+		/** The cached normal map data from the last normal map render */
+		TArray<FFloat16Color> CachedNormalData;
+
+		/** A texture that contains the normal map, for visualization purposes */
+		TWeakObjectPtr<UTexture2D> NormalMapTexture;
+
+		/** The width of the normal map. */
+		uint32 SizeX = 0;
+
+		/** The height of the normal map. */
+		uint32 SizeY = 0;
+	};
+
+	/** The amount to offset a light card's DistanceFromCenter when making it flush with a screen */
+	static constexpr float LightCardFlushOffset = -0.5f;
+
 public:
 	FDisplayClusterLightCardEditorViewportClient(FAdvancedPreviewScene& InPreviewScene, const TWeakPtr<SEditorViewport>& InEditorViewportWidget,
 		TWeakPtr<SDisplayClusterLightCardEditor> InLightCardEditor);
@@ -117,6 +168,9 @@ private:
 	/** Gets the scene view init options to use to create scene views for the preview scene */
 	void GetSceneViewInitOptions(FSceneViewInitOptions& OutViewInitOptions);
 
+	/** Gets the scene view init options to use when rendering the normal map cache */
+	void GetNormalMapSceneViewInitOptions(FIntPoint NormalMapSize, float NormalMapFOV, const FVector& ViewDirection, FSceneViewInitOptions& OutViewInitOptions);
+
 	/** Gets the viewport that is attached to the specified primitive component */
 	UDisplayClusterConfigurationViewport* FindViewportForPrimitiveComponent(UPrimitiveComponent* PrimitiveComponent);
 
@@ -156,12 +210,21 @@ private:
 	/** Calculates the world transform to render the editor widget with to align it with the selected light card */
 	bool CalcEditorWidgetTransform(FTransform& WidgetTransformBeforeMapProjection, FTransform& WidgetTransformAfterMapProjection);
 
+	/** Renders the viewport's normal map and stores the texture data to be used later */
+	void RenderNormalMap(FNormalMap& NormalMap, const FVector& NormalMapDirection);
+
+	/** Invalidates the viewport's normal map, forcing it to be rerendered on the next draw call */
+	void InvalidateNormalMap();
+
 private:
 	TWeakPtr<FSceneViewport> SceneViewportPtr;
 	TWeakPtr<SDisplayClusterLightCardEditor> LightCardEditorPtr;
 	TWeakObjectPtr<ADisplayClusterRootActor> RootActorProxy;
 	TWeakObjectPtr<ADisplayClusterRootActor> RootActorLevelInstance;
 
+	/** The radius of the bounding sphere that entirely encapsulates the root actor */
+	float RootActorBoundingRadius = 0.0f;
+
 	struct FLightCardProxy
 	{
 		TWeakObjectPtr<ADisplayClusterLightCardActor> LevelInstance;
@@ -207,4 +270,16 @@ private:
 
 	/** The increment to change the FOV by when using the scroll wheel */
 	float FOVScrollIncrement = 5.0f;
+
+	/** The render target used to render a map of the screens' normals for the northern hemisphere of the view */
+	FNormalMap NorthNormalMap;
+
+	/** The render target used to render a map of the screens' normals for the southern hemisphere of the view */
+	FNormalMap SouthNormalMap;
+
+	/** Indicates if the cached normal map is invalid and needs to be redrawn */
+	bool bNormalMapInvalid = false;
+
+	/** Indicates if the normal map should be displayed to the screen */
+	bool bDisplayNormalMapVisualization = false;
 };
\ No newline at end of file
diff --git a/Engine/Plugins/Runtime/nDisplay/Source/DisplayClusterLightCardEditorShaders/Private/DisplayClusterMeshProjectionRenderer.cpp b/Engine/Plugins/Runtime/nDisplay/Source/DisplayClusterLightCardEditorShaders/Private/DisplayClusterMeshProjectionRenderer.cpp
--- a/Engine/Plugins/Runtime/nDisplay/Source/DisplayClusterLightCardEditorShaders/Private/DisplayClusterMeshProjectionRenderer.cpp
+++ b/Engine/Plugins/Runtime/nDisplay/Source/DisplayClusterLightCardEditorShaders/Private/DisplayClusterMeshProjectionRenderer.cpp
@@ -48,6 +48,13 @@ using FMeshAzimuthalProjectionVS = FMeshProjectionVS<EDisplayClusterMeshProjecti
 IMPLEMENT_MATERIAL_SHADER_TYPE(template<>, FMeshPerspectiveProjectionVS, TEXT("/Plugin/nDisplay/Private/MeshProjectionShaders.usf"), TEXT("PerspectiveVS"), SF_Vertex);
 IMPLEMENT_MATERIAL_SHADER_TYPE(template<>, FMeshAzimuthalProjectionVS, TEXT("/Plugin/nDisplay/Private/MeshProjectionShaders.usf"), TEXT("AzimuthalVS"), SF_Vertex);
 
+enum EMeshProjectionShaderOutput
+{
+	Default,
+	Normals
+};
+
+template<EMeshProjectionShaderOutput OutputType>
 class FMeshProjectionPS : public FMeshMaterialShader
 {
 public:
@@ -68,7 +75,11 @@ public:
 	}
 };
 
-IMPLEMENT_MATERIAL_SHADER_TYPE(, FMeshProjectionPS, TEXT("/Plugin/nDisplay/Private/MeshProjectionShaders.usf"), TEXT("MainPS"), SF_Pixel);
+using FMeshProjectionDefaultPS = FMeshProjectionPS<EMeshProjectionShaderOutput::Default>;
+using FMeshProjectionNormalPS = FMeshProjectionPS<EMeshProjectionShaderOutput::Normals>;
+
+IMPLEMENT_MATERIAL_SHADER_TYPE(template<>, FMeshProjectionDefaultPS, TEXT("/Plugin/nDisplay/Private/MeshProjectionShaders.usf"), TEXT("MainPS"), SF_Pixel);
+IMPLEMENT_MATERIAL_SHADER_TYPE(template<>, FMeshProjectionNormalPS, TEXT("/Plugin/nDisplay/Private/MeshProjectionShaders.usf"), TEXT("NormalPS"), SF_Pixel);
 
 BEGIN_SHADER_PARAMETER_STRUCT(FMeshProjectionPassParameters, )
 	SHADER_PARAMETER_STRUCT_REF(FViewUniformShaderParameters, View)
@@ -77,7 +88,7 @@ BEGIN_SHADER_PARAMETER_STRUCT(FMeshProjectionPassParameters, )
 	RENDER_TARGET_BINDING_SLOTS()
 END_SHADER_PARAMETER_STRUCT()
 
-template<EDisplayClusterMeshProjectionType ProjectionType> 
+template<EDisplayClusterMeshProjectionType ProjectionType, EMeshProjectionShaderOutput OutputType = EMeshProjectionShaderOutput::Default> 
 class FLightCardEditorMeshPassProcessor : public FMeshPassProcessor
 {
 public:
@@ -85,6 +96,7 @@ public:
 		: FMeshPassProcessor(InScene, GMaxRHIFeatureLevel, InView, InDrawListContext)
 		, DrawRenderState(*InView)
 		, bTranslucencyPass(bIsTranslucencyPass)
+		, bIgnoreTranslucency(false)
 	{
 
 		if (bTranslucencyPass)
@@ -94,7 +106,7 @@ public:
 		}
 		else
 		{
-			DrawRenderState.SetDepthStencilState(TStaticDepthStencilState<true, CF_DepthNearOrEqual>::GetRHI());
+			DrawRenderState.SetDepthStencilState(TStaticDepthStencilState<true, CF_DepthNearOrEqual, true, CF_Always, SO_Keep, SO_Keep, SO_Replace>::GetRHI());
 			DrawRenderState.SetBlendState(TStaticBlendStateWriteMask<CW_RGBA, CW_RGBA, CW_RGBA, CW_RGBA>::GetRHI());
 		}
 	}
@@ -112,11 +124,11 @@ public:
 
 		const FVertexFactory* VertexFactory = MeshBatch.VertexFactory;
 
-		TMeshProcessorShaders<FMeshProjectionVS<ProjectionType>, FMeshProjectionPS> PassShaders;
+		TMeshProcessorShaders<FMeshProjectionVS<ProjectionType>, FMeshProjectionPS<OutputType>> PassShaders;
 
 		FMaterialShaderTypes ShaderTypes;
 		ShaderTypes.AddShaderType<FMeshProjectionVS<ProjectionType>>();
-		ShaderTypes.AddShaderType<FMeshProjectionPS>();
+		ShaderTypes.AddShaderType<FMeshProjectionPS<OutputType>>();
 
 		FMaterialShaders Shaders;
 		if (!Material.TryGetShaders(ShaderTypes, VertexFactory->GetType(), Shaders))
@@ -136,6 +148,8 @@ public:
 
 		FMeshDrawCommandSortKey SortKey = CreateMeshSortKey(MeshBatch, PrimitiveSceneProxy, Material, PassShaders.VertexShader.GetShader(), PassShaders.PixelShader.GetShader());
 
+		DrawRenderState.SetStencilRef(StencilValue);
+
 		BuildMeshDrawCommands(
 			MeshBatch,
 			BatchElementMask,
@@ -151,6 +165,16 @@ public:
 			ShaderElementData);
 	}
 
+	void SetStencilValue(uint32 InStencilValue)
+	{
+		StencilValue = InStencilValue;
+	}
+
+	void SetIgnoreTranslucency(bool bInIgnoreTranslucency)
+	{
+		bIgnoreTranslucency = bInIgnoreTranslucency;
+	}
+
 private:
 	bool CanDrawMeshBatch(const FMeshBatch& RESTRICT MeshBatch, const FMaterial& Material)
 	{
@@ -162,7 +186,7 @@ private:
 		}
 		else
 		{
-			return !bIsTranslucent;
+			return bIgnoreTranslucency || !bIsTranslucent;
 		}
 	}
 
@@ -202,7 +226,9 @@ private:
 
 private:
 	FMeshPassProcessorRenderState DrawRenderState;
+	uint32 StencilValue;
 	bool bTranslucencyPass;
+	bool bIgnoreTranslucency;
 };
 
 
@@ -347,6 +373,124 @@ private:
 	FMeshPassProcessorRenderState DrawRenderState;
 };
 
+//////////////////////////////////////////////////////////////////////////
+// Normals Render Pass
+
+class FMeshProjectionNormalsCreateRWTexturesCS : public FGlobalShader
+{
+public:
+	DECLARE_GLOBAL_SHADER(FMeshProjectionNormalsCreateRWTexturesCS);
+	SHADER_USE_PARAMETER_STRUCT(FMeshProjectionNormalsCreateRWTexturesCS, FGlobalShader);
+
+	BEGIN_SHADER_PARAMETER_STRUCT(FParameters, )
+		SHADER_PARAMETER_STRUCT_REF(FViewUniformShaderParameters, View)
+		SHADER_PARAMETER_RDG_TEXTURE(Texture2D, SceneColor)
+		SHADER_PARAMETER_RDG_TEXTURE(Texture2D, SceneDepth)
+		SHADER_PARAMETER_RDG_TEXTURE_SRV(Texture2D, SceneStencil)
+		SHADER_PARAMETER_RDG_TEXTURE_UAV(RWTexture2D<float4>, RWColor)
+		SHADER_PARAMETER_RDG_TEXTURE_UAV(RWTexture2D<float>, RWDepth)
+		SHADER_PARAMETER_RDG_TEXTURE_UAV(RWTexture2D<uint>, RWStencil)
+	END_SHADER_PARAMETER_STRUCT()
+
+	static bool ShouldCompilePermutation(const FGlobalShaderPermutationParameters& Parameters)
+	{
+		return IsPCPlatform(Parameters.Platform) && EnumHasAllFlags(Parameters.Flags, EShaderPermutationFlags::HasEditorOnlyData);
+	}
+
+	static void ModifyCompilationEnvironment(const FGlobalShaderPermutationParameters& Parameters, FShaderCompilerEnvironment& OutEnvironment)
+	{
+		FGlobalShader::ModifyCompilationEnvironment(Parameters, OutEnvironment);
+		OutEnvironment.CompilerFlags.Add(CFLAG_AllowTypedUAVLoads);
+	}
+};
+
+IMPLEMENT_GLOBAL_SHADER(FMeshProjectionNormalsCreateRWTexturesCS, "/Plugin/nDisplay/Private/MeshProjectionNormalSmoothing.usf", "CreateRWTexturesCS", SF_Compute);
+
+#define FILTER_KERNEL_SIZE 25
+
+enum EMeshProjectionFilterType
+{
+	DepthDilate,
+	Dilate,
+	Blur
+};
+
+BEGIN_SHADER_PARAMETER_STRUCT(FMeshProjectionNormalsFilterParameters, )
+	SHADER_PARAMETER_STRUCT_REF(FViewUniformShaderParameters, View)
+	SHADER_PARAMETER_STRUCT(FScreenPassTextureViewportParameters, Input)
+	SHADER_PARAMETER_RDG_TEXTURE_SRV(Texture2D, SceneStencil)
+	SHADER_PARAMETER_RDG_TEXTURE_UAV(RWTexture2D<float4>, RWColor)
+	SHADER_PARAMETER_RDG_TEXTURE_UAV(RWTexture2D<float>, RWDepth)
+	SHADER_PARAMETER_RDG_TEXTURE_UAV(RWTexture2D<uint>, RWStencil)
+	SHADER_PARAMETER_SCALAR_ARRAY(float, SpatialKernel, [FILTER_KERNEL_SIZE])
+	SHADER_PARAMETER_SCALAR_ARRAY(float, InteriorSpatialKernel, [FILTER_KERNEL_SIZE])
+	SHADER_PARAMETER(FVector2f, SampleDirection)
+	SHADER_PARAMETER(FVector2f, SampleOffsetScale)
+END_SHADER_PARAMETER_STRUCT()
+
+template<EMeshProjectionFilterType FilterType>
+class FMeshProjectionNormalsFilterCS : public FGlobalShader
+{
+public:
+	using FParameters = FMeshProjectionNormalsFilterParameters;
+
+	DECLARE_GLOBAL_SHADER(FMeshProjectionNormalsFilterCS);
+
+	FMeshProjectionNormalsFilterCS() { }
+	FMeshProjectionNormalsFilterCS(const ShaderMetaType::CompiledShaderInitializerType& Initializer)
+		: FGlobalShader(Initializer)
+	{
+		BindForLegacyShaderParameters<FParameters>(this, Initializer.PermutationId, Initializer.ParameterMap, true);
+	}
+
+	static inline const FShaderParametersMetadata* GetRootParametersMetadata() { return FParameters::FTypeInfo::GetStructMetadata(); }
+
+	static bool ShouldCompilePermutation(const FGlobalShaderPermutationParameters& Parameters)
+	{
+		return IsPCPlatform(Parameters.Platform) && EnumHasAllFlags(Parameters.Flags, EShaderPermutationFlags::HasEditorOnlyData);
+	}
+
+	static void ModifyCompilationEnvironment(const FGlobalShaderPermutationParameters& Parameters, FShaderCompilerEnvironment& OutEnvironment)
+	{
+		FGlobalShader::ModifyCompilationEnvironment(Parameters, OutEnvironment);
+		OutEnvironment.CompilerFlags.Add(CFLAG_AllowTypedUAVLoads);
+	}
+};
+
+using FMeshProjectionNormalsDepthDilationCS = FMeshProjectionNormalsFilterCS<EMeshProjectionFilterType::DepthDilate>;
+using FMeshProjectionNormalsDilationCS = FMeshProjectionNormalsFilterCS<EMeshProjectionFilterType::Dilate>;
+using FMeshProjectionNormalsBlurCS = FMeshProjectionNormalsFilterCS<EMeshProjectionFilterType::Blur>;
+
+IMPLEMENT_SHADER_TYPE(template<>, FMeshProjectionNormalsDepthDilationCS, TEXT("/Plugin/nDisplay/Private/MeshProjectionNormalSmoothing.usf"), TEXT("DepthDilationCS"), SF_Compute);
+IMPLEMENT_SHADER_TYPE(template<>, FMeshProjectionNormalsDilationCS, TEXT("/Plugin/nDisplay/Private/MeshProjectionNormalSmoothing.usf"), TEXT("DilationCS"), SF_Compute);
+IMPLEMENT_SHADER_TYPE(template<>, FMeshProjectionNormalsBlurCS, TEXT("/Plugin/nDisplay/Private/MeshProjectionNormalSmoothing.usf"), TEXT("BlurCS"), SF_Compute);
+
+class FMeshProjectionNormalsOutputPS : public FGlobalShader
+{
+public:
+	DECLARE_GLOBAL_SHADER(FMeshProjectionNormalsOutputPS);
+	SHADER_USE_PARAMETER_STRUCT(FMeshProjectionNormalsOutputPS, FGlobalShader);
+
+	BEGIN_SHADER_PARAMETER_STRUCT(FParameters, )
+		SHADER_PARAMETER_STRUCT(FScreenPassTextureViewportParameters, Input)
+		SHADER_PARAMETER_RDG_TEXTURE_UAV(RWTexture2D<float4>, RWColor)
+		SHADER_PARAMETER_RDG_TEXTURE_UAV(RWTexture2D<float>, RWDepth)
+		RENDER_TARGET_BINDING_SLOTS()
+	END_SHADER_PARAMETER_STRUCT()
+
+	static bool ShouldCompilePermutation(const FGlobalShaderPermutationParameters& Parameters)
+	{
+		return IsPCPlatform(Parameters.Platform) && EnumHasAllFlags(Parameters.Flags, EShaderPermutationFlags::HasEditorOnlyData);
+	}
+
+	static void ModifyCompilationEnvironment(const FGlobalShaderPermutationParameters& Parameters, FShaderCompilerEnvironment& OutEnvironment)
+	{
+		FGlobalShader::ModifyCompilationEnvironment(Parameters, OutEnvironment);
+		OutEnvironment.CompilerFlags.Add(CFLAG_AllowTypedUAVLoads);
+	}
+};
+
+IMPLEMENT_GLOBAL_SHADER(FMeshProjectionNormalsOutputPS, "/Plugin/nDisplay/Private/MeshProjectionNormalSmoothing.usf", "OutputNormalMapPS", SF_Pixel);
 
 //////////////////////////////////////////////////////////////////////////
 // Selection Outline Render Pass
@@ -419,11 +563,11 @@ public:
 				
 				const FVertexFactory* VertexFactory = MeshBatch.VertexFactory;
 
-				TMeshProcessorShaders<FMeshProjectionVS<ProjectionType>, FMeshProjectionPS> PassShaders;
+				TMeshProcessorShaders<FMeshProjectionVS<ProjectionType>, FMeshProjectionDefaultPS> PassShaders;
 
 				FMaterialShaderTypes ShaderTypes;
 				ShaderTypes.AddShaderType<FMeshProjectionVS<ProjectionType>>();
-				ShaderTypes.AddShaderType<FMeshProjectionPS>();
+				ShaderTypes.AddShaderType<FMeshProjectionDefaultPS>();
 
 				FMaterialShaders Shaders;
 				if (!Material->TryGetShaders(ShaderTypes, VertexFactory->GetType(), Shaders))
@@ -535,6 +679,19 @@ private:
 	FHitProxyConsumer* HitProxyConsumer;
 };
 
+bool FDisplayClusterMeshProjectionPrimitiveFilter::IsPrimitiveComponentFiltered(const UPrimitiveComponent* InPrimitiveComponent) const
+{
+	if (PrimitiveFilterDelegate.IsBound())
+	{
+		if (!PrimitiveFilterDelegate.Execute(InPrimitiveComponent))
+		{
+			return true;
+		}
+	}
+
+	return false;
+}
+
 FVector FDisplayClusterMeshProjectionRenderer::ProjectViewPosition(const FVector& ViewPosition, EDisplayClusterMeshProjectionType  ProjectionType)
 {
 	FVector ProjectedViewPosition(ViewPosition);
@@ -784,6 +941,66 @@ void FDisplayClusterMeshProjectionRenderer::Render(FCanvas* Canvas, FSceneInterf
 		});
 }
 
+void FDisplayClusterMeshProjectionRenderer::RenderNormals(FCanvas* Canvas,
+	FSceneInterface* Scene,
+	const FSceneViewInitOptions& ViewInitOptions,
+	const FEngineShowFlags& EngineShowFlags,
+	EDisplayClusterMeshProjectionType  ProjectionType,
+	FDisplayClusterMeshProjectionPrimitiveFilter* PrimitiveFilter)
+{
+	Canvas->Flush_GameThread();
+	FRenderTarget* RenderTarget = Canvas->GetRenderTarget();
+
+	ENQUEUE_RENDER_COMMAND(FDrawProjectedMeshes)(
+		[RenderTarget, Scene, ViewInitOptions, EngineShowFlags, ProjectionType, PrimitiveFilter, this](FRHICommandListImmediate& RHICmdList)
+		{
+			FMemMark Mark(FMemStack::Get());
+			FRDGBuilder GraphBuilder(RHICmdList);
+
+			FSceneViewFamily ViewFamily(FSceneViewFamily::ConstructionValues(
+				RenderTarget,
+				Scene,
+				EngineShowFlags)
+				.SetTime(FGameTime::GetTimeSinceAppStart())
+				.SetGammaCorrection(1.0f));
+
+			if (Scene)
+			{
+				Scene->IncrementFrameNumber();
+				ViewFamily.FrameNumber = Scene->GetFrameNumber();
+			}
+			else
+			{
+				ViewFamily.FrameNumber = GFrameNumber;
+			}
+
+			FScenePrimitiveRenderingContextScopeHelper ScenePrimitiveRenderingContextScopeHelper(GetRendererModule().BeginScenePrimitiveRendering(GraphBuilder, &ViewFamily));
+
+			FSceneViewInitOptions NewInitOptions(ViewInitOptions);
+			NewInitOptions.ViewFamily = &ViewFamily;
+
+			GetRendererModule().CreateAndInitSingleView(RHICmdList, &ViewFamily, &NewInitOptions);
+			FViewInfo* View = (FViewInfo*)ViewFamily.Views[0];
+
+			FRDGTextureRef OutputTexture = GraphBuilder.RegisterExternalTexture(CreateRenderTarget(RenderTarget->GetRenderTargetTexture(), TEXT("ViewRenderTarget")));
+			FRenderTargetBinding OutputRenderTargetBinding(OutputTexture, ERenderTargetLoadAction::ELoad);
+
+			const FRDGTextureDesc NormalsDesc(FRDGTextureDesc::Create2D(OutputTexture->Desc.Extent, OutputTexture->Desc.Format, FClearValueBinding::Black, TexCreate_RenderTargetable | TexCreate_ShaderResource));
+			FRDGTextureRef NormalsTexture = GraphBuilder.CreateTexture(NormalsDesc, TEXT("DisplayClusterMeshProjection.NormalsTexture"));
+
+			const FRDGTextureDesc DepthDesc = FRDGTextureDesc::Create2D(OutputTexture->Desc.Extent, PF_DepthStencil, FClearValueBinding::DepthFar, TexCreate_DepthStencilTargetable | TexCreate_ShaderResource);
+			FRDGTextureRef NormalsDepthTexture = GraphBuilder.CreateTexture(DepthDesc, TEXT("DisplayClusterMeshProjection.NormalsDepthTexture"));
+
+			FRenderTargetBinding NormalsRenderTargetBinding(NormalsTexture, ERenderTargetLoadAction::EClear);
+			FDepthStencilBinding NormalsDepthStencilBinding(NormalsDepthTexture, ERenderTargetLoadAction::EClear, ERenderTargetLoadAction::EClear, FExclusiveDepthStencil::DepthWrite_StencilWrite);
+
+			AddNormalsRenderPass(GraphBuilder, View, ProjectionType, PrimitiveFilter, NormalsRenderTargetBinding, NormalsDepthStencilBinding);
+			AddNormalsFilterPass(GraphBuilder, View, OutputRenderTargetBinding, NormalsTexture, NormalsDepthTexture);
+
+			GraphBuilder.Execute();
+		});
+}
+
 // Helper macros that generate appropriate switch cases for each projection type the mesh projection renderer supports, allowing new projection types to be easily added
 #define PROJECTION_TYPE_CASE(FuncName, ProjectionType, ...) case ProjectionType: \
 	FuncName<ProjectionType>(__VA_ARGS__); \
@@ -856,6 +1073,226 @@ void FDisplayClusterMeshProjectionRenderer::AddHitProxyRenderPass(FRDGBuilder& G
 	});
 }
 
+void FDisplayClusterMeshProjectionRenderer::AddNormalsRenderPass(FRDGBuilder& GraphBuilder,
+	const FViewInfo* View,
+	EDisplayClusterMeshProjectionType ProjectionType,
+	FDisplayClusterMeshProjectionPrimitiveFilter* PrimitiveFilter,
+	FRenderTargetBinding& OutputRenderTargetBinding,
+	FDepthStencilBinding& OutputDepthStencilBinding)
+{
+	FMeshProjectionPassParameters* MeshPassParameters = GraphBuilder.AllocParameters<FMeshProjectionPassParameters>();
+	MeshPassParameters->View = View->ViewUniformBuffer;
+	MeshPassParameters->InstanceCulling = FInstanceCullingContext::CreateDummyInstanceCullingUniformBuffer(GraphBuilder);
+	MeshPassParameters->RenderTargets[0] = OutputRenderTargetBinding;
+	MeshPassParameters->RenderTargets.DepthStencil = OutputDepthStencilBinding;
+
+	GraphBuilder.AddPass(RDG_EVENT_NAME("MeshProjectionRenderer::Normals"),
+		MeshPassParameters,
+		ERDGPassFlags::Raster | ERDGPassFlags::NeverCull,
+		[View, ProjectionType, PrimitiveFilter, this](FRHICommandList& RHICmdList)
+		{
+			FIntRect ViewRect = View->UnscaledViewRect;
+			RHICmdList.SetViewport(ViewRect.Min.X, ViewRect.Min.Y, 0.0f, ViewRect.Max.X, ViewRect.Max.Y, 1.0f);
+
+			SWITCH_ON_PROJECTION_TYPE(RenderNormals_RenderThread, ProjectionType, View, RHICmdList, PrimitiveFilter);
+		});
+}
+
+void GaussianSpatialFilter(float Sigma, TArray<float>& OutKernel)
+{
+	constexpr uint32 KernelSize = FILTER_KERNEL_SIZE;
+	constexpr uint32 KernelRadius = (KernelSize - 1) / 2;
+
+	OutKernel.Init(0.0, KernelSize);
+
+	const float InvSigma = 1.f / Sigma;
+
+	float Sum = 0.0f;
+	for (int Index = 0; Index < KernelSize; ++Index)
+	{
+		int32 X = Index - KernelRadius;
+
+		OutKernel[Index] = 0.39894 * FMath::Exp(-0.5f * (X * X) * InvSigma * InvSigma) * InvSigma;
+		Sum += OutKernel[Index];
+	}
+
+	// Normalize the kernel so that there is no change in intensity of the filtered elements
+	for (int Index = 0; Index < KernelSize; ++Index)
+	{
+		OutKernel[Index] /= Sum;
+	}
+}
+
+template<EMeshProjectionFilterType FilterType>
+void AddSeparableFilterPass(FRDGBuilder& GraphBuilder,
+	const FViewInfo* View,
+	FRDGTexture* SceneDepth,
+	FRDGTextureUAV* RWColorUAV,
+	FRDGTextureUAV* RWDepthUAV,
+	FRDGTextureUAV* RWStencilUAV)
+{
+	TArray<float> SpatialKernel;
+	TArray<float> InteriorSpatialKernel;
+
+	GaussianSpatialFilter(100.0f, SpatialKernel);
+	GaussianSpatialFilter(1.0f, InteriorSpatialKernel);
+
+	const FScreenPassTextureViewport InputViewport(RWColorUAV->Desc.Texture);
+
+	FMeshProjectionNormalsFilterParameters* HorizontalPassParameters = GraphBuilder.AllocParameters<FMeshProjectionNormalsFilterParameters>();
+	HorizontalPassParameters->View = View->ViewUniformBuffer;
+	HorizontalPassParameters->Input = GetScreenPassTextureViewportParameters(InputViewport);
+	HorizontalPassParameters->SceneStencil = GraphBuilder.CreateSRV(FRDGTextureSRVDesc::CreateWithPixelFormat(SceneDepth, PF_X24_G8));
+	HorizontalPassParameters->RWColor = RWColorUAV;
+	HorizontalPassParameters->RWDepth = RWDepthUAV;
+	HorizontalPassParameters->RWStencil = RWStencilUAV;
+	HorizontalPassParameters->SampleDirection = FVector2f(1.0f, 0.0f);
+	HorizontalPassParameters->SampleOffsetScale = FVector2f(1.f);
+
+	for (int Index = 0; Index < FILTER_KERNEL_SIZE; ++Index)
+	{
+		GET_SCALAR_ARRAY_ELEMENT(HorizontalPassParameters->SpatialKernel, Index) = SpatialKernel[Index];
+		GET_SCALAR_ARRAY_ELEMENT(HorizontalPassParameters->InteriorSpatialKernel, Index) = InteriorSpatialKernel[Index];
+	}
+
+	TShaderMapRef<FMeshProjectionNormalsFilterCS<FilterType>> ComputeShader(View->ShaderMap);
+
+	FComputeShaderUtils::AddPass(
+		GraphBuilder,
+		RDG_EVENT_NAME("MeshProjectionRenderer::NormalsFilterHorizontal"),
+		ComputeShader,
+		HorizontalPassParameters,
+		FComputeShaderUtils::GetGroupCount(InputViewport.Extent, FIntPoint(8, 8)));
+
+	FMeshProjectionNormalsFilterParameters* VerticalPassParameters = GraphBuilder.AllocParameters<FMeshProjectionNormalsFilterParameters>();
+	VerticalPassParameters->View = View->ViewUniformBuffer;
+	VerticalPassParameters->Input = GetScreenPassTextureViewportParameters(InputViewport);
+	VerticalPassParameters->SceneStencil = GraphBuilder.CreateSRV(FRDGTextureSRVDesc::CreateWithPixelFormat(SceneDepth, PF_X24_G8));
+	VerticalPassParameters->RWColor = RWColorUAV;
+	VerticalPassParameters->RWDepth = RWDepthUAV;
+	VerticalPassParameters->RWStencil = RWStencilUAV;
+	VerticalPassParameters->SampleDirection = FVector2f(0.0f, 1.0f);
+	VerticalPassParameters->SampleOffsetScale = FVector2f(1.f);
+
+	for (int Index = 0; Index < FILTER_KERNEL_SIZE; ++Index)
+	{
+		GET_SCALAR_ARRAY_ELEMENT(VerticalPassParameters->SpatialKernel, Index) = SpatialKernel[Index];
+		GET_SCALAR_ARRAY_ELEMENT(VerticalPassParameters->InteriorSpatialKernel, Index) = InteriorSpatialKernel[Index];
+	}
+
+	FComputeShaderUtils::AddPass(
+		GraphBuilder,
+		RDG_EVENT_NAME("MeshProjectionRenderer::NormalsFilterVertical"),
+		ComputeShader,
+		VerticalPassParameters,
+		FComputeShaderUtils::GetGroupCount(InputViewport.Extent, FIntPoint(8, 8)));
+}
+
+void FDisplayClusterMeshProjectionRenderer::AddNormalsFilterPass(FRDGBuilder& GraphBuilder,
+	const FViewInfo* View,
+	FRenderTargetBinding& OutputRenderTargetBinding,
+	FRDGTexture* SceneColor,
+	FRDGTexture* SceneDepth)
+{
+	// The normal map filter pass is made up of several separate passes:
+	// * First, the scene color, depth, and stencil textures must be copied into UAV read-write textures for the compute shaders to operate on. This pass also
+	//   converts the normal vectors from world space to radial space, so that the filtering operates on the relative normals instead of the absolute normals
+	// * Second, the depth and normals are dilated so they bleed into empty space. There is a separate interior depth dilate pass which dilates closer overlapping
+	//   screen depths into further away screens, so that when blurred, there is a continuous depth change between overlapping screens
+	// * Third, the normals and depths are blurred using a Gaussian blur. There are two separate gaussian spatial filters used, one used for blurring empty space regions
+	//   and one used for blurring the screens. The empty space filter has a large sigma, allowing a large blur, while the interior filter has a smaller sigma, so that
+	//   the blurring doesn't cause too much feature loss
+	// * Finally, the RW textures are rendered back to a render target. Here, the normals are placed into the RGB components, while the depth is placed in the A component
+
+	const FRDGTextureDesc RWColorDesc(FRDGTextureDesc::Create2D(SceneColor->Desc.Extent, PF_B8G8R8A8, FClearValueBinding::Black, TexCreate_ShaderResource | TexCreate_UAV));
+	FRDGTextureRef RWColor = GraphBuilder.CreateTexture(RWColorDesc, TEXT("DisplayClusterMeshProjection.NormalsRWTexture"));
+	FRDGTextureUAVRef RWColorUAV = GraphBuilder.CreateUAV(RWColor);
+
+	const FRDGTextureDesc RWDepthDesc = FRDGTextureDesc::Create2D(SceneDepth->Desc.Extent, PF_G16, FClearValueBinding::White, TexCreate_UAV | TexCreate_ShaderResource);
+	FRDGTextureRef RWDepth = GraphBuilder.CreateTexture(RWDepthDesc, TEXT("DisplayClusterMeshProjection.RWDepth"));
+	FRDGTextureUAVRef RWDepthUAV = GraphBuilder.CreateUAV(RWDepth);
+
+	const FRDGTextureDesc RWStencilDesc = FRDGTextureDesc::Create2D(SceneDepth->Desc.Extent, PF_R32_UINT, FClearValueBinding::Black, TexCreate_UAV | TexCreate_ShaderResource);
+	FRDGTextureRef RWStencil = GraphBuilder.CreateTexture(RWStencilDesc, TEXT("DisplayClusterMeshProjection.RWStencil"));
+	FRDGTextureUAVRef RWStencilUAV = GraphBuilder.CreateUAV(RWStencil);
+
+	// Copy the scene color and depth to RW buffers
+	{
+		FMeshProjectionNormalsCreateRWTexturesCS::FParameters* PassParameters = GraphBuilder.AllocParameters<FMeshProjectionNormalsCreateRWTexturesCS::FParameters>();
+		PassParameters->View = View->ViewUniformBuffer;
+		PassParameters->SceneColor = SceneColor;
+		PassParameters->SceneDepth = SceneDepth;
+		PassParameters->SceneStencil = GraphBuilder.CreateSRV(FRDGTextureSRVDesc::CreateWithPixelFormat(SceneDepth, PF_X24_G8));
+		PassParameters->RWColor = RWColorUAV;
+		PassParameters->RWDepth = RWDepthUAV;
+		PassParameters->RWStencil = RWStencilUAV;
+
+		TShaderMapRef<FMeshProjectionNormalsCreateRWTexturesCS> ComputeShader(View->ShaderMap);
+
+		FComputeShaderUtils::AddPass(
+			GraphBuilder,
+			RDG_EVENT_NAME("MeshProjectionRenderer::CreateRWTextures"),
+			ComputeShader,
+			PassParameters,
+			FComputeShaderUtils::GetGroupCount(SceneColor->Desc.Extent, FIntPoint(8, 8)));
+	}
+
+	static constexpr int32 DepthDilatePasses = 1;
+	static constexpr int32 DilatePasses = 8;
+	static constexpr int32 BlurPasses = 32;
+
+	// Filter the RW buffers, dilating and blurring them appropriately
+	for (int32 Index = 0; Index < DepthDilatePasses; Index++)
+	{
+		AddSeparableFilterPass<EMeshProjectionFilterType::DepthDilate>(GraphBuilder, View, SceneDepth, RWColorUAV, RWDepthUAV, RWStencilUAV);
+	}
+
+	for (int32 Index = 0; Index < DilatePasses; ++Index)
+	{
+		AddSeparableFilterPass<EMeshProjectionFilterType::Dilate>(GraphBuilder, View, SceneDepth, RWColorUAV, RWDepthUAV, RWStencilUAV);
+	}
+
+	for (int32 Index = 0; Index < BlurPasses; Index++)
+	{
+		AddSeparableFilterPass<EMeshProjectionFilterType::Blur>(GraphBuilder, View, SceneDepth, RWColorUAV, RWDepthUAV, RWStencilUAV);
+	}
+
+	// Copy the RW buffers to the output render target
+	{
+		const FScreenPassTextureViewport InputViewport(RWColor);
+		const FScreenPassTextureViewport OutputViewport(OutputRenderTargetBinding.GetTexture());
+
+		FMeshProjectionNormalsOutputPS::FParameters* ScreenPassParameters = GraphBuilder.AllocParameters<FMeshProjectionNormalsOutputPS::FParameters>();
+		ScreenPassParameters->Input = GetScreenPassTextureViewportParameters(InputViewport);
+		ScreenPassParameters->RWColor = RWColorUAV;
+		ScreenPassParameters->RWDepth = RWDepthUAV;
+		ScreenPassParameters->RenderTargets[0] = OutputRenderTargetBinding;
+
+		TShaderMapRef<FScreenPassVS> ScreenPassVS(View->ShaderMap);
+		TShaderMapRef<FMeshProjectionNormalsOutputPS> OutputNormalsPS(View->ShaderMap);
+
+		FRHIBlendState* DefaultBlendState = FScreenPassPipelineState::FDefaultBlendState::GetRHI();
+
+		GraphBuilder.AddPass(
+			RDG_EVENT_NAME("MeshProjectionRenderer::OutputNormals"),
+			ScreenPassParameters,
+			ERDGPassFlags::Raster,
+			[View, ScreenPassVS, OutputNormalsPS, InputViewport, OutputViewport, ScreenPassParameters, DefaultBlendState](FRHICommandList& RHICmdList)
+		{
+			DrawScreenPass(
+				RHICmdList,
+				*View,
+				InputViewport,
+				OutputViewport,
+				FScreenPassPipelineState(ScreenPassVS, OutputNormalsPS, DefaultBlendState),
+				[&](FRHICommandList&)
+			{
+				SetShaderParameters(RHICmdList, OutputNormalsPS, OutputNormalsPS.GetPixelShader(), *ScreenPassParameters);
+			});
+		});
+	}
+}
+
 #if WITH_EDITOR
 void FDisplayClusterMeshProjectionRenderer::AddSelectionDepthRenderPass(FRDGBuilder& GraphBuilder, 
 		const FViewInfo* View,
@@ -1016,6 +1453,42 @@ void FDisplayClusterMeshProjectionRenderer::RenderHitProxies_RenderThread(const
 	});
 }
 
+template<EDisplayClusterMeshProjectionType ProjectionType>
+void FDisplayClusterMeshProjectionRenderer::RenderNormals_RenderThread(const FSceneView* View, FRHICommandList& RHICmdList, FDisplayClusterMeshProjectionPrimitiveFilter* PrimitiveFilter)
+{
+	DrawDynamicMeshPass(*View, RHICmdList, [View, PrimitiveFilter, this](FDynamicPassMeshDrawListContext* DynamicMeshPassContext)
+	{
+		TArray<FPrimitiveSceneProxy*> PrimitiveSceneProxies;
+		for (const TWeakObjectPtr<UPrimitiveComponent>& PrimitiveComponent : PrimitiveComponents)
+		{
+			if (PrimitiveComponent.IsValid() && PrimitiveComponent->SceneProxy)
+			{
+				if (PrimitiveFilter && PrimitiveFilter->IsPrimitiveComponentFiltered(PrimitiveComponent.Get()))
+				{
+					continue;
+				}
+
+				PrimitiveSceneProxies.Add(PrimitiveComponent->SceneProxy);
+			}
+		}
+
+		FLightCardEditorMeshPassProcessor<ProjectionType, EMeshProjectionShaderOutput::Normals> MeshProcessor(nullptr, View, DynamicMeshPassContext);
+		MeshProcessor.SetStencilValue(1);
+		MeshProcessor.SetIgnoreTranslucency(true);
+
+		for (FPrimitiveSceneProxy* PrimitiveProxy : PrimitiveSceneProxies)
+		{
+			if (const FMeshBatch* MeshBatch = PrimitiveProxy->GetPrimitiveSceneInfo()->GetMeshBatch(PrimitiveProxy->GetPrimitiveSceneInfo()->StaticMeshes.Num() - 1))
+			{
+				MeshBatch->MaterialRenderProxy->UpdateUniformExpressionCacheIfNeeded(View->GetFeatureLevel());
+
+				const uint64 BatchElementMask = ~0ull;
+				MeshProcessor.AddMeshBatch(*MeshBatch, BatchElementMask, PrimitiveProxy);
+			}
+		}
+	});
+}
+
 #if WITH_EDITOR
 template<EDisplayClusterMeshProjectionType ProjectionType>
 void FDisplayClusterMeshProjectionRenderer::RenderSelection_RenderThread(const FSceneView* View, FRHICommandList& RHICmdList)
@@ -1065,6 +1538,7 @@ bool FDisplayClusterMeshProjectionRenderer::IsPrimitiveComponentSelected(const U
 
 CREATE_PROJECTION_TYPE_TEMPLATE_SPECIALIZATIONS(RenderPrimitives_RenderThread, const FSceneView* View, FRHICommandList& RHICmdList, bool bTranslucencyPass)
 CREATE_PROJECTION_TYPE_TEMPLATE_SPECIALIZATIONS(RenderHitProxies_RenderThread, const FSceneView* View, FRHICommandList& RHICmdList)
+CREATE_PROJECTION_TYPE_TEMPLATE_SPECIALIZATIONS(RenderNormals_RenderThread, const FSceneView* View, FRHICommandList& RHICmdList, FDisplayClusterMeshProjectionPrimitiveFilter* PrimitiveFilter)
 
 #if WITH_EDITOR
 CREATE_PROJECTION_TYPE_TEMPLATE_SPECIALIZATIONS(RenderSelection_RenderThread, const FSceneView* View, FRHICommandList& RHICmdList)
diff --git a/Engine/Plugins/Runtime/nDisplay/Source/DisplayClusterLightCardEditorShaders/Public/DisplayClusterMeshProjectionRenderer.h b/Engine/Plugins/Runtime/nDisplay/Source/DisplayClusterLightCardEditorShaders/Public/DisplayClusterMeshProjectionRenderer.h
--- a/Engine/Plugins/Runtime/nDisplay/Source/DisplayClusterLightCardEditorShaders/Public/DisplayClusterMeshProjectionRenderer.h
+++ b/Engine/Plugins/Runtime/nDisplay/Source/DisplayClusterLightCardEditorShaders/Public/DisplayClusterMeshProjectionRenderer.h
@@ -24,6 +24,18 @@ enum DISPLAYCLUSTERLIGHTCARDEDITORSHADERS_API EDisplayClusterMeshProjectionType
 	Azimuthal
 };
 
+/** A filter that allows specific primitive components to be filtered from a render pass */
+class DISPLAYCLUSTERLIGHTCARDEDITORSHADERS_API FDisplayClusterMeshProjectionPrimitiveFilter
+{
+public:
+	/** A delegate that returns true if the primitive component should be included in the render pass */
+	DECLARE_DELEGATE_RetVal_OneParam(bool, FPrimitiveFilter, const UPrimitiveComponent*);
+	FPrimitiveFilter PrimitiveFilterDelegate;
+
+	/** Gets whether a primitive component should be filtered out of the render pass or not */
+	bool IsPrimitiveComponentFiltered(const UPrimitiveComponent* InPrimitiveComponent) const;
+};
+
 /** A renderer that projects meshes to screen space using non-linear projection methods */
 class DISPLAYCLUSTERLIGHTCARDEDITORSHADERS_API FDisplayClusterMeshProjectionRenderer
 {
@@ -49,6 +61,14 @@ public:
 	/** Renders its list of primitive components to the specified canvas using the desired projection type. Can be called from the game thread */
 	void Render(FCanvas* Canvas, FSceneInterface* Scene, const FSceneViewInitOptions& ViewInitOptions, const FEngineShowFlags& EngineShowFlags, EDisplayClusterMeshProjectionType  ProjectionType);
 
+	/** Renders the normals and depth of the scene's primitives and performs blurring to create a continuous normal map */
+	void RenderNormals(FCanvas* Canvas,
+		FSceneInterface* Scene,
+		const FSceneViewInitOptions& ViewInitOptions,
+		const FEngineShowFlags& EngineShowFlags,
+		EDisplayClusterMeshProjectionType ProjectionType, 
+		FDisplayClusterMeshProjectionPrimitiveFilter* PrimitiveFilter = nullptr);
+
 private:
 	/** Adds a pass to perform the base render for the mesh projection. */
 	void AddBaseRenderPass(FRDGBuilder& GraphBuilder, 
@@ -71,6 +91,22 @@ private:
 		FRenderTargetBinding& OutputRenderTargetBinding,
 		FDepthStencilBinding& OutputDepthStencilBinding);
 
+	/** Adds a pass to perform a render of the primitives' normals for the mesh projection */
+	void AddNormalsRenderPass(FRDGBuilder& GraphBuilder,
+		const FViewInfo* View,
+		EDisplayClusterMeshProjectionType ProjectionType,
+		FDisplayClusterMeshProjectionPrimitiveFilter* PrimitiveFilter,
+		FRenderTargetBinding& OutputRenderTargetBinding,
+		FDepthStencilBinding& OutputDepthStencilBinding);
+
+	/** Adds a pass to filter the output normal map for the mesh projection, dilating and blurring it
+	 * to obtain a continuous normal map from primitives into empty space */
+	void AddNormalsFilterPass(FRDGBuilder& GraphBuilder,
+		const FViewInfo* View,
+		FRenderTargetBinding& OutputRenderTargetBinding,
+		FRDGTexture* SceneColor,
+		FRDGTexture* SceneDepth);
+
 #if WITH_EDITOR
 	/** Adds a pass to perform the depth render for any selected primitives for the mesh projection */
 	void AddSelectionDepthRenderPass(FRDGBuilder& GraphBuilder, 
@@ -101,6 +137,10 @@ private:
 	template<EDisplayClusterMeshProjectionType ProjectionType>
 	void RenderHitProxies_RenderThread(const FSceneView* View, FRHICommandList& RHICmdList);
 
+	/** Renders the normals of the list of primitive components using the appropriate mesh pass processor given by the template parameter */
+	template<EDisplayClusterMeshProjectionType ProjectionType>
+	void RenderNormals_RenderThread(const FSceneView* View, FRHICommandList& RHICmdList, FDisplayClusterMeshProjectionPrimitiveFilter* PrimitiveFilter);
+
 #if WITH_EDITOR
 	/** Renders the list of primitive components using the appropriate mesh pass processor given by the template parameter */
 	template<EDisplayClusterMeshProjectionType ProjectionType>
