nDisplay: fixed in-editor performance degradation due to missing flush in the ndisplay editor behaviour when frustum preview rendering is enabled

diff --git a/Engine/Plugins/Runtime/nDisplay/Source/DisplayClusterConfigurator/Private/Views/Viewport/DisplayClusterConfiguratorSCSEditorViewportClient.cpp b/Engine/Plugins/Runtime/nDisplay/Source/DisplayClusterConfigurator/Private/Views/Viewport/DisplayClusterConfiguratorSCSEditorViewportClient.cpp
--- a/Engine/Plugins/Runtime/nDisplay/Source/DisplayClusterConfigurator/Private/Views/Viewport/DisplayClusterConfiguratorSCSEditorViewportClient.cpp
+++ b/Engine/Plugins/Runtime/nDisplay/Source/DisplayClusterConfigurator/Private/Views/Viewport/DisplayClusterConfiguratorSCSEditorViewportClient.cpp
@@ -12,6 +12,7 @@
 #include "DisplayClusterRootActor.h"
 #include "Components/DisplayClusterPreviewComponent.h"
 #include "Components/DisplayClusterCameraComponent.h"
+#include "Components/LineBatchComponent.h"
 
 #include "AssetViewerSettings.h"
 #include "CanvasItem.h"
@@ -289,6 +290,12 @@ void FDisplayClusterConfiguratorSCSEditorViewportClient::Draw(const FSceneView*
 {
 	FEditorViewportClient::Draw(View, PDI);
 
+	// mimic the behaviour of the FEditorViewportClient call used in editor viewport
+	if (PreviewScene->GetWorld()->LineBatcher != NULL && (PreviewScene->GetWorld()->LineBatcher->BatchedLines.Num() || PreviewScene->GetWorld()->LineBatcher->BatchedPoints.Num() || PreviewScene->GetWorld()->LineBatcher->BatchedMeshes.Num()))
+	{
+		PreviewScene->GetWorld()->LineBatcher->Flush();
+	}
+	
 	bool bHitTesting = PDI->IsHitTesting();
 	AActor* PreviewActor = GetPreviewActor();
 	if (PreviewActor)
