nDisplay: Implement auto tilt and auto pan in the light card editor. Implemented ResetCamera command to restore rotation and FOV.

diff --git a/Engine/Plugins/Runtime/nDisplay/Source/DisplayClusterLightCardEditor/Private/Viewport/DisplayClusterLightCardEditorViewportClient.cpp b/Engine/Plugins/Runtime/nDisplay/Source/DisplayClusterLightCardEditor/Private/Viewport/DisplayClusterLightCardEditorViewportClient.cpp
--- a/Engine/Plugins/Runtime/nDisplay/Source/DisplayClusterLightCardEditor/Private/Viewport/DisplayClusterLightCardEditorViewportClient.cpp
+++ b/Engine/Plugins/Runtime/nDisplay/Source/DisplayClusterLightCardEditor/Private/Viewport/DisplayClusterLightCardEditorViewportClient.cpp
@@ -38,6 +38,8 @@
 #include "Engine/TextureRenderTarget2D.h"
 #include "CameraController.h"
 #include "ImageUtils.h"
+#include "Kismet/KismetMathLibrary.h"
+
 
 #include "Renderer/Private/SceneRendering.h"
 
@@ -205,9 +207,7 @@ FDisplayClusterLightCardEditorViewportClient::FDisplayClusterLightCardEditorView
 	SetRealtime(true);
 	SetShowStats(true);
 
-	ProjectionFOVs.AddZeroed(2);
-	ProjectionFOVs[(int32)EDisplayClusterMeshProjectionType::Perspective] = 90.0f;
-	ProjectionFOVs[(int32)EDisplayClusterMeshProjectionType::Azimuthal] = 130.0f;
+	ResetFOVs();
 
 	//This seems to be needed to get the correct world time in the preview.
 	SetIsSimulateInEditorViewport(true);
@@ -236,13 +236,7 @@ void FDisplayClusterLightCardEditorViewportClient::Tick(float DeltaSeconds)
 	FEditorViewportClient::Tick(DeltaSeconds);
 
 	// Camera position is locked to a specific location
-	FVector Location = FVector::ZeroVector;
-	if (ProjectionOriginComponent.IsValid())
-	{
-		Location = ProjectionOriginComponent->GetComponentLocation();
-	}
-
-	SetViewLocation(Location);
+	ResetCamera(/* bLocationOnly */ true);
 
 	CalcEditorWidgetTransform(CachedEditorWidgetTransformBeforeMapProjection, CachedEditorWidgetTransformAfterMapProjection);
 
@@ -290,6 +284,19 @@ void FDisplayClusterLightCardEditorViewportClient::Tick(float DeltaSeconds)
 
 		ApplyRequiredCursorVisibility(true);
 	}
+
+	if (DesiredLookAtLocation.IsSet())
+	{
+		const FRotator LookAtRotation = UKismetMathLibrary::FindLookAtRotation(GetViewLocation(), *DesiredLookAtLocation);
+		const FRotator NewRotation = FMath::RInterpTo(GetViewRotation(), LookAtRotation, DeltaSeconds, DesiredLookAtSpeed);
+		SetViewRotation(NewRotation);
+
+		if (NewRotation.Equals(LookAtRotation, 2.f) ||
+			!IsLocationCloseToEdge(CachedEditorWidgetTransformAfterMapProjection.GetTranslation()))
+		{
+			DesiredLookAtLocation.Reset();
+		}
+	}
 }
 
 void FDisplayClusterLightCardEditorViewportClient::Draw(FViewport* InViewport, FCanvas* Canvas)
@@ -571,6 +578,8 @@ void FDisplayClusterLightCardEditorViewportClient::TrackingStarted(const FInputE
 
 		BeginTransaction(LOCTEXT("MoveLightCard", "Move Light Card"));
 		bDraggingActor = true;
+		
+		DesiredLookAtLocation.Reset();
 
 		// Compute and store the delta between the widget's origin and the place the user clicked on it,
 		// in order to factor it out when transforming the selected actor
@@ -964,6 +973,26 @@ void FDisplayClusterLightCardEditorViewportClient::SetProjectionModeFOV(EDisplay
 	bShouldCheckHitProxy = true;
 }
 
+void FDisplayClusterLightCardEditorViewportClient::ResetCamera(bool bLocationOnly)
+{
+	FVector Location = FVector::ZeroVector;
+	if (ProjectionOriginComponent.IsValid())
+	{
+		Location = ProjectionOriginComponent->GetComponentLocation();
+	}
+
+	SetViewLocation(Location);
+
+	if (bLocationOnly)
+	{
+		return;
+	}
+	
+	SetProjectionMode(GetProjectionMode());
+
+	ResetFOVs();
+}
+
 void FDisplayClusterLightCardEditorViewportClient::BeginTransaction(const FText& Description)
 {
 	GEditor->BeginTransaction(Description);
@@ -1437,6 +1466,18 @@ void FDisplayClusterLightCardEditorViewportClient::MoveSelectedLightCards(FViewp
 
 			PropagateLightCardTransform(LightCard.Get());
 		}
+
+		FVector LightCardWorldLocation = CachedEditorWidgetTransformAfterMapProjection.GetTranslation();
+		FVector2D ScreenPercentage;
+		if (IsLocationCloseToEdge(LightCardWorldLocation, InViewport, View, &ScreenPercentage))
+		{
+			DesiredLookAtSpeed = FMath::Max(ScreenPercentage.X, ScreenPercentage.Y) * MaxDesiredLookAtSpeed;
+			DesiredLookAtLocation = LightCardWorldLocation;
+		}
+		else
+		{
+			DesiredLookAtLocation.Reset();
+		}
 	}
 }
 
@@ -1681,4 +1722,75 @@ void FDisplayClusterLightCardEditorViewportClient::InvalidateNormalMap()
 	bNormalMapInvalid = true;
 }
 
+bool FDisplayClusterLightCardEditorViewportClient::IsLocationCloseToEdge(const FVector& InPosition, const FViewport* InViewport,
+                                                                         const FSceneView* InView, FVector2D* OutPercentageToEdge)
+{
+	if (InViewport == nullptr)
+	{
+		InViewport = Viewport;
+	}
+
+	check(InViewport);
+	const FIntPoint ViewportSize = InViewport->GetSizeXY();
+
+	FPlane Projection;
+	if (InView == nullptr)
+	{
+		FSceneViewFamilyContext ViewFamily(FSceneViewFamily::ConstructionValues(
+		Viewport,
+		GetScene(),
+		EngineShowFlags)
+		.SetRealtimeUpdate(IsRealtime()));
+		
+		InView = CalcSceneView(&ViewFamily);
+		Projection = InView->Project(InPosition);
+
+		// InView will be deleted here
+	}
+	else
+	{
+		Projection = InView->Project(InPosition);
+	}
+	
+	if (Projection.W > 0)
+	{
+		const float HighThreshold = 1.f - EdgePercentageLookAtThreshold;
+		
+		const int32 HalfX = 0.5f * Viewport->GetSizeXY().X;
+		const int32 HalfY = 0.5f * Viewport->GetSizeXY().Y;
+		const int32 XPos = HalfX + (HalfX * Projection.X);
+		const int32 YPos = HalfY + (HalfY * (Projection.Y * -1));
+			
+		auto GetPercentToEdge = [&](int32 CurrentPos, int32 MaxPos) -> float
+		{
+			const float Center = static_cast<float>(MaxPos) / 2.f;
+			const float RelativePosition = FMath::Abs(CurrentPos - Center);
+			return RelativePosition / Center;
+		};
+			
+		const float XPercent = GetPercentToEdge(XPos, ViewportSize.X);
+		const float YPercent = GetPercentToEdge(YPos, ViewportSize.Y);
+			
+		if (OutPercentageToEdge)
+		{
+			*OutPercentageToEdge = FVector2D(XPercent, YPercent);
+		}
+			
+		return XPercent >= HighThreshold || YPercent >= HighThreshold;
+	}
+
+	return false;
+}
+
+void FDisplayClusterLightCardEditorViewportClient::ResetFOVs()
+{
+	constexpr int32 MaxFOVs = 2;
+	if (ProjectionFOVs.Num() < MaxFOVs)
+	{
+		ProjectionFOVs.AddDefaulted(MaxFOVs - ProjectionFOVs.Num());
+	}
+	ProjectionFOVs[static_cast<int32>(EDisplayClusterMeshProjectionType::Perspective)] = 90.0f;
+	ProjectionFOVs[static_cast<int32>(EDisplayClusterMeshProjectionType::Azimuthal)] = 130.0f;
+}
+
 #undef LOCTEXT_NAMESPACE
diff --git a/Engine/Plugins/Runtime/nDisplay/Source/DisplayClusterLightCardEditor/Private/Viewport/DisplayClusterLightCardEditorViewportClient.h b/Engine/Plugins/Runtime/nDisplay/Source/DisplayClusterLightCardEditor/Private/Viewport/DisplayClusterLightCardEditorViewportClient.h
--- a/Engine/Plugins/Runtime/nDisplay/Source/DisplayClusterLightCardEditor/Private/Viewport/DisplayClusterLightCardEditorViewportClient.h
+++ b/Engine/Plugins/Runtime/nDisplay/Source/DisplayClusterLightCardEditor/Private/Viewport/DisplayClusterLightCardEditorViewportClient.h
@@ -155,6 +155,9 @@ public:
 	/** Sets the field of view of the specified projection mode */
 	void SetProjectionModeFOV(EDisplayClusterMeshProjectionType InProjectionMode, float NewFOV);
 
+	/** Resets the camera to the initial rotation / position */
+	void ResetCamera(bool bLocationOnly = false);
+	
 private:
 	/** Initiates a transaction. */
 	void BeginTransaction(const FText& Description);
@@ -209,13 +212,19 @@ private:
 
 	/** Calculates the world transform to render the editor widget with to align it with the selected light card */
 	bool CalcEditorWidgetTransform(FTransform& WidgetTransformBeforeMapProjection, FTransform& WidgetTransformAfterMapProjection);
-
+	
 	/** Renders the viewport's normal map and stores the texture data to be used later */
 	void RenderNormalMap(FNormalMap& NormalMap, const FVector& NormalMapDirection);
 
 	/** Invalidates the viewport's normal map, forcing it to be rerendered on the next draw call */
 	void InvalidateNormalMap();
 
+	/** Checks if the location is approaching the edge of the view space */
+	bool IsLocationCloseToEdge(const FVector& InPosition, const FViewport* InViewport = nullptr, const FSceneView* InView = nullptr, FVector2D* OutPercentageToEdge = nullptr);
+
+	/** Resets the camera FOVs */
+	void ResetFOVs();
+
 private:
 	TWeakPtr<FSceneViewport> SceneViewportPtr;
 	TWeakPtr<SDisplayClusterLightCardEditor> LightCardEditorPtr;
@@ -259,6 +268,18 @@ private:
 	/** The offset between the widget's origin and the place it was clicked when a drag action was started */
 	FVector DragWidgetOffset;
 
+	/** The location the camera should be looking at */
+	TOptional<FVector> DesiredLookAtLocation;
+
+	/** The current speed to use when looking at a location */
+	float DesiredLookAtSpeed = 1.f;
+
+	/** The maximum speed to use when looking at a location */
+	float MaxDesiredLookAtSpeed = 5.f;
+
+	/** The percentage to the edge of the view which should trigger auto tilt and auto pan */
+	float EdgePercentageLookAtThreshold = 0.1f;
+
 	/** The current projection mode the 3D viewport is being displayed with */
 	EDisplayClusterMeshProjectionType ProjectionMode = EDisplayClusterMeshProjectionType::Azimuthal;
 
diff --git a/Engine/Plugins/Runtime/nDisplay/Source/DisplayClusterLightCardEditor/Private/Viewport/DisplayClusterLightcardEditorViewport.cpp b/Engine/Plugins/Runtime/nDisplay/Source/DisplayClusterLightCardEditor/Private/Viewport/DisplayClusterLightcardEditorViewport.cpp
--- a/Engine/Plugins/Runtime/nDisplay/Source/DisplayClusterLightCardEditor/Private/Viewport/DisplayClusterLightcardEditorViewport.cpp
+++ b/Engine/Plugins/Runtime/nDisplay/Source/DisplayClusterLightCardEditor/Private/Viewport/DisplayClusterLightcardEditorViewport.cpp
@@ -231,6 +231,10 @@ void SDisplayClusterLightCardEditorViewport::BindCommands()
 		FCanExecuteAction(),
 		FIsActionChecked::CreateSP(this, &SDisplayClusterLightCardEditorViewport::IsProjectionModeSelected, EDisplayClusterMeshProjectionType::Azimuthal));
 
+	CommandList->MapAction(
+		Commands.ResetCamera,
+		FExecuteAction::CreateSP(ViewportClient.Get(), &FDisplayClusterLightCardEditorViewportClient::ResetCamera, false));
+	
 	SEditorViewport::BindCommands();
 }
 
