nDisplay: Restored ability to kill UE process in case some critical nDisplay errors occur on very early initialization stage

diff --git a/Engine/Plugins/Runtime/nDisplay/Source/DisplayCluster/Private/Game/EngineClasses/Basics/DisplayClusterGameEngine.cpp b/Engine/Plugins/Runtime/nDisplay/Source/DisplayCluster/Private/Game/EngineClasses/Basics/DisplayClusterGameEngine.cpp
--- a/Engine/Plugins/Runtime/nDisplay/Source/DisplayCluster/Private/Game/EngineClasses/Basics/DisplayClusterGameEngine.cpp
+++ b/Engine/Plugins/Runtime/nDisplay/Source/DisplayCluster/Private/Game/EngineClasses/Basics/DisplayClusterGameEngine.cpp
@@ -58,7 +58,7 @@ void UDisplayClusterGameEngine::Init(class IEngineLoop* InEngineLoop)
 	// Initialize Display Cluster
 	if (!GDisplayCluster->Init(OperationMode))
 	{
-		FDisplayClusterAppExit::ExitApplication(FString("Couldn't initialize DisplayCluster module"));
+		FDisplayClusterAppExit::ExitApplication(FString("Couldn't initialize DisplayCluster module"), FDisplayClusterAppExit::EExitType::KillImmediately);
 	}
 
 	// This is required to prevent GameThread dead-locking when Alt+F4 is used
@@ -99,7 +99,7 @@ void UDisplayClusterGameEngine::Init(class IEngineLoop* InEngineLoop)
 		FString ConfigPath;
 		if (!ParseCommandArg(FCommandLine::Get(), DisplayClusterStrings::args::Config, ConfigPath))
 		{
-			FDisplayClusterAppExit::ExitApplication(FString("No config file specified. Cluster operation mode requires config file."));
+			FDisplayClusterAppExit::ExitApplication(FString("No config file specified. Cluster operation mode requires config file."), FDisplayClusterAppExit::EExitType::KillImmediately);
 		}
 
 		// Clean the file path before using it
@@ -109,14 +109,14 @@ void UDisplayClusterGameEngine::Init(class IEngineLoop* InEngineLoop)
 		const EDisplayClusterConfigurationVersion ConfigVersion = IDisplayClusterConfiguration::Get().GetConfigVersion(ConfigPath);
 		if (!ValidateConfigFile(ConfigPath))
 		{
-			FDisplayClusterAppExit::ExitApplication(FString("An invalid or outdated configuration file was specified. Please consider using nDisplay configurator to update the config files."));
+			FDisplayClusterAppExit::ExitApplication(FString("An invalid or outdated configuration file was specified. Please consider using nDisplay configurator to update the config files."), FDisplayClusterAppExit::EExitType::KillImmediately);
 		}
 
 		// Load config data
 		UDisplayClusterConfigurationData* ConfigData = IDisplayClusterConfiguration::Get().LoadConfig(ConfigPath);
 		if (!ConfigData)
 		{
-			FDisplayClusterAppExit::ExitApplication(FString("An error occurred during loading the configuration file"));
+			FDisplayClusterAppExit::ExitApplication(FString("An error occurred during loading the configuration file"), FDisplayClusterAppExit::EExitType::KillImmediately);
 		}
 
 		// Extract node ID from command line
@@ -128,7 +128,7 @@ void UDisplayClusterGameEngine::Init(class IEngineLoop* InEngineLoop)
 			// Find node ID based on the host address
 			if (!GetResolvedNodeId(ConfigData, NodeId))
 			{
-				FDisplayClusterAppExit::ExitApplication(FString("Couldn't resolve node ID. Try to specify host addresses explicitly."));
+				FDisplayClusterAppExit::ExitApplication(FString("Couldn't resolve node ID. Try to specify host addresses explicitly."), FDisplayClusterAppExit::EExitType::KillImmediately);
 			}
 
 			UE_LOG(LogDisplayClusterEngine, Log, TEXT("Node ID has been successfully resolved: %s"), *NodeId);
@@ -140,7 +140,7 @@ void UDisplayClusterGameEngine::Init(class IEngineLoop* InEngineLoop)
 		// Start game session
 		if (!GDisplayCluster->StartSession(ConfigData, NodeId))
 		{
-			FDisplayClusterAppExit::ExitApplication(FString("Couldn't start DisplayCluster session"));
+			FDisplayClusterAppExit::ExitApplication(FString("Couldn't start DisplayCluster session"), FDisplayClusterAppExit::EExitType::KillImmediately);
 		}
 
 		// Initialize internals
diff --git a/Engine/Plugins/Runtime/nDisplay/Source/DisplayCluster/Private/Misc/DisplayClusterAppExit.cpp b/Engine/Plugins/Runtime/nDisplay/Source/DisplayCluster/Private/Misc/DisplayClusterAppExit.cpp
--- a/Engine/Plugins/Runtime/nDisplay/Source/DisplayCluster/Private/Misc/DisplayClusterAppExit.cpp
+++ b/Engine/Plugins/Runtime/nDisplay/Source/DisplayCluster/Private/Misc/DisplayClusterAppExit.cpp
@@ -12,7 +12,7 @@
 #endif
 
 
-void FDisplayClusterAppExit::ExitApplication(const FString& Msg)
+void FDisplayClusterAppExit::ExitApplication(const FString& Msg, EExitType ExitType)
 {
 	if (GEngine && GEngine->IsEditor())
 	{
@@ -23,24 +23,32 @@ void FDisplayClusterAppExit::ExitApplication(const FString& Msg)
 	}
 	else
 	{
-		if (!IsEngineExitRequested())
+		if (ExitType == EExitType::Normal)
 		{
-			UE_LOG(LogDisplayClusterModule, Log, TEXT("Exit requested - %s"), *Msg);
-
-			if (IsInGameThread())
-			{
-				FPlatformMisc::RequestExit(false);
-			}
-			else
+			if (!IsEngineExitRequested())
 			{
-				// For some reason UE4 generates crash info if FPlatformMisc::RequestExit gets called
-				// from a thread other than GameThread. Since it may be called from the networking
-				// session threads (failover pipeline), we don't want to generate unnecessary crash reports.
-				AsyncTask(ENamedThreads::GameThread, []()
+				UE_LOG(LogDisplayClusterModule, Log, TEXT("Exit requested - %s"), *Msg);
+
+				if (IsInGameThread())
 				{
 					FPlatformMisc::RequestExit(false);
-				});
+				}
+				else
+				{
+					// For some reason UE4 generates crash info if FPlatformMisc::RequestExit gets called
+					// from a thread other than GameThread. Since it may be called from the networking
+					// session threads (failover pipeline), we don't want to generate unnecessary crash reports.
+					AsyncTask(ENamedThreads::GameThread, []()
+						{
+							FPlatformMisc::RequestExit(false);
+						});
+				}
 			}
 		}
+		else if(ExitType == EExitType::KillImmediately)
+		{
+			FProcHandle hProc = FPlatformProcess::OpenProcess(FPlatformProcess::GetCurrentProcessId());
+			FPlatformProcess::TerminateProc(hProc, true);
+		}
 	}
 }
diff --git a/Engine/Plugins/Runtime/nDisplay/Source/DisplayCluster/Private/Misc/DisplayClusterAppExit.h b/Engine/Plugins/Runtime/nDisplay/Source/DisplayCluster/Private/Misc/DisplayClusterAppExit.h
--- a/Engine/Plugins/Runtime/nDisplay/Source/DisplayCluster/Private/Misc/DisplayClusterAppExit.h
+++ b/Engine/Plugins/Runtime/nDisplay/Source/DisplayCluster/Private/Misc/DisplayClusterAppExit.h
@@ -11,5 +11,14 @@
 class FDisplayClusterAppExit
 {
 public:
-	static void ExitApplication(const FString& Msg);
+	enum class EExitType
+	{
+		// Delicate UE exit, full resource cleaning.
+		Normal,
+		// Kills current process. No resource cleaning performed.
+		KillImmediately,
+	};
+
+public:
+	static void ExitApplication(const FString& Msg, EExitType ExitType = FDisplayClusterAppExit::EExitType::Normal);
 };
