nDisplay: force DCRA sublevel visibility on

diff --git a/Engine/Plugins/Runtime/nDisplay/Source/DisplayCluster/Private/Game/DisplayClusterGameManager.cpp b/Engine/Plugins/Runtime/nDisplay/Source/DisplayCluster/Private/Game/DisplayClusterGameManager.cpp
--- a/Engine/Plugins/Runtime/nDisplay/Source/DisplayCluster/Private/Game/DisplayClusterGameManager.cpp
+++ b/Engine/Plugins/Runtime/nDisplay/Source/DisplayCluster/Private/Game/DisplayClusterGameManager.cpp
@@ -25,8 +25,10 @@
 
 #include "GameFramework/Actor.h"
 #include "GameFramework/PlayerStart.h"
+#include "Engine/Level.h"
 #include "Engine/LevelStreaming.h"
 #include "EngineUtils.h"
+#include "LevelUtils.h"
 
 
 FDisplayClusterGameManager::FDisplayClusterGameManager()
@@ -146,6 +148,26 @@ void FDisplayClusterGameManager::EndScene()
 	CurrentWorld = nullptr;
 }
 
+void FDisplayClusterGameManager::PreTick(float DeltaSeconds)
+{
+	// Here we check if the active DCRA instance is in a sublevel, and whether
+	// the sublevel visibility is off. If so, we forcibly make the sublevel visible.
+	if (const AActor* const DCRA = DisplayClusterRootActorRef.GetOrFindSceneActor())
+	{
+		if (const ULevel* const Level = DCRA->GetLevel())
+		{
+			if (!Level->bIsVisible)
+			{
+				if (ULevelStreaming* const StreamingLevel = FLevelUtils::FindStreamingLevel(Level))
+				{
+					UE_LOG(LogDisplayClusterGame, Warning, TEXT("A streaming level containing the active DCRA instance is currently invisible. Forcing level visibility on."));
+					StreamingLevel->SetShouldBeVisible(true);
+				}
+			}
+		}
+	}
+}
+
 
 //////////////////////////////////////////////////////////////////////////////////////////////
 // IDisplayClusterGameManager
diff --git a/Engine/Plugins/Runtime/nDisplay/Source/DisplayCluster/Private/Game/DisplayClusterGameManager.h b/Engine/Plugins/Runtime/nDisplay/Source/DisplayCluster/Private/Game/DisplayClusterGameManager.h
--- a/Engine/Plugins/Runtime/nDisplay/Source/DisplayCluster/Private/Game/DisplayClusterGameManager.h
+++ b/Engine/Plugins/Runtime/nDisplay/Source/DisplayCluster/Private/Game/DisplayClusterGameManager.h
@@ -35,6 +35,7 @@ public:
 	virtual void EndSession() override;
 	virtual bool StartScene(UWorld* World) override;
 	virtual void EndScene() override;
+	virtual void PreTick(float DeltaSeconds) override;
 
 public:
 	//////////////////////////////////////////////////////////////////////////////////////////////
