nDisplay: Added support for the resize tool in the OutputMapping UI to lock the aspect ratio when being dragged by holding shift

diff --git a/Engine/Plugins/Runtime/nDisplay/Source/DisplayClusterConfigurator/Private/Views/OutputMapping/Widgets/SDisplayClusterConfiguratorResizer.cpp b/Engine/Plugins/Runtime/nDisplay/Source/DisplayClusterConfigurator/Private/Views/OutputMapping/Widgets/SDisplayClusterConfiguratorResizer.cpp
--- a/Engine/Plugins/Runtime/nDisplay/Source/DisplayClusterConfigurator/Private/Views/OutputMapping/Widgets/SDisplayClusterConfiguratorResizer.cpp
+++ b/Engine/Plugins/Runtime/nDisplay/Source/DisplayClusterConfigurator/Private/Views/OutputMapping/Widgets/SDisplayClusterConfiguratorResizer.cpp
@@ -95,7 +95,7 @@ FReply SDisplayClusterConfiguratorResizer::OnMouseMove(const FGeometry& MyGeomet
 
 		FVector2D NewNodeSize = MouseEvent.GetScreenSpacePosition() - BaseNode->GetTickSpaceGeometry().GetAbsolutePosition();
 
-		bool bIsFixedAspectRatio = IsFixedAspectRatio.Get(false);
+		bool bIsFixedAspectRatio = IsAspectRatioFixed();
 		if (bIsFixedAspectRatio)
 		{
 			// If the aspect ratio is fixed, first get the node's current size to compute the ratio from,
@@ -147,4 +147,12 @@ float SDisplayClusterConfiguratorResizer::GetDPIScale() const
 	return DPIScale;
 }
 
+bool SDisplayClusterConfiguratorResizer::IsAspectRatioFixed() const
+{
+	bool bIsShiftPressed = FSlateApplication::Get().GetModifierKeys().IsShiftDown();
+	bool bIsFixedAspectRatio = IsFixedAspectRatio.Get(false);
+
+	return bIsFixedAspectRatio || bIsShiftPressed;
+}
+
 #undef LOCTEXT_NAMESPACE
diff --git a/Engine/Plugins/Runtime/nDisplay/Source/DisplayClusterConfigurator/Private/Views/OutputMapping/Widgets/SDisplayClusterConfiguratorResizer.h b/Engine/Plugins/Runtime/nDisplay/Source/DisplayClusterConfigurator/Private/Views/OutputMapping/Widgets/SDisplayClusterConfiguratorResizer.h
--- a/Engine/Plugins/Runtime/nDisplay/Source/DisplayClusterConfigurator/Private/Views/OutputMapping/Widgets/SDisplayClusterConfiguratorResizer.h
+++ b/Engine/Plugins/Runtime/nDisplay/Source/DisplayClusterConfigurator/Private/Views/OutputMapping/Widgets/SDisplayClusterConfiguratorResizer.h
@@ -36,6 +36,7 @@ public:
 
 private:
 	float GetDPIScale() const;
+	bool IsAspectRatioFixed() const;
 
 private:
 	TWeakPtr<FDisplayClusterConfiguratorBlueprintEditor> ToolkitPtr;
