nDisplay: Removed 2.2 gamma correction from LC Editor.

diff --git a/Engine/Plugins/Runtime/nDisplay/Shaders/Private/MeshProjectionShaders.usf b/Engine/Plugins/Runtime/nDisplay/Shaders/Private/MeshProjectionShaders.usf
--- a/Engine/Plugins/Runtime/nDisplay/Shaders/Private/MeshProjectionShaders.usf
+++ b/Engine/Plugins/Runtime/nDisplay/Shaders/Private/MeshProjectionShaders.usf
@@ -98,9 +98,6 @@ void MainPS(FMeshProjectionVSToPS Input OPTIONAL_IsFrontFace, out float4 OutputC
 	
 	OutputColor = float4(Emissive, Opacity);
 	
-	static const float Gamma = 2.2;
-	OutputColor.rgb = pow(OutputColor.rgb, Gamma);
-	
 	static const float AzimuthFeatherStart = radians(180 - 45);
 	static const float AzimuthFeatherEnd   = radians(180 - 30);
 	
