nDisplay: Add New LC now parents under root actor.

Since the typical intention is for it to move with the stage.

diff --git a/Engine/Plugins/Runtime/nDisplay/Source/DisplayClusterLightCardEditor/Private/SDisplayClusterLightCardList.cpp b/Engine/Plugins/Runtime/nDisplay/Source/DisplayClusterLightCardEditor/Private/SDisplayClusterLightCardList.cpp
--- a/Engine/Plugins/Runtime/nDisplay/Source/DisplayClusterLightCardEditor/Private/SDisplayClusterLightCardList.cpp
+++ b/Engine/Plugins/Runtime/nDisplay/Source/DisplayClusterLightCardEditor/Private/SDisplayClusterLightCardList.cpp
@@ -363,6 +363,9 @@ void SDisplayClusterLightCardList::AddNewLightCard()
 	ensure(NewActor->GetLevel() == RootActor->GetLevel());
 	NewActor->SetActorLabel(NewActor->GetName());
 
+	// Parent to the root actor, since the typical intention is for the LC to move with the stage.
+	NewActor->AttachToActor(RootActor.Get(), FAttachmentTransformRules::KeepWorldTransform);
+
 	AddLightCardToActor(NewActor);
 }
 
