nDisplay: Fix operator details panel closing on modifying light cards. The previous light card selection is now maintained when switching to a new root actor if they share light cards.

diff --git a/Engine/Plugins/Runtime/nDisplay/Source/DisplayClusterLightCardEditor/Private/SDisplayClusterLightCardList.cpp b/Engine/Plugins/Runtime/nDisplay/Source/DisplayClusterLightCardEditor/Private/SDisplayClusterLightCardList.cpp
--- a/Engine/Plugins/Runtime/nDisplay/Source/DisplayClusterLightCardEditor/Private/SDisplayClusterLightCardList.cpp
+++ b/Engine/Plugins/Runtime/nDisplay/Source/DisplayClusterLightCardEditor/Private/SDisplayClusterLightCardList.cpp
@@ -193,6 +193,9 @@ void SDisplayClusterLightCardList::Construct(const FArguments& InArgs, TSharedPt
 
 void SDisplayClusterLightCardList::SetRootActor(ADisplayClusterRootActor* NewRootActor)
 {
+	TArray<TSharedPtr<FLightCardTreeItem>> PreviouslySelectedLightCards;
+	LightCardTreeView->GetSelectedItems(PreviouslySelectedLightCards);
+	
 	RootActor = NewRootActor;
 	if (FillLightCardList())
 	{
@@ -206,6 +209,21 @@ void SDisplayClusterLightCardList::SetRootActor(ADisplayClusterRootActor* NewRoo
 			LightCardTreeView->SetItemExpansion(LightCardTreeItem, true);
 		}
 	}
+
+	// Select previously selected light cards. This fixes an issue where the details panel may clear in some cases
+	// and also supports maintaining the selection of shared light cards across different root actors.
+	{
+		TArray<AActor*> LightCardActorsToSelect;
+		for (const TSharedPtr<FLightCardTreeItem>& LightCard : PreviouslySelectedLightCards)
+		{
+			if (LightCard.IsValid() && LightCard->LightCardActor.IsValid())
+			{
+				LightCardActorsToSelect.Add(LightCard.Get()->LightCardActor.Get());
+			}
+		}
+
+		SelectLightCards(MoveTemp(LightCardActorsToSelect));
+	}
 }
 
 void SDisplayClusterLightCardList::SelectLightCards(const TArray<AActor*>& LightCardsToSelect)
