nDisplay: New lightcards now spawn in the middle of the LC Editor view.

diff --git a/Engine/Plugins/Runtime/nDisplay/Source/DisplayClusterLightCardEditor/Private/SDisplayClusterLightCardEditor.cpp b/Engine/Plugins/Runtime/nDisplay/Source/DisplayClusterLightCardEditor/Private/SDisplayClusterLightCardEditor.cpp
--- a/Engine/Plugins/Runtime/nDisplay/Source/DisplayClusterLightCardEditor/Private/SDisplayClusterLightCardEditor.cpp
+++ b/Engine/Plugins/Runtime/nDisplay/Source/DisplayClusterLightCardEditor/Private/SDisplayClusterLightCardEditor.cpp
@@ -134,6 +134,12 @@ void SDisplayClusterLightCardEditor::SelectLightCardProxies(const TArray<AActor*
 	ViewportView->GetLightCardEditorViewportClient()->SelectLightCards(LightCardsToSelect);
 }
 
+void SDisplayClusterLightCardEditor::CenterLightCardInView(ADisplayClusterLightCardActor& LightCard)
+{
+	check(ViewportView);
+	ViewportView->GetLightCardEditorViewportClient()->CenterLightCardInView(LightCard);
+}
+
 void SDisplayClusterLightCardEditor::OnActiveRootActorChanged(ADisplayClusterRootActor* NewRootActor)
 {
 	RemoveCompileDelegates();
diff --git a/Engine/Plugins/Runtime/nDisplay/Source/DisplayClusterLightCardEditor/Private/SDisplayClusterLightCardEditor.h b/Engine/Plugins/Runtime/nDisplay/Source/DisplayClusterLightCardEditor/Private/SDisplayClusterLightCardEditor.h
--- a/Engine/Plugins/Runtime/nDisplay/Source/DisplayClusterLightCardEditor/Private/SDisplayClusterLightCardEditor.h
+++ b/Engine/Plugins/Runtime/nDisplay/Source/DisplayClusterLightCardEditor/Private/SDisplayClusterLightCardEditor.h
@@ -16,6 +16,7 @@ class SDockTab;
 class SDisplayClusterLightCardList;
 class SDisplayClusterLightCardEditorViewport;
 class ADisplayClusterRootActor;
+class ADisplayClusterLightCardActor;
 
 /** A panel that can be spawned in a tab that contains all the UI elements that make up the 2D light cards editor */
 class SDisplayClusterLightCardEditor : public SCompoundWidget
@@ -56,6 +57,9 @@ public:
 	/** Selects the light card proxies that correspond to the specified light cards */
 	void SelectLightCardProxies(const TArray<AActor*>& LightCardsToSelect);
 
+	/** Places the given light card in the middle of the viewport */
+	void CenterLightCardInView(ADisplayClusterLightCardActor& LightCard);
+
 private:
 	/** Raised when the active Display cluster root actor has been changed in the operator panel */
 	void OnActiveRootActorChanged(ADisplayClusterRootActor* NewRootActor);
diff --git a/Engine/Plugins/Runtime/nDisplay/Source/DisplayClusterLightCardEditor/Private/SDisplayClusterLightCardList.cpp b/Engine/Plugins/Runtime/nDisplay/Source/DisplayClusterLightCardEditor/Private/SDisplayClusterLightCardList.cpp
--- a/Engine/Plugins/Runtime/nDisplay/Source/DisplayClusterLightCardEditor/Private/SDisplayClusterLightCardList.cpp
+++ b/Engine/Plugins/Runtime/nDisplay/Source/DisplayClusterLightCardEditor/Private/SDisplayClusterLightCardList.cpp
@@ -386,6 +386,12 @@ void SDisplayClusterLightCardList::AddNewLightCard()
 	NewActor->AttachToActor(RootActor.Get(), FAttachmentTransformRules::KeepWorldTransform);
 
 	AddLightCardToActor(NewActor);
+
+	// When adding a new lightcard, usually the desired location is in the middle of the viewport
+	if (LightCardEditorPtr.IsValid())
+	{
+		LightCardEditorPtr.Pin()->CenterLightCardInView(*NewActor);
+	}
 }
 
 void SDisplayClusterLightCardList::AddExistingLightCard()
diff --git a/Engine/Plugins/Runtime/nDisplay/Source/DisplayClusterLightCardEditor/Private/Viewport/DisplayClusterLightCardEditorViewportClient.cpp b/Engine/Plugins/Runtime/nDisplay/Source/DisplayClusterLightCardEditor/Private/Viewport/DisplayClusterLightCardEditorViewportClient.cpp
--- a/Engine/Plugins/Runtime/nDisplay/Source/DisplayClusterLightCardEditor/Private/Viewport/DisplayClusterLightCardEditorViewportClient.cpp
+++ b/Engine/Plugins/Runtime/nDisplay/Source/DisplayClusterLightCardEditor/Private/Viewport/DisplayClusterLightCardEditorViewportClient.cpp
@@ -87,7 +87,7 @@ void FDisplayClusterLightCardEditorViewportClient::FNormalMap::Release()
 		});
 }
 
-bool FDisplayClusterLightCardEditorViewportClient::FNormalMap::GetNormalAndDistanceAtPosition(FVector Position, FVector& OutNormal, float& OutDistance)
+bool FDisplayClusterLightCardEditorViewportClient::FNormalMap::GetNormalAndDistanceAtPosition(FVector Position, FVector& OutNormal, float& OutDistance) const
 {
 	auto GetPixel = [this](uint32 InX, uint32 InY)
 	{
@@ -790,7 +790,7 @@ void FDisplayClusterLightCardEditorViewportClient::UpdatePreviewActor(ADisplayCl
 				RootActorProxy->SetActorLocation(FVector::ZeroVector);
 				RootActorProxy->SetActorRotation(FRotator::ZeroRotator);
 
-				FindProjectionOriginComponent();
+				ProjectionOriginComponent = FindProjectionOriginComponent(RootActorProxy.Get());
 
 				RootActorProxy->UpdatePreviewComponents();
 				RootActorProxy->EnableEditorRender(false);
@@ -936,7 +936,7 @@ void FDisplayClusterLightCardEditorViewportClient::SetProjectionMode(EDisplayClu
 		EditorWidget->SetWidgetScale(0.5f);
 	}
 
-	FindProjectionOriginComponent();
+	ProjectionOriginComponent = FindProjectionOriginComponent(RootActorProxy.Get());
 
 	if (Viewport)
 	{
@@ -995,6 +995,47 @@ void FDisplayClusterLightCardEditorViewportClient::ResetCamera(bool bLocationOnl
 	ResetFOVs();
 }
 
+void FDisplayClusterLightCardEditorViewportClient::MoveLightCardTo(ADisplayClusterLightCardActor& LightCard, const FSphericalCoordinates& SphericalCoords) const
+{
+	const FVector Origin = GetViewLocation(); // Assumed the same as the LC origin in proxy space. Call VerifyAndFixLightCardOrigin to ensure this.
+	const FVector LightCardPosition = Origin + SphericalCoords.AsCartesian();
+
+	FVector DesiredNormal;
+	float DesiredDistance;
+
+	// If the light card is in the southern hemisphere of the view origin, use the southern normal map; otherwise, use the north normal map
+	if (LightCardPosition.Z < Origin.Z)
+	{
+		SouthNormalMap.GetNormalAndDistanceAtPosition(LightCardPosition, DesiredNormal, DesiredDistance);
+	}
+	else
+	{
+		NorthNormalMap.GetNormalAndDistanceAtPosition(LightCardPosition, DesiredNormal, DesiredDistance);
+	}
+
+	const FRotator Rotation = FRotationMatrix::MakeFromX(-DesiredNormal).Rotator();
+
+	LightCard.DistanceFromCenter = FMath::Min(DesiredDistance, RootActorBoundingRadius) + LightCardFlushOffset;
+
+	LightCard.Longitude = FRotator::ClampAxis(FMath::RadiansToDegrees(SphericalCoords.Azimuth) - 180);
+	LightCard.Latitude = 90.f - FMath::RadiansToDegrees(SphericalCoords.Inclination);
+
+	LightCard.Pitch = Rotation.Pitch;
+	LightCard.Yaw = Rotation.Yaw;
+}
+
+void FDisplayClusterLightCardEditorViewportClient::CenterLightCardInView(ADisplayClusterLightCardActor& LightCard)
+{
+	VerifyAndFixLightCardOrigin(&LightCard);
+	MoveLightCardTo(LightCard, FSphericalCoordinates(GetViewRotation().RotateVector(FVector::ForwardVector)));
+
+	// If this is a proxy light, propagate to its counterpart in the level.
+	if (LightCard.bIsProxy)
+	{
+		PropagateLightCardTransform(&LightCard);
+	}
+}
+
 void FDisplayClusterLightCardEditorViewportClient::BeginTransaction(const FText& Description)
 {
 	GEditor->BeginTransaction(Description);
@@ -1224,26 +1265,22 @@ UDisplayClusterConfigurationViewport* FDisplayClusterLightCardEditorViewportClie
 	return nullptr;
 }
 
-void FDisplayClusterLightCardEditorViewportClient::FindProjectionOriginComponent()
+USceneComponent* FDisplayClusterLightCardEditorViewportClient::FindProjectionOriginComponent(const ADisplayClusterRootActor* InRootActor) const
 {
-	if (RootActorProxy.IsValid())
+	if (!InRootActor)
 	{
-		TArray<UDisplayClusterCameraComponent*> ViewOriginComponents;
-		RootActorProxy->GetComponents<UDisplayClusterCameraComponent>(ViewOriginComponents);
-
-		if (ViewOriginComponents.Num())
-		{
-			ProjectionOriginComponent = ViewOriginComponents[0];
-		}
-		else
-		{
-			ProjectionOriginComponent = RootActorProxy->GetRootComponent();
-		}
+		return nullptr;
 	}
-	else
+
+	TArray<UDisplayClusterCameraComponent*> ViewOriginComponents;
+	InRootActor->GetComponents<UDisplayClusterCameraComponent>(ViewOriginComponents);
+
+	if (ViewOriginComponents.Num())
 	{
-		ProjectionOriginComponent.Reset();
+		return ViewOriginComponents[0];
 	}
+
+	return InRootActor->GetRootComponent();
 }
 
 void FDisplayClusterLightCardEditorViewportClient::FindLightCardsForRootActor(ADisplayClusterRootActor* RootActor, TArray<TWeakObjectPtr<ADisplayClusterLightCardActor>>& OutLightCards)
@@ -1383,6 +1420,41 @@ void FDisplayClusterLightCardEditorViewportClient::PropagateLightCardTransform(A
 	}
 }
 
+void FDisplayClusterLightCardEditorViewportClient::VerifyAndFixLightCardOrigin(ADisplayClusterLightCardActor* LightCard) const
+{
+	// Center lightcard on the current view origin, let it keep its current world placement 
+	// (but not its spin/yaw/pitch since that will be happen later using the cache).
+
+	if (!LightCard)
+	{
+		return;
+	}
+
+	const ADisplayClusterRootActor* RootActor = LightCard->bIsProxy ? RootActorProxy.Get() : RootActorLevelInstance.Get();
+	const USceneComponent* OriginComponent = LightCard->bIsProxy ? ProjectionOriginComponent.Get() : FindProjectionOriginComponent(RootActor);
+
+	// Set location at the view origin
+	const FVector& NewLightCardActorLocation = OriginComponent ? OriginComponent->GetComponentLocation() : FVector::ZeroVector;
+
+	// Set rotation to match the root actor
+	const FRotator& NewLightCardActorRotation = RootActor ? RootActor->GetActorRotation() : FRotator::ZeroRotator;
+
+	const FVector LightCardEndEffectorLocation = LightCard->GetLightCardTransform().GetLocation();
+
+	LightCard->SetActorLocation(NewLightCardActorLocation);
+	LightCard->SetActorRotation(NewLightCardActorRotation);
+
+	// Update the light card spherical coordinates to match its current world coordinates
+
+	const FVector LightCardRelativeLocation = NewLightCardActorRotation.UnrotateVector(LightCardEndEffectorLocation - NewLightCardActorLocation);
+
+	const FSphericalCoordinates SphericalCoords(LightCardRelativeLocation);
+
+	LightCard->DistanceFromCenter = SphericalCoords.Radius;
+	LightCard->Longitude = FMath::RadiansToDegrees(SphericalCoords.Radius) - 180;
+	LightCard->Latitude = 90.f - FMath::RadiansToDegrees(SphericalCoords.Radius);
+}
+
 void FDisplayClusterLightCardEditorViewportClient::MoveSelectedLightCards(FViewport* InViewport, EAxisList::Type CurrentAxis)
 {
 	UWorld* PreviewWorld = PreviewScene->GetWorld();
@@ -1417,55 +1489,23 @@ void FDisplayClusterLightCardEditorViewportClient::MoveSelectedLightCards(FViewp
 				continue;
 			}
 
-			// Light cards should be centered on the current view origin, so set the light card position to match the current view origin. Update the light card
-			// spherical coordinates to match its current coordinates
-			if (ProjectionOriginComponent.IsValid() && ProjectionOriginComponent->GetComponentLocation() != LightCard->GetActorLocation())
-			{
-				const FVector DesiredLightCardOffset = LightCard->GetLightCardTransform().GetTranslation() - ProjectionOriginComponent->GetComponentLocation();
-
-				LightCard->SetActorLocation(ProjectionOriginComponent->GetComponentLocation());
-
-				FSphericalCoordinates SphericalCoords(DesiredLightCardOffset);
-
-				LightCard->DistanceFromCenter = SphericalCoords.Radius;
-				LightCard->Longitude = FMath::RadiansToDegrees(SphericalCoords.Radius) - 180;
-				LightCard->Latitude = 90.f - FMath::RadiansToDegrees(SphericalCoords.Radius);
-			}
-
-			const FSphericalCoordinates CurrentCoords = GetLightCardCoordinates(LightCard.Get());
+			VerifyAndFixLightCardOrigin(LightCard.Get());
 
 			FSphericalCoordinates NewCoords;
-			NewCoords.Radius = CurrentCoords.Radius + DeltaCoords.Radius;
-			NewCoords.Azimuth = CurrentCoords.Azimuth + DeltaCoords.Azimuth;
-			NewCoords.Inclination = CurrentCoords.Inclination + DeltaCoords.Inclination;
-
-			LightCard->DistanceFromCenter = NewCoords.Radius;
-			LightCard->Longitude = FRotator::ClampAxis(FMath::RadiansToDegrees(NewCoords.Azimuth) - 180);
-			LightCard->Latitude = 90.f - FMath::RadiansToDegrees(NewCoords.Inclination);
-
 			{
-				const FVector LightCardPosition = LightCard->GetLightCardTransform().GetTranslation();
+				const FSphericalCoordinates CurrentCoords = GetLightCardCoordinates(LightCard.Get());
 
-				FVector DesiredNormal;
-				float DesiredDistance;
+				NewCoords.Radius = CurrentCoords.Radius + DeltaCoords.Radius;
+				NewCoords.Azimuth = CurrentCoords.Azimuth + DeltaCoords.Azimuth;
+				NewCoords.Inclination = CurrentCoords.Inclination + DeltaCoords.Inclination;
 
-				// If the light card is in the southern hemisphere of the view origin, use the southern normal map; otherwise, use the north normal map
-				if (LightCardPosition.Z < Origin.Z)
-				{
-					SouthNormalMap.GetNormalAndDistanceAtPosition(LightCardPosition, DesiredNormal, DesiredDistance);
-				}
-				else
-				{
-					NorthNormalMap.GetNormalAndDistanceAtPosition(LightCardPosition, DesiredNormal, DesiredDistance);
-				}
-
-				const FRotator Rotation = FRotationMatrix::MakeFromX(-DesiredNormal).Rotator();
-
-				LightCard->Pitch = Rotation.Pitch;
-				LightCard->Yaw = Rotation.Yaw;
-				LightCard->DistanceFromCenter = FMath::Min(DesiredDistance, RootActorBoundingRadius) + LightCardFlushOffset;
+				LightCard->DistanceFromCenter = NewCoords.Radius;
+				LightCard->Longitude = FRotator::ClampAxis(FMath::RadiansToDegrees(NewCoords.Azimuth) - 180);
+				LightCard->Latitude = 90.f - FMath::RadiansToDegrees(NewCoords.Inclination);
 			}
 
+			MoveLightCardTo(*LightCard.Get(), NewCoords);
+
 			PropagateLightCardTransform(LightCard.Get());
 		}
 
@@ -1557,7 +1597,7 @@ FDisplayClusterLightCardEditorViewportClient::FSphericalCoordinates FDisplayClus
 	// If the light card inclination is 0 or 180, the spherical coordinates will have an
 	// "undefined" azimuth value. For continuity when dragging a light card positioned there, we can manually
 	// set the azimuthal value to match the light card's configured longitude
-	if (LightCardCoords.Inclination == 0.f ||LightCardCoords.Inclination == PI)
+	if (LightCardCoords.Inclination == 0.f || LightCardCoords.Inclination == PI)
 	{
 		LightCardCoords.Azimuth = FMath::DegreesToRadians(LightCard->Longitude + 180);
 	}
diff --git a/Engine/Plugins/Runtime/nDisplay/Source/DisplayClusterLightCardEditor/Private/Viewport/DisplayClusterLightCardEditorViewportClient.h b/Engine/Plugins/Runtime/nDisplay/Source/DisplayClusterLightCardEditor/Private/Viewport/DisplayClusterLightCardEditorViewportClient.h
--- a/Engine/Plugins/Runtime/nDisplay/Source/DisplayClusterLightCardEditor/Private/Viewport/DisplayClusterLightCardEditorViewportClient.h
+++ b/Engine/Plugins/Runtime/nDisplay/Source/DisplayClusterLightCardEditor/Private/Viewport/DisplayClusterLightCardEditorViewportClient.h
@@ -34,21 +34,36 @@ private:
 			}
 			else
 			{
-				Inclination = 0.f;
+				Inclination = 0;
 			}
 
 			Azimuth = FMath::Atan2(CartesianPosition.Y, CartesianPosition.X);
 		}
 
 		FSphericalCoordinates()
-			: Radius(0.f)
-			, Inclination(0.f)
-			, Azimuth(0.f)
+			: Radius(0)
+			, Inclination(0)
+			, Azimuth(0)
 		{ }
 
-		float Radius = 0.f;
-		float Inclination = 0.f;
-		float Azimuth = 0.f;
+		FVector AsCartesian() const
+		{
+			const double SinAzimuth = FMath::Sin(Azimuth);
+			const double CosAzimuth = FMath::Cos(Azimuth);
+
+			const double SinInclination = FMath::Sin(Inclination);
+			const double CosInclination = FMath::Cos(Inclination);
+
+			return FVector(
+				Radius * CosAzimuth * SinInclination,
+				Radius * SinAzimuth * SinInclination,
+				Radius * CosInclination
+			);
+		}
+
+		double Radius = 0;
+		double Inclination = 0;
+		double Azimuth = 0;
 	};
 
 	/** Custom render target that stores the normal data for the stage */
@@ -75,7 +90,7 @@ private:
 		TArray<FFloat16Color>& GetCachedNormalData() { return CachedNormalData; }
 
 		/** Gets the normal vector and distance at the specified world location. The normal and distance are bilinearly interpolated from the nearest pixels in the normal map */
-		bool GetNormalAndDistanceAtPosition(FVector Position, FVector& OutNormal, float& OutDistance);
+		bool GetNormalAndDistanceAtPosition(FVector Position, FVector& OutNormal, float& OutDistance) const;
 
 		/** Generates a texture object that can be used to visualize the normal map */
 		UTexture2D* GenerateNormalMapTexture(const FString& TextureName);
@@ -157,6 +172,16 @@ public:
 
 	/** Resets the camera to the initial rotation / position */
 	void ResetCamera(bool bLocationOnly = false);
+
+	/** Moves specified card to desired coordinates 
+	 *
+	 * @param LightCard The light card that we are moving
+	 * @param SphericalCoords specifies desired location of light card in spherical coordinates with respect to view origin
+	*/
+	void MoveLightCardTo(ADisplayClusterLightCardActor& LightCard, const FSphericalCoordinates& SphericalCoords) const;
+
+	/** Places the given light card in the middle of the current viewport */
+	void CenterLightCardInView(ADisplayClusterLightCardActor& LightCard);
 	
 private:
 	/** Initiates a transaction. */
@@ -178,7 +203,7 @@ private:
 	UDisplayClusterConfigurationViewport* FindViewportForPrimitiveComponent(UPrimitiveComponent* PrimitiveComponent);
 
 	/** Finds a suitable primitive component on the stage actor to use as a projection origin */
-	void FindProjectionOriginComponent();
+	USceneComponent* FindProjectionOriginComponent(const ADisplayClusterRootActor* InRootActor) const;
 
 	/** Gets a list of all light card actors on the level linked to the specified root actor */
 	void FindLightCardsForRootActor(ADisplayClusterRootActor* RootActor, TArray<TWeakObjectPtr<ADisplayClusterLightCardActor>>& OutLightCards);
@@ -198,6 +223,9 @@ private:
 	/** Moves the currently selected light cards */
 	void MoveSelectedLightCards(FViewport* InViewport, EAxisList::Type CurrentAxis);
 
+	/** Ensures that the light card root component is at the same location as the projection/view origin */
+	void VerifyAndFixLightCardOrigin(ADisplayClusterLightCardActor* LightCard) const;
+
 	/** Determines the appropriate delta in spherical coordinates needed to move the specified light card to the mouse's location */
 	FSphericalCoordinates GetLightCardTranslationDelta(FViewport* InViewport, ADisplayClusterLightCardActor* LightCard, EAxisList::Type CurrentAxis);
 
