Replaced use of CopyToResolveTarget with CopyTexture.

diff --git a/Engine/Plugins/Media/PixelStreaming/Source/PixelStreaming/Private/TextureSourceComputeI420.cpp b/Engine/Plugins/Media/PixelStreaming/Source/PixelStreaming/Private/TextureSourceComputeI420.cpp
--- a/Engine/Plugins/Media/PixelStreaming/Source/PixelStreaming/Private/TextureSourceComputeI420.cpp
+++ b/Engine/Plugins/Media/PixelStreaming/Source/PixelStreaming/Private/TextureSourceComputeI420.cpp
@@ -160,8 +160,10 @@ namespace
 
 			TempFBCopy = GDynamicRHI->RHICreateTexture(Desc);
 		}
-		RHICmdList.Transition(FRHITransitionInfo(SourceTexture, ERHIAccess::Unknown, ERHIAccess::CopySrc));
-		RHICmdList.Transition(FRHITransitionInfo(TempFBCopy, ERHIAccess::Unknown, ERHIAccess::CopyDest));
+		RHICmdList.Transition({
+			FRHITransitionInfo(SourceTexture, ERHIAccess::Unknown, ERHIAccess::CopySrc),
+			FRHITransitionInfo(TempFBCopy, ERHIAccess::Unknown, ERHIAccess::CopyDest)
+		});
 		RHICmdList.CopyTexture(SourceTexture, TempFBCopy, {});
 		RHICmdList.Transition(FRHITransitionInfo(TempFBCopy, ERHIAccess::Unknown, ERHIAccess::SRVMask));
 		
@@ -174,9 +176,9 @@ namespace
 		ShaderParameters.DestPlaneV = I420ReadbackTexture->TextureVUAV;
 		FRGBToYUVShader::Dispatch(RHICmdList, ShaderParameters);
 
-		RHICmdList.CopyToResolveTarget(I420ReadbackTexture->TextureY, I420ReadbackTexture->StagingTextureY, {});
-		RHICmdList.CopyToResolveTarget(I420ReadbackTexture->TextureU, I420ReadbackTexture->StagingTextureU, {});
-		RHICmdList.CopyToResolveTarget(I420ReadbackTexture->TextureV, I420ReadbackTexture->StagingTextureV, {});
+		RHICmdList.CopyTexture(I420ReadbackTexture->TextureY, I420ReadbackTexture->StagingTextureY, {});
+		RHICmdList.CopyTexture(I420ReadbackTexture->TextureU, I420ReadbackTexture->StagingTextureU, {});
+		RHICmdList.CopyTexture(I420ReadbackTexture->TextureV, I420ReadbackTexture->StagingTextureV, {});
 
 		RHICmdList.WriteGPUFence(Fence);
 	}
diff --git a/Engine/Plugins/MeshPainting/Source/MeshPaintingToolset/Private/MeshTexturePaintingTool.cpp b/Engine/Plugins/MeshPainting/Source/MeshPaintingToolset/Private/MeshTexturePaintingTool.cpp
--- a/Engine/Plugins/MeshPainting/Source/MeshPaintingToolset/Private/MeshTexturePaintingTool.cpp
+++ b/Engine/Plugins/MeshPainting/Source/MeshPaintingToolset/Private/MeshTexturePaintingTool.cpp
@@ -998,35 +998,21 @@ void UMeshTexturePaintingTool::PaintTexture(FMeshPaintParameters& InParams, TArr
 		TextureData->bIsPaintingTexture2DModified = true;
 	}
 
+	ENQUEUE_RENDER_COMMAND(UpdateMeshPaintRTCommand1)(
+		[BrushRenderTargetResource](FRHICommandListImmediate& RHICmdList)
 	{
-		ENQUEUE_RENDER_COMMAND(UpdateMeshPaintRTCommand1)(
-			[BrushRenderTargetResource](FRHICommandListImmediate& RHICmdList)
-		{
-			// Copy (resolve) the rendered image from the frame buffer to its render target texture
-			RHICmdList.CopyToResolveTarget(
-				BrushRenderTargetResource->GetRenderTargetTexture(),	// Source texture
-				BrushRenderTargetResource->TextureRHI,
-				FResolveParams());										// Resolve parameters
-		});
-	}
-
+		CopyTextureWithTransitions(RHICmdList, BrushRenderTargetResource->GetRenderTargetTexture(), BrushRenderTargetResource->TextureRHI, {});
+	});
 
 	if (bEnableSeamPainting)
 	{
 		BrushMaskCanvas->Flush_GameThread(true);
 
+		ENQUEUE_RENDER_COMMAND(UpdateMeshPaintRTCommand2)(
+			[BrushMaskRenderTargetResource](FRHICommandListImmediate& RHICmdList)
 		{
-			ENQUEUE_RENDER_COMMAND(UpdateMeshPaintRTCommand2)(
-				[BrushMaskRenderTargetResource](FRHICommandListImmediate& RHICmdList)
-			{
-				// Copy (resolve) the rendered image from the frame buffer to its render target texture
-				RHICmdList.CopyToResolveTarget(
-					BrushMaskRenderTargetResource->GetRenderTargetTexture(),		// Source texture
-					BrushMaskRenderTargetResource->TextureRHI,
-					FResolveParams());									// Resolve parameters
-			});
-
-		}
+			CopyTextureWithTransitions(RHICmdList, BrushMaskRenderTargetResource->GetRenderTargetTexture(), BrushMaskRenderTargetResource->TextureRHI, {});
+		});
 	}
 
 	if (!bEnableSeamPainting)
@@ -1116,19 +1102,11 @@ void UMeshTexturePaintingTool::PaintTexture(FMeshPaintParameters& InParams, TArr
 
 		}
 
+		ENQUEUE_RENDER_COMMAND(UpdateMeshPaintRTCommand3)(
+			[RenderTargetResource](FRHICommandListImmediate& RHICmdList)
 		{
-			ENQUEUE_RENDER_COMMAND(UpdateMeshPaintRTCommand3)(
-				[RenderTargetResource](FRHICommandListImmediate& RHICmdList)
-			{
-				// Copy (resolve) the rendered image from the frame buffer to its render target texture
-				RHICmdList.CopyToResolveTarget(
-					RenderTargetResource->GetRenderTargetTexture(),		// Source texture
-					RenderTargetResource->TextureRHI,
-					FResolveParams());									// Resolve parameters
-			});
-
-		}
-
+			CopyTextureWithTransitions(RHICmdList, RenderTargetResource->GetRenderTargetTexture(), RenderTargetResource->TextureRHI, {});
+		});
 	}
 	FlushRenderingCommands();
 }
diff --git a/Engine/Plugins/MeshPainting/Source/MeshPaintingToolset/Private/TexturePaintToolset.cpp b/Engine/Plugins/MeshPainting/Source/MeshPaintingToolset/Private/TexturePaintToolset.cpp
--- a/Engine/Plugins/MeshPainting/Source/MeshPaintingToolset/Private/TexturePaintToolset.cpp
+++ b/Engine/Plugins/MeshPainting/Source/MeshPaintingToolset/Private/TexturePaintToolset.cpp
@@ -96,12 +96,8 @@ void UTexturePaintToolset::CopyTextureToRenderTargetTexture(UTexture* SourceText
 	ENQUEUE_RENDER_COMMAND(UpdateMeshPaintRTCommand)(
 		[RenderTargetResource](FRHICommandListImmediate& RHICmdList)
 		{
-			// Copy (resolve) the rendered image from the frame buffer to its render target texture
-			RHICmdList.CopyToResolveTarget(
-				RenderTargetResource->GetRenderTargetTexture(),		// Source texture
-				RenderTargetResource->TextureRHI,					// Dest texture
-				FResolveParams());									// Resolve parameters
-		});		
+			CopyTextureWithTransitions(RHICmdList, RenderTargetResource->GetRenderTargetTexture(), RenderTargetResource->TextureRHI, {});
+		});
 }
 
 bool UTexturePaintToolset::GenerateSeamMask(UMeshComponent* MeshComponent, int32 UVSet, UTextureRenderTarget2D* SeamRenderTexture, UTexture2D* Texture, UTextureRenderTarget2D* RenderTargetTexture)
@@ -308,14 +304,9 @@ bool UTexturePaintToolset::GenerateSeamMask(UMeshComponent* MeshComponent, int32
 	{
 		ENQUEUE_RENDER_COMMAND(UpdateMeshPaintRTCommand5)(
 			[RenderTargetResource](FRHICommandListImmediate& RHICmdList)
-			{
-				// Copy (resolve) the rendered image from the frame buffer to its render target texture
-				RHICmdList.CopyToResolveTarget(
-					RenderTargetResource->GetRenderTargetTexture(),		// Source texture
-					RenderTargetResource->TextureRHI,
-					FResolveParams());									// Resolve parameters
-			});
-
+		{
+			CopyTextureWithTransitions(RHICmdList, RenderTargetResource->GetRenderTargetTexture(), RenderTargetResource->TextureRHI, {});
+		});
 	}
 
 	return RetVal;
diff --git a/Engine/Plugins/Runtime/TextureShare/Source/TextureShare/Private/TextureShareRHI.cpp b/Engine/Plugins/Runtime/TextureShare/Source/TextureShare/Private/TextureShareRHI.cpp
--- a/Engine/Plugins/Runtime/TextureShare/Source/TextureShare/Private/TextureShareRHI.cpp
+++ b/Engine/Plugins/Runtime/TextureShare/Source/TextureShare/Private/TextureShareRHI.cpp
@@ -60,24 +60,18 @@ namespace
 		IsSizeResampleRequired(SrcTexture, DstTexture, SrcTextureRect, DstTextureRect, SrcRect, DstRect);
 
 		// Copy with resolved params
-		FResolveParams Params;
+		FRHICopyTextureInfo Params;
 
-		Params.DestArrayIndex = 0;
-		Params.SourceArrayIndex = 0;
+		Params.SourcePosition.X = SrcRect.Min.X;
+		Params.SourcePosition.Y = SrcRect.Min.Y;
 
-		Params.Rect.X1 = SrcRect.Min.X;
-		Params.Rect.X2 = SrcRect.Max.X;
+		Params.DestPosition.X = DstRect.Min.X;
+		Params.DestPosition.Y = DstRect.Min.Y;
 
-		Params.Rect.Y1 = SrcRect.Min.Y;
-		Params.Rect.Y2 = SrcRect.Max.Y;
+		Params.Size.X = SrcRect.Width();
+		Params.Size.Y = SrcRect.Height();
 
-		Params.DestRect.X1 = DstRect.Min.X;
-		Params.DestRect.X2 = DstRect.Max.X;
-
-		Params.DestRect.Y1 = DstRect.Min.Y;
-		Params.DestRect.Y2 = DstRect.Max.Y;
-
-		RHICmdList.CopyToResolveTarget(SrcTexture, DstTexture, Params);
+		CopyTextureWithTransitions(RHICmdList, SrcTexture, DstTexture, {});
 	}
 
 	static void ResampleCopyTextureImpl_RenderThread(FRHICommandListImmediate& RHICmdList, FRHITexture* SrcTexture, FRHITexture* DstTexture, const FIntRect* SrcTextureRect, const FIntRect* DstTextureRect)
diff --git a/Engine/Plugins/Runtime/TextureShare/Source/TextureShareD3D12/Private/CrossGPUHeap/D3D12CrossGPUHeap.cpp b/Engine/Plugins/Runtime/TextureShare/Source/TextureShareD3D12/Private/CrossGPUHeap/D3D12CrossGPUHeap.cpp
--- a/Engine/Plugins/Runtime/TextureShare/Source/TextureShareD3D12/Private/CrossGPUHeap/D3D12CrossGPUHeap.cpp
+++ b/Engine/Plugins/Runtime/TextureShare/Source/TextureShareD3D12/Private/CrossGPUHeap/D3D12CrossGPUHeap.cpp
@@ -183,7 +183,7 @@ bool FD3D12CrossGPUHeap::EndCrossGPUSession(FRHICommandListImmediate& RHICmdList
 static void CopyDirectTextureImpl_RenderThread(FRHICommandListImmediate& RHICmdList, FRHITexture* SrcTexture, FRHITexture* DstTexture, const FIntRect* SrcTextureRect, const FIntRect* DstTextureRect)
 {
 	// Copy direct, format identical
-	FResolveParams Params;
+	FRHICopyTextureInfo Params;
 	if (SrcTextureRect || DstTextureRect)
 	{
 		FIntVector SrcSizeXYZ = SrcTexture->GetSizeXYZ();
@@ -195,23 +195,17 @@ static void CopyDirectTextureImpl_RenderThread(FRHICommandListImmediate& RHICmdL
 		FIntRect SrcRect = SrcTextureRect ? (*SrcTextureRect) : (FIntRect(FIntPoint(0, 0), SrcSize));
 		FIntRect DstRect = DstTextureRect ? (*DstTextureRect) : (FIntRect(FIntPoint(0, 0), DstSize));
 
-		Params.DestArrayIndex = 0;
-		Params.SourceArrayIndex = 0;
+		Params.SourcePosition.X = SrcRect.Min.X;
+		Params.SourcePosition.Y = SrcRect.Min.Y;
 
-		Params.Rect.X1 = SrcRect.Min.X;
-		Params.Rect.X2 = SrcRect.Max.X;
+		Params.Size.X = SrcRect.Width();
+		Params.Size.Y = SrcRect.Height();
 
-		Params.Rect.Y1 = SrcRect.Min.Y;
-		Params.Rect.Y2 = SrcRect.Max.Y;
-
-		Params.DestRect.X1 = DstRect.Min.X;
-		Params.DestRect.X2 = DstRect.Max.X;
-
-		Params.DestRect.Y1 = DstRect.Min.Y;
-		Params.DestRect.Y2 = DstRect.Max.Y;
+		Params.DestPosition.X = DstRect.Min.X;
+		Params.DestPosition.Y = DstRect.Min.Y;
 	}
 
-	RHICmdList.CopyToResolveTarget(SrcTexture, DstTexture, Params);
+	CopyTextureWithTransitions(RHICmdList, SrcTexture, DstTexture, {});
 }
 
 bool FD3D12CrossGPUHeap::SendCrossGPUResource(FRHICommandListImmediate& RHICmdList, const FString& ResourceID, FRHITexture2D* SrcResource, const FIntRect* SrcTextureRect)
diff --git a/Engine/Plugins/Runtime/nDisplay/Source/DisplayCluster/Private/Render/Device/DisplayClusterDeviceBase.cpp b/Engine/Plugins/Runtime/nDisplay/Source/DisplayCluster/Private/Render/Device/DisplayClusterDeviceBase.cpp
--- a/Engine/Plugins/Runtime/nDisplay/Source/DisplayCluster/Private/Render/Device/DisplayClusterDeviceBase.cpp
+++ b/Engine/Plugins/Runtime/nDisplay/Source/DisplayCluster/Private/Render/Device/DisplayClusterDeviceBase.cpp
@@ -476,11 +476,11 @@ void FDisplayClusterDeviceBase::RenderTexture_RenderThread(FRHICommandListImmedi
 			const FIntPoint SrcSize = SrcTexture->GetSizeXY();
 			const FIntPoint DstSize = BackBuffer->GetSizeXY();
 
-			FResolveRect CopyRect(0, 0, FMath::Min(SrcSize.X, DstSize.X), FMath::Min(SrcSize.Y, DstSize.Y));
-			FResolveParams CopyParams = {};
-			CopyParams.Rect = CopyParams.DestRect = CopyRect;
+			FRHICopyTextureInfo CopyInfo;
+			CopyInfo.Size.X = FMath::Min(SrcSize.X, DstSize.X);
+			CopyInfo.Size.Y = FMath::Min(SrcSize.Y, DstSize.Y);
 
-			RHICmdList.CopyToResolveTarget(SrcTexture, BackBuffer, CopyParams);
+			CopyTextureWithTransitions(RHICmdList, SrcTexture, BackBuffer, CopyInfo);
 		}
 
 		if (RenderFrameMode == EDisplayClusterRenderFrameMode::Stereo && ViewportManagerProxyPtr)
@@ -493,10 +493,8 @@ void FDisplayClusterDeviceBase::RenderTexture_RenderThread(FRHICommandListImmedi
 		if (bClearTextureEnabled)
 		{
 			// Clear render target before out frame resolving, help to make things look better visually for console/resize, etc.
-			FRHIRenderPassInfo RPInfo(SrcTexture, ERenderTargetActions::Clear_Store);
 			RHICmdList.Transition(FRHITransitionInfo(SrcTexture, ERHIAccess::Unknown, ERHIAccess::RTV));
-			RHICmdList.BeginRenderPass(RPInfo, TEXT("ClearTexture"));
-			RHICmdList.EndRenderPass();
+			ClearRenderTarget(RHICmdList, SrcTexture);
 		}
 	}
 }
diff --git a/Engine/Plugins/Runtime/nDisplay/Source/DisplayCluster/Private/Render/Viewport/DisplayClusterViewportManagerProxy.cpp b/Engine/Plugins/Runtime/nDisplay/Source/DisplayCluster/Private/Render/Viewport/DisplayClusterViewportManagerProxy.cpp
--- a/Engine/Plugins/Runtime/nDisplay/Source/DisplayCluster/Private/Render/Viewport/DisplayClusterViewportManagerProxy.cpp
+++ b/Engine/Plugins/Runtime/nDisplay/Source/DisplayCluster/Private/Render/Viewport/DisplayClusterViewportManagerProxy.cpp
@@ -547,22 +547,18 @@ bool FDisplayClusterViewportManagerProxy::ResolveFrameTargetToBackBuffer_RenderT
 			DstRect.Max.X = FMath::Min(DstSize.X, DstRect.Max.X);
 			DstRect.Max.Y = FMath::Min(DstSize.Y, DstRect.Max.Y);
 
-			FResolveParams CopyParams;
+			FRHICopyTextureInfo CopyInfo;
 
-			CopyParams.SourceArrayIndex = 0;
-			CopyParams.DestArrayIndex = DestArrayIndex;
+			CopyInfo.SourceSliceIndex = 0;
+			CopyInfo.DestSliceIndex = DestArrayIndex;
 
-			CopyParams.Rect.X1 = 0;
-			CopyParams.Rect.Y1 = 0;
-			CopyParams.Rect.X2 = DstRect.Width();
-			CopyParams.Rect.Y2 = DstRect.Height();
+			CopyInfo.Size.X = DstRect.Width();
+			CopyInfo.Size.Y = DstRect.Height();
 
-			CopyParams.DestRect.X1 = DstRect.Min.X;
-			CopyParams.DestRect.Y1 = DstRect.Min.Y;
-			CopyParams.DestRect.X2 = DstRect.Max.X;
-			CopyParams.DestRect.Y2 = DstRect.Max.Y;
+			CopyInfo.DestPosition.X = DstRect.Min.X;
+			CopyInfo.DestPosition.Y = DstRect.Min.Y;
 
-			RHICmdList.CopyToResolveTarget(FrameTexture, DestTexture, CopyParams);
+			CopyTextureWithTransitions(RHICmdList, FrameTexture, DestTexture, CopyInfo);
 
 			return true;
 		}
diff --git a/Engine/Plugins/Runtime/nDisplay/Source/DisplayCluster/Private/Render/Viewport/DisplayClusterViewportProxy.cpp b/Engine/Plugins/Runtime/nDisplay/Source/DisplayCluster/Private/Render/Viewport/DisplayClusterViewportProxy.cpp
--- a/Engine/Plugins/Runtime/nDisplay/Source/DisplayCluster/Private/Render/Viewport/DisplayClusterViewportProxy.cpp
+++ b/Engine/Plugins/Runtime/nDisplay/Source/DisplayCluster/Private/Render/Viewport/DisplayClusterViewportProxy.cpp
@@ -510,27 +510,6 @@ void FDisplayClusterViewportProxy::UpdateDeferredResources(FRHICommandListImmedi
 	}
 }
 
-static void DirectCopyTextureImpl_RenderThread(FRHICommandListImmediate& RHICmdList, FRHITexture* SrcTexture, FRHITexture* DstTexture, const FIntRect& SrcRect, const FIntRect& DstRect)
-{
-	// Copy with resolved params
-	FResolveParams Params;
-
-	Params.DestArrayIndex = 0;
-	Params.SourceArrayIndex = 0;
-
-	Params.Rect.X1 = SrcRect.Min.X;
-	Params.Rect.Y1 = SrcRect.Min.Y;
-	Params.Rect.X2 = SrcRect.Max.X;
-	Params.Rect.Y2 = SrcRect.Max.Y;
-
-	Params.DestRect.X1 = DstRect.Min.X;
-	Params.DestRect.Y1 = DstRect.Min.Y;
-	Params.DestRect.X2 = DstRect.Max.X;
-	Params.DestRect.Y2 = DstRect.Max.Y;
-
-	RHICmdList.CopyToResolveTarget(SrcTexture, DstTexture, Params);
-}
-
 template<class TScreenPixelShader>
 void ResampleCopyTextureImpl_RenderThread(FRHICommandListImmediate& RHICmdList, FRHITexture* SrcTexture, FRHITexture* DstTexture, const FIntRect& SrcRect, const FIntRect& DstRect)
 {
@@ -599,11 +578,13 @@ void ImplResolveResource(FRHICommandListImmediate& RHICmdList, FRHITexture2D* In
 	check(InputResource);
 	check(OutputResource);
 
-	if (bOutputIsMipsResource)
+	if (bOutputIsPreviewResource)
+	{
+		// Preview require a normal alpha (re-invert)
+		ResampleCopyTextureImpl_RenderThread<FScreenPSInvertAlpha>(RHICmdList, InputResource, OutputResource, InputRect, OutputRect);
+	}
+	else if (InputRect.Size() == OutputRect.Size() && InputResource->GetFormat() == OutputResource->GetFormat())
 	{
-		// Special copy for output texture with NumMips (copy to dest mips = 0)
-		// Copy first Mip from SrcTexture to CameraRTT(other mips is black)
-
 		FRHICopyTextureInfo CopyInfo;
 		CopyInfo.Size = FIntVector(InputRect.Width(), InputRect.Height(), 0);
 		CopyInfo.SourcePosition.X = InputRect.Min.X;
@@ -611,43 +592,11 @@ void ImplResolveResource(FRHICommandListImmediate& RHICmdList, FRHITexture2D* In
 		CopyInfo.DestPosition.X = OutputRect.Min.X;
 		CopyInfo.DestPosition.Y = OutputRect.Min.Y;
 
-		// Copy texture
-		FRHITransitionInfo TransitionsBefore[] =
-		{
-			FRHITransitionInfo(InputResource, ERHIAccess::Unknown, ERHIAccess::CopySrc),
-			FRHITransitionInfo(OutputResource, ERHIAccess::Unknown, ERHIAccess::CopyDest)
-		};
-		RHICmdList.Transition(MakeArrayView(TransitionsBefore, UE_ARRAY_COUNT(TransitionsBefore)));
-
-		RHICmdList.CopyTexture(InputResource, OutputResource, CopyInfo);
-
-		FRHITransitionInfo TransitionsAfter[] =
-		{
-			FRHITransitionInfo(InputResource, ERHIAccess::CopySrc, ERHIAccess::SRVMask),
-			FRHITransitionInfo(OutputResource, ERHIAccess::CopyDest, ERHIAccess::SRVMask)
-		};
-		RHICmdList.Transition(MakeArrayView(TransitionsAfter, UE_ARRAY_COUNT(TransitionsAfter)));
+		CopyTextureWithTransitions(RHICmdList, InputResource, OutputResource, CopyInfo);
 	}
 	else
 	{
-		if (InputRect.Size() != OutputRect.Size() || InputResource->GetFormat() != OutputResource->GetFormat() || bOutputIsPreviewResource)
-		{
-			// Resample size with shader
-			if(bOutputIsPreviewResource)
-			{
-				// Preview require a normal alpha (re-invert)
-				ResampleCopyTextureImpl_RenderThread<FScreenPSInvertAlpha>(RHICmdList, InputResource, OutputResource, InputRect, OutputRect);
-			}
-			else
-			{
-				ResampleCopyTextureImpl_RenderThread<FScreenPS>(RHICmdList, InputResource, OutputResource, InputRect, OutputRect);
-			}
-		}
-		else
-		{
-			// Use RHICmdList.CopyToResolveTarget inside
-			DirectCopyTextureImpl_RenderThread(RHICmdList, InputResource, OutputResource, InputRect, OutputRect);
-		}
+		ResampleCopyTextureImpl_RenderThread<FScreenPS>(RHICmdList, InputResource, OutputResource, InputRect, OutputRect);
 	}
 }
 
diff --git a/Engine/Plugins/Runtime/nDisplay/Source/DisplayCluster/Private/Render/Viewport/Postprocess/DisplayClusterViewportPostProcessOutputRemap.cpp b/Engine/Plugins/Runtime/nDisplay/Source/DisplayCluster/Private/Render/Viewport/Postprocess/DisplayClusterViewportPostProcessOutputRemap.cpp
--- a/Engine/Plugins/Runtime/nDisplay/Source/DisplayCluster/Private/Render/Viewport/Postprocess/DisplayClusterViewportPostProcessOutputRemap.cpp
+++ b/Engine/Plugins/Runtime/nDisplay/Source/DisplayCluster/Private/Render/Viewport/Postprocess/DisplayClusterViewportPostProcessOutputRemap.cpp
@@ -270,8 +270,7 @@ void FDisplayClusterViewportPostProcessOutputRemap::PerformPostProcessFrame_Rend
 
 				if (ShadersAPI.RenderPostprocess_OutputRemap(RHICmdList, InOutTexture, TempTargetableTexture, *MeshProxy))
 				{
-					//Copy remapped result back
-					RHICmdList.CopyToResolveTarget(TempTargetableTexture, InOutTexture, FResolveParams());
+					CopyTextureWithTransitions(RHICmdList, TempTargetableTexture, InOutTexture, {});
 				}
 			}
 		}
diff --git a/Engine/Plugins/Runtime/nDisplay/Source/DisplayClusterShaders/Private/Shaders/DisplayClusterShadersPostprocess_Blur.cpp b/Engine/Plugins/Runtime/nDisplay/Source/DisplayClusterShaders/Private/Shaders/DisplayClusterShadersPostprocess_Blur.cpp
--- a/Engine/Plugins/Runtime/nDisplay/Source/DisplayClusterShaders/Private/Shaders/DisplayClusterShadersPostprocess_Blur.cpp
+++ b/Engine/Plugins/Runtime/nDisplay/Source/DisplayClusterShaders/Private/Shaders/DisplayClusterShadersPostprocess_Blur.cpp
@@ -157,10 +157,25 @@ static void PicpBlurPostProcess_RenderThread(
 		// Blur X
 		PixelShader->SetParameters(RHICmdList, InShaderTexture, FVector2f(InSettings.KernelScale / TargetSizeXY.X, 0.0f), InSettings.KernelRadius);
 		FPixelShaderUtils::DrawFullscreenQuad(RHICmdList, 1);
+	
+		RHICmdList.EndRenderPass();
 
-		RHICmdList.CopyToResolveTarget(OutRenderTargetableTexture, InShaderTexture, FResolveParams());
+		RHICmdList.Transition({
+			FRHITransitionInfo(OutRenderTargetableTexture, ERHIAccess::RTV, ERHIAccess::CopySrc),
+			FRHITransitionInfo(InShaderTexture, ERHIAccess::Unknown, ERHIAccess::CopyDest),
+		});
+
+		RHICmdList.CopyTexture(OutRenderTargetableTexture, InShaderTexture, {});
+
+		RHICmdList.Transition({
+			FRHITransitionInfo(OutRenderTargetableTexture, ERHIAccess::CopySrc, ERHIAccess::RTV),
+			FRHITransitionInfo(InShaderTexture, ERHIAccess::CopyDest, ERHIAccess::SRVMask),
+		});
+
+		RHICmdList.BeginRenderPass(RPInfo, TEXT("nDisplay_PicpPostProcessBlur"));
 
 		// Blur Y
+		SetGraphicsPipelineState(RHICmdList, GraphicsPSOInit, 0);
 		PixelShader->SetParameters(RHICmdList, InShaderTexture, FVector2f(0.0f, InSettings.KernelScale / TargetSizeXY.Y), InSettings.KernelRadius);
 		FPixelShaderUtils::DrawFullscreenQuad(RHICmdList, 1);
 	}
diff --git a/Engine/Plugins/VirtualProduction/DMX/DMXPixelMapping/Source/DMXPixelMappingRenderer/Private/DMXPixelMappingRenderer.cpp b/Engine/Plugins/VirtualProduction/DMX/DMXPixelMapping/Source/DMXPixelMappingRenderer/Private/DMXPixelMappingRenderer.cpp
--- a/Engine/Plugins/VirtualProduction/DMX/DMXPixelMapping/Source/DMXPixelMappingRenderer/Private/DMXPixelMappingRenderer.cpp
+++ b/Engine/Plugins/VirtualProduction/DMX/DMXPixelMapping/Source/DMXPixelMappingRenderer/Private/DMXPixelMappingRenderer.cpp
@@ -136,11 +136,11 @@ void FDMXPixelMappingRenderer::DownsampleRender(
 			SCOPED_GPU_STAT(RHICmdList, DMXPixelMappingShadersStat);
 			SCOPED_DRAW_EVENTF(RHICmdList, DMXPixelMappingShadersStat, DMXPixelMappingRenderer::RenderPassName);
 
-			const FTextureRHIRef RenderTargetRef = DstTextureTargetResource->TextureRHI;
-			const FTextureRHIRef DstTextureRef = DstTexture->TextureRHI;
-			const FTexture2DRHIRef ResolveRenderTarget = DstTextureTargetResource->GetRenderTargetTexture();
+			FRHITexture* RenderTargetRef = DstTextureTargetResource->TextureRHI;
+			FRHITexture* DstTextureRef = DstTexture->TextureRHI;
+			FRHITexture* ResolveRenderTarget = DstTextureTargetResource->GetRenderTargetTexture();
 
-			if (!RenderTargetRef.IsValid() || !DstTextureRef.IsValid() || !ResolveRenderTarget.IsValid())
+			if (!RenderTargetRef || !DstTextureRef || !ResolveRenderTarget)
 			{
 				ensure(false);
 				return;
@@ -214,11 +214,7 @@ void FDMXPixelMappingRenderer::DownsampleRender(
 
 			// Copy texture from GPU to CPU
 			{
-				// Copies the contents of the given surface to its resolve target texture.
-				FResolveParams ResolveParams;
-				ResolveParams.SourceAccessFinal = ERHIAccess::SRVMask;
-				ResolveParams.DestAccessFinal = ERHIAccess::SRVMask;
-				RHICmdList.CopyToResolveTarget(ResolveRenderTarget, RenderTargetRef, ResolveParams);
+				CopyTextureWithTransitions(RHICmdList, ResolveRenderTarget, RenderTargetRef, {});
 
 				// Read the contents of a texture to an output CPU buffer
 				TArray<FLinearColor> ColorArray;
diff --git a/Engine/Source/Developer/MaterialBaking/Private/MaterialBakingModule.cpp b/Engine/Source/Developer/MaterialBaking/Private/MaterialBakingModule.cpp
--- a/Engine/Source/Developer/MaterialBaking/Private/MaterialBakingModule.cpp
+++ b/Engine/Source/Developer/MaterialBaking/Private/MaterialBakingModule.cpp
@@ -610,8 +610,9 @@ void FMaterialBakingModule::BakeMaterials(const TArray<FMaterialDataEx*>& Materi
 							FTexture2DRHIRef StagingBufferRef = StagingBufferPool.CreateStagingBuffer_RenderThread(RHICmdList, RenderTargetResource->GetSizeX(), RenderTargetResource->GetSizeY(), RenderTarget->GetFormat(), RenderTarget->IsSRGB());
 							FGPUFenceRHIRef GPUFence = RHICreateGPUFence(TEXT("MaterialBackingFence"));
 
-							FResolveRect Rect(0, 0, RenderTargetResource->GetSizeX(), RenderTargetResource->GetSizeY());
-							RHICmdList.CopyToResolveTarget(RenderTargetResource->GetRenderTargetTexture(), StagingBufferRef, FResolveParams(Rect));	
+							RHICmdList.Transition(FRHITransitionInfo(RenderTargetResource->GetRenderTargetTexture(), ERHIAccess::Unknown, ERHIAccess::CopySrc));
+							RHICmdList.CopyTexture(RenderTargetResource->GetRenderTargetTexture(), StagingBufferRef, {});
+							RHICmdList.Transition(FRHITransitionInfo(RenderTargetResource->GetRenderTargetTexture(), ERHIAccess::CopySrc, ERHIAccess::SRVMask));
 							RHICmdList.WriteGPUFence(GPUFence);
 
 							// Prepare a lambda for final processing that will be executed asynchronously
diff --git a/Engine/Source/Editor/MeshPaint/Private/TexturePaintHelpers.cpp b/Engine/Source/Editor/MeshPaint/Private/TexturePaintHelpers.cpp
--- a/Engine/Source/Editor/MeshPaint/Private/TexturePaintHelpers.cpp
+++ b/Engine/Source/Editor/MeshPaint/Private/TexturePaintHelpers.cpp
@@ -96,11 +96,7 @@ void TexturePaintHelpers::CopyTextureToRenderTargetTexture(UTexture* SourceTextu
 	ENQUEUE_RENDER_COMMAND(UpdateMeshPaintRTCommand)(
 		[RenderTargetResource](FRHICommandListImmediate& RHICmdList)
 		{
-			// Copy (resolve) the rendered image from the frame buffer to its render target texture
-			RHICmdList.CopyToResolveTarget(
-				RenderTargetResource->GetRenderTargetTexture(),		// Source texture
-				RenderTargetResource->TextureRHI,					// Dest texture
-				FResolveParams());									// Resolve parameters
+			CopyTextureWithTransitions(RHICmdList, RenderTargetResource->GetRenderTargetTexture(), RenderTargetResource->TextureRHI, {});
 		});		
 }
 
@@ -309,13 +305,8 @@ bool TexturePaintHelpers::GenerateSeamMask(UMeshComponent* MeshComponent, int32
 		ENQUEUE_RENDER_COMMAND(UpdateMeshPaintRTCommand5)(
 			[RenderTargetResource](FRHICommandListImmediate& RHICmdList)
 			{
-				// Copy (resolve) the rendered image from the frame buffer to its render target texture
-				RHICmdList.CopyToResolveTarget(
-					RenderTargetResource->GetRenderTargetTexture(),		// Source texture
-					RenderTargetResource->TextureRHI,
-					FResolveParams());									// Resolve parameters
+				CopyTextureWithTransitions(RHICmdList, RenderTargetResource->GetRenderTargetTexture(), RenderTargetResource->TextureRHI, {});
 			});
-
 	}
 
 	return RetVal;
diff --git a/Engine/Source/Editor/UnrealEd/Private/ObjectTools.cpp b/Engine/Source/Editor/UnrealEd/Private/ObjectTools.cpp
--- a/Engine/Source/Editor/UnrealEd/Private/ObjectTools.cpp
+++ b/Engine/Source/Editor/UnrealEd/Private/ObjectTools.cpp
@@ -4790,11 +4790,7 @@ namespace ThumbnailTools
 				ENQUEUE_RENDER_COMMAND(UpdateThumbnailRTCommand)(
 					[RenderTargetResource](FRHICommandListImmediate& RHICmdList)
 					{
-						// Copy (resolve) the rendered thumbnail from the render target to its texture
-						RHICmdList.CopyToResolveTarget(
-							RenderTargetResource->GetRenderTargetTexture(),		// Source texture
-							RenderTargetResource->TextureRHI,					// Dest texture
-							FResolveParams() );									// Resolve parameters
+						CopyTextureWithTransitions(RHICmdList, RenderTargetResource->GetRenderTargetTexture(), RenderTargetResource->TextureRHI, {});
 					});
 
 				if(OutThumbnail)
diff --git a/Engine/Source/Runtime/Engine/Private/KismetRenderingLibrary.cpp b/Engine/Source/Runtime/Engine/Private/KismetRenderingLibrary.cpp
--- a/Engine/Source/Runtime/Engine/Private/KismetRenderingLibrary.cpp
+++ b/Engine/Source/Runtime/Engine/Private/KismetRenderingLibrary.cpp
@@ -779,9 +779,9 @@ void UKismetRenderingLibrary::EndDrawCanvasToRenderTarget(UObject* WorldContextO
 			FTextureRenderTargetResource* RenderTargetResource = Context.RenderTarget->GameThread_GetRenderTargetResource();
 
 			ENQUEUE_RENDER_COMMAND(CanvasRenderTargetResolveCommand)(
-				[RenderTargetResource](FRHICommandList& RHICmdList)
+				[RenderTargetResource](FRHICommandListImmediate& RHICmdList)
 				{
-					RHICmdList.CopyToResolveTarget(RenderTargetResource->GetRenderTargetTexture(), RenderTargetResource->TextureRHI, FResolveParams());
+					CopyTextureWithTransitions(RHICmdList, RenderTargetResource->GetRenderTargetTexture(), RenderTargetResource->TextureRHI, {});
 				}
 			);
 
diff --git a/Engine/Source/Runtime/Engine/Private/UnrealClient.cpp b/Engine/Source/Runtime/Engine/Private/UnrealClient.cpp
--- a/Engine/Source/Runtime/Engine/Private/UnrealClient.cpp
+++ b/Engine/Source/Runtime/Engine/Private/UnrealClient.cpp
@@ -1715,18 +1715,7 @@ const TArray<FColor>& FViewport::GetRawHitProxyData(FIntRect InRect)
 		ENQUEUE_RENDER_COMMAND(UpdateHitProxyRTCommand)(
 			[HitProxyMapPtr](FRHICommandListImmediate& RHICmdList)
 			{
-				FTexture2DRHIRef RenderTargetTexture = HitProxyMapPtr->GetRenderTargetTexture();
-
-				// Should not be multisampled, so can safely use copy src as state instead of taking care of ResolveSrc here
-				check(!RenderTargetTexture->IsMultisampled());
-				
-				// Don't touch the copy source state after the first call to CopyToResolveTarget
-				FResolveParams ResolveParams;
-				ResolveParams.SourceAccessFinal = ERHIAccess::Unknown;
-
-				// Copy (resolve) the rendered thumbnail from the render target to its texture
-				RHICmdList.Transition(FRHITransitionInfo(RenderTargetTexture, ERHIAccess::Unknown, ERHIAccess::CopySrc));
-				RHICmdList.CopyToResolveTarget(RenderTargetTexture, HitProxyMapPtr->GetHitProxyCPUTexture(), ResolveParams);
+				CopyTextureWithTransitions(RHICmdList, HitProxyMapPtr->GetRenderTargetTexture(), HitProxyMapPtr->GetHitProxyCPUTexture(), {});
 			});
 
 		ENQUEUE_RENDER_COMMAND(EndDrawingCommand)(
diff --git a/Engine/Source/Runtime/GameplayMediaEncoder/Private/GameplayMediaEncoder.cpp b/Engine/Source/Runtime/GameplayMediaEncoder/Private/GameplayMediaEncoder.cpp
--- a/Engine/Source/Runtime/GameplayMediaEncoder/Private/GameplayMediaEncoder.cpp
+++ b/Engine/Source/Runtime/GameplayMediaEncoder/Private/GameplayMediaEncoder.cpp
@@ -683,7 +683,7 @@ void FGameplayMediaEncoder::CopyTexture(const FTexture2DRHIRef& SourceTexture, F
 
 	if(SourceTexture->GetFormat() == DestinationTexture->GetFormat() && SourceTexture->GetSizeXY() == DestinationTexture->GetSizeXY())
 	{
-		RHICmdList.CopyToResolveTarget(SourceTexture, DestinationTexture, FResolveParams{});
+		CopyTextureWithTransitions(RHICmdList, SourceTexture, DestinationTexture, {});
 	}
 	else // Texture format mismatch, use a shader to do the copy.
 	{
@@ -692,6 +692,7 @@ void FGameplayMediaEncoder::CopyTexture(const FTexture2DRHIRef& SourceTexture, F
 		// #todo-renderpasses there's no explicit resolve here? Do we need one?
 		FRHIRenderPassInfo RPInfo(DestinationTexture, ERenderTargetActions::Load_Store);
 
+		RHICmdList.Transition(FRHITransitionInfo(DestinationTexture, ERHIAccess::Unknown, ERHIAccess::RTV));
 		RHICmdList.BeginRenderPass(RPInfo, TEXT("CopyBackbuffer"));
 
 		{
@@ -736,5 +737,6 @@ void FGameplayMediaEncoder::CopyTexture(const FTexture2DRHIRef& SourceTexture, F
 		}
 
 		RHICmdList.EndRenderPass();
+		RHICmdList.Transition(FRHITransitionInfo(DestinationTexture, ERHIAccess::RTV, ERHIAccess::SRVMask));
 	}
 }
diff --git a/Engine/Source/Runtime/Landscape/Private/LandscapeEdit.cpp b/Engine/Source/Runtime/Landscape/Private/LandscapeEdit.cpp
--- a/Engine/Source/Runtime/Landscape/Private/LandscapeEdit.cpp
+++ b/Engine/Source/Runtime/Landscape/Private/LandscapeEdit.cpp
@@ -7533,11 +7533,7 @@ bool ALandscapeProxy::LandscapeExportHeightmapToRenderTarget(UTextureRenderTarge
 	ENQUEUE_RENDER_COMMAND(DrawHeightmapRTCommand)(
 		[RenderTargetResource](FRHICommandListImmediate& RHICmdList)
 		{
-			// Copy (resolve) the rendered image from the frame buffer to its render target texture
-			RHICmdList.CopyToResolveTarget(
-				RenderTargetResource->GetRenderTargetTexture(),		// Source texture
-				RenderTargetResource->TextureRHI,					// Dest texture
-				FResolveParams());									// Resolve parameters
+			CopyTextureWithTransitions(RHICmdList, RenderTargetResource->GetRenderTargetTexture(), RenderTargetResource->TextureRHI, {});
 		});
 
 
diff --git a/Engine/Source/Runtime/Landscape/Private/LandscapeEditLayers.cpp b/Engine/Source/Runtime/Landscape/Private/LandscapeEditLayers.cpp
--- a/Engine/Source/Runtime/Landscape/Private/LandscapeEditLayers.cpp
+++ b/Engine/Source/Runtime/Landscape/Private/LandscapeEditLayers.cpp
@@ -3269,8 +3269,7 @@ void ALandscape::PrintLayersDebugRT(const FString& InContext, UTextureRenderTarg
 		[RenderTargetResource](FRHICommandListImmediate& RHICmdList) mutable
 	{
 		TRACE_CPUPROFILER_EVENT_SCOPE(LandscapeLayers_RT_DebugResolve);
-		// Copy (resolve) the rendered image from the frame buffer to its render target texture
-		RHICmdList.CopyToResolveTarget(RenderTargetResource->GetRenderTargetTexture(), RenderTargetResource->TextureRHI, FResolveParams());
+		CopyTextureWithTransitions(RHICmdList, RenderTargetResource->GetRenderTargetTexture(), RenderTargetResource->TextureRHI, {});
 	});
 
 	{
diff --git a/Engine/Source/Runtime/RHI/Private/RHIGPUReadback.cpp b/Engine/Source/Runtime/RHI/Private/RHIGPUReadback.cpp
--- a/Engine/Source/Runtime/RHI/Private/RHIGPUReadback.cpp
+++ b/Engine/Source/Runtime/RHI/Private/RHIGPUReadback.cpp
@@ -119,14 +119,11 @@ FRHIGPUTextureReadback::FRHIGPUTextureReadback(FName RequestName) : FRHIGPUMemor
 
 void FRHIGPUTextureReadback::EnqueueCopyRDG(FRHICommandList& RHICmdList, FRHITexture* SourceTexture, FResolveRect Rect)
 {
-	// SourceTexture is already in CopySrc state (handled by RDG)
 	EnqueueCopyInternal(RHICmdList, SourceTexture, FResolveParams(Rect));
 }
 
 void FRHIGPUTextureReadback::EnqueueCopy(FRHICommandList& RHICmdList, FRHITexture* SourceTexture, FResolveRect Rect)
 {
-	// In the non-RDG version, we don't know what state the source texture will already be in, so transition it to CopySrc.
-	RHICmdList.Transition(FRHITransitionInfo(SourceTexture, ERHIAccess::Unknown, ERHIAccess::CopySrc));
 	EnqueueCopyInternal(RHICmdList, SourceTexture, FResolveParams(Rect));
 }
 
@@ -136,8 +133,8 @@ void FRHIGPUTextureReadback::EnqueueCopyInternal(FRHICommandList& RHICmdList, FR
 
 	if (SourceTexture)
 	{
-		// We only support 2d textures for now.
-		ensure(SourceTexture->GetTexture2D());
+		check(SourceTexture->GetTexture2D());
+		check(!SourceTexture->IsMultisampled());
 
 		// Assume for now that every enqueue happens on a texture of the same format and size (when reused).
 		if (!DestinationStagingTexture)
@@ -149,18 +146,9 @@ void FRHIGPUTextureReadback::EnqueueCopyInternal(FRHICommandList& RHICmdList, FR
 			DestinationStagingTexture = RHICreateTexture2D(TextureSize.X, TextureSize.Y, SourceTexture->GetFormat(), 1, 1, TexCreate_CPUReadback | TexCreate_HideInVisualizeTexture, ERHIAccess::CPURead, CreateInfo);
 		}
 
-		// We need the destination texture to be writable from a copy operation
-		RHICmdList.Transition(FRHITransitionInfo(DestinationStagingTexture, ERHIAccess::CPURead, ERHIAccess::CopyDest));
-
-		// Ensure this copy call does not perform any transitions. We're handling them manually.
-		ResolveParams.SourceAccessFinal = ERHIAccess::Unknown;
-		ResolveParams.DestAccessFinal = ERHIAccess::Unknown;
-
 		// Transfer memory GPU -> CPU
-		RHICmdList.CopyToResolveTarget(SourceTexture, DestinationStagingTexture, ResolveParams);
+		RHICmdList.CopyTexture(SourceTexture, DestinationStagingTexture, {});
 
-		// Transition the dest to CPURead *before* signaling the fence, otherwise ordering is not guaranteed.
-		RHICmdList.Transition(FRHITransitionInfo(DestinationStagingTexture, ERHIAccess::CopyDest, ERHIAccess::CPURead));
 		RHICmdList.WriteGPUFence(Fence);
 
 		LastCopyGPUMask = RHICmdList.GetGPUMask();
diff --git a/Engine/Source/Runtime/RHI/Public/RHIUtilities.h b/Engine/Source/Runtime/RHI/Public/RHIUtilities.h
--- a/Engine/Source/Runtime/RHI/Public/RHIUtilities.h
+++ b/Engine/Source/Runtime/RHI/Public/RHIUtilities.h
@@ -802,6 +802,33 @@ inline void ClearRenderTarget(FRHICommandList& RHICmdList, FRHITexture* Texture,
 	RHICmdList.EndRenderPass();
 }
 
+inline void CopyTextureWithTransitions(FRHICommandListImmediate& RHICmdList, FRHITexture* SrcTexture, FRHITexture* DstTexture, const FRHICopyTextureInfo& Info)
+{
+	check(SrcTexture && DstTexture);
+	check(SrcTexture->GetNumSamples() == DstTexture->GetNumSamples());
+
+	if (SrcTexture == DstTexture)
+	{
+		RHICmdList.Transition({
+			FRHITransitionInfo(SrcTexture, ERHIAccess::Unknown, ERHIAccess::SRVMask),
+			FRHITransitionInfo(DstTexture, ERHIAccess::Unknown, ERHIAccess::SRVMask)
+		});
+		return;
+	}
+
+	RHICmdList.Transition({
+		FRHITransitionInfo(SrcTexture, ERHIAccess::Unknown, ERHIAccess::CopySrc),
+		FRHITransitionInfo(DstTexture, ERHIAccess::Unknown, ERHIAccess::CopyDest)
+	});
+
+	RHICmdList.CopyTexture(SrcTexture, DstTexture, Info);
+
+	RHICmdList.Transition({
+		FRHITransitionInfo(SrcTexture, ERHIAccess::CopySrc,  ERHIAccess::SRVMask),
+		FRHITransitionInfo(DstTexture, ERHIAccess::CopyDest, ERHIAccess::SRVMask)
+	});
+}
+
 /**
  * Computes the vertex count for a given number of primitives of the specified type.
  * @param NumPrimitives The number of primitives.
diff --git a/Engine/Source/Runtime/RenderCore/Private/DumpGPU.cpp b/Engine/Source/Runtime/RenderCore/Private/DumpGPU.cpp
--- a/Engine/Source/Runtime/RenderCore/Private/DumpGPU.cpp
+++ b/Engine/Source/Runtime/RenderCore/Private/DumpGPU.cpp
@@ -946,13 +946,8 @@ public:
 
 				RHICmdList.Transition(FRHITransitionInfo(StagingTexture, ERHIAccess::Unknown, ERHIAccess::CopyDest));
 
-				// Ensure this copy call does not perform any transitions. We're handling them manually.
-				FResolveParams ResolveParams;
-				ResolveParams.SourceAccessFinal = ERHIAccess::Unknown;
-				ResolveParams.DestAccessFinal = ERHIAccess::Unknown;
-
 				// Transfer memory GPU -> CPU
-				RHICmdList.CopyToResolveTarget(StagingSrcTexture, StagingTexture, ResolveParams);
+				RHICmdList.CopyTexture(StagingSrcTexture, StagingTexture, {});
 
 				RHICmdList.Transition(FRHITransitionInfo(StagingTexture, ERHIAccess::CopyDest, ERHIAccess::CPURead));
 			}
diff --git a/Engine/Source/Runtime/Renderer/Private/ReflectionEnvironment.cpp b/Engine/Source/Runtime/Renderer/Private/ReflectionEnvironment.cpp
--- a/Engine/Source/Runtime/Renderer/Private/ReflectionEnvironment.cpp
+++ b/Engine/Source/Runtime/Renderer/Private/ReflectionEnvironment.cpp
@@ -481,37 +481,43 @@ void FReflectionEnvironmentCubemapArray::ResizeCubemapArrayGPU(uint32 InMaxCubem
 
 	{
 		SCOPED_DRAW_EVENT(RHICmdList, ReflectionEnvironment_ResizeCubemapArray);
-		//SCOPED_GPU_STAT(RHICmdList, ReflectionEnvironment)
+
+		RHICmdList.Transition({
+			FRHITransitionInfo(OldReflectionEnvs->GetRHI(), ERHIAccess::Unknown, ERHIAccess::CopySrc),
+			FRHITransitionInfo(ReflectionEnvs->GetRHI(), ERHIAccess::Unknown, ERHIAccess::CopyDest)
+		});
+
+		FRHICopyTextureInfo CopyInfo;
+		CopyInfo.Size = TexRef->GetSizeXYZ();
 
 		// Copy the cubemaps, remapping the elements as necessary
-		FResolveParams ResolveParams;
-		ResolveParams.Rect = FResolveRect();
 		for (int32 SourceCubemapIndex = 0; SourceCubemapIndex < OldMaxCubemaps; SourceCubemapIndex++)
 		{
 			int32 DestCubemapIndex = IndexRemapping[SourceCubemapIndex];
 			if (DestCubemapIndex != -1)
 			{
-				ResolveParams.SourceArrayIndex = SourceCubemapIndex;
-				ResolveParams.DestArrayIndex = DestCubemapIndex;
-
 				check(SourceCubemapIndex < OldMaxCubemaps);
 				check(DestCubemapIndex < (int32)MaxCubemaps);
 
-				for (int Face = 0; Face < 6; Face++)
+				for (int32 Face = 0; Face < CubeFace_MAX; Face++)
 				{
-					ResolveParams.CubeFace = (ECubeFace)Face;
-					for (int Mip = 0; Mip < NumMips; Mip++)
+					CopyInfo.SourceSliceIndex = SourceCubemapIndex * CubeFace_MAX + Face;
+					CopyInfo.DestSliceIndex   = DestCubemapIndex   * CubeFace_MAX + Face;
+
+					for (int32 Mip = 0; Mip < NumMips; Mip++)
 					{
-						ResolveParams.MipIndex = Mip;
-						//@TODO: We should use an explicit copy method for this rather than CopyToResolveTarget, but that doesn't exist right now. 
-						// For now, we'll just do this on RHIs where we know CopyToResolveTarget does the right thing. In future we should look to 
-						// add a a new RHI method
-						check(GRHISupportsResolveCubemapFaces);
-						RHICmdList.CopyToResolveTarget(OldReflectionEnvs->GetRHI(), ReflectionEnvs->GetRHI(), ResolveParams);
+						CopyInfo.SourceMipIndex = CopyInfo.DestMipIndex = Mip;
+
+						RHICmdList.CopyTexture(OldReflectionEnvs->GetRHI(), ReflectionEnvs->GetRHI(), CopyInfo);
 					}
 				}
 			}
 		}
+		
+		RHICmdList.Transition({
+			FRHITransitionInfo(OldReflectionEnvs->GetRHI(), ERHIAccess::CopySrc, ERHIAccess::SRVMask),
+			FRHITransitionInfo(ReflectionEnvs->GetRHI(), ERHIAccess::CopyDest, ERHIAccess::SRVMask)
+		});
 	}
 	GRenderTargetPool.FreeUnusedResource(OldReflectionEnvs);
 }
diff --git a/Engine/Source/Runtime/SlateRHIRenderer/Private/SlateRHIRenderer.cpp b/Engine/Source/Runtime/SlateRHIRenderer/Private/SlateRHIRenderer.cpp
--- a/Engine/Source/Runtime/SlateRHIRenderer/Private/SlateRHIRenderer.cpp
+++ b/Engine/Source/Runtime/SlateRHIRenderer/Private/SlateRHIRenderer.cpp
@@ -922,15 +922,8 @@ void FSlateRHIRenderer::DrawWindow_RenderThread(FRHICommandListImmediate& RHICmd
 			{
 				bClear = true; // Force a clear of the UI buffer to black
 
-				FRHITransitionInfo Transitions[] = {
-					FRHITransitionInfo(FinalBuffer, ERHIAccess::Unknown, ERHIAccess::ResolveSrc),
-					FRHITransitionInfo(ViewportInfo.HDRSourceRT->GetRHI(), ERHIAccess::Unknown, ERHIAccess::ResolveDst)
-				};
-				RHICmdList.Transition(MakeArrayView(Transitions, UE_ARRAY_COUNT(Transitions)));
-
 				// Grab HDR backbuffer
-				FResolveParams ResolveParams;
-				RHICmdList.CopyToResolveTarget(FinalBuffer, ViewportInfo.HDRSourceRT->GetRHI(), ResolveParams);
+				CopyTextureWithTransitions(RHICmdList, FinalBuffer, ViewportInfo.HDRSourceRT->GetRHI(), {});
 
 				// UI backbuffer is temp target
 				BackBuffer = ViewportInfo.UITargetRT->GetRHI();
