nDisplay: removed public 'Interfaces' folder for better file structuring before extending the API (DisplayClusterConfigurator)
[preparation before media integration and UE-145807]

#swarm https://p4-swarm.epicgames.net/reviews/20174552

diff --git a/Engine/Plugins/Runtime/nDisplay/Source/DisplayClusterConfigurator/Private/DisplayClusterConfiguratorBlueprintEditor.cpp b/Engine/Plugins/Runtime/nDisplay/Source/DisplayClusterConfigurator/Private/DisplayClusterConfiguratorBlueprintEditor.cpp
--- a/Engine/Plugins/Runtime/nDisplay/Source/DisplayClusterConfigurator/Private/DisplayClusterConfiguratorBlueprintEditor.cpp
+++ b/Engine/Plugins/Runtime/nDisplay/Source/DisplayClusterConfigurator/Private/DisplayClusterConfiguratorBlueprintEditor.cpp
@@ -2,7 +2,7 @@
 
 #include "DisplayClusterConfiguratorBlueprintEditor.h"
 
-#include "Interfaces/IDisplayClusterConfigurator.h"
+#include "IDisplayClusterConfigurator.h"
 #include "DisplayClusterConfigurationTypes.h"
 #include "DisplayClusterConfiguratorCommands.h"
 #include "DisplayClusterConfiguratorEditorMode.h"
diff --git a/Engine/Plugins/Runtime/nDisplay/Source/DisplayClusterConfigurator/Private/DisplayClusterConfiguratorBlueprintEditor.h b/Engine/Plugins/Runtime/nDisplay/Source/DisplayClusterConfigurator/Private/DisplayClusterConfiguratorBlueprintEditor.h
--- a/Engine/Plugins/Runtime/nDisplay/Source/DisplayClusterConfigurator/Private/DisplayClusterConfiguratorBlueprintEditor.h
+++ b/Engine/Plugins/Runtime/nDisplay/Source/DisplayClusterConfigurator/Private/DisplayClusterConfiguratorBlueprintEditor.h
@@ -3,7 +3,7 @@
 #pragma once
 
 #include "BlueprintEditorModes.h"
-#include "Interfaces/IDisplayClusterConfiguratorBlueprintEditor.h"
+#include "IDisplayClusterConfiguratorBlueprintEditor.h"
 
 class ADisplayClusterRootActor;
 class UActorComponent;
diff --git a/Engine/Plugins/Runtime/nDisplay/Source/DisplayClusterConfigurator/Private/DisplayClusterConfiguratorModule.h b/Engine/Plugins/Runtime/nDisplay/Source/DisplayClusterConfigurator/Private/DisplayClusterConfiguratorModule.h
--- a/Engine/Plugins/Runtime/nDisplay/Source/DisplayClusterConfigurator/Private/DisplayClusterConfiguratorModule.h
+++ b/Engine/Plugins/Runtime/nDisplay/Source/DisplayClusterConfigurator/Private/DisplayClusterConfiguratorModule.h
@@ -4,7 +4,7 @@
 
 #include "DisplayClusterConfiguratorCompiler.h"
 #include "IAssetTypeActions.h"
-#include "Interfaces/IDisplayClusterConfigurator.h"
+#include "IDisplayClusterConfigurator.h"
 
 class FPlacementModeID;
 class IAssetTools;
diff --git a/Engine/Plugins/Runtime/nDisplay/Source/DisplayClusterConfigurator/Private/DisplayClusterConfiguratorTabSpawners.cpp b/Engine/Plugins/Runtime/nDisplay/Source/DisplayClusterConfigurator/Private/DisplayClusterConfiguratorTabSpawners.cpp
--- a/Engine/Plugins/Runtime/nDisplay/Source/DisplayClusterConfigurator/Private/DisplayClusterConfiguratorTabSpawners.cpp
+++ b/Engine/Plugins/Runtime/nDisplay/Source/DisplayClusterConfigurator/Private/DisplayClusterConfiguratorTabSpawners.cpp
@@ -5,8 +5,8 @@
 #include "DisplayClusterConfiguratorStyle.h"
 #include "DisplayClusterConfiguratorBlueprintEditor.h"
 #include "Views/Viewport/DisplayClusterConfiguratorSCSEditorViewport.h"
-#include "Interfaces/Views/OutputMapping/IDisplayClusterConfiguratorViewOutputMapping.h"
-#include "Interfaces/Views/TreeViews/IDisplayClusterConfiguratorViewTree.h"
+#include "Views/OutputMapping/IDisplayClusterConfiguratorViewOutputMapping.h"
+#include "Views/TreeViews/IDisplayClusterConfiguratorViewTree.h"
 
 #include "BlueprintEditorTabs.h"
 #include "Widgets/Docking/SDockTab.h"
diff --git a/Engine/Plugins/Runtime/nDisplay/Source/DisplayClusterConfigurator/Private/Views/DisplayClusterConfiguratorToolbar.cpp b/Engine/Plugins/Runtime/nDisplay/Source/DisplayClusterConfigurator/Private/Views/DisplayClusterConfiguratorToolbar.cpp
--- a/Engine/Plugins/Runtime/nDisplay/Source/DisplayClusterConfigurator/Private/Views/DisplayClusterConfiguratorToolbar.cpp
+++ b/Engine/Plugins/Runtime/nDisplay/Source/DisplayClusterConfigurator/Private/Views/DisplayClusterConfiguratorToolbar.cpp
@@ -8,7 +8,7 @@
 #include "DisplayClusterConfiguratorStyle.h"
 #include "DisplayClusterConfiguratorBlueprintEditor.h"
 #include "Blueprints/DisplayClusterBlueprint.h"
-#include "Interfaces/IDisplayClusterConfigurator.h"
+#include "IDisplayClusterConfigurator.h"
 
 #include "WorkflowOrientedApp/SModeWidget.h"
 #include "Framework/MultiBox/MultiBoxBuilder.h"
diff --git a/Engine/Plugins/Runtime/nDisplay/Source/DisplayClusterConfigurator/Private/Views/Log/DisplayClusterConfiguratorViewLog.h b/Engine/Plugins/Runtime/nDisplay/Source/DisplayClusterConfigurator/Private/Views/Log/DisplayClusterConfiguratorViewLog.h
--- a/Engine/Plugins/Runtime/nDisplay/Source/DisplayClusterConfigurator/Private/Views/Log/DisplayClusterConfiguratorViewLog.h
+++ b/Engine/Plugins/Runtime/nDisplay/Source/DisplayClusterConfigurator/Private/Views/Log/DisplayClusterConfiguratorViewLog.h
@@ -2,7 +2,7 @@
 
 #pragma once
 
-#include "Interfaces/Views/Log/IDisplayClusterConfiguratorViewLog.h"
+#include "Views/Log/IDisplayClusterConfiguratorViewLog.h"
 
 class SDisplayClusterConfiguratorViewLog;
 class FDisplayClusterConfiguratorBlueprintEditor;
diff --git a/Engine/Plugins/Runtime/nDisplay/Source/DisplayClusterConfigurator/Private/Views/OutputMapping/Alignment/DisplayClusterConfiguratorHostNodeArrangementHelper.h b/Engine/Plugins/Runtime/nDisplay/Source/DisplayClusterConfigurator/Private/Views/OutputMapping/Alignment/DisplayClusterConfiguratorHostNodeArrangementHelper.h
--- a/Engine/Plugins/Runtime/nDisplay/Source/DisplayClusterConfigurator/Private/Views/OutputMapping/Alignment/DisplayClusterConfiguratorHostNodeArrangementHelper.h
+++ b/Engine/Plugins/Runtime/nDisplay/Source/DisplayClusterConfigurator/Private/Views/OutputMapping/Alignment/DisplayClusterConfiguratorHostNodeArrangementHelper.h
@@ -4,7 +4,7 @@
 
 #include "CoreMinimal.h"
 
-#include "Interfaces/Views/OutputMapping/IDisplayClusterConfiguratorViewOutputMapping.h"
+#include "Views/OutputMapping/IDisplayClusterConfiguratorViewOutputMapping.h"
 
 class UDisplayClusterConfiguratorHostNode;
 
diff --git a/Engine/Plugins/Runtime/nDisplay/Source/DisplayClusterConfigurator/Private/Views/OutputMapping/DisplayClusterConfiguratorViewOutputMapping.cpp b/Engine/Plugins/Runtime/nDisplay/Source/DisplayClusterConfigurator/Private/Views/OutputMapping/DisplayClusterConfiguratorViewOutputMapping.cpp
--- a/Engine/Plugins/Runtime/nDisplay/Source/DisplayClusterConfigurator/Private/Views/OutputMapping/DisplayClusterConfiguratorViewOutputMapping.cpp
+++ b/Engine/Plugins/Runtime/nDisplay/Source/DisplayClusterConfigurator/Private/Views/OutputMapping/DisplayClusterConfiguratorViewOutputMapping.cpp
@@ -3,7 +3,7 @@
 #include "Views/OutputMapping/DisplayClusterConfiguratorViewOutputMapping.h"
 
 #include "DisplayClusterConfiguratorBlueprintEditor.h"
-#include "Interfaces/IDisplayClusterConfigurator.h"
+#include "IDisplayClusterConfigurator.h"
 #include "Views/OutputMapping/DisplayClusterConfiguratorGraph.h"
 #include "Views/OutputMapping/DisplayClusterConfiguratorGraphSchema.h"
 #include "Views/OutputMapping/DisplayClusterConfiguratorOutputMappingCommands.h"
diff --git a/Engine/Plugins/Runtime/nDisplay/Source/DisplayClusterConfigurator/Private/Views/OutputMapping/DisplayClusterConfiguratorViewOutputMapping.h b/Engine/Plugins/Runtime/nDisplay/Source/DisplayClusterConfigurator/Private/Views/OutputMapping/DisplayClusterConfiguratorViewOutputMapping.h
--- a/Engine/Plugins/Runtime/nDisplay/Source/DisplayClusterConfigurator/Private/Views/OutputMapping/DisplayClusterConfiguratorViewOutputMapping.h
+++ b/Engine/Plugins/Runtime/nDisplay/Source/DisplayClusterConfigurator/Private/Views/OutputMapping/DisplayClusterConfiguratorViewOutputMapping.h
@@ -2,7 +2,7 @@
 
 #pragma once
 
-#include "Interfaces/Views/OutputMapping/IDisplayClusterConfiguratorViewOutputMapping.h"
+#include "Views/OutputMapping/IDisplayClusterConfiguratorViewOutputMapping.h"
 
 #include "UObject/StrongObjectPtr.h"
 #include "Styling/SlateTypes.h"
diff --git a/Engine/Plugins/Runtime/nDisplay/Source/DisplayClusterConfigurator/Private/Views/OutputMapping/EdNodes/DisplayClusterConfiguratorBaseNode.cpp b/Engine/Plugins/Runtime/nDisplay/Source/DisplayClusterConfigurator/Private/Views/OutputMapping/EdNodes/DisplayClusterConfiguratorBaseNode.cpp
--- a/Engine/Plugins/Runtime/nDisplay/Source/DisplayClusterConfigurator/Private/Views/OutputMapping/EdNodes/DisplayClusterConfiguratorBaseNode.cpp
+++ b/Engine/Plugins/Runtime/nDisplay/Source/DisplayClusterConfigurator/Private/Views/OutputMapping/EdNodes/DisplayClusterConfiguratorBaseNode.cpp
@@ -2,9 +2,9 @@
 
 #include "Views/OutputMapping/EdNodes/DisplayClusterConfiguratorBaseNode.h"
 
-#include "Interfaces/Views/OutputMapping/IDisplayClusterConfiguratorViewOutputMapping.h"
-#include "Interfaces/Views/TreeViews/IDisplayClusterConfiguratorViewTree.h"
-#include "Interfaces/Views/TreeViews/IDisplayClusterConfiguratorTreeItem.h"
+#include "Views/OutputMapping/IDisplayClusterConfiguratorViewOutputMapping.h"
+#include "Views/TreeViews/IDisplayClusterConfiguratorViewTree.h"
+#include "Views/TreeViews/IDisplayClusterConfiguratorTreeItem.h"
 
 #include "DisplayClusterConfiguratorBlueprintEditor.h"
 
diff --git a/Engine/Plugins/Runtime/nDisplay/Source/DisplayClusterConfigurator/Private/Views/OutputMapping/EdNodes/DisplayClusterConfiguratorBaseNode.h b/Engine/Plugins/Runtime/nDisplay/Source/DisplayClusterConfigurator/Private/Views/OutputMapping/EdNodes/DisplayClusterConfiguratorBaseNode.h
--- a/Engine/Plugins/Runtime/nDisplay/Source/DisplayClusterConfigurator/Private/Views/OutputMapping/EdNodes/DisplayClusterConfiguratorBaseNode.h
+++ b/Engine/Plugins/Runtime/nDisplay/Source/DisplayClusterConfigurator/Private/Views/OutputMapping/EdNodes/DisplayClusterConfiguratorBaseNode.h
@@ -2,7 +2,7 @@
 
 #pragma once
 
-#include "Interfaces/Views/OutputMapping/IDisplayClusterConfiguratorOutputMappingItem.h"
+#include "Views/OutputMapping/IDisplayClusterConfiguratorOutputMappingItem.h"
 #include "EdGraph/EdGraphNode.h"
 
 #include "Views/OutputMapping/Alignment/DisplayClusterConfiguratorNodeAlignmentHelper.h"
diff --git a/Engine/Plugins/Runtime/nDisplay/Source/DisplayClusterConfigurator/Private/Views/OutputMapping/EdNodes/DisplayClusterConfiguratorCanvasNode.cpp b/Engine/Plugins/Runtime/nDisplay/Source/DisplayClusterConfigurator/Private/Views/OutputMapping/EdNodes/DisplayClusterConfiguratorCanvasNode.cpp
--- a/Engine/Plugins/Runtime/nDisplay/Source/DisplayClusterConfigurator/Private/Views/OutputMapping/EdNodes/DisplayClusterConfiguratorCanvasNode.cpp
+++ b/Engine/Plugins/Runtime/nDisplay/Source/DisplayClusterConfigurator/Private/Views/OutputMapping/EdNodes/DisplayClusterConfiguratorCanvasNode.cpp
@@ -3,7 +3,7 @@
 #include "Views/OutputMapping/EdNodes/DisplayClusterConfiguratorCanvasNode.h"
 
 #include "DisplayClusterConfiguratorBlueprintEditor.h"
-#include "Interfaces/Views/OutputMapping/IDisplayClusterConfiguratorViewOutputMapping.h"
+#include "Views/OutputMapping/IDisplayClusterConfiguratorViewOutputMapping.h"
 #include "Views/OutputMapping/EdNodes/DisplayClusterConfiguratorHostNode.h"
 #include "Views/OutputMapping/EdNodes/DisplayClusterConfiguratorWindowNode.h"
 #include "Views/OutputMapping/GraphNodes/SDisplayClusterConfiguratorCanvasNode.h"
diff --git a/Engine/Plugins/Runtime/nDisplay/Source/DisplayClusterConfigurator/Private/Views/OutputMapping/EdNodes/DisplayClusterConfiguratorHostNode.cpp b/Engine/Plugins/Runtime/nDisplay/Source/DisplayClusterConfigurator/Private/Views/OutputMapping/EdNodes/DisplayClusterConfiguratorHostNode.cpp
--- a/Engine/Plugins/Runtime/nDisplay/Source/DisplayClusterConfigurator/Private/Views/OutputMapping/EdNodes/DisplayClusterConfiguratorHostNode.cpp
+++ b/Engine/Plugins/Runtime/nDisplay/Source/DisplayClusterConfigurator/Private/Views/OutputMapping/EdNodes/DisplayClusterConfiguratorHostNode.cpp
@@ -5,9 +5,9 @@
 #include "DisplayClusterConfiguratorBlueprintEditor.h"
 #include "DisplayClusterConfigurationTypes.h"
 #include "ClusterConfiguration/DisplayClusterConfiguratorClusterUtils.h"
-#include "Interfaces/Views/OutputMapping/IDisplayClusterConfiguratorViewOutputMapping.h"
-#include "Interfaces/Views/TreeViews/IDisplayClusterConfiguratorViewTree.h"
-#include "Interfaces/Views/TreeViews/IDisplayClusterConfiguratorTreeItem.h"
+#include "Views/OutputMapping/IDisplayClusterConfiguratorViewOutputMapping.h"
+#include "Views/TreeViews/IDisplayClusterConfiguratorViewTree.h"
+#include "Views/TreeViews/IDisplayClusterConfiguratorTreeItem.h"
 
 #include "Views/OutputMapping/GraphNodes/SDisplayClusterConfiguratorHostNode.h"
 
diff --git a/Engine/Plugins/Runtime/nDisplay/Source/DisplayClusterConfigurator/Private/Views/OutputMapping/EdNodes/DisplayClusterConfiguratorViewportNode.cpp b/Engine/Plugins/Runtime/nDisplay/Source/DisplayClusterConfigurator/Private/Views/OutputMapping/EdNodes/DisplayClusterConfiguratorViewportNode.cpp
--- a/Engine/Plugins/Runtime/nDisplay/Source/DisplayClusterConfigurator/Private/Views/OutputMapping/EdNodes/DisplayClusterConfiguratorViewportNode.cpp
+++ b/Engine/Plugins/Runtime/nDisplay/Source/DisplayClusterConfigurator/Private/Views/OutputMapping/EdNodes/DisplayClusterConfiguratorViewportNode.cpp
@@ -6,8 +6,8 @@
 #include "ClusterConfiguration/ViewModels/DisplayClusterConfiguratorViewportViewModel.h"
 #include "Views/OutputMapping/EdNodes/DisplayClusterConfiguratorWindowNode.h"
 #include "Views/OutputMapping/GraphNodes/SDisplayClusterConfiguratorViewportNode.h"
-#include "Interfaces/Views/TreeViews/IDisplayClusterConfiguratorViewTree.h"
-#include "Interfaces/Views/TreeViews/IDisplayClusterConfiguratorTreeItem.h"
+#include "Views/TreeViews/IDisplayClusterConfiguratorViewTree.h"
+#include "Views/TreeViews/IDisplayClusterConfiguratorTreeItem.h"
 
 #include "DisplayClusterConfigurationTypes.h"
 #include "DisplayClusterConfigurationTypes_ViewportRemap.h"
diff --git a/Engine/Plugins/Runtime/nDisplay/Source/DisplayClusterConfigurator/Private/Views/OutputMapping/EdNodes/DisplayClusterConfiguratorWindowNode.cpp b/Engine/Plugins/Runtime/nDisplay/Source/DisplayClusterConfigurator/Private/Views/OutputMapping/EdNodes/DisplayClusterConfiguratorWindowNode.cpp
--- a/Engine/Plugins/Runtime/nDisplay/Source/DisplayClusterConfigurator/Private/Views/OutputMapping/EdNodes/DisplayClusterConfiguratorWindowNode.cpp
+++ b/Engine/Plugins/Runtime/nDisplay/Source/DisplayClusterConfigurator/Private/Views/OutputMapping/EdNodes/DisplayClusterConfiguratorWindowNode.cpp
@@ -3,9 +3,9 @@
 #include "Views/OutputMapping/EdNodes/DisplayClusterConfiguratorWindowNode.h"
 
 #include "DisplayClusterConfiguratorBlueprintEditor.h"
-#include "Interfaces/Views/OutputMapping/IDisplayClusterConfiguratorViewOutputMapping.h"
 #include "ClusterConfiguration/DisplayClusterConfiguratorClusterUtils.h"
 #include "ClusterConfiguration/ViewModels/DisplayClusterConfiguratorClusterNodeViewModel.h"
+#include "Views/OutputMapping/IDisplayClusterConfiguratorViewOutputMapping.h"
 #include "Views/OutputMapping/EdNodes/DisplayClusterConfiguratorCanvasNode.h"
 #include "Views/OutputMapping/EdNodes/DisplayClusterConfiguratorViewportNode.h"
 #include "Views/OutputMapping/GraphNodes/SDisplayClusterConfiguratorWindowNode.h"
diff --git a/Engine/Plugins/Runtime/nDisplay/Source/DisplayClusterConfigurator/Private/Views/OutputMapping/GraphNodes/SDisplayClusterConfiguratorBaseNode.cpp b/Engine/Plugins/Runtime/nDisplay/Source/DisplayClusterConfigurator/Private/Views/OutputMapping/GraphNodes/SDisplayClusterConfiguratorBaseNode.cpp
--- a/Engine/Plugins/Runtime/nDisplay/Source/DisplayClusterConfigurator/Private/Views/OutputMapping/GraphNodes/SDisplayClusterConfiguratorBaseNode.cpp
+++ b/Engine/Plugins/Runtime/nDisplay/Source/DisplayClusterConfigurator/Private/Views/OutputMapping/GraphNodes/SDisplayClusterConfiguratorBaseNode.cpp
@@ -4,9 +4,9 @@
 
 #include "DisplayClusterConfiguratorStyle.h"
 #include "DisplayClusterConfiguratorBlueprintEditor.h"
-#include "Interfaces/Views/OutputMapping/IDisplayClusterConfiguratorViewOutputMapping.h"
-#include "Interfaces/Views/TreeViews/IDisplayClusterConfiguratorViewTree.h"
-#include "Interfaces/Views/TreeViews/IDisplayClusterConfiguratorTreeItem.h"
+#include "Views/TreeViews/IDisplayClusterConfiguratorViewTree.h"
+#include "Views/TreeViews/IDisplayClusterConfiguratorTreeItem.h"
+#include "Views/OutputMapping/IDisplayClusterConfiguratorViewOutputMapping.h"
 #include "Views/OutputMapping/EdNodes/DisplayClusterConfiguratorBaseNode.h"
 #include "Views/OutputMapping/Widgets/SDisplayClusterConfiguratorResizer.h"
 #include "Views/OutputMapping/Widgets/SDisplayClusterConfiguratorLayeringBox.h"
diff --git a/Engine/Plugins/Runtime/nDisplay/Source/DisplayClusterConfigurator/Private/Views/OutputMapping/GraphNodes/SDisplayClusterConfiguratorCanvasNode.cpp b/Engine/Plugins/Runtime/nDisplay/Source/DisplayClusterConfigurator/Private/Views/OutputMapping/GraphNodes/SDisplayClusterConfiguratorCanvasNode.cpp
--- a/Engine/Plugins/Runtime/nDisplay/Source/DisplayClusterConfigurator/Private/Views/OutputMapping/GraphNodes/SDisplayClusterConfiguratorCanvasNode.cpp
+++ b/Engine/Plugins/Runtime/nDisplay/Source/DisplayClusterConfigurator/Private/Views/OutputMapping/GraphNodes/SDisplayClusterConfiguratorCanvasNode.cpp
@@ -7,7 +7,7 @@
 
 #include "DisplayClusterConfiguratorBlueprintEditor.h"
 #include "DisplayClusterConfiguratorStyle.h"
-#include "Interfaces/Views/TreeViews/IDisplayClusterConfiguratorTreeItem.h"
+#include "Views/TreeViews/IDisplayClusterConfiguratorTreeItem.h"
 #include "Views/OutputMapping/EdNodes/DisplayClusterConfiguratorCanvasNode.h"
 #include "SGraphPanel.h"
 
diff --git a/Engine/Plugins/Runtime/nDisplay/Source/DisplayClusterConfigurator/Private/Views/OutputMapping/GraphNodes/SDisplayClusterConfiguratorViewportNode.cpp b/Engine/Plugins/Runtime/nDisplay/Source/DisplayClusterConfigurator/Private/Views/OutputMapping/GraphNodes/SDisplayClusterConfiguratorViewportNode.cpp
--- a/Engine/Plugins/Runtime/nDisplay/Source/DisplayClusterConfigurator/Private/Views/OutputMapping/GraphNodes/SDisplayClusterConfiguratorViewportNode.cpp
+++ b/Engine/Plugins/Runtime/nDisplay/Source/DisplayClusterConfigurator/Private/Views/OutputMapping/GraphNodes/SDisplayClusterConfiguratorViewportNode.cpp
@@ -5,8 +5,8 @@
 #include "DisplayClusterConfiguratorStyle.h"
 #include "DisplayClusterConfiguratorBlueprintEditor.h"
 #include "DisplayClusterConfigurationTypes.h"
-#include "Interfaces/Views/TreeViews/IDisplayClusterConfiguratorTreeItem.h"
-#include "Interfaces/Views/OutputMapping/IDisplayClusterConfiguratorViewOutputMapping.h"
+#include "Views/TreeViews/IDisplayClusterConfiguratorTreeItem.h"
+#include "Views/OutputMapping/IDisplayClusterConfiguratorViewOutputMapping.h"
 #include "Views/OutputMapping/EdNodes/DisplayClusterConfiguratorViewportNode.h"
 #include "Views/OutputMapping/EdNodes/DisplayClusterConfiguratorWindowNode.h"
 #include "Views/OutputMapping/Widgets/SDisplayClusterConfiguratorResizer.h"
diff --git a/Engine/Plugins/Runtime/nDisplay/Source/DisplayClusterConfigurator/Private/Views/OutputMapping/GraphNodes/SDisplayClusterConfiguratorWindowNode.cpp b/Engine/Plugins/Runtime/nDisplay/Source/DisplayClusterConfigurator/Private/Views/OutputMapping/GraphNodes/SDisplayClusterConfiguratorWindowNode.cpp
--- a/Engine/Plugins/Runtime/nDisplay/Source/DisplayClusterConfigurator/Private/Views/OutputMapping/GraphNodes/SDisplayClusterConfiguratorWindowNode.cpp
+++ b/Engine/Plugins/Runtime/nDisplay/Source/DisplayClusterConfigurator/Private/Views/OutputMapping/GraphNodes/SDisplayClusterConfiguratorWindowNode.cpp
@@ -5,8 +5,8 @@
 #include "DisplayClusterConfiguratorStyle.h"
 #include "DisplayClusterConfiguratorBlueprintEditor.h"
 #include "DisplayClusterConfigurationTypes.h"
-#include "Interfaces/Views/TreeViews/IDisplayClusterConfiguratorTreeItem.h"
-#include "Interfaces/Views/OutputMapping/IDisplayClusterConfiguratorViewOutputMapping.h"
+#include "Views/TreeViews/IDisplayClusterConfiguratorTreeItem.h"
+#include "Views/OutputMapping/IDisplayClusterConfiguratorViewOutputMapping.h"
 #include "Views/OutputMapping/EdNodes/DisplayClusterConfiguratorWindowNode.h"
 #include "Views/OutputMapping/EdNodes/DisplayClusterConfiguratorViewportNode.h"
 #include "Views/OutputMapping/GraphNodes/SDisplayClusterConfiguratorViewportNode.h"
diff --git a/Engine/Plugins/Runtime/nDisplay/Source/DisplayClusterConfigurator/Private/Views/OutputMapping/SDisplayClusterConfiguratorGraphEditor.cpp b/Engine/Plugins/Runtime/nDisplay/Source/DisplayClusterConfigurator/Private/Views/OutputMapping/SDisplayClusterConfiguratorGraphEditor.cpp
--- a/Engine/Plugins/Runtime/nDisplay/Source/DisplayClusterConfigurator/Private/Views/OutputMapping/SDisplayClusterConfiguratorGraphEditor.cpp
+++ b/Engine/Plugins/Runtime/nDisplay/Source/DisplayClusterConfigurator/Private/Views/OutputMapping/SDisplayClusterConfiguratorGraphEditor.cpp
@@ -5,7 +5,7 @@
 #include "DisplayClusterConfiguratorCommands.h"
 #include "DisplayClusterConfiguratorBlueprintEditor.h"
 #include "DisplayClusterConfigurationTypes.h"
-#include "Interfaces/IDisplayClusterConfigurator.h"
+#include "IDisplayClusterConfigurator.h"
 #include "ClusterConfiguration/DisplayClusterConfiguratorClusterUtils.h"
 #include "ClusterConfiguration/SDisplayClusterConfiguratorNewClusterItemDialog.h"
 #include "Views/OutputMapping/DisplayClusterConfiguratorGraph.h"
@@ -25,7 +25,7 @@
 #include "Misc/ConfigCacheIni.h"
 #include "IDocumentation.h"
 
-#include "Interfaces/Views/OutputMapping/IDisplayClusterConfiguratorViewOutputMapping.h"
+#include "Views/OutputMapping/IDisplayClusterConfiguratorViewOutputMapping.h"
 
 
 #define LOCTEXT_NAMESPACE "SDisplayClusterConfiguratorGraphEditor"
diff --git a/Engine/Plugins/Runtime/nDisplay/Source/DisplayClusterConfigurator/Private/Views/OutputMapping/SDisplayClusterConfiguratorViewOutputMapping.cpp b/Engine/Plugins/Runtime/nDisplay/Source/DisplayClusterConfigurator/Private/Views/OutputMapping/SDisplayClusterConfiguratorViewOutputMapping.cpp
--- a/Engine/Plugins/Runtime/nDisplay/Source/DisplayClusterConfigurator/Private/Views/OutputMapping/SDisplayClusterConfiguratorViewOutputMapping.cpp
+++ b/Engine/Plugins/Runtime/nDisplay/Source/DisplayClusterConfigurator/Private/Views/OutputMapping/SDisplayClusterConfiguratorViewOutputMapping.cpp
@@ -4,7 +4,7 @@
 
 #include "DisplayClusterConfiguratorCommands.h"
 #include "DisplayClusterConfiguratorBlueprintEditor.h"
-#include "Interfaces/IDisplayClusterConfigurator.h"
+#include "IDisplayClusterConfigurator.h"
 #include "Views/OutputMapping/DisplayClusterConfiguratorViewOutputMapping.h"
 #include "Views/OutputMapping/SDisplayClusterConfiguratorGraphEditor.h"
 #include "Views/OutputMapping/Widgets/SDisplayClusterConfiguratorOutputMappingToolbar.h"
diff --git a/Engine/Plugins/Runtime/nDisplay/Source/DisplayClusterConfigurator/Private/Views/OutputMapping/Widgets/SDisplayClusterConfiguratorOutputMappingToolbar.cpp b/Engine/Plugins/Runtime/nDisplay/Source/DisplayClusterConfigurator/Private/Views/OutputMapping/Widgets/SDisplayClusterConfiguratorOutputMappingToolbar.cpp
--- a/Engine/Plugins/Runtime/nDisplay/Source/DisplayClusterConfigurator/Private/Views/OutputMapping/Widgets/SDisplayClusterConfiguratorOutputMappingToolbar.cpp
+++ b/Engine/Plugins/Runtime/nDisplay/Source/DisplayClusterConfigurator/Private/Views/OutputMapping/Widgets/SDisplayClusterConfiguratorOutputMappingToolbar.cpp
@@ -3,7 +3,7 @@
 #include "Views/OutputMapping/Widgets/SDisplayClusterConfiguratorOutputMappingToolbar.h"
 
 #include "DisplayClusterConfiguratorStyle.h"
-#include "Interfaces/IDisplayClusterConfigurator.h"
+#include "IDisplayClusterConfigurator.h"
 #include "Views/OutputMapping/DisplayClusterConfiguratorOutputMappingCommands.h"
 #include "Views/OutputMapping/DisplayClusterConfiguratorViewOutputMapping.h"
 
diff --git a/Engine/Plugins/Runtime/nDisplay/Source/DisplayClusterConfigurator/Private/Views/TreeViews/Cluster/DisplayClusterConfiguratorViewCluster.cpp b/Engine/Plugins/Runtime/nDisplay/Source/DisplayClusterConfigurator/Private/Views/TreeViews/Cluster/DisplayClusterConfiguratorViewCluster.cpp
--- a/Engine/Plugins/Runtime/nDisplay/Source/DisplayClusterConfigurator/Private/Views/TreeViews/Cluster/DisplayClusterConfiguratorViewCluster.cpp
+++ b/Engine/Plugins/Runtime/nDisplay/Source/DisplayClusterConfigurator/Private/Views/TreeViews/Cluster/DisplayClusterConfiguratorViewCluster.cpp
@@ -7,7 +7,7 @@
 #include "DisplayClusterConfiguratorCommands.h"
 #include "ClusterConfiguration/DisplayClusterConfiguratorClusterUtils.h"
 #include "ClusterConfiguration/SDisplayClusterConfiguratorNewClusterItemDialog.h"
-#include "Interfaces/Views/TreeViews/IDisplayClusterConfiguratorTreeItem.h"
+#include "Views/TreeViews/IDisplayClusterConfiguratorTreeItem.h"
 #include "Views/TreeViews/Cluster/DisplayClusterConfiguratorViewClusterBuilder.h"
 #include "Views/TreeViews/DisplayClusterConfiguratorTreeViewCommands.h"
 #include "Views/TreeViews/SDisplayClusterConfiguratorViewTree.h"
diff --git a/Engine/Plugins/Runtime/nDisplay/Source/DisplayClusterConfigurator/Private/Views/TreeViews/Cluster/DisplayClusterConfiguratorViewClusterBuilder.cpp b/Engine/Plugins/Runtime/nDisplay/Source/DisplayClusterConfigurator/Private/Views/TreeViews/Cluster/DisplayClusterConfiguratorViewClusterBuilder.cpp
--- a/Engine/Plugins/Runtime/nDisplay/Source/DisplayClusterConfigurator/Private/Views/TreeViews/Cluster/DisplayClusterConfiguratorViewClusterBuilder.cpp
+++ b/Engine/Plugins/Runtime/nDisplay/Source/DisplayClusterConfigurator/Private/Views/TreeViews/Cluster/DisplayClusterConfiguratorViewClusterBuilder.cpp
@@ -5,8 +5,8 @@
 #include "DisplayClusterConfigurationTypes.h"
 #include "DisplayClusterConfiguratorBlueprintEditor.h"
 #include "ClusterConfiguration/DisplayClusterConfiguratorClusterUtils.h"
-#include "Interfaces/Views/TreeViews/IDisplayClusterConfiguratorTreeItem.h"
-#include "Interfaces/Views/TreeViews/IDisplayClusterConfiguratorViewTree.h"
+#include "Views/TreeViews/IDisplayClusterConfiguratorTreeItem.h"
+#include "Views/TreeViews/IDisplayClusterConfiguratorViewTree.h"
 #include "Views/TreeViews/Cluster/TreeItems/DisplayClusterConfiguratorTreeItemCluster.h"
 #include "Views/TreeViews/Cluster/TreeItems/DisplayClusterConfiguratorTreeItemHost.h"
 #include "Views/TreeViews/Cluster/TreeItems/DisplayClusterConfiguratorTreeItemClusterNode.h"
diff --git a/Engine/Plugins/Runtime/nDisplay/Source/DisplayClusterConfigurator/Private/Views/TreeViews/Cluster/TreeItems/DisplayClusterConfiguratorTreeItemCluster.cpp b/Engine/Plugins/Runtime/nDisplay/Source/DisplayClusterConfigurator/Private/Views/TreeViews/Cluster/TreeItems/DisplayClusterConfiguratorTreeItemCluster.cpp
--- a/Engine/Plugins/Runtime/nDisplay/Source/DisplayClusterConfigurator/Private/Views/TreeViews/Cluster/TreeItems/DisplayClusterConfiguratorTreeItemCluster.cpp
+++ b/Engine/Plugins/Runtime/nDisplay/Source/DisplayClusterConfigurator/Private/Views/TreeViews/Cluster/TreeItems/DisplayClusterConfiguratorTreeItemCluster.cpp
@@ -4,8 +4,8 @@
 
 #include "DisplayClusterConfiguratorBlueprintEditor.h"
 #include "DisplayClusterConfiguratorStyle.h"
-#include "Interfaces/Views/TreeViews/IDisplayClusterConfiguratorViewTree.h"
-#include "Interfaces/Views/OutputMapping/IDisplayClusterConfiguratorViewOutputMapping.h"
+#include "Views/OutputMapping/IDisplayClusterConfiguratorViewOutputMapping.h"
+#include "Views/TreeViews/IDisplayClusterConfiguratorViewTree.h"
 #include "Views/TreeViews/Cluster/DisplayClusterConfiguratorViewCluster.h"
 
 #include "Widgets/Images/SImage.h"
diff --git a/Engine/Plugins/Runtime/nDisplay/Source/DisplayClusterConfigurator/Private/Views/TreeViews/Cluster/TreeItems/DisplayClusterConfiguratorTreeItemClusterNode.cpp b/Engine/Plugins/Runtime/nDisplay/Source/DisplayClusterConfigurator/Private/Views/TreeViews/Cluster/TreeItems/DisplayClusterConfiguratorTreeItemClusterNode.cpp
--- a/Engine/Plugins/Runtime/nDisplay/Source/DisplayClusterConfigurator/Private/Views/TreeViews/Cluster/TreeItems/DisplayClusterConfiguratorTreeItemClusterNode.cpp
+++ b/Engine/Plugins/Runtime/nDisplay/Source/DisplayClusterConfigurator/Private/Views/TreeViews/Cluster/TreeItems/DisplayClusterConfiguratorTreeItemClusterNode.cpp
@@ -13,7 +13,7 @@
 #include "Views/DragDrop/DisplayClusterConfiguratorViewportDragDropOp.h"
 #include "Views/TreeViews/Cluster/DisplayClusterConfiguratorViewCluster.h"
 #include "Views/TreeViews/Cluster/TreeItems/DisplayClusterConfiguratorTreeItemViewport.h"
-#include "Interfaces/Views/TreeViews/IDisplayClusterConfiguratorViewTree.h"
+#include "Views/TreeViews/IDisplayClusterConfiguratorViewTree.h"
 
 #include "ScopedTransaction.h"
 #include "Widgets/Images/SImage.h"
diff --git a/Engine/Plugins/Runtime/nDisplay/Source/DisplayClusterConfigurator/Private/Views/TreeViews/Cluster/TreeItems/DisplayClusterConfiguratorTreeItemViewport.cpp b/Engine/Plugins/Runtime/nDisplay/Source/DisplayClusterConfigurator/Private/Views/TreeViews/Cluster/TreeItems/DisplayClusterConfiguratorTreeItemViewport.cpp
--- a/Engine/Plugins/Runtime/nDisplay/Source/DisplayClusterConfigurator/Private/Views/TreeViews/Cluster/TreeItems/DisplayClusterConfiguratorTreeItemViewport.cpp
+++ b/Engine/Plugins/Runtime/nDisplay/Source/DisplayClusterConfigurator/Private/Views/TreeViews/Cluster/TreeItems/DisplayClusterConfiguratorTreeItemViewport.cpp
@@ -12,7 +12,7 @@
 #include "Views/DragDrop/DisplayClusterConfiguratorViewportDragDropOp.h"
 #include "Views/TreeViews/Cluster/DisplayClusterConfiguratorViewCluster.h"
 #include "Views/TreeViews/Cluster/TreeItems/DisplayClusterConfiguratorTreeItemClusterNode.h"
-#include "Interfaces/Views/TreeViews/IDisplayClusterConfiguratorViewTree.h"
+#include "Views/TreeViews/IDisplayClusterConfiguratorViewTree.h"
 
 #include "ScopedTransaction.h"
 #include "Widgets/Images/SImage.h"
diff --git a/Engine/Plugins/Runtime/nDisplay/Source/DisplayClusterConfigurator/Private/Views/TreeViews/DisplayClusterConfiguratorTreeBuilder.cpp b/Engine/Plugins/Runtime/nDisplay/Source/DisplayClusterConfigurator/Private/Views/TreeViews/DisplayClusterConfiguratorTreeBuilder.cpp
--- a/Engine/Plugins/Runtime/nDisplay/Source/DisplayClusterConfigurator/Private/Views/TreeViews/DisplayClusterConfiguratorTreeBuilder.cpp
+++ b/Engine/Plugins/Runtime/nDisplay/Source/DisplayClusterConfigurator/Private/Views/TreeViews/DisplayClusterConfiguratorTreeBuilder.cpp
@@ -1,8 +1,7 @@
 // Copyright Epic Games, Inc. All Rights Reserved.
 
 #include "Views/TreeViews/DisplayClusterConfiguratorTreeBuilder.h"
-
-#include "Interfaces/Views/TreeViews/IDisplayClusterConfiguratorTreeItem.h"
+#include "Views/TreeViews/IDisplayClusterConfiguratorTreeItem.h"
 
 void FDisplayClusterConfiguratorTreeBuilderOutput::Add(const TSharedPtr<IDisplayClusterConfiguratorTreeItem>& InItem, const FName& InParentName, const FName& InParentType, bool bAddToHead)
 {
diff --git a/Engine/Plugins/Runtime/nDisplay/Source/DisplayClusterConfigurator/Private/Views/TreeViews/DisplayClusterConfiguratorTreeBuilder.h b/Engine/Plugins/Runtime/nDisplay/Source/DisplayClusterConfigurator/Private/Views/TreeViews/DisplayClusterConfiguratorTreeBuilder.h
--- a/Engine/Plugins/Runtime/nDisplay/Source/DisplayClusterConfigurator/Private/Views/TreeViews/DisplayClusterConfiguratorTreeBuilder.h
+++ b/Engine/Plugins/Runtime/nDisplay/Source/DisplayClusterConfigurator/Private/Views/TreeViews/DisplayClusterConfiguratorTreeBuilder.h
@@ -2,7 +2,7 @@
 
 #pragma once
 
-#include "Interfaces/Views/TreeViews/IDisplayClusterConfiguratorTreeBuilder.h"
+#include "Views/TreeViews/IDisplayClusterConfiguratorTreeBuilder.h"
 
 class FDisplayClusterConfiguratorBlueprintEditor;
 class IDisplayClusterConfiguratorTreeItem;
diff --git a/Engine/Plugins/Runtime/nDisplay/Source/DisplayClusterConfigurator/Private/Views/TreeViews/DisplayClusterConfiguratorTreeItem.cpp b/Engine/Plugins/Runtime/nDisplay/Source/DisplayClusterConfigurator/Private/Views/TreeViews/DisplayClusterConfiguratorTreeItem.cpp
--- a/Engine/Plugins/Runtime/nDisplay/Source/DisplayClusterConfigurator/Private/Views/TreeViews/DisplayClusterConfiguratorTreeItem.cpp
+++ b/Engine/Plugins/Runtime/nDisplay/Source/DisplayClusterConfigurator/Private/Views/TreeViews/DisplayClusterConfiguratorTreeItem.cpp
@@ -4,7 +4,7 @@
 
 #include "DisplayClusterConfiguratorStyle.h"
 #include "DisplayClusterConfiguratorBlueprintEditor.h"
-#include "Interfaces/Views/TreeViews/IDisplayClusterConfiguratorViewTree.h"
+#include "Views/TreeViews/IDisplayClusterConfiguratorViewTree.h"
 #include "Views/TreeViews/SDisplayClusterConfiguratorTreeItemRow.h"
 
 #include "Misc/TextFilterExpressionEvaluator.h"
diff --git a/Engine/Plugins/Runtime/nDisplay/Source/DisplayClusterConfigurator/Private/Views/TreeViews/DisplayClusterConfiguratorTreeItem.h b/Engine/Plugins/Runtime/nDisplay/Source/DisplayClusterConfigurator/Private/Views/TreeViews/DisplayClusterConfiguratorTreeItem.h
--- a/Engine/Plugins/Runtime/nDisplay/Source/DisplayClusterConfigurator/Private/Views/TreeViews/DisplayClusterConfiguratorTreeItem.h
+++ b/Engine/Plugins/Runtime/nDisplay/Source/DisplayClusterConfigurator/Private/Views/TreeViews/DisplayClusterConfiguratorTreeItem.h
@@ -2,7 +2,7 @@
 
 #pragma once
 
-#include "Interfaces/Views/TreeViews/IDisplayClusterConfiguratorTreeItem.h"
+#include "Views/TreeViews/IDisplayClusterConfiguratorTreeItem.h"
 
 #include "Widgets/SNullWidget.h"
 #include "Widgets/SWidget.h"
diff --git a/Engine/Plugins/Runtime/nDisplay/Source/DisplayClusterConfigurator/Private/Views/TreeViews/DisplayClusterConfiguratorViewTree.h b/Engine/Plugins/Runtime/nDisplay/Source/DisplayClusterConfigurator/Private/Views/TreeViews/DisplayClusterConfiguratorViewTree.h
--- a/Engine/Plugins/Runtime/nDisplay/Source/DisplayClusterConfigurator/Private/Views/TreeViews/DisplayClusterConfiguratorViewTree.h
+++ b/Engine/Plugins/Runtime/nDisplay/Source/DisplayClusterConfigurator/Private/Views/TreeViews/DisplayClusterConfiguratorViewTree.h
@@ -3,7 +3,7 @@
 #pragma once
 
 #include "DisplayClusterConfigurationTypes.h"
-#include "Interfaces/Views/TreeViews/IDisplayClusterConfiguratorViewTree.h"
+#include "Views/TreeViews/IDisplayClusterConfiguratorViewTree.h"
 
 #include "EditorUndoClient.h"
 
diff --git a/Engine/Plugins/Runtime/nDisplay/Source/DisplayClusterConfigurator/Private/Views/TreeViews/SDisplayClusterConfiguratorTreeItemRow.cpp b/Engine/Plugins/Runtime/nDisplay/Source/DisplayClusterConfigurator/Private/Views/TreeViews/SDisplayClusterConfiguratorTreeItemRow.cpp
--- a/Engine/Plugins/Runtime/nDisplay/Source/DisplayClusterConfigurator/Private/Views/TreeViews/SDisplayClusterConfiguratorTreeItemRow.cpp
+++ b/Engine/Plugins/Runtime/nDisplay/Source/DisplayClusterConfigurator/Private/Views/TreeViews/SDisplayClusterConfiguratorTreeItemRow.cpp
@@ -2,8 +2,8 @@
 
 #include "Views/TreeViews/SDisplayClusterConfiguratorTreeItemRow.h"
 
-#include "Interfaces/Views/TreeViews/IDisplayClusterConfiguratorTreeItem.h"
-#include "Interfaces/Views/TreeViews/IDisplayClusterConfiguratorViewTree.h"
+#include "Views/TreeViews/IDisplayClusterConfiguratorTreeItem.h"
+#include "Views/TreeViews/IDisplayClusterConfiguratorViewTree.h"
 
 #include "Widgets/Views/STableViewBase.h"
 #include "Widgets/Views/STableRow.h"
diff --git a/Engine/Plugins/Runtime/nDisplay/Source/DisplayClusterConfigurator/Private/Views/TreeViews/SDisplayClusterConfiguratorViewTree.cpp b/Engine/Plugins/Runtime/nDisplay/Source/DisplayClusterConfigurator/Private/Views/TreeViews/SDisplayClusterConfiguratorViewTree.cpp
--- a/Engine/Plugins/Runtime/nDisplay/Source/DisplayClusterConfigurator/Private/Views/TreeViews/SDisplayClusterConfiguratorViewTree.cpp
+++ b/Engine/Plugins/Runtime/nDisplay/Source/DisplayClusterConfigurator/Private/Views/TreeViews/SDisplayClusterConfiguratorViewTree.cpp
@@ -3,9 +3,9 @@
 #include "Views/TreeViews/SDisplayClusterConfiguratorViewTree.h"
 
 #include "DisplayClusterConfiguratorBlueprintEditor.h"
-#include "Interfaces/Views/TreeViews/IDisplayClusterConfiguratorTreeBuilder.h"
-#include "Interfaces/Views/TreeViews/IDisplayClusterConfiguratorTreeItem.h"
-#include "Interfaces/Views/TreeViews/IDisplayClusterConfiguratorViewTree.h"
+#include "Views/TreeViews/IDisplayClusterConfiguratorTreeBuilder.h"
+#include "Views/TreeViews/IDisplayClusterConfiguratorTreeItem.h"
+#include "Views/TreeViews/IDisplayClusterConfiguratorViewTree.h"
 #include "Views/TreeViews/DisplayClusterConfiguratorTreeViewCommands.h"
 
 #include "Framework/MultiBox/MultiBoxBuilder.h"
diff --git a/Engine/Plugins/Runtime/nDisplay/Source/DisplayClusterConfigurator/Private/Views/TreeViews/SDisplayClusterConfiguratorViewTree.h b/Engine/Plugins/Runtime/nDisplay/Source/DisplayClusterConfigurator/Private/Views/TreeViews/SDisplayClusterConfiguratorViewTree.h
--- a/Engine/Plugins/Runtime/nDisplay/Source/DisplayClusterConfigurator/Private/Views/TreeViews/SDisplayClusterConfiguratorViewTree.h
+++ b/Engine/Plugins/Runtime/nDisplay/Source/DisplayClusterConfigurator/Private/Views/TreeViews/SDisplayClusterConfiguratorViewTree.h
@@ -3,7 +3,7 @@
 #pragma once
 
 #include "Views/SDisplayClusterConfiguratorViewBase.h"
-#include "Interfaces/Views/TreeViews/IDisplayClusterConfiguratorViewTree.h"
+#include "Views/TreeViews/IDisplayClusterConfiguratorViewTree.h"
 
 #include "EditorUndoClient.h"
 
diff --git a/Engine/Plugins/Runtime/nDisplay/Source/DisplayClusterConfigurator/Public/IDisplayClusterConfigurator.h b/Engine/Plugins/Runtime/nDisplay/Source/DisplayClusterConfigurator/Public/IDisplayClusterConfigurator.h
new file mode 100644
--- /dev/null
+++ b/Engine/Plugins/Runtime/nDisplay/Source/DisplayClusterConfigurator/Public/IDisplayClusterConfigurator.h
@@ -0,0 +1,52 @@
+// Copyright Epic Games, Inc. All Rights Reserved.
+
+#pragma once
+
+#include "Modules/ModuleInterface.h"
+#include "Modules/ModuleManager.h"
+
+class FDisplayClusterConfiguratorCommands;
+
+// Read only delegate
+DECLARE_MULTICAST_DELEGATE_OneParam(FOnDisplayClusterConfiguratorReadOnlyChanged, bool /* bInEnabled */);
+using FOnDisplayClusterConfiguratorReadOnlyChangedDelegate = FOnDisplayClusterConfiguratorReadOnlyChanged::FDelegate;
+
+/**
+ * Configurator editor module
+ */
+class IDisplayClusterConfigurator :
+	public IModuleInterface
+{
+public:
+	static constexpr const TCHAR* ModuleName = TEXT("DisplayClusterConfigurator");
+
+public:
+	virtual ~IDisplayClusterConfigurator() = default;
+
+public:
+	/**
+	* Singleton-like access to this module's interface.  This is just for convenience!
+	* Beware of calling this during the shutdown phase, though.  Your module might have been unloaded already.
+	*
+	* @return Returns singleton instance, loading the module on demand if needed
+	*/
+	static inline IDisplayClusterConfigurator& Get()
+	{
+		return FModuleManager::GetModuleChecked<IDisplayClusterConfigurator>(IDisplayClusterConfigurator::ModuleName);
+	}
+
+	/**
+	* Checks to see if this module is loaded and ready.  It is only valid to call Get() if IsAvailable() returns true.
+	*
+	* @return True if the module is loaded and ready to use
+	*/
+	static inline bool IsAvailable()
+	{
+		return FModuleManager::Get().IsModuleLoaded(IDisplayClusterConfigurator::ModuleName);
+	}
+
+public:
+	DISPLAYCLUSTERCONFIGURATOR_API virtual const FDisplayClusterConfiguratorCommands& GetCommands() const = 0;
+	DISPLAYCLUSTERCONFIGURATOR_API virtual TSharedPtr<class FExtensibilityManager> GetMenuExtensibilityManager() const = 0;
+	DISPLAYCLUSTERCONFIGURATOR_API virtual TSharedPtr<class FExtensibilityManager> GetToolBarExtensibilityManager() const = 0;
+};
diff --git a/Engine/Plugins/Runtime/nDisplay/Source/DisplayClusterConfigurator/Public/IDisplayClusterConfiguratorBlueprintEditor.h b/Engine/Plugins/Runtime/nDisplay/Source/DisplayClusterConfigurator/Public/IDisplayClusterConfiguratorBlueprintEditor.h
new file mode 100644
--- /dev/null
+++ b/Engine/Plugins/Runtime/nDisplay/Source/DisplayClusterConfigurator/Public/IDisplayClusterConfiguratorBlueprintEditor.h
@@ -0,0 +1,34 @@
+// Copyright Epic Games, Inc. All Rights Reserved.
+
+#pragma once
+
+#include "CoreMinimal.h"
+#include "BlueprintEditor.h"
+
+/**
+ * TODO: Remove this completely, don't expose the BP Editor outside of the configuration module.
+ */
+class IDisplayClusterConfiguratorBlueprintEditor
+	: public FBlueprintEditor
+{
+public:
+	virtual ~IDisplayClusterConfiguratorBlueprintEditor() = default;
+
+public:
+	DECLARE_MULTICAST_DELEGATE(FOnConfigReloaded);
+	DECLARE_MULTICAST_DELEGATE(FOnObjectSelected);
+	DECLARE_MULTICAST_DELEGATE(FOnInvalidateViews);
+	DECLARE_MULTICAST_DELEGATE(FOnClearViewportSelection);
+	DECLARE_MULTICAST_DELEGATE(FOnClusterChanged);
+
+	using FOnConfigReloadedDelegate = FOnConfigReloaded::FDelegate;
+	using FOnObjectSelectedDelegate = FOnObjectSelected::FDelegate;
+	using FOnInvalidateViewsDelegate = FOnInvalidateViews::FDelegate;
+	using FOnClearViewportSelectionDelegate = FOnClearViewportSelection::FDelegate;
+	using FOnClusterChangedDelegate = FOnClusterChanged::FDelegate;
+
+public:
+	virtual TArray<UObject*> GetSelectedObjects() const = 0;
+
+	virtual bool IsObjectSelected(UObject* Obj) const = 0;
+};
diff --git a/Engine/Plugins/Runtime/nDisplay/Source/DisplayClusterConfigurator/Public/Interfaces/IDisplayClusterConfigurator.h b/Engine/Plugins/Runtime/nDisplay/Source/DisplayClusterConfigurator/Public/Interfaces/IDisplayClusterConfigurator.h
deleted file mode 100644
diff --git a/Engine/Plugins/Runtime/nDisplay/Source/DisplayClusterConfigurator/Public/Interfaces/IDisplayClusterConfiguratorBlueprintEditor.h b/Engine/Plugins/Runtime/nDisplay/Source/DisplayClusterConfigurator/Public/Interfaces/IDisplayClusterConfiguratorBlueprintEditor.h
deleted file mode 100644
diff --git a/Engine/Plugins/Runtime/nDisplay/Source/DisplayClusterConfigurator/Public/Interfaces/Views/Details/IDisplayClusterConfiguratorViewDetails.h b/Engine/Plugins/Runtime/nDisplay/Source/DisplayClusterConfigurator/Public/Interfaces/Views/Details/IDisplayClusterConfiguratorViewDetails.h
deleted file mode 100644
diff --git a/Engine/Plugins/Runtime/nDisplay/Source/DisplayClusterConfigurator/Public/Interfaces/Views/IDisplayClusterConfiguratorBuilder.h b/Engine/Plugins/Runtime/nDisplay/Source/DisplayClusterConfigurator/Public/Interfaces/Views/IDisplayClusterConfiguratorBuilder.h
deleted file mode 100644
diff --git a/Engine/Plugins/Runtime/nDisplay/Source/DisplayClusterConfigurator/Public/Interfaces/Views/IDisplayClusterConfiguratorItem.h b/Engine/Plugins/Runtime/nDisplay/Source/DisplayClusterConfigurator/Public/Interfaces/Views/IDisplayClusterConfiguratorItem.h
deleted file mode 100644
diff --git a/Engine/Plugins/Runtime/nDisplay/Source/DisplayClusterConfigurator/Public/Interfaces/Views/IDisplayClusterConfiguratorView.h b/Engine/Plugins/Runtime/nDisplay/Source/DisplayClusterConfigurator/Public/Interfaces/Views/IDisplayClusterConfiguratorView.h
deleted file mode 100644
diff --git a/Engine/Plugins/Runtime/nDisplay/Source/DisplayClusterConfigurator/Public/Interfaces/Views/Log/IDisplayClusterConfiguratorViewLog.h b/Engine/Plugins/Runtime/nDisplay/Source/DisplayClusterConfigurator/Public/Interfaces/Views/Log/IDisplayClusterConfiguratorViewLog.h
deleted file mode 100644
diff --git a/Engine/Plugins/Runtime/nDisplay/Source/DisplayClusterConfigurator/Public/Interfaces/Views/OutputMapping/IDisplayClusterConfiguratorOutputMappingBuilder.h b/Engine/Plugins/Runtime/nDisplay/Source/DisplayClusterConfigurator/Public/Interfaces/Views/OutputMapping/IDisplayClusterConfiguratorOutputMappingBuilder.h
deleted file mode 100644
diff --git a/Engine/Plugins/Runtime/nDisplay/Source/DisplayClusterConfigurator/Public/Interfaces/Views/OutputMapping/IDisplayClusterConfiguratorOutputMappingItem.h b/Engine/Plugins/Runtime/nDisplay/Source/DisplayClusterConfigurator/Public/Interfaces/Views/OutputMapping/IDisplayClusterConfiguratorOutputMappingItem.h
deleted file mode 100644
diff --git a/Engine/Plugins/Runtime/nDisplay/Source/DisplayClusterConfigurator/Public/Interfaces/Views/OutputMapping/IDisplayClusterConfiguratorViewOutputMapping.h b/Engine/Plugins/Runtime/nDisplay/Source/DisplayClusterConfigurator/Public/Interfaces/Views/OutputMapping/IDisplayClusterConfiguratorViewOutputMapping.h
deleted file mode 100644
diff --git a/Engine/Plugins/Runtime/nDisplay/Source/DisplayClusterConfigurator/Public/Interfaces/Views/TreeViews/IDisplayClusterConfiguratorTreeBuilder.h b/Engine/Plugins/Runtime/nDisplay/Source/DisplayClusterConfigurator/Public/Interfaces/Views/TreeViews/IDisplayClusterConfiguratorTreeBuilder.h
deleted file mode 100644
diff --git a/Engine/Plugins/Runtime/nDisplay/Source/DisplayClusterConfigurator/Public/Interfaces/Views/TreeViews/IDisplayClusterConfiguratorTreeItem.h b/Engine/Plugins/Runtime/nDisplay/Source/DisplayClusterConfigurator/Public/Interfaces/Views/TreeViews/IDisplayClusterConfiguratorTreeItem.h
deleted file mode 100644
diff --git a/Engine/Plugins/Runtime/nDisplay/Source/DisplayClusterConfigurator/Public/Interfaces/Views/TreeViews/IDisplayClusterConfiguratorViewTree.h b/Engine/Plugins/Runtime/nDisplay/Source/DisplayClusterConfigurator/Public/Interfaces/Views/TreeViews/IDisplayClusterConfiguratorViewTree.h
deleted file mode 100644
diff --git a/Engine/Plugins/Runtime/nDisplay/Source/DisplayClusterConfigurator/Public/Interfaces/Views/Viewport/IDisplayClusterConfiguratorPreviewScene.h b/Engine/Plugins/Runtime/nDisplay/Source/DisplayClusterConfigurator/Public/Interfaces/Views/Viewport/IDisplayClusterConfiguratorPreviewScene.h
deleted file mode 100644
diff --git a/Engine/Plugins/Runtime/nDisplay/Source/DisplayClusterConfigurator/Public/Interfaces/Views/Viewport/IDisplayClusterConfiguratorViewViewport.h b/Engine/Plugins/Runtime/nDisplay/Source/DisplayClusterConfigurator/Public/Interfaces/Views/Viewport/IDisplayClusterConfiguratorViewViewport.h
deleted file mode 100644
diff --git a/Engine/Plugins/Runtime/nDisplay/Source/DisplayClusterConfigurator/Public/Interfaces/Views/Viewport/IDisplayClusterConfiguratorViewportBuilder.h b/Engine/Plugins/Runtime/nDisplay/Source/DisplayClusterConfigurator/Public/Interfaces/Views/Viewport/IDisplayClusterConfiguratorViewportBuilder.h
deleted file mode 100644
diff --git a/Engine/Plugins/Runtime/nDisplay/Source/DisplayClusterConfigurator/Public/Views/Details/IDisplayClusterConfiguratorViewDetails.h b/Engine/Plugins/Runtime/nDisplay/Source/DisplayClusterConfigurator/Public/Views/Details/IDisplayClusterConfiguratorViewDetails.h
new file mode 100644
--- /dev/null
+++ b/Engine/Plugins/Runtime/nDisplay/Source/DisplayClusterConfigurator/Public/Views/Details/IDisplayClusterConfiguratorViewDetails.h
@@ -0,0 +1,13 @@
+// Copyright Epic Games, Inc. All Rights Reserved.
+
+#pragma once
+
+#include "Views/IDisplayClusterConfiguratorView.h"
+
+
+class IDisplayClusterConfiguratorViewDetails
+	: public IDisplayClusterConfiguratorView
+{
+public:
+	virtual ~IDisplayClusterConfiguratorViewDetails() = default;
+};
diff --git a/Engine/Plugins/Runtime/nDisplay/Source/DisplayClusterConfigurator/Public/Views/IDisplayClusterConfiguratorBuilder.h b/Engine/Plugins/Runtime/nDisplay/Source/DisplayClusterConfigurator/Public/Views/IDisplayClusterConfiguratorBuilder.h
new file mode 100644
--- /dev/null
+++ b/Engine/Plugins/Runtime/nDisplay/Source/DisplayClusterConfigurator/Public/Views/IDisplayClusterConfiguratorBuilder.h
@@ -0,0 +1,12 @@
+// Copyright Epic Games, Inc. All Rights Reserved.
+
+#pragma once
+
+#include "CoreMinimal.h"
+
+class IDisplayClusterConfiguratorBuilder
+	: public TSharedFromThis<IDisplayClusterConfiguratorBuilder>
+{
+public:
+	virtual ~IDisplayClusterConfiguratorBuilder() = default;
+};
diff --git a/Engine/Plugins/Runtime/nDisplay/Source/DisplayClusterConfigurator/Public/Views/IDisplayClusterConfiguratorItem.h b/Engine/Plugins/Runtime/nDisplay/Source/DisplayClusterConfigurator/Public/Views/IDisplayClusterConfiguratorItem.h
new file mode 100644
--- /dev/null
+++ b/Engine/Plugins/Runtime/nDisplay/Source/DisplayClusterConfigurator/Public/Views/IDisplayClusterConfiguratorItem.h
@@ -0,0 +1,21 @@
+// Copyright Epic Games, Inc. All Rights Reserved.
+
+#pragma once
+
+#include "CoreMinimal.h"
+
+
+class IDisplayClusterConfiguratorItem
+	: public TSharedFromThis<IDisplayClusterConfiguratorItem>
+{
+public:
+	virtual ~IDisplayClusterConfiguratorItem() = default;
+
+public:
+	/** Get the object represented by this item, if any */
+	virtual UObject* GetObject() const = 0;
+
+	virtual void OnSelection() = 0;
+
+	virtual bool IsSelected() = 0;
+};
diff --git a/Engine/Plugins/Runtime/nDisplay/Source/DisplayClusterConfigurator/Public/Views/IDisplayClusterConfiguratorView.h b/Engine/Plugins/Runtime/nDisplay/Source/DisplayClusterConfigurator/Public/Views/IDisplayClusterConfiguratorView.h
new file mode 100644
--- /dev/null
+++ b/Engine/Plugins/Runtime/nDisplay/Source/DisplayClusterConfigurator/Public/Views/IDisplayClusterConfiguratorView.h
@@ -0,0 +1,22 @@
+// Copyright Epic Games, Inc. All Rights Reserved.
+
+#pragma once
+
+#include "CoreMinimal.h"
+
+class SWidget;
+
+
+class IDisplayClusterConfiguratorView
+	: public TSharedFromThis<IDisplayClusterConfiguratorView>
+{
+public:
+	virtual ~IDisplayClusterConfiguratorView() = default;
+
+public:
+	virtual TSharedRef<SWidget> CreateWidget() = 0;
+	virtual TSharedRef<SWidget> GetWidget() = 0;
+	
+	virtual void SetEnabled(bool bInEnabled) {};
+	virtual bool GetIsEnabled() const { return false; };
+};
diff --git a/Engine/Plugins/Runtime/nDisplay/Source/DisplayClusterConfigurator/Public/Views/Log/IDisplayClusterConfiguratorViewLog.h b/Engine/Plugins/Runtime/nDisplay/Source/DisplayClusterConfigurator/Public/Views/Log/IDisplayClusterConfiguratorViewLog.h
new file mode 100644
--- /dev/null
+++ b/Engine/Plugins/Runtime/nDisplay/Source/DisplayClusterConfigurator/Public/Views/Log/IDisplayClusterConfiguratorViewLog.h
@@ -0,0 +1,29 @@
+// Copyright Epic Games, Inc. All Rights Reserved.
+
+#pragma once
+
+#include "Views/IDisplayClusterConfiguratorView.h"
+
+class IMessageLogListing;
+
+class IDisplayClusterConfiguratorViewLog
+	: public IDisplayClusterConfiguratorView
+{
+public:
+	virtual ~IDisplayClusterConfiguratorViewLog() = default;
+
+public:
+	enum class EVerbosityLevel
+	{
+		Log,
+		Warning,
+		Error
+	};
+
+public:
+	virtual TSharedRef<IMessageLogListing> GetMessageLogListing() const = 0;
+
+	virtual TSharedRef<IMessageLogListing> CreateLogListing() = 0;
+
+	virtual void Log(const FText& Message, EVerbosityLevel Verbosity = EVerbosityLevel::Log) = 0;
+};
diff --git a/Engine/Plugins/Runtime/nDisplay/Source/DisplayClusterConfigurator/Public/Views/OutputMapping/IDisplayClusterConfiguratorOutputMappingBuilder.h b/Engine/Plugins/Runtime/nDisplay/Source/DisplayClusterConfigurator/Public/Views/OutputMapping/IDisplayClusterConfiguratorOutputMappingBuilder.h
new file mode 100644
--- /dev/null
+++ b/Engine/Plugins/Runtime/nDisplay/Source/DisplayClusterConfigurator/Public/Views/OutputMapping/IDisplayClusterConfiguratorOutputMappingBuilder.h
@@ -0,0 +1,33 @@
+// Copyright Epic Games, Inc. All Rights Reserved.
+
+#pragma once
+
+#include "Views/IDisplayClusterConfiguratorBuilder.h"
+
+class IDisplayClusterConfiguratorOutputMappingSlot;
+class SWidget;
+class UObject;
+
+/**
+ * The Interface for generation of nodes from config
+ */
+class IDisplayClusterConfiguratorOutputMappingBuilder
+	: public IDisplayClusterConfiguratorBuilder
+{
+public:
+	virtual ~IDisplayClusterConfiguratorOutputMappingBuilder() = default;
+
+public:
+	struct FSlot
+	{
+		static const FName Canvas;
+		static const FName Window;
+		static const FName Viewport;
+	};
+
+	virtual void Build() = 0;
+
+	virtual TSharedRef<SWidget> GetCanvasWidget() const = 0;
+
+	virtual const TArray<TSharedPtr<IDisplayClusterConfiguratorOutputMappingSlot>>& GetAllSlots() const = 0;
+};
diff --git a/Engine/Plugins/Runtime/nDisplay/Source/DisplayClusterConfigurator/Public/Views/OutputMapping/IDisplayClusterConfiguratorOutputMappingItem.h b/Engine/Plugins/Runtime/nDisplay/Source/DisplayClusterConfigurator/Public/Views/OutputMapping/IDisplayClusterConfiguratorOutputMappingItem.h
new file mode 100644
--- /dev/null
+++ b/Engine/Plugins/Runtime/nDisplay/Source/DisplayClusterConfigurator/Public/Views/OutputMapping/IDisplayClusterConfiguratorOutputMappingItem.h
@@ -0,0 +1,20 @@
+// Copyright Epic Games, Inc. All Rights Reserved.
+
+#pragma once
+
+#include "Views/IDisplayClusterConfiguratorItem.h"
+
+class IDisplayClusterConfiguratorOutputMappingSlot;
+
+/**
+ * The Interface holsd object and setting from Details View
+ */
+class IDisplayClusterConfiguratorOutputMappingItem
+	: public IDisplayClusterConfiguratorItem
+{
+public:
+	virtual ~IDisplayClusterConfiguratorOutputMappingItem() = default;
+
+public:
+	virtual const FString& GetNodeName() const = 0;
+};
diff --git a/Engine/Plugins/Runtime/nDisplay/Source/DisplayClusterConfigurator/Public/Views/OutputMapping/IDisplayClusterConfiguratorViewOutputMapping.h b/Engine/Plugins/Runtime/nDisplay/Source/DisplayClusterConfigurator/Public/Views/OutputMapping/IDisplayClusterConfiguratorViewOutputMapping.h
new file mode 100644
--- /dev/null
+++ b/Engine/Plugins/Runtime/nDisplay/Source/DisplayClusterConfigurator/Public/Views/OutputMapping/IDisplayClusterConfiguratorViewOutputMapping.h
@@ -0,0 +1,106 @@
+// Copyright Epic Games, Inc. All Rights Reserved.
+
+#pragma once
+
+#include "Views/IDisplayClusterConfiguratorView.h"
+
+class UTexture;
+
+struct FOutputMappingSettings
+{
+	bool bShowRuler;
+	bool bShowWindowInfo;
+	bool bShowWindowCornerImage;
+	bool bShowOutsideViewports;
+	bool bAllowClusterItemOverlap;
+	bool bKeepClusterNodesInHosts;
+	bool bLockViewports;
+	bool bLockClusterNodes;
+	bool bTintSelectedViewports;
+	float ViewScale;
+
+	FOutputMappingSettings() :
+		bShowRuler(true),
+		bShowWindowInfo(true),
+		bShowWindowCornerImage(true),
+		bShowOutsideViewports(false),
+		bAllowClusterItemOverlap(false),
+		bKeepClusterNodesInHosts(true),
+		bLockViewports(false),
+		bLockClusterNodes(false),
+		bTintSelectedViewports(true),
+		ViewScale(1.0f)
+	{ }
+};
+
+enum class EHostArrangementType : uint8
+{
+	Horizontal,
+	Vertical,
+	Wrap,
+	Grid
+};
+
+struct FHostNodeArrangementSettings
+{
+	EHostArrangementType ArrangementType;
+	float WrapThreshold;
+	int GridSize;
+
+	FHostNodeArrangementSettings() :
+		ArrangementType(EHostArrangementType::Wrap),
+		WrapThreshold(5400),
+		GridSize(4)
+	{ }
+};
+
+struct FNodeAlignmentSettings
+{
+	int32 SnapProximity;
+	int32 AdjacentEdgesSnapPadding;
+
+	bool bSnapAdjacentEdges;
+	bool bSnapSameEdges;
+
+	FNodeAlignmentSettings() :
+		SnapProximity(25),
+		AdjacentEdgesSnapPadding(0),
+		bSnapAdjacentEdges(true),
+		bSnapSameEdges(true)
+	{ }
+};
+
+/**
+ * The Interface for controll the Output Mapping Мiew
+ */
+class IDisplayClusterConfiguratorViewOutputMapping
+	: public IDisplayClusterConfiguratorView
+{
+public:
+	virtual ~IDisplayClusterConfiguratorViewOutputMapping() = default;
+
+public:
+	DECLARE_MULTICAST_DELEGATE(FOnOutputMappingBuilt);
+
+	using FOnOutputMappingBuiltDelegate = FOnOutputMappingBuilt::FDelegate;
+
+public:
+	virtual void Cleanup() = 0;
+
+	virtual const FOutputMappingSettings& GetOutputMappingSettings() const = 0;
+	virtual FOutputMappingSettings& GetOutputMappingSettings() = 0;
+
+	virtual const FHostNodeArrangementSettings& GetHostArrangementSettings() const = 0;
+	virtual FHostNodeArrangementSettings& GetHostArrangementSettings() = 0;
+
+	virtual const FNodeAlignmentSettings& GetNodeAlignmentSettings() const = 0;
+	virtual FNodeAlignmentSettings& GetNodeAlignmentSettings() = 0;
+
+	virtual FOnOutputMappingBuilt& GetOnOutputMappingBuiltDelegate() = 0;
+	virtual FDelegateHandle RegisterOnOutputMappingBuilt(const FOnOutputMappingBuiltDelegate& Delegate) = 0;
+	virtual void UnregisterOnOutputMappingBuilt(FDelegateHandle DelegateHandle) = 0;
+
+	virtual void FindAndSelectObjects(const TArray<UObject*>& ObjectsToSelect) = 0;
+
+	virtual void JumpToObject(UObject* InObject) = 0;
+};
diff --git a/Engine/Plugins/Runtime/nDisplay/Source/DisplayClusterConfigurator/Public/Views/TreeViews/IDisplayClusterConfiguratorTreeBuilder.h b/Engine/Plugins/Runtime/nDisplay/Source/DisplayClusterConfigurator/Public/Views/TreeViews/IDisplayClusterConfiguratorTreeBuilder.h
new file mode 100644
--- /dev/null
+++ b/Engine/Plugins/Runtime/nDisplay/Source/DisplayClusterConfigurator/Public/Views/TreeViews/IDisplayClusterConfiguratorTreeBuilder.h
@@ -0,0 +1,102 @@
+// Copyright Epic Games, Inc. All Rights Reserved.
+
+#pragma once
+
+#include "Views/IDisplayClusterConfiguratorBuilder.h"
+#include "Containers/ArrayView.h"
+
+class IDisplayClusterConfiguratorTreeItem;
+class IDisplayClusterConfiguratorViewTree;
+
+/** Output struct for builders to use */
+struct DISPLAYCLUSTERCONFIGURATOR_API FDisplayClusterConfiguratorTreeBuilderOutput
+{
+	FDisplayClusterConfiguratorTreeBuilderOutput(TArray<TSharedPtr<IDisplayClusterConfiguratorTreeItem>>& InItems, TArray<TSharedPtr<IDisplayClusterConfiguratorTreeItem>>& InLinearItems)
+		: Items(InItems)
+		, LinearItems(InLinearItems)
+	{ }
+
+	/**
+	 * Add an item to the output
+	 * @param	InItem			The item to add
+	 * @param	InParentName	The name of the item's parent
+	 * @param	InParentTypes	The types of items to search. If this is empty all items will be searched.
+	 * @param	bAddToHead		Whether to add the item to the start or end of the parent's children array
+	 */
+	void Add(const TSharedPtr<IDisplayClusterConfiguratorTreeItem>& InItem, const FName& InParentName, TArrayView<const FName> InParentTypes, bool bAddToHead = false);
+
+	/**
+	 * Add an item to the output
+	 * @param	InItem			The item to add
+	 * @param	InParentName	The name of the item's parent
+	 * @param	InParentTypes	The types of items to search. If this is empty all items will be searched.
+	 * @param	bAddToHead		Whether to add the item to the start or end of the parent's children array
+	 */
+	void Add(const TSharedPtr<IDisplayClusterConfiguratorTreeItem>& InItem, const FName& InParentName, std::initializer_list<FName> InParentTypes, bool bAddToHead = false)
+	{
+		Add(InItem, InParentName, MakeArrayView(InParentTypes), bAddToHead);
+	}
+
+	/**
+	 * Add an item to the output
+	 * @param	InItem			The item to add
+	 * @param	InParentName	The name of the item's parent
+	 * @param	InParentType	The type of items to search. If this is NAME_None all items will be searched.
+	 * @param	bAddToHead		Whether to add the item to the start or end of the parent's children array
+	 */
+	void Add(const TSharedPtr<IDisplayClusterConfiguratorTreeItem>& InItem, const FName& InParentName, const FName& InParentType, bool bAddToHead = false);
+
+	/**
+	 * Find the item with the specified name
+	 * @param	InName	The item's name
+	 * @param	InTypes	The types of items to search. If this is empty all items will be searched.
+	 * @return the item found, or an invalid ptr if it was not found.
+	 */
+	TSharedPtr<IDisplayClusterConfiguratorTreeItem> Find(const FName& InName, TArrayView<const FName> InTypes) const;
+
+	/**
+	 * Find the item with the specified name
+	 * @param	InName	The item's name
+	 * @param	InTypes	The types of items to search. If this is empty all items will be searched.
+	 * @return the item found, or an invalid ptr if it was not found.
+	 */
+	TSharedPtr<IDisplayClusterConfiguratorTreeItem> Find(const FName& InName, std::initializer_list<FName> InTypes) const
+	{
+		return Find(InName, MakeArrayView(InTypes));
+	}
+
+	/**
+	 * Find the item with the specified name
+	 * @param	InName	The item's name
+	 * @param	InType	The type of items to search. If this is NAME_None all items will be searched.
+	 * @return the item found, or an invalid ptr if it was not found.
+	 */
+	TSharedPtr<IDisplayClusterConfiguratorTreeItem> Find(const FName& InName, const FName& InType) const;
+
+private:
+	/** The items that are built by this builder */
+	TArray<TSharedPtr<IDisplayClusterConfiguratorTreeItem>>& Items;
+
+	/** A linearized list of all items in OutItems (for easier searching) */
+	TArray<TSharedPtr<IDisplayClusterConfiguratorTreeItem>>& LinearItems;
+};
+
+/** 
+ * Interface to implement to provide custom build logic to tree views 
+ */
+class IDisplayClusterConfiguratorTreeBuilder
+	: public IDisplayClusterConfiguratorBuilder
+{
+public:
+	virtual ~IDisplayClusterConfiguratorTreeBuilder() = default;
+
+public:
+	/** Setup this builder with links to the tree */
+	virtual void Initialize(const TSharedRef<IDisplayClusterConfiguratorViewTree>& InConfiguratorTree) = 0;
+
+	/**
+	 * Build an array of item tree items to display in the tree.
+	 * @param	Output			The items that are built by this builder
+	 */
+	virtual void Build(FDisplayClusterConfiguratorTreeBuilderOutput& Output) = 0;
+};
diff --git a/Engine/Plugins/Runtime/nDisplay/Source/DisplayClusterConfigurator/Public/Views/TreeViews/IDisplayClusterConfiguratorTreeItem.h b/Engine/Plugins/Runtime/nDisplay/Source/DisplayClusterConfigurator/Public/Views/TreeViews/IDisplayClusterConfiguratorTreeItem.h
new file mode 100644
--- /dev/null
+++ b/Engine/Plugins/Runtime/nDisplay/Source/DisplayClusterConfigurator/Public/Views/TreeViews/IDisplayClusterConfiguratorTreeItem.h
@@ -0,0 +1,190 @@
+// Copyright Epic Games, Inc. All Rights Reserved.
+
+#pragma once
+
+#include "Views/IDisplayClusterConfiguratorItem.h"
+
+#include "Framework/SlateDelegates.h"
+
+enum class EItemDropZone;
+class IDisplayClusterConfiguratorViewTree;
+class ITableRow;
+class SHorizontalBox;
+class STableViewBase;
+class SWidget;
+class FTextFilterExpressionEvaluator;
+class FDragDropEvent;
+
+template< typename ObjectType >
+class TAttribute;
+
+
+enum class EDisplayClusterConfiguratorTreeFilterResult : uint8
+{
+	/** Hide the item */
+	Hidden,
+
+	/** Show the item because child items were shown */
+	ShownDescendant,
+
+	/** Show the item */
+	Shown,
+
+	/** Show the item and highlight search text */
+	ShownHighlighted,
+};
+
+
+#define NDISPLAY_TREE_BASE_ITEM_TYPE(TYPE) \
+	static const FName& GetTypeId() { static FName Type(TEXT(#TYPE)); return Type; } \
+	virtual bool IsOfTypeByName(const FName& Type) const { return TYPE::GetTypeId() == Type; } \
+	virtual FName GetTypeName() const { return TYPE::GetTypeId(); }
+
+
+class IDisplayClusterConfiguratorTreeItem
+	: public IDisplayClusterConfiguratorItem
+{
+public:
+	virtual ~IDisplayClusterConfiguratorTreeItem() = default;
+
+public:
+	NDISPLAY_TREE_BASE_ITEM_TYPE(IDisplayClusterConfiguratorTreeItem)
+
+	/** Initializes the tree item */
+	virtual void Initialize() = 0;
+
+	/** Builds the table row widget to display this info */
+	virtual TSharedRef<ITableRow> MakeTreeRowWidget(const TSharedRef<STableViewBase>& InOwnerTable, const TAttribute<FText>& InFilterText) = 0;
+
+	/**
+	 * Builds a widget for the specified column.
+	 * @param ColumnName - The name of the column to build the widget for.
+	 * @param TableRow - The table row the widget is being generated for.
+	 * @param FilterText - The current filter being applied to the tree view.
+	 * @param InIsSelected - The FIsSelected delegate for the current table row.
+	 */
+	virtual TSharedRef<SWidget> GenerateWidgetForColumn(const FName& ColumnName, TSharedPtr<ITableRow> TableRow, const TAttribute<FText>& FilterText, FIsSelected InIsSelected) = 0;
+
+	/** Builds the slate widget for any inline data editing */
+	virtual TSharedRef< SWidget > GenerateInlineEditWidget(const TAttribute<FText>& FilterText, FIsSelected InIsSelected) = 0;
+
+	/** @return true if the item has an in-line editor widget */
+	virtual bool HasInlineEditor() const = 0;
+
+	/** Toggle the expansion state of the inline editor */
+	virtual void ToggleInlineEditorExpansion() = 0;
+
+	/** Get the expansion state of the inline editor */
+	virtual bool IsInlineEditorExpanded() const = 0;
+
+	/** Get the name of the item that this row represents */
+	virtual FName GetRowItemName() const = 0;
+
+	/** Get the style for image icon */
+	virtual FString GetIconStyle() const = 0;
+
+	/** Return the name used to attach to this item */
+	virtual FName GetAttachName() const = 0;
+
+	/** @return true if this item can be renamed */
+	virtual bool CanRenameItem() const = 0;
+
+	/** Requests a rename on the the tree row item */
+	virtual void RequestRename() = 0;
+
+	/** @return true if this item can be deleted */
+	virtual bool CanDeleteItem() const = 0;
+
+	/** Deletes the tree row item and its backing object. */
+	virtual void DeleteItem() const = 0;
+
+	/** @return true if this item can be duplicated */
+	virtual bool CanDuplicateItem() const = 0;
+
+	/** @return true if this item can be hidden */
+	virtual bool CanHideItem() const = 0;
+	
+	/** Sets the item's visibility */
+	virtual void SetItemHidden(bool bIsHidden) = 0;
+
+	/** Get the objects of all parents of this item */
+	virtual void GetParentObjectsRecursive(TArray<UObject*>& OutObjects) const = 0;
+
+	/** Get all Children object  */
+	virtual void GetChildrenObjectsRecursive(TArray<UObject*>& OutObjects) const = 0;
+
+	/** Handler for when the user double clicks on this item in the tree */
+	virtual void OnItemDoubleClicked() = 0;
+
+	/** Handler for when the user enter the mouse on this item in the tree */
+	virtual void OnMouseEnter() = 0;
+
+	/** Handler for when the user leave the mouse on this item in the tree */
+	virtual void OnMouseLeave() = 0;
+
+	/** Return true if tree item is hovered */
+	virtual bool IsHovered() const = 0;
+
+	/** Return true if that is root tree item */
+	virtual bool IsRoot() const = 0;
+
+	/** Return true if this item is a child of the given parent tree item */
+	virtual bool IsChildOfRecursive(const TSharedRef<IDisplayClusterConfiguratorTreeItem>& InParentTreeItem) const = 0;
+
+	/** Handle a drag and drop detected event */
+	virtual FReply HandleDragDetected(const FGeometry& MyGeometry, const FPointerEvent& MouseEvent) = 0;
+
+	/** Handle a drag and drop enter event */
+	virtual void HandleDragEnter(const FDragDropEvent& DragDropEvent) = 0;
+
+	/** Handle a drag and drop leave event */
+	virtual void HandleDragLeave(const FDragDropEvent& DragDropEvent) = 0;
+
+	/** Handles the CanAcceptDrop callback for the tree item */
+	virtual TOptional<EItemDropZone> HandleCanAcceptDrop(const FDragDropEvent& DragDropEvent, EItemDropZone DropZone, TSharedPtr<IDisplayClusterConfiguratorTreeItem> TargetItem) = 0;
+
+	/** Handle Accep drag and drop drop event */
+	virtual FReply HandleAcceptDrop(const FDragDropEvent& DragDropEvent, TSharedPtr<IDisplayClusterConfiguratorTreeItem> TargetItem) = 0;
+
+	/** Handle a drag and drop drop event */
+	virtual FReply HandleDrop(FDragDropEvent const& DragDropEvent) = 0;
+
+	/** Get this item's parent  */
+	virtual TSharedPtr<IDisplayClusterConfiguratorTreeItem> GetParent() const = 0;
+
+	/** Set this item's parent */
+	virtual void SetParent(TSharedPtr<IDisplayClusterConfiguratorTreeItem> InParent) = 0;
+
+	/** The array of children for this item */
+	virtual TArray<TSharedPtr<IDisplayClusterConfiguratorTreeItem>>& GetChildren() = 0;
+
+	/** The const array of children for this item */
+	virtual const TArray<TSharedPtr<IDisplayClusterConfiguratorTreeItem>>& GetChildrenConst() const = 0;
+
+	/** The filtered array of children for this item */
+	virtual TArray<TSharedPtr<IDisplayClusterConfiguratorTreeItem>>& GetFilteredChildren() = 0;
+
+	/** The owning view tree */
+	virtual TSharedRef<IDisplayClusterConfiguratorViewTree> GetConfiguratorTree() const = 0;
+
+	/**
+	 * Applies the specified text filter to the tree item and returns the results
+	 * @param TextFilter - The filter to apply
+	 * @return The results of applying the filter to this tree item
+	 */
+	virtual EDisplayClusterConfiguratorTreeFilterResult ApplyFilter(const TSharedPtr<FTextFilterExpressionEvaluator>& TextFilter) = 0;
+
+	/** Get the current filter result */
+	virtual EDisplayClusterConfiguratorTreeFilterResult GetFilterResult() const = 0;
+
+	/** Set the current filter result */
+	virtual void SetFilterResult(EDisplayClusterConfiguratorTreeFilterResult InResult) = 0;
+
+	/** Get whether this item begins expanded or not */
+	virtual bool IsInitiallyExpanded() const = 0;
+};
+
+#define NDISPLAY_TREE_ITEM_TYPE(TYPE, BASE) \
+	static const FName& GetTypeId() { static FName Type(TEXT(#TYPE)); return Type; } \
+	virtual bool IsOfTypeByName(const FName& Type) const override { return GetTypeId() == Type || BASE::IsOfTypeByName(Type); } \
+	virtual FName GetTypeName() const { return TYPE::GetTypeId(); }
diff --git a/Engine/Plugins/Runtime/nDisplay/Source/DisplayClusterConfigurator/Public/Views/TreeViews/IDisplayClusterConfiguratorViewTree.h b/Engine/Plugins/Runtime/nDisplay/Source/DisplayClusterConfigurator/Public/Views/TreeViews/IDisplayClusterConfiguratorViewTree.h
new file mode 100644
--- /dev/null
+++ b/Engine/Plugins/Runtime/nDisplay/Source/DisplayClusterConfigurator/Public/Views/TreeViews/IDisplayClusterConfiguratorViewTree.h
@@ -0,0 +1,222 @@
+// Copyright Epic Games, Inc. All Rights Reserved.
+
+#pragma once
+
+#include "Views/IDisplayClusterConfiguratorView.h"
+
+#include "Widgets/Views/SHeaderRow.h"
+
+class IDisplayClusterConfiguratorTreeBuilder;
+class IDisplayClusterConfiguratorTreeItem;
+class UDisplayClusterConfigurationData;
+struct FDisplayClusterConfiguratorTreeFilterArgs;
+class FUICommandList_Pinnable;
+class FMenuBuilder;
+class FTextFilterExpressionEvaluator;
+
+enum class EDisplayClusterConfiguratorTreeFilterResult : uint8;
+
+enum class EDisplayClusterConfiguratorTreeMode
+{
+	Editor,
+	Picker
+};
+
+struct FDisplayClusterConfiguratorTreeArgs
+{
+	FDisplayClusterConfiguratorTreeArgs()
+		: Mode(EDisplayClusterConfiguratorTreeMode::Editor)
+		, ContextName(TEXT("ConfiguratorTree"))
+	{ }
+
+	/** Optional builder to allow for custom tree construction */
+	TSharedPtr<IDisplayClusterConfiguratorTreeBuilder> Builder;
+
+	/** The mode that this items tree is in */
+	EDisplayClusterConfiguratorTreeMode Mode;
+
+	/** Context name used to persist settings */
+	FName ContextName;
+};
+
+
+/** Basic filter used when re-filtering the tree */
+struct FDisplayClusterConfiguratorTreeFilterArgs
+{
+	FDisplayClusterConfiguratorTreeFilterArgs(TSharedPtr<FTextFilterExpressionEvaluator> InTextFilter)
+		: TextFilter(InTextFilter)
+		, bFlattenHierarchyOnFilter(false)
+	{}
+
+	/** The text filter we are using, if any */
+	TSharedPtr<FTextFilterExpressionEvaluator> TextFilter;
+
+	/** Whether to flatten the hierarchy so filtered items appear in a linear list */
+	bool bFlattenHierarchyOnFilter;
+};
+
+/*
+ * Base interface for editor tree view
+ */
+class IDisplayClusterConfiguratorViewTree
+	: public IDisplayClusterConfiguratorView
+{
+public:
+	virtual ~IDisplayClusterConfiguratorViewTree() = default;
+
+public:
+	DECLARE_MULTICAST_DELEGATE_OneParam(FOnHoveredItemSet, const TSharedRef<IDisplayClusterConfiguratorTreeItem>&);
+	DECLARE_MULTICAST_DELEGATE(FOnHoveredItemCleared);
+	DECLARE_MULTICAST_DELEGATE_OneParam(FOnItemsSelected, const TArray<TSharedPtr<IDisplayClusterConfiguratorTreeItem>>&);
+	DECLARE_MULTICAST_DELEGATE(FOnSelectionCleared);
+
+	using FOnHoveredItemSetDelegate = FOnHoveredItemSet::FDelegate;
+	using FOnHoveredItemClearedDelegate = FOnHoveredItemCleared::FDelegate;
+	using FOnItemsSelectedDelegate = FOnItemsSelected::FDelegate;
+	using FOnSelectionClearedDelegate = FOnSelectionCleared::FDelegate;
+
+	/**
+	 * Type of the tree item
+	 */
+	struct Columns
+	{
+		static const FName Item;
+		static const FName Group;
+	};
+
+	/**
+	 * Rebuild the tree from the config
+	 */
+	virtual void RebuildTree() = 0;
+
+	/**
+	 * @return The Editor editing UObject
+	 */
+	virtual UDisplayClusterConfigurationData* GetEditorData() const = 0;
+
+	/**
+	 * Generates a list of columns to add to the tree view.
+	 * @param OutColumnArgs - The list of column arguments to construct the columns from.
+	 */
+	virtual void ConstructColumns(TArray<SHeaderRow::FColumn::FArguments>& OutColumnArgs) const = 0;
+
+	/**
+	 * Sets currently hovered tree item
+	 *
+	 * @param InTreeItem	 hovered tree item
+	 */
+	virtual void SetHoveredItem(const TSharedRef<IDisplayClusterConfiguratorTreeItem>& InTreeItem) = 0;
+
+	/**
+	 * Remove hovered Item from this tree
+	 */
+	virtual void ClearHoveredItem() = 0;
+
+	/**
+	 * @return The hovered tree item
+	 */
+	virtual TSharedPtr<IDisplayClusterConfiguratorTreeItem> GetHoveredItem() const = 0;
+
+	/**
+	 * Sets currently selected tree items
+	 *
+	 * @param InTreeItems - List of selected tree items
+	 */
+	virtual void SetSelectedItems(const TArray<TSharedPtr<IDisplayClusterConfiguratorTreeItem>>& InTreeItems) = 0;
+
+	/**
+	 * Removes all tree items from list of selected items
+	 */
+	virtual void ClearSelection() = 0;
+
+	/**
+	 * @return The list of currently selected tree items
+	 */
+	virtual TArray<TSharedPtr<IDisplayClusterConfiguratorTreeItem>> GetSelectedItems() const = 0;
+
+	/** Gets the objects that are currently selected in the tree view. */
+	virtual void GetSelectedObjects(TArray<UObject*>& OutObjects) const = 0;
+
+	/**
+	 * Finds any tree items that represet the specified objects and selects them.
+	 * @param ObjectsToSelect - The objects to select
+	 */
+	virtual void FindAndSelectObjects(const TArray<UObject*>& ObjectsToSelect) = 0;
+
+	/**
+	 * Filters the specified list of items based on the filter arguments, and outputs the filtered list.
+	 * @param InArgs - The filtering arguments
+	 * @param InItems - The list of items to filter
+	 * @param OutFilteredItems - The output list of items that matched the filter
+	 */
+	virtual void Filter(const FDisplayClusterConfiguratorTreeFilterArgs& InArgs, const TArray<TSharedPtr<IDisplayClusterConfiguratorTreeItem>>& InItems, TArray<TSharedPtr<IDisplayClusterConfiguratorTreeItem>>& OutFilteredItems) = 0;
+
+
+	/**
+	 * Register the hovered tree item delegate
+	 */
+	virtual FDelegateHandle RegisterOnHoveredItemSet(const FOnHoveredItemSetDelegate& Delegate) = 0;
+
+	/**
+	 * Remove the hovered tree item delegate
+	 */
+	virtual void UnregisterOnHoveredItemSet(FDelegateHandle DelegateHandle) = 0;
+
+	/**
+	 * Register clear hovered tree item delegate
+	 */
+	virtual FDelegateHandle RegisterOnHoveredItemCleared(const FOnHoveredItemClearedDelegate& Delegate) = 0;
+
+	/**
+	 * Unregister clear hovered tree item delegate
+	 */
+	virtual void UnregisterOnHoveredItemCleared(FDelegateHandle DelegateHandle) = 0;
+
+
+	/**
+	 * Raised when the context menu is being shown for an item in the tree view.
+	 * @param MenuBuilder - The menu builder of the context menu
+	 */
+	virtual void FillContextMenu(FMenuBuilder& MenuBuilder) = 0;
+
+	/**
+	 * Raised on tree view creation when binding UI commands
+	 * @param CommandList - The command list to bind commands to
+	 */
+	virtual void BindPinnableCommands(FUICommandList_Pinnable& CommandList) = 0;
+
+	/**
+	 * @return true if the tree view should show an "Add New" button in its top bar
+	 */
+	virtual bool ShowAddNewButton() const = 0;
+
+	/**
+	 * Raised when the "Add New" combo button is clicked and its menu is being opened
+	 */
+	virtual void FillAddNewMenu(FMenuBuilder& MenuBuilder) = 0;
+
+	/**
+	 * @return true if the tree view should show a filter options button in its top bar
+	 */
+	virtual bool ShowFilterOptionsButton() const = 0;
+
+	/**
+	 * Raised when the filter options combo button is clicked and its menu is being opened
+	 */
+	virtual void FillFilterOptionsMenu(FMenuBuilder& MenuBuilder) = 0;
+
+	/**
+	 * @return true if the tree view should show a "View Options" button in its bottom bar
+	 */
+	virtual bool ShowViewOptionsButton() const = 0;
+
+	/**
+	 * Raised when the "View Options" combo button is clicked and its menu is being opened
+	 */
+	virtual void FillViewOptionsMenu(FMenuBuilder& MenuBuilder) = 0;
+
+	/**
+	 * @return The watermark text to display in the bottom right corner of the tree view
+	 */
+	virtual FText GetCornerText() const = 0;
+};
diff --git a/Engine/Plugins/Runtime/nDisplay/Source/DisplayClusterConfigurator/Public/Views/Viewport/IDisplayClusterConfiguratorPreviewScene.h b/Engine/Plugins/Runtime/nDisplay/Source/DisplayClusterConfigurator/Public/Views/Viewport/IDisplayClusterConfiguratorPreviewScene.h
new file mode 100644
--- /dev/null
+++ b/Engine/Plugins/Runtime/nDisplay/Source/DisplayClusterConfigurator/Public/Views/Viewport/IDisplayClusterConfiguratorPreviewScene.h
@@ -0,0 +1,19 @@
+// Copyright Epic Games, Inc. All Rights Reserved.
+
+#pragma once
+
+#include "CoreMinimal.h"
+#include "AdvancedPreviewScene.h"
+
+
+class IDisplayClusterConfiguratorPreviewScene 
+	: public FAdvancedPreviewScene
+{
+public:
+	/** Constructor only here to pass ConstructionValues to base constructor */
+	IDisplayClusterConfiguratorPreviewScene(ConstructionValues CVS)
+		: FAdvancedPreviewScene(CVS)
+	{ }
+
+	virtual ~IDisplayClusterConfiguratorPreviewScene() = default;
+};
diff --git a/Engine/Plugins/Runtime/nDisplay/Source/DisplayClusterConfigurator/Public/Views/Viewport/IDisplayClusterConfiguratorViewViewport.h b/Engine/Plugins/Runtime/nDisplay/Source/DisplayClusterConfigurator/Public/Views/Viewport/IDisplayClusterConfiguratorViewViewport.h
new file mode 100644
--- /dev/null
+++ b/Engine/Plugins/Runtime/nDisplay/Source/DisplayClusterConfigurator/Public/Views/Viewport/IDisplayClusterConfiguratorViewViewport.h
@@ -0,0 +1,18 @@
+// Copyright Epic Games, Inc. All Rights Reserved.
+
+#pragma once
+
+#include "Views/IDisplayClusterConfiguratorView.h"
+
+class IDisplayClusterConfiguratorPreviewScene;
+
+
+class IDisplayClusterConfiguratorViewViewport
+	: public IDisplayClusterConfiguratorView
+{
+public:
+	virtual ~IDisplayClusterConfiguratorViewViewport() = default;
+
+public:
+	virtual TSharedRef<IDisplayClusterConfiguratorPreviewScene> GetPreviewScene() const = 0;
+};
diff --git a/Engine/Plugins/Runtime/nDisplay/Source/DisplayClusterConfigurator/Public/Views/Viewport/IDisplayClusterConfiguratorViewportBuilder.h b/Engine/Plugins/Runtime/nDisplay/Source/DisplayClusterConfigurator/Public/Views/Viewport/IDisplayClusterConfiguratorViewportBuilder.h
new file mode 100644
--- /dev/null
+++ b/Engine/Plugins/Runtime/nDisplay/Source/DisplayClusterConfigurator/Public/Views/Viewport/IDisplayClusterConfiguratorViewportBuilder.h
@@ -0,0 +1,18 @@
+// Copyright Epic Games, Inc. All Rights Reserved.
+
+#pragma once
+
+#include "Views/IDisplayClusterConfiguratorBuilder.h"
+
+
+class IDisplayClusterConfiguratorViewportBuilder
+	: public IDisplayClusterConfiguratorBuilder
+{
+public:
+	virtual ~IDisplayClusterConfiguratorViewportBuilder() = default;
+
+public:
+	virtual void BuildViewport() = 0;
+
+	virtual void ClearViewportSelection() = 0;
+};
