nDisplay: LC proxy xform propagation now includes missing actor rotation.

diff --git a/Engine/Plugins/Runtime/nDisplay/Source/DisplayClusterLightCardEditor/Private/Viewport/DisplayClusterLightCardEditorViewportClient.cpp b/Engine/Plugins/Runtime/nDisplay/Source/DisplayClusterLightCardEditor/Private/Viewport/DisplayClusterLightCardEditorViewportClient.cpp
--- a/Engine/Plugins/Runtime/nDisplay/Source/DisplayClusterLightCardEditor/Private/Viewport/DisplayClusterLightCardEditorViewportClient.cpp
+++ b/Engine/Plugins/Runtime/nDisplay/Source/DisplayClusterLightCardEditor/Private/Viewport/DisplayClusterLightCardEditorViewportClient.cpp
@@ -1402,6 +1402,16 @@ void FDisplayClusterLightCardEditorViewportClient::PropagateLightCardTransform(A
 			}
 		};
 		
+		// Here we count on the fact that the root actor proxy has zero loc/rot.
+		// If that ever changes then the math below will need to be updated to use the 
+		// relative loc/rot of the LC proxies wrt the root actor proxy.
+
+		const FVector RootActorLevelInstanceLocation = RootActorLevelInstance.IsValid() ? RootActorLevelInstance->GetActorLocation() : FVector::ZeroVector;
+		LevelInstance->SetActorLocation(RootActorLevelInstanceLocation + LightCardProxy->GetActorLocation());
+
+		const FRotator RootActorLevelInstanceRotation = RootActorLevelInstance.IsValid() ? RootActorLevelInstance->GetActorRotation() : FRotator::ZeroRotator;
+		LevelInstance->SetActorRotation(RootActorLevelInstanceRotation.Quaternion() * LightCardProxy->GetActorRotation().Quaternion());
+
 		TryChangeProperty(GET_MEMBER_NAME_CHECKED(ADisplayClusterLightCardActor, Longitude));
 		TryChangeProperty(GET_MEMBER_NAME_CHECKED(ADisplayClusterLightCardActor, Latitude));
 		TryChangeProperty(GET_MEMBER_NAME_CHECKED(ADisplayClusterLightCardActor, DistanceFromCenter));
@@ -1409,9 +1419,6 @@ void FDisplayClusterLightCardEditorViewportClient::PropagateLightCardTransform(A
 		TryChangeProperty(GET_MEMBER_NAME_CHECKED(ADisplayClusterLightCardActor, Pitch));
 		TryChangeProperty(GET_MEMBER_NAME_CHECKED(ADisplayClusterLightCardActor, Yaw));
 		
-		const FVector RootActorLevelInstanceLocation = RootActorLevelInstance.IsValid() ? RootActorLevelInstance->GetActorLocation() : FVector::ZeroVector;
-		LevelInstance->SetActorLocation(RootActorLevelInstanceLocation + LightCardProxy->GetActorLocation());
-
 		// Snapshot the changed properties so multi-user can update while dragging.
 		if (ChangedProperties.Num() > 0)
 		{
