nDisplay: Added EngineShowFlags.SetGlobalIllumination(0) in LC and CK renders.

diff --git a/Engine/Plugins/Runtime/nDisplay/Source/DisplayCluster/Private/Render/Viewport/DisplayClusterViewportManager.cpp b/Engine/Plugins/Runtime/nDisplay/Source/DisplayCluster/Private/Render/Viewport/DisplayClusterViewportManager.cpp
--- a/Engine/Plugins/Runtime/nDisplay/Source/DisplayCluster/Private/Render/Viewport/DisplayClusterViewportManager.cpp
+++ b/Engine/Plugins/Runtime/nDisplay/Source/DisplayCluster/Private/Render/Viewport/DisplayClusterViewportManager.cpp
@@ -581,6 +581,7 @@ FSceneViewFamily::ConstructionValues FDisplayClusterViewportManager::CreateViewF
 
 			InEngineShowFlags.SetLumenReflections(0);
 			InEngineShowFlags.SetLumenGlobalIllumination(0);
+			InEngineShowFlags.SetGlobalIllumination(0);
 
 			break;
 
