nDisplay: Added more triangles to LC so that the proxies and the selection outlines match more closely in the projection mapping.

This comes at the expense rendering 4x more LC triangles. For the time being the tradeoff seems ok.

diff --git a/Engine/Plugins/Runtime/nDisplay/Source/DisplayCluster/Private/Game/EngineClasses/Scene/DisplayClusterLightCardActor.cpp b/Engine/Plugins/Runtime/nDisplay/Source/DisplayCluster/Private/Game/EngineClasses/Scene/DisplayClusterLightCardActor.cpp
--- a/Engine/Plugins/Runtime/nDisplay/Source/DisplayCluster/Private/Game/EngineClasses/Scene/DisplayClusterLightCardActor.cpp
+++ b/Engine/Plugins/Runtime/nDisplay/Source/DisplayCluster/Private/Game/EngineClasses/Scene/DisplayClusterLightCardActor.cpp
@@ -36,7 +36,7 @@ ADisplayClusterLightCardActor::ADisplayClusterLightCardActor(const FObjectInitia
 	LightCardTransformerComponent = CreateDefaultSubobject<USceneComponent>(TEXT("LightCardTransformer"));
 	LightCardTransformerComponent->AttachToComponent(MainSpringArmComponent, FAttachmentTransformRules::KeepRelativeTransform);
 
-	static ConstructorHelpers::FObjectFinder<UStaticMesh> PlaneObj(TEXT("/Engine/BasicShapes/Plane"));
+	static ConstructorHelpers::FObjectFinder<UStaticMesh> PlaneObj(TEXT("/nDisplay/LightCard/SM_LightCardPlane"));
 	static ConstructorHelpers::FObjectFinder<UMaterial> LightCardMatObj(TEXT("/nDisplay/LightCard/M_LightCard"));
 
 	LightCardComponent = CreateDefaultSubobject<UStaticMeshComponent>(TEXT("LightCard"));
