nDisplay: LC Editor now uses a subdivided plane to mitigate map projection straight light artifacts.

nDisplay proper still uses undivided planes.

diff --git a/Engine/Plugins/Runtime/nDisplay/Source/DisplayCluster/Private/Game/EngineClasses/Scene/DisplayClusterLightCardActor.cpp b/Engine/Plugins/Runtime/nDisplay/Source/DisplayCluster/Private/Game/EngineClasses/Scene/DisplayClusterLightCardActor.cpp
--- a/Engine/Plugins/Runtime/nDisplay/Source/DisplayCluster/Private/Game/EngineClasses/Scene/DisplayClusterLightCardActor.cpp
+++ b/Engine/Plugins/Runtime/nDisplay/Source/DisplayCluster/Private/Game/EngineClasses/Scene/DisplayClusterLightCardActor.cpp
@@ -169,4 +169,14 @@ void ADisplayClusterLightCardActor::UpdateLightCardMaterialInstance()
 		LightCardMaterialInstance->SetScalarParameterValue(TEXT("Opacity"), bIsProxy ? ProxyOpacity : Opacity);
 		LightCardMaterialInstance->SetScalarParameterValue(TEXT("Feather"), Feathering);
 	}
-}
\ No newline at end of file
+}
+
+UStaticMesh* ADisplayClusterLightCardActor::GetStaticMesh() const
+{
+	return LightCardComponent->GetStaticMesh();
+}
+
+void ADisplayClusterLightCardActor::SetStaticMesh(UStaticMesh* InStaticMesh)
+{
+	LightCardComponent->SetStaticMesh(InStaticMesh);
+}
diff --git a/Engine/Plugins/Runtime/nDisplay/Source/DisplayCluster/Public/DisplayClusterLightCardActor.h b/Engine/Plugins/Runtime/nDisplay/Source/DisplayCluster/Public/DisplayClusterLightCardActor.h
--- a/Engine/Plugins/Runtime/nDisplay/Source/DisplayCluster/Public/DisplayClusterLightCardActor.h
+++ b/Engine/Plugins/Runtime/nDisplay/Source/DisplayCluster/Public/DisplayClusterLightCardActor.h
@@ -10,6 +10,7 @@
 class USceneComponent;
 class USpringArmComponent;
 class UStaticMeshComponent;
+class UStaticMesh;
 
 UENUM(BlueprintType)
 enum class EDisplayClusterLightCardMask : uint8
@@ -40,6 +41,12 @@ public:
 	 * @param bIgnoreSpinYawPitch - If the light card component's spin, yaw, and pitch should be ignored when computing the transform
 	 */
 	FTransform GetLightCardTransform(bool bIgnoreSpinYawPitch = false) const;
+
+	/** Returns the current static mesh used by this light card */
+	UStaticMesh* GetStaticMesh() const;
+
+	/** Sets a new static mesh for the light card */
+	void SetStaticMesh(UStaticMesh* InStaticMesh);
 	
 protected:
 	void UpdateLightCardTransform();
diff --git a/Engine/Plugins/Runtime/nDisplay/Source/DisplayClusterLightCardEditor/Private/Viewport/DisplayClusterLightCardEditorViewportClient.cpp b/Engine/Plugins/Runtime/nDisplay/Source/DisplayClusterLightCardEditor/Private/Viewport/DisplayClusterLightCardEditorViewportClient.cpp
--- a/Engine/Plugins/Runtime/nDisplay/Source/DisplayClusterLightCardEditor/Private/Viewport/DisplayClusterLightCardEditorViewportClient.cpp
+++ b/Engine/Plugins/Runtime/nDisplay/Source/DisplayClusterLightCardEditor/Private/Viewport/DisplayClusterLightCardEditorViewportClient.cpp
@@ -825,6 +825,24 @@ void FDisplayClusterLightCardEditorViewportClient::UpdatePreviewActor(ADisplayCl
 					LightCardProxy->SetActorRotation(LightCard->GetActorRotation() - RootActor->GetActorRotation());
 					LightCardProxy->bIsProxy = true;
 
+					// Change mesh to proxy mesh with more vertices
+					if (const UStaticMesh* LightCardMesh = LightCardProxy->GetStaticMesh())
+					{
+						// Only change the mesh if we are using the default one.
+					
+						const FString DefaultPlanePath = TEXT("/nDisplay/LightCard/SM_LightCardPlane.SM_LightCardPlane");
+						const UStaticMesh* DefaultPlane = Cast<UStaticMesh>(FSoftObjectPath(DefaultPlanePath).TryLoad());
+
+						if (DefaultPlane == LightCardMesh)
+						{
+							const FString LightCardPlanePath = TEXT("/nDisplay/LightCard/SM_LightCardPlaneSubdivided.SM_LightCardPlaneSubdivided");
+							if (UStaticMesh* LightCardPlaneMesh = Cast<UStaticMesh>(FSoftObjectPath(LightCardPlanePath).TryLoad()))
+							{
+								LightCardProxy->SetStaticMesh(LightCardPlaneMesh);
+							}
+						}
+					}
+
 					LightCardProxies.Add(FLightCardProxy(LightCard.Get(), LightCardProxy));
 					
 					if (LastSelectedLightCardLevelInstances.Contains(LightCard.Get()))
