nDisplay: Added ability to scale light cards in light card editor with scale widget that closely resembles the scaling gizmo used elsewhere in Unreal Editor. Also fixed several issues with transform widget rendering, including correct scaling with viewport size and DPI, as well as fixing frame lag when in azimuthal mode when moving the view around.

diff --git a/Engine/Plugins/Runtime/nDisplay/Shaders/Private/MeshProjectionNormalSmoothing.usf b/Engine/Plugins/Runtime/nDisplay/Shaders/Private/MeshProjectionNormalSmoothing.usf
--- a/Engine/Plugins/Runtime/nDisplay/Shaders/Private/MeshProjectionNormalSmoothing.usf
+++ b/Engine/Plugins/Runtime/nDisplay/Shaders/Private/MeshProjectionNormalSmoothing.usf
@@ -22,10 +22,10 @@ float2 SampleDirection;
 float2 SampleOffsetScale;
 float InvSigma;
 
-/** Wraps the specified pixel position around the input viewport's extent */
-uint2 WrapPixelPos(uint2 PixelPos)
+/** Adjusts the specified pixel position so as not to be outside of the viewport bounds */
+uint2 AdjustPixelPos(uint2 PixelPos)
 {
-	return uint2(PixelPos.x % (uint)Input_Extent.x, PixelPos.y % (uint)Input_Extent.y);
+	return uint2(clamp(PixelPos.x, 0.0, Input_Extent.x), clamp(PixelPos.y, 0.0, Input_Extent.y));
 }
 
 /** Converts a vector in normal space to color space */
@@ -98,19 +98,19 @@ void DepthDilationCS(uint2 DispatchThreadId : SV_DispatchThreadID)
 	const int kRadius = (FILTER_KERNEL_SIZE - 1) / 2;
 
 	float InitialDepth = RWDepth[PixelPos];
-	uint Stencil = RWStencil[PixelPos];
+	uint InitialStencil = RWStencil[PixelPos];
 
 	float FinalDepth = 0.0;
 	float W = 0.0;
 	
 	// Only dilate onto geometry; if the current pixel's stencil is zero, the dilation can be skipped
-	if (Stencil > 0)
+	if (InitialStencil > 0)
 	{
 		UNROLL
 		for (int kIndex = 0; kIndex < FILTER_KERNEL_SIZE; ++kIndex)
 		{
 			int2 PixelOffset = int2(kIndex - kRadius, kIndex - kRadius) * SampleOffsetScale * SampleDirection;
-			uint2 SamplePos = WrapPixelPos(PixelPos + PixelOffset);
+			uint2 SamplePos = AdjustPixelPos(PixelPos + PixelOffset);
 			
 			float Depth = RWDepth[SamplePos];
 			uint Stencil = RWStencil[SamplePos];
@@ -157,7 +157,7 @@ void DilationCS(uint2 DispatchThreadId : SV_DispatchThreadID)
 		for (int kIndex = 0; kIndex < FILTER_KERNEL_SIZE; ++kIndex)
 		{
 			int2 PixelOffset = int2(kIndex - kRadius, kIndex - kRadius) * SampleOffsetScale * SampleDirection;
-			uint2 SamplePos = WrapPixelPos(PixelPos + PixelOffset);
+			uint2 SamplePos = AdjustPixelPos(PixelPos + PixelOffset);
 		
 			float3 Color = RWColor[SamplePos].rgb;
 			float Depth = RWDepth[SamplePos];
@@ -206,7 +206,7 @@ void BlurCS(uint2 DispatchThreadId : SV_DispatchThreadID)
 	for (int kIndex = 0; kIndex < FILTER_KERNEL_SIZE; ++kIndex)
 	{
 		int2 PixelOffset = int2(kIndex - kRadius, kIndex - kRadius) * SampleOffsetScale * SampleDirection;
-		uint2 SamplePos = WrapPixelPos(PixelPos + PixelOffset);
+		uint2 SamplePos = AdjustPixelPos(PixelPos + PixelOffset);
 		
 		float3 Normal = NormalFromColorSpace(RWColor[SamplePos].rgb);
 		float Depth = RWDepth[SamplePos];
diff --git a/Engine/Plugins/Runtime/nDisplay/Source/DisplayCluster/Private/Game/EngineClasses/Scene/DisplayClusterLightCardActor.cpp b/Engine/Plugins/Runtime/nDisplay/Source/DisplayCluster/Private/Game/EngineClasses/Scene/DisplayClusterLightCardActor.cpp
--- a/Engine/Plugins/Runtime/nDisplay/Source/DisplayCluster/Private/Game/EngineClasses/Scene/DisplayClusterLightCardActor.cpp
+++ b/Engine/Plugins/Runtime/nDisplay/Source/DisplayCluster/Private/Game/EngineClasses/Scene/DisplayClusterLightCardActor.cpp
@@ -8,6 +8,8 @@
 #include "Engine/StaticMesh.h"
 #include "UObject/ConstructorHelpers.h"
 
+const FRotator ADisplayClusterLightCardActor::PlaneMeshRotation = FRotator(0.0f, -90.0f, 90.0f);
+
 ADisplayClusterLightCardActor::ADisplayClusterLightCardActor(const FObjectInitializer& ObjectInitializer)
 	: Super(ObjectInitializer)
 	, DistanceFromCenter(300.f)
@@ -130,24 +132,46 @@ FTransform ADisplayClusterLightCardActor::GetLightCardTransform(bool bIgnoreSpin
 
 	Transform.SetLocation(LightCardComponent->GetComponentLocation());
 
-	FRotator LightCardRotation = LightCardTransformerComponent->GetComponentRotation();
+	FQuat LightCardOrientation;
 
-	if (bIgnoreSpinYawPitch)
+	if (!bIgnoreSpinYawPitch)
+	{
+		// Use the light card component's orientation, but remove the plane mesh rotation so that the returned transform's local x axis
+		// points radially inwards to match engine convention
+		LightCardOrientation = LightCardComponent->GetComponentQuat() * PlaneMeshRotation.Quaternion().Inverse();
+	}
+	else
 	{
-		LightCardRotation -= FRotator(Spin, -Yaw, -Pitch);
+		LightCardOrientation = LightCardTransformerComponent->GetComponentQuat();
 	}
 
-	Transform.SetRotation(LightCardRotation.Quaternion());
+	Transform.SetRotation(LightCardOrientation);
 
 	return Transform;
 }
 
+FBox ADisplayClusterLightCardActor::GetLightCardBounds(bool bLocalSpace) const
+{
+	FTransform ObjectOrientedTransform;
+	ObjectOrientedTransform.SetRotation(PlaneMeshRotation.Quaternion());
+
+	if (!bLocalSpace)
+	{
+		ObjectOrientedTransform.SetTranslation(LightCardComponent->GetComponentLocation());
+		ObjectOrientedTransform.SetScale3D(LightCardComponent->GetComponentScale());
+	}
+
+	return LightCardComponent->CalcBounds(ObjectOrientedTransform).GetBox();
+}
+
 void ADisplayClusterLightCardActor::UpdateLightCardTransform()
 {
 	MainSpringArmComponent->TargetArmLength = DistanceFromCenter;
 	MainSpringArmComponent->SetRelativeRotation(FRotator(-Latitude, Longitude, 0.0));
 
-	LightCardComponent->SetRelativeRotation(FRotator(-Spin, -90.f + Yaw, 90.f + Pitch));
+	FRotator LightCardOrientation = FRotator(-Pitch, Yaw, Spin);
+
+	LightCardComponent->SetRelativeRotation((LightCardOrientation.Quaternion() * PlaneMeshRotation.Quaternion()).Rotator());
 	LightCardComponent->SetRelativeScale3D(FVector(Scale, 1.f));
 }
 
diff --git a/Engine/Plugins/Runtime/nDisplay/Source/DisplayCluster/Public/DisplayClusterLightCardActor.h b/Engine/Plugins/Runtime/nDisplay/Source/DisplayCluster/Public/DisplayClusterLightCardActor.h
--- a/Engine/Plugins/Runtime/nDisplay/Source/DisplayCluster/Public/DisplayClusterLightCardActor.h
+++ b/Engine/Plugins/Runtime/nDisplay/Source/DisplayCluster/Public/DisplayClusterLightCardActor.h
@@ -25,6 +25,10 @@ class DISPLAYCLUSTER_API ADisplayClusterLightCardActor : public AActor
 {
 	GENERATED_BODY()
 
+public:
+	/** The rotation used to orient the plane mesh used for the light card so that its normal points radially inwards */
+	static const FRotator PlaneMeshRotation;
+
 public:
 	ADisplayClusterLightCardActor(const FObjectInitializer& ObjectInitializer);
 
@@ -42,6 +46,9 @@ public:
 	 */
 	FTransform GetLightCardTransform(bool bIgnoreSpinYawPitch = false) const;
 
+	/** Gets the object oriented bounding box of the light card component */
+	FBox GetLightCardBounds(bool bLocalSpace = false) const;
+
 	/** Returns the current static mesh used by this light card */
 	UStaticMesh* GetStaticMesh() const;
 
diff --git a/Engine/Plugins/Runtime/nDisplay/Source/DisplayClusterLightCardEditor/Private/Viewport/DisplayClusterLightCardEditorViewportClient.cpp b/Engine/Plugins/Runtime/nDisplay/Source/DisplayClusterLightCardEditor/Private/Viewport/DisplayClusterLightCardEditorViewportClient.cpp
--- a/Engine/Plugins/Runtime/nDisplay/Source/DisplayClusterLightCardEditor/Private/Viewport/DisplayClusterLightCardEditorViewportClient.cpp
+++ b/Engine/Plugins/Runtime/nDisplay/Source/DisplayClusterLightCardEditor/Private/Viewport/DisplayClusterLightCardEditorViewportClient.cpp
@@ -91,10 +91,10 @@ bool FDisplayClusterLightCardEditorViewportClient::FNormalMap::GetNormalAndDista
 {
 	auto GetPixel = [this](uint32 InX, uint32 InY)
 	{
-		uint32 WrappedX = InX % SizeX;
-		uint32 WrappedY = InY % SizeY;
+		uint32 ClampedX = FMath::Clamp(InX, (uint32)0, SizeX);
+		uint32 ClampedY = FMath::Clamp(InY, (uint32)0, SizeY);
 
-		return CachedNormalData[WrappedY * SizeX + WrappedX].GetFloats();
+		return CachedNormalData[ClampedY * SizeX + ClampedX].GetFloats();
 	};
 
 	const FVector ViewPos = FVector(ViewMatrices.GetViewMatrix().TransformFVector4(FVector4(Position, 1)));
@@ -240,7 +240,7 @@ void FDisplayClusterLightCardEditorViewportClient::Tick(float DeltaSeconds)
 	// Camera position is locked to a specific location
 	ResetCamera(/* bLocationOnly */ true);
 
-	CalcEditorWidgetTransform(CachedEditorWidgetTransformBeforeMapProjection, CachedEditorWidgetTransformAfterMapProjection);
+	CalcEditorWidgetTransform(CachedEditorWidgetWorldTransform);
 
 	// Tick the preview scene world.
 	if (!GIntraFrameDebuggingGameThread)
@@ -293,8 +293,14 @@ void FDisplayClusterLightCardEditorViewportClient::Tick(float DeltaSeconds)
 		const FRotator NewRotation = FMath::RInterpTo(GetViewRotation(), LookAtRotation, DeltaSeconds, DesiredLookAtSpeed);
 		SetViewRotation(NewRotation);
 
+		FSceneViewInitOptions SceneVewInitOptions;
+		GetSceneViewInitOptions(SceneVewInitOptions);
+		FViewMatrices ViewMatrices(SceneVewInitOptions);
+
+		FVector ProjectedEditorWidgetPosition = ProjectWorldPosition(CachedEditorWidgetWorldTransform.GetTranslation(), ViewMatrices);
+
 		if (NewRotation.Equals(LookAtRotation, 2.f) ||
-			!IsLocationCloseToEdge(CachedEditorWidgetTransformAfterMapProjection.GetTranslation()))
+			!IsLocationCloseToEdge(ProjectedEditorWidgetPosition))
 		{
 			DesiredLookAtLocation.Reset();
 		}
@@ -491,8 +497,19 @@ void FDisplayClusterLightCardEditorViewportClient::Draw(const FSceneView* View,
 {
 	if (SelectedLightCards.Num())
 	{
-		EditorWidget->SetTransform(CachedEditorWidgetTransformAfterMapProjection);
-		EditorWidget->Draw(View, PDI);
+		// Project the editor widget's world position into projection space so that it renders at the appropriate screen location.
+		// This needs to be computed on the render thread using the render thread's scene view, which will be behind the game thread's scene view
+		// by at least one frame
+
+		FTransform ProjectedEditorWidgetTransform(CachedEditorWidgetWorldTransform);
+
+		if (ProjectionMode != EDisplayClusterMeshProjectionType::Perspective)
+		{
+			ProjectedEditorWidgetTransform.SetTranslation(ProjectWorldPosition(CachedEditorWidgetWorldTransform.GetTranslation(), View->ViewMatrices));
+		}
+
+		EditorWidget->SetTransform(ProjectedEditorWidgetTransform);
+		EditorWidget->Draw(View, this, PDI);
 	}
 }
 
@@ -557,9 +574,19 @@ bool FDisplayClusterLightCardEditorViewportClient::InputWidgetDelta(FViewport* I
 	{
 		if (CurrentAxis != EAxisList::Type::None && SelectedLightCards.Num())
 		{
-			MoveSelectedLightCards(InViewport, CurrentAxis);
+			if (EditorWidget->GetWidgetMode() == FDisplayClusterLightCardEditorWidget::EWidgetMode::WM_Translate)
+			{
+				MoveSelectedLightCards(InViewport, CurrentAxis);
+			}
+			else if (EditorWidget->GetWidgetMode() == FDisplayClusterLightCardEditorWidget::EWidgetMode::WM_Scale)
+			{
+				ScaleSelectedLightCards(InViewport, CurrentAxis);
+			}
+
 			bHandled = true;
 		}
+
+		InViewport->GetMousePos(LastWidgetMousePos);
 	}
 
 	return bHandled;
@@ -599,7 +626,8 @@ void FDisplayClusterLightCardEditorViewportClient::TrackingStarted(const FInputE
 		FVector Direction;
 		PixelToWorld(*View, MousePos, Origin, Direction);
 
-		DragWidgetOffset = Direction - (CachedEditorWidgetTransformBeforeMapProjection.GetTranslation() - Origin).GetSafeNormal();
+		DragWidgetOffset = Direction - (CachedEditorWidgetWorldTransform.GetTranslation() - Origin).GetSafeNormal();
+		LastWidgetMousePos = MousePos;
 	}
 
 	FEditorViewportClient::TrackingStarted(InInputState, bIsDraggingWidget, bNudge);
@@ -1436,6 +1464,7 @@ void FDisplayClusterLightCardEditorViewportClient::PropagateLightCardTransform(A
 		TryChangeProperty(GET_MEMBER_NAME_CHECKED(ADisplayClusterLightCardActor, Spin));
 		TryChangeProperty(GET_MEMBER_NAME_CHECKED(ADisplayClusterLightCardActor, Pitch));
 		TryChangeProperty(GET_MEMBER_NAME_CHECKED(ADisplayClusterLightCardActor, Yaw));
+		TryChangeProperty(GET_MEMBER_NAME_CHECKED(ADisplayClusterLightCardActor, Scale));
 		
 		// Snapshot the changed properties so multi-user can update while dragging.
 		if (ChangedProperties.Num() > 0)
@@ -1482,9 +1511,6 @@ void FDisplayClusterLightCardEditorViewportClient::VerifyAndFixLightCardOrigin(A
 
 void FDisplayClusterLightCardEditorViewportClient::MoveSelectedLightCards(FViewport* InViewport, EAxisList::Type CurrentAxis)
 {
-	UWorld* PreviewWorld = PreviewScene->GetWorld();
-	check(PreviewWorld);
-
 	FIntPoint MousePos;
 	InViewport->GetMousePos(MousePos);
 
@@ -1499,9 +1525,6 @@ void FDisplayClusterLightCardEditorViewportClient::MoveSelectedLightCards(FViewp
 	FVector Direction;
 	PixelToWorld(*View, MousePos, Origin, Direction);
 
-	const FVector CursorRayStart = Origin;
-	const FVector CursorRayEnd = CursorRayStart + Direction * HALF_WORLD_MAX;
-
 	const TWeakObjectPtr<ADisplayClusterLightCardActor>& LastSelectedLightCard = SelectedLightCards.Last();
 	if (LastSelectedLightCard.IsValid())
 	{
@@ -1534,12 +1557,16 @@ void FDisplayClusterLightCardEditorViewportClient::MoveSelectedLightCards(FViewp
 			PropagateLightCardTransform(LightCard.Get());
 		}
 
-		FVector LightCardWorldLocation = CachedEditorWidgetTransformAfterMapProjection.GetTranslation();
+		FSceneViewInitOptions SceneVewInitOptions;
+		GetSceneViewInitOptions(SceneVewInitOptions);
+		FViewMatrices ViewMatrices(SceneVewInitOptions);
+
+		FVector LightCardProjectedLocation = ProjectWorldPosition(CachedEditorWidgetWorldTransform.GetTranslation(), ViewMatrices);
 		FVector2D ScreenPercentage;
-		if (IsLocationCloseToEdge(LightCardWorldLocation, InViewport, View, &ScreenPercentage))
+		if (IsLocationCloseToEdge(LightCardProjectedLocation, InViewport, View, &ScreenPercentage))
 		{
 			DesiredLookAtSpeed = FMath::Max(ScreenPercentage.X, ScreenPercentage.Y) * MaxDesiredLookAtSpeed;
-			DesiredLookAtLocation = LightCardWorldLocation;
+			DesiredLookAtLocation = LightCardProjectedLocation;
 		}
 		else
 		{
@@ -1613,6 +1640,91 @@ FDisplayClusterLightCardEditorViewportClient::FSphericalCoordinates FDisplayClus
 	return DeltaCoords;
 }
 
+void FDisplayClusterLightCardEditorViewportClient::ScaleSelectedLightCards(FViewport* InViewport, EAxisList::Type CurrentAxis)
+{
+	const TWeakObjectPtr<ADisplayClusterLightCardActor>& LastSelectedLightCard = SelectedLightCards.Last();
+	if (LastSelectedLightCard.IsValid())
+	{
+		const FVector2D DeltaScale = GetLightCardScaleDelta(InViewport, LastSelectedLightCard.Get(), CurrentAxis);
+		for (const TWeakObjectPtr<ADisplayClusterLightCardActor>& LightCard : SelectedLightCards)
+		{
+			if (!LightCard.IsValid())
+			{
+				continue;
+			}
+
+			LightCard->Scale += DeltaScale;
+
+			PropagateLightCardTransform(LightCard.Get());
+		}
+	}
+}
+
+FVector2D FDisplayClusterLightCardEditorViewportClient::GetLightCardScaleDelta(FViewport* InViewport, ADisplayClusterLightCardActor* LightCard, EAxisList::Type CurrentAxis)
+{
+	FIntPoint MousePos;
+	InViewport->GetMousePos(MousePos);
+
+	FSceneViewFamilyContext ViewFamily(FSceneViewFamily::ConstructionValues(
+		InViewport,
+		GetScene(),
+		EngineShowFlags)
+		.SetRealtimeUpdate(IsRealtime()));
+	FSceneView* View = CalcSceneView(&ViewFamily);
+
+	const FVector2D DragDir = MousePos - LastWidgetMousePos;
+
+	const FVector WorldWidgetOrigin = CachedEditorWidgetWorldTransform.GetTranslation();
+	const FVector WorldWidgetXAxis = CachedEditorWidgetWorldTransform.GetRotation().RotateVector(FVector::XAxisVector);
+	const FVector WorldWidgetYAxis = CachedEditorWidgetWorldTransform.GetRotation().RotateVector(FVector::YAxisVector);
+
+	FVector2D ScreenXAxis = FVector2D::ZeroVector;
+	FVector2D ScreenYAxis = FVector2D::ZeroVector;
+
+	WorldToScreenDirection(*View, WorldWidgetOrigin, WorldWidgetXAxis, ScreenXAxis);
+	WorldToScreenDirection(*View, WorldWidgetOrigin, WorldWidgetYAxis, ScreenYAxis);
+
+	FVector2D ScaleDir = FVector2D::ZeroVector;
+
+	if (CurrentAxis & EAxisList::Type::X)
+	{
+		ScaleDir += ScreenXAxis;
+	}
+
+	if (CurrentAxis & EAxisList::Type::Y)
+	{
+		ScaleDir += ScreenYAxis;
+	}
+
+	ScaleDir.Normalize();
+
+	// To give the scale delta a nice feel to it when dragging, we want the distance the user is dragging the mouse to be proportional to the change in size
+	// of the light card when the scale delta is applied. So, if the user moves the mouse a distance d, then the light card's bounds along the direction of scale
+	// should change by 2d (one d for each side), and the new scale should be s' = (h + 2d) / h_0, where h is the current length of the side and h_0 is the unscaled
+	// length of the side. Since the current scale s = h / h_0, this means that the scale delta is s' - s = 2d / h_0.
+
+	// First, obtain the size of the unscaled light card in the direction the user is scaling. Convert to screen space as the scale and drag vectors are in screen space
+	const bool bLocalSpace = true;
+	const FVector LightCardSize3D = LightCard->GetLightCardBounds(bLocalSpace).GetSize();
+	const FVector2D SizeToScale = FVector2D((CurrentAxis & EAxisList::X) * LightCardSize3D.Y, (CurrentAxis & EAxisList::Y) * LightCardSize3D.Z);
+	const float DistanceFromCamera = FMath::Max(FVector::Dist(WorldWidgetOrigin, View->ViewMatrices.GetViewOrigin()), 1.0f);
+	const float ScreenSize = View->ViewMatrices.GetScreenScale() * SizeToScale.Length() / DistanceFromCamera;
+
+	// Compute the scale delta as s' - s = 2d / h_0
+	const float ScaleMagnitude = 2.0f * (ScaleDir | DragDir) / ScreenSize;
+
+	FVector2D ScaleDelta = FVector2D((CurrentAxis & EAxisList::X) * ScaleMagnitude, (CurrentAxis & EAxisList::Y) * ScaleMagnitude);
+
+	// If both axes are being scaled at the same time, preserve the aspect ratio of the scale delta
+	if (CurrentAxis == EAxisList::Type::XY)
+	{
+		// Ensure the signs of the deltas remain the same, and avoid potential divide by zero
+		ScaleDelta.Y = ScaleDelta.X * FMath::Abs(LightCard->Scale.Y) / FMath::Max(0.001, FMath::Abs(LightCard->Scale.X));
+	}
+
+	return ScaleDelta;
+}
+
 FDisplayClusterLightCardEditorViewportClient::FSphericalCoordinates FDisplayClusterLightCardEditorViewportClient::GetLightCardCoordinates(ADisplayClusterLightCardActor* LightCard) const
 {
 	const FVector LightCardLocation = LightCard->GetLightCardTransform().GetTranslation() - LightCard->GetActorLocation();
@@ -1699,6 +1811,20 @@ ADisplayClusterLightCardActor* FDisplayClusterLightCardEditorViewportClient::Tra
 	return nullptr;
 }
 
+FVector FDisplayClusterLightCardEditorViewportClient::ProjectWorldPosition(const FVector& UnprojectedWorldPosition, const FViewMatrices& ViewMatrices) const
+{
+	FVector ProjectedPosition(UnprojectedWorldPosition);
+
+	if (ProjectionMode != EDisplayClusterMeshProjectionType::Perspective)
+	{
+		const FVector ViewPos = ViewMatrices.GetViewMatrix().TransformPosition(CachedEditorWidgetWorldTransform.GetTranslation());
+		const FVector ProjectedViewPos = FDisplayClusterMeshProjectionRenderer::ProjectViewPosition(ViewPos, ProjectionMode);
+		ProjectedPosition = ViewMatrices.GetInvViewMatrix().TransformPosition(ProjectedViewPos);
+	}
+
+	return ProjectedPosition;
+}
+
 void FDisplayClusterLightCardEditorViewportClient::PixelToWorld(const FSceneView& View, const FIntPoint& PixelPos, FVector& OutOrigin, FVector& OutDirection)
 {
 	const FMatrix& InvProjMatrix = View.ViewMatrices.GetInvProjectionMatrix();
@@ -1712,7 +1838,48 @@ void FDisplayClusterLightCardEditorViewportClient::PixelToWorld(const FSceneView
 	OutDirection = InvViewMatrix.TransformVector(UnprojectedViewPos).GetSafeNormal();
 }
 
-bool FDisplayClusterLightCardEditorViewportClient::CalcEditorWidgetTransform(FTransform& WidgetTransformBeforeMapProjection, FTransform& WidgetTransformAfterMapProjection)
+bool FDisplayClusterLightCardEditorViewportClient::WorldToPixel(const FSceneView& View, const FVector& WorldPos, FVector2D& OutPixelPos) const
+{
+	const FMatrix& ViewMatrix = View.ViewMatrices.GetViewMatrix();
+	const FMatrix& ProjMatrix = View.ViewMatrices.GetProjectionMatrix();
+
+	const FVector ViewPos = ViewMatrix.TransformPosition(WorldPos);
+	const FVector ProjectedViewPos = FDisplayClusterMeshProjectionRenderer::ProjectViewPosition(ViewPos, ProjectionMode);
+	const FVector4 ScreenPos = ProjMatrix.TransformFVector4(FVector4(ProjectedViewPos, 1));
+
+	return View.ScreenToPixel(ScreenPos, OutPixelPos);
+}
+
+bool FDisplayClusterLightCardEditorViewportClient::WorldToScreenDirection(const FSceneView& View, const FVector& WorldPos, const FVector& WorldDirection, FVector2D& OutScreenDir)
+{
+	FVector2D ScreenVectorStart = FVector2D::ZeroVector;
+	FVector2D ScreenVectorEnd = FVector2D::ZeroVector;
+
+	if (WorldToPixel(View, WorldPos, ScreenVectorStart) && WorldToPixel(View, WorldPos + WorldDirection, ScreenVectorEnd))
+	{
+		OutScreenDir = (ScreenVectorEnd - ScreenVectorStart).GetSafeNormal();
+		return true;
+	}
+	else
+	{
+		// If either the start or end of the vector is not onscreen, translate the vector to be in front of the camera to approximate the screen space direction
+		const FMatrix InvViewMatrix = View.ViewMatrices.GetInvViewMatrix();
+		const FVector ViewLocation = InvViewMatrix.GetOrigin();
+		const FVector ViewDirection = InvViewMatrix.GetUnitAxis(EAxis::Z);
+		const FVector Offset = ViewDirection * (FVector::DotProduct(ViewLocation - WorldPos, ViewDirection) + 100.0f);
+		const FVector AdjustedWorldPos = WorldPos + Offset;
+
+		if (WorldToPixel(View, AdjustedWorldPos, ScreenVectorStart) && WorldToPixel(View, AdjustedWorldPos + WorldDirection, ScreenVectorEnd))
+		{
+			OutScreenDir = -(ScreenVectorEnd - ScreenVectorStart).GetSafeNormal();
+			return true;
+		}
+	}
+
+	return false;
+}
+
+bool FDisplayClusterLightCardEditorViewportClient::CalcEditorWidgetTransform(FTransform& WidgetTransformBeforeMapProjection)
 {
 	if (!SelectedLightCards.Num())
 	{
@@ -1728,30 +1895,32 @@ bool FDisplayClusterLightCardEditorViewportClient::CalcEditorWidgetTransform(FTr
 
 	FVector LightCardPosition = LastSelected->GetLightCardTransform().GetTranslation();
 
-	WidgetTransformAfterMapProjection = WidgetTransformBeforeMapProjection = FTransform(FRotator::ZeroRotator, LightCardPosition, FVector::OneVector);
+	WidgetTransformBeforeMapProjection = FTransform(FRotator::ZeroRotator, LightCardPosition, FVector::OneVector);
 
-	if (ProjectionMode != EDisplayClusterMeshProjectionType::Perspective)
+	FQuat WidgetOrientation;
+	if (EditorWidget->GetWidgetMode() == FDisplayClusterLightCardEditorWidget::EWidgetMode::WM_Translate)
 	{
-		FSceneViewInitOptions SceneVewInitOptions;
-		GetSceneViewInitOptions(SceneVewInitOptions);
-		FViewMatrices ViewMatrices(SceneVewInitOptions);
-
-		const FVector ViewPos = ViewMatrices.GetViewMatrix().TransformPosition(LightCardPosition);
-		const FVector ProjectedViewPos = FDisplayClusterMeshProjectionRenderer::ProjectViewPosition(ViewPos, ProjectionMode);
-		const FVector ProjectedPosition = ViewMatrices.GetInvViewMatrix().TransformPosition(ProjectedViewPos);
+		// The translation widget should be oriented to show the x axis pointing in the longitudinal direction and the y axis pointing in the latitudinal direction
+		const FVector ProjectionOrigin = ProjectionOriginComponent.IsValid() ? ProjectionOriginComponent->GetComponentLocation() : FVector::ZeroVector;
+		const FVector RadialVector = (LightCardPosition - ProjectionOrigin).GetSafeNormal();
+		const FVector AzimuthalVector = (FVector::ZAxisVector ^ RadialVector).GetSafeNormal();
+		const FVector InclinationVector = RadialVector ^ AzimuthalVector;
 
-		WidgetTransformAfterMapProjection.SetTranslation(ProjectedPosition);
+		WidgetOrientation = FMatrix(AzimuthalVector, InclinationVector, RadialVector, FVector::ZeroVector).ToQuat();
 	}
+	else
+	{
+		// Otherwise, orient the widget to match the light card's local orientation (spin, pitch, yaw)
+		const FQuat LightCardOrientation = LastSelected->GetLightCardTransform().GetRotation();
+		const FVector Normal = LightCardOrientation.RotateVector(FVector::XAxisVector);
+		const FVector Tangent = LightCardOrientation.RotateVector(FVector::YAxisVector);
+		const FVector Binormal = LightCardOrientation.RotateVector(FVector::ZAxisVector);
 
-	const FVector ProjectionOrigin = ProjectionOriginComponent.IsValid() ? ProjectionOriginComponent->GetComponentLocation() : FVector::ZeroVector;
-	const FVector RadialVector = (LightCardPosition - ProjectionOrigin).GetSafeNormal();
-	const FVector AzimuthalVector = (FVector::ZAxisVector ^ RadialVector).GetSafeNormal();
-	const FVector InclinationVector = RadialVector ^ AzimuthalVector;
-
-	FRotator Orientation = FMatrix(AzimuthalVector, InclinationVector, RadialVector, FVector::ZeroVector).Rotator();
+		// Reorder the orientation basis so that the x axis and the y axis point along the light card's width and height, respectively
+		WidgetOrientation = FMatrix(-Tangent, Binormal, -Normal, FVector::ZeroVector).ToQuat();
+	}
 
-	WidgetTransformBeforeMapProjection.SetRotation(Orientation.Quaternion());
-	WidgetTransformAfterMapProjection.SetRotation(Orientation.Quaternion());
+	WidgetTransformBeforeMapProjection.SetRotation(WidgetOrientation);
 
 	return true;
 }
diff --git a/Engine/Plugins/Runtime/nDisplay/Source/DisplayClusterLightCardEditor/Private/Viewport/DisplayClusterLightCardEditorViewportClient.h b/Engine/Plugins/Runtime/nDisplay/Source/DisplayClusterLightCardEditor/Private/Viewport/DisplayClusterLightCardEditorViewportClient.h
--- a/Engine/Plugins/Runtime/nDisplay/Source/DisplayClusterLightCardEditor/Private/Viewport/DisplayClusterLightCardEditorViewportClient.h
+++ b/Engine/Plugins/Runtime/nDisplay/Source/DisplayClusterLightCardEditor/Private/Viewport/DisplayClusterLightCardEditorViewportClient.h
@@ -10,11 +10,11 @@
 #include "DisplayClusterLightCardEditorProxyType.h"
 #include "DisplayClusterMeshProjectionRenderer.h"
 #include "DisplayClusterLightCardActor.h"
+#include "DisplayClusterLightCardEditorWidget.h"
 
 class ADisplayClusterRootActor;
 class SDisplayClusterLightCardEditor;
 class FScopedTransaction;
-class FDisplayClusterLightCardEditorWidget;
 class UDisplayClusterConfigurationViewport;
 
 /** Viewport Client for the preview viewport */
@@ -161,6 +161,9 @@ public:
 	/** Selects the light card proxies that correspond to the specified light cards */
 	void SelectLightCards(const TArray<AActor*>& LightCardsToSelect);
 
+	FDisplayClusterLightCardEditorWidget::EWidgetMode GetEditorWidgetMode() const { return EditorWidget->GetWidgetMode(); }
+	void SetEditorWidgetMode(FDisplayClusterLightCardEditorWidget::EWidgetMode InWidgetMode) { EditorWidget->SetWidgetMode(InWidgetMode); }
+
 	EDisplayClusterMeshProjectionType GetProjectionMode() const { return ProjectionMode; }
 	void SetProjectionMode(EDisplayClusterMeshProjectionType InProjectionMode);
 
@@ -229,17 +232,32 @@ private:
 	/** Determines the appropriate delta in spherical coordinates needed to move the specified light card to the mouse's location */
 	FSphericalCoordinates GetLightCardTranslationDelta(FViewport* InViewport, ADisplayClusterLightCardActor* LightCard, EAxisList::Type CurrentAxis);
 
+	/** Scales the currently selected light cards */
+	void ScaleSelectedLightCards(FViewport* InViewport, EAxisList::Type CurrentAxis);
+
+	/** Determines the appropriate scale delta needed to scale the light card */
+	FVector2D GetLightCardScaleDelta(FViewport* InViewport, ADisplayClusterLightCardActor* LightCard, EAxisList::Type CurrentAxis);
+
 	/** Gets the spherical coordinates of the specified light card */
 	FSphericalCoordinates GetLightCardCoordinates(ADisplayClusterLightCardActor* LightCard) const;
 
 	/** Traces to find the light card corresponding to a click on a stage screen */
 	ADisplayClusterLightCardActor* TraceScreenForLightCard(const FSceneView& View, int32 HitX, int32 HitY);
 
+	/** Projects the specified world position to the viewport's current projection space */
+	FVector ProjectWorldPosition(const FVector& UnprojectedWorldPosition, const FViewMatrices& ViewMatrices) const;
+
 	/** Converts a pixel coordinate into a point and direction vector in world space */
 	void PixelToWorld(const FSceneView& View, const FIntPoint& PixelPos, FVector& OutOrigin, FVector& OutDirection);
 
+	/** Converts a world coordinate into a point in screen space, and returns true if the world position is on the screen */
+	bool WorldToPixel(const FSceneView& View, const FVector& WorldPos, FVector2D& OutPixelPos) const;
+
+	/** Converts a direction vector from world space to screen screen space, and returns true of the direction vector is on the screen */
+	bool WorldToScreenDirection(const FSceneView& View, const FVector& WorldPos, const FVector& WorldDirection, FVector2D& OutScreenDir);
+
 	/** Calculates the world transform to render the editor widget with to align it with the selected light card */
-	bool CalcEditorWidgetTransform(FTransform& WidgetTransformBeforeMapProjection, FTransform& WidgetTransformAfterMapProjection);
+	bool CalcEditorWidgetTransform(FTransform& WidgetTransformBeforeMapProjection);
 	
 	/** Renders the viewport's normal map and stores the texture data to be used later */
 	void RenderNormalMap(FNormalMap& NormalMap, const FVector& NormalMapDirection);
@@ -287,15 +305,15 @@ private:
 	/** The LC editor widget used to manipulate light cards */
 	TSharedPtr<FDisplayClusterLightCardEditorWidget> EditorWidget;
 
-	/** The cached editor widget transform after map projection calculated for the editor widget on the last tick */
-	FTransform CachedEditorWidgetTransformAfterMapProjection;
-
-	/** The cached editor widget transform before map projection calculated for the editor widget on the last tick */
-	FTransform CachedEditorWidgetTransformBeforeMapProjection;
+	/** The cached editor widget transform in unprojected world space */
+	FTransform CachedEditorWidgetWorldTransform;
 
 	/** The offset between the widget's origin and the place it was clicked when a drag action was started */
 	FVector DragWidgetOffset;
 
+	/** The mouse position registered on the previous widget input pass */
+	FIntPoint LastWidgetMousePos;
+
 	/** The location the camera should be looking at */
 	TOptional<FVector> DesiredLookAtLocation;
 
diff --git a/Engine/Plugins/Runtime/nDisplay/Source/DisplayClusterLightCardEditor/Private/Viewport/DisplayClusterLightCardEditorWidget.cpp b/Engine/Plugins/Runtime/nDisplay/Source/DisplayClusterLightCardEditor/Private/Viewport/DisplayClusterLightCardEditorWidget.cpp
--- a/Engine/Plugins/Runtime/nDisplay/Source/DisplayClusterLightCardEditor/Private/Viewport/DisplayClusterLightCardEditorWidget.cpp
+++ b/Engine/Plugins/Runtime/nDisplay/Source/DisplayClusterLightCardEditor/Private/Viewport/DisplayClusterLightCardEditorWidget.cpp
@@ -4,15 +4,19 @@
 
 #include "SceneManagement.h"
 #include "UnrealWidget.h"
+#include "EditorViewportClient.h"
 
-void FDisplayClusterLightCardEditorWidget::Draw(const FSceneView* View, FPrimitiveDrawInterface* PDI)
+void FDisplayClusterLightCardEditorWidget::Draw(const FSceneView* View, const FEditorViewportClient* ViewportClient, FPrimitiveDrawInterface* PDI)
 {
-	DrawAxis(View, PDI, EAxisList::Type::X);
-	DrawAxis(View, PDI, EAxisList::Type::Y);
-	DrawOrigin(View, PDI);
+	const float SizeScalar = GetSizeScreenScalar(View, ViewportClient);
+	const float LengthScalar = GetLengthScreenScalar(View, ViewportClient, Transform.GetTranslation());
+
+	DrawAxis(View, PDI, EAxisList::Type::X, SizeScalar, LengthScalar);
+	DrawAxis(View, PDI, EAxisList::Type::Y, SizeScalar, LengthScalar);
+	DrawOrigin(PDI, SizeScalar);
 }
 
-void FDisplayClusterLightCardEditorWidget::DrawAxis(const FSceneView* View, FPrimitiveDrawInterface* PDI, EAxisList::Type Axis)
+void FDisplayClusterLightCardEditorWidget::DrawAxis(const FSceneView* View, FPrimitiveDrawInterface* PDI, EAxisList::Type Axis, float SizeScalar, float LengthScalar)
 {
 	const FVector Origin = Transform.GetTranslation();
 
@@ -20,26 +24,28 @@ void FDisplayClusterLightCardEditorWidget::DrawAxis(const FSceneView* View, FPri
 	const FVector TransformedAxis = Transform.TransformVector(GlobalAxis);
 	const FLinearColor AxisColor = HighlightedAxis == Axis ? HighlightColor : GetAxisColor(Axis);
 
-	// Compute the scalar needed to make the axis length a fixed size on the screen regardless of distance and field of view
-	const float DistanceScreenSpaceScalar = WidgetScale * FMath::Max(FVector::Dist(Origin, View->ViewMatrices.GetViewOrigin()), 1.f) / View->ViewMatrices.GetScreenScale();
-	const float ThicknessScreenSpaceScalar = 100.f / View->ViewMatrices.GetScreenScale();
+	PDI->SetHitProxy(new HWidgetAxis(Axis));
+	{
+		const FVector AxisStart = Origin;
+		const FVector AxisEnd = Origin + TransformedAxis * AxisLength * WidgetScale * LengthScalar;
 
+		PDI->DrawLine(AxisStart, AxisEnd, AxisColor, ESceneDepthPriorityGroup::SDPG_Foreground, AxisThickness * SizeScalar, 0.0, true);
 
-	PDI->SetHitProxy(new HWidgetAxis(Axis));
-	PDI->DrawLine(Origin, Origin + TransformedAxis * AxisLength * DistanceScreenSpaceScalar, AxisColor, ESceneDepthPriorityGroup::SDPG_Foreground, AxisThickness * ThicknessScreenSpaceScalar, 0.0, true);
+		if (WidgetMode == EWidgetMode::WM_Scale)
+		{
+			PDI->DrawPoint(AxisEnd, AxisColor, AxisCapSize * SizeScalar, ESceneDepthPriorityGroup::SDPG_Foreground);
+		}
+	}
 	PDI->SetHitProxy(nullptr);
 }
 
-void FDisplayClusterLightCardEditorWidget::DrawOrigin(const FSceneView* View, FPrimitiveDrawInterface* PDI)
+void FDisplayClusterLightCardEditorWidget::DrawOrigin(FPrimitiveDrawInterface* PDI, float SizeScalar)
 {
 	const FVector Origin = Transform.GetTranslation();
 	const FLinearColor Color = HighlightedAxis == EAxisList::Type::XY ? HighlightColor : FLinearColor::Black;
 
-	// Compute the scalar needed to make the axis length a fixed size on the screen regardless of distance and field of view
-	const float ScreenSpaceScalar = 100.f / View->ViewMatrices.GetScreenScale();
-
 	PDI->SetHitProxy(new HWidgetAxis(EAxisList::Type::XY));
-	PDI->DrawPoint(Origin, Color, OriginSize * ScreenSpaceScalar, ESceneDepthPriorityGroup::SDPG_Foreground);
+	PDI->DrawPoint(Origin, Color, OriginSize * SizeScalar, ESceneDepthPriorityGroup::SDPG_Foreground);
 	PDI->SetHitProxy(nullptr);
 }
 
@@ -75,4 +81,47 @@ FLinearColor FDisplayClusterLightCardEditorWidget::GetAxisColor(EAxisList::Type
 	}
 
 	return FLinearColor::Transparent;
+}
+
+float FDisplayClusterLightCardEditorWidget::GetLengthScreenScalar(const FSceneView* View, const FEditorViewportClient* ViewportClient, const FVector& Origin) const
+{
+	// The ideal behavior for the length of the widget axes is to remain a fixed length regardless of the field of view or distance from the camera,
+	// but change proportionally with the size of the viewport, so that the widget takes up the same percentage of viewport space. 
+	// To accomplish this, three factors go into the size scalar: 
+	// * DPIScalar ensures the length is the same, relative to viewport size, on high density monitors
+	// * ResolutionScale ensures the length remains the same, relative to viewport size, regardless of the size of the viewport.
+	// * ProjectionScale ensures that the length remains the same regardless of field of view or distance from the camera. For lengths specifically,
+	//   the distance from the camera needs to be taken into account as well. The ViewMatrices ScreenScale contains the appropraite FOV and viewport size factors.
+	// Note we use the x axis sizes to compute the scalars as the viewport will only scale smaller when the width is resized; when the 
+	// height is resized, the view is clamped
+
+	const float DPIScale = ViewportClient->GetDPIScale();
+	const float ResolutionScale = View->UnconstrainedViewRect.Size().X / (DPIScale * ReferenceResolution);
+
+	const float DistanceFromView = FMath::Max(FVector::Dist(Origin, View->ViewMatrices.GetViewOrigin()), 1.f);
+	const float ProjectionScale = DistanceFromView / View->ViewMatrices.GetScreenScale();
+
+	const float FinalScalar = DPIScale * ResolutionScale * ProjectionScale;
+	return FinalScalar;
+}
+
+float FDisplayClusterLightCardEditorWidget::GetSizeScreenScalar(const FSceneView* View, const FEditorViewportClient* ViewportClient) const
+{
+	// The ideal behavior for the size of the widget is to remain a fixed size regardless of the field of view or distance from the camera,
+	// but change proportionally with the size of the viewport, so that the widget takes up the same percentage of viewport space. 
+	// To accomplish this, three factors go into the size scalar: 
+	// * DPIScalar ensures the widget is the same size, relative to viewport size, on high density monitors
+	// * ResolutionScale ensures the widget remains the same size, relative to viewport size, regardless of the size of the viewport.
+	// * ProjectionScale ensures that the widget remains the same size regardless of field of view or distance from the camera
+	// Note we use the x axis sizes to compute the scalars as the viewport will only scale smaller when the width is resized; when the 
+	// height is resized, the view is clamped
+
+	const float DPIScale = ViewportClient->GetDPIScale();
+	const float ResolutionScale = View->UnconstrainedViewRect.Size().X / (DPIScale * ReferenceResolution);
+
+	const FMatrix& ProjectionMatrix = View->ViewMatrices.GetProjectionMatrix();
+	const float ProjectionScale = FMath::Abs(ProjectionMatrix.M[0][0]);
+
+	const float FinalScalar = DPIScale * ResolutionScale / ProjectionScale;
+	return FinalScalar;
 }
\ No newline at end of file
diff --git a/Engine/Plugins/Runtime/nDisplay/Source/DisplayClusterLightCardEditor/Private/Viewport/DisplayClusterLightCardEditorWidget.h b/Engine/Plugins/Runtime/nDisplay/Source/DisplayClusterLightCardEditor/Private/Viewport/DisplayClusterLightCardEditorWidget.h
--- a/Engine/Plugins/Runtime/nDisplay/Source/DisplayClusterLightCardEditor/Private/Viewport/DisplayClusterLightCardEditorWidget.h
+++ b/Engine/Plugins/Runtime/nDisplay/Source/DisplayClusterLightCardEditor/Private/Viewport/DisplayClusterLightCardEditorWidget.h
@@ -6,34 +6,73 @@
 #include "Math/Axis.h"
 
 class FSceneView;
+class FEditorViewportClient;
 class FPrimitiveDrawInterface;
 
 /** Widget used to manipulate entities in the light card editor's viewport */
 class FDisplayClusterLightCardEditorWidget
 {
 public:
+	/** The default reference resolution that the size and length constants are defined relative to */
+	constexpr static float ReferenceResolution = 1920.0f;
+
+	/** The length of the widget's axes, in pixels */
 	constexpr static float AxisLength = 100.f;
-	constexpr static float AxisThickness = 35.f;
-	constexpr static float OriginSize = 200.f;
+
+	/** The thickness of the widget's axes, in pixels */
+	constexpr static float AxisThickness = 5.f;
+
+	/** The size of the axes caps */
+	constexpr static float AxisCapSize = 15.f;
+
+	/** The size of the axes origin */
+	constexpr static float OriginSize = 35.f;
+
 	constexpr static FLinearColor AxisColorX = FLinearColor(0.594f, 0.0197f, 0.0f);
 	constexpr static FLinearColor AxisColorY = FLinearColor(0.1349f, 0.3959f, 0.0f);
 	constexpr static FLinearColor AxisColorZ = FLinearColor(0.0251f, 0.207f, 0.85f);
 	constexpr static FLinearColor HighlightColor = FLinearColor(1.0f, 1.0f, 0.0f);
 
-	void Draw(const FSceneView* View, FPrimitiveDrawInterface* PDI);
+	enum EWidgetMode
+	{
+		WM_Translate,
+		WM_Scale,
+
+		WM_Max
+	};
+
+	void Draw(const FSceneView* View, const FEditorViewportClient* ViewportClient, FPrimitiveDrawInterface* PDI);
+
+	EWidgetMode GetWidgetMode() const { return WidgetMode; }
+	void SetWidgetMode(EWidgetMode NewWidgetMode) { WidgetMode = NewWidgetMode; }
 
 	void SetTransform(const FTransform& NewTransform) { Transform = FTransform(NewTransform); }
 	void SetHighlightedAxis(EAxisList::Type InAxis) { HighlightedAxis = InAxis; }
 	void SetWidgetScale(float InWidgetScale) { WidgetScale = InWidgetScale; }
 
 private:
-	void DrawAxis(const FSceneView* View, FPrimitiveDrawInterface* PDI, EAxisList::Type Axis);
-	void DrawOrigin(const FSceneView* View, FPrimitiveDrawInterface* PDI);
+	/** Draws the specified axis to the PDI, sizing and coloring it appropriately */
+	void DrawAxis(const FSceneView* View, FPrimitiveDrawInterface* PDI, EAxisList::Type Axis, float SizeScalar, float LengthScalar);
+
+	/** Draws the widget origin to the PDI, sizing and coloring it appropriately */
+	void DrawOrigin(FPrimitiveDrawInterface* PDI, float SizeScalar);
 
+	/** Gets the axis vector in global coordinates for the specified axis */
 	FVector GetGlobalAxis(EAxisList::Type Axis) const;
+
+	/** Gets the color for the specified axis */
 	FLinearColor GetAxisColor(EAxisList::Type Axis) const;
 
+	/** Calculates a scalar to use to keep lengths and lines a fixed size on the screen regardless of DPI, viewport size, FOV, or distance from the view origin */
+	float GetLengthScreenScalar(const FSceneView* View, const FEditorViewportClient* ViewportClient, const FVector& Origin) const;
+
+	/** Calculates a scalar to use to keep sizes a fixed size on the screen regardless of DPI, viewport size, FOV, or distance from the view origin */
+	float GetSizeScreenScalar(const FSceneView* View, const FEditorViewportClient* ViewportClient) const;
+
 private:
+	/** The current mode the widget is in */
+	EWidgetMode WidgetMode = WM_Translate;
+
 	/** The world transform to apply to the widget when rendering */
 	FTransform Transform;
 
diff --git a/Engine/Plugins/Runtime/nDisplay/Source/DisplayClusterLightCardEditor/Private/Viewport/DisplayClusterLightcardEditorViewport.cpp b/Engine/Plugins/Runtime/nDisplay/Source/DisplayClusterLightCardEditor/Private/Viewport/DisplayClusterLightcardEditorViewport.cpp
--- a/Engine/Plugins/Runtime/nDisplay/Source/DisplayClusterLightCardEditor/Private/Viewport/DisplayClusterLightcardEditorViewport.cpp
+++ b/Engine/Plugins/Runtime/nDisplay/Source/DisplayClusterLightCardEditor/Private/Viewport/DisplayClusterLightcardEditorViewport.cpp
@@ -59,9 +59,7 @@ public:
 				.Padding( 3.0f, 1.0f )
 				.HAlign( HAlign_Right )
 				[
-					SNew(STransformViewportToolBar)
-					.Viewport(EditorViewport.Pin().ToSharedRef())
-					.CommandList(EditorViewport.Pin()->GetCommandList())
+					MakeTransformToolBar()
 				]
 			]
 		];
@@ -147,6 +145,40 @@ public:
 		return Label;
 	}
 
+	TSharedRef<SWidget> MakeTransformToolBar()
+	{
+		FSlimHorizontalToolBarBuilder ToolbarBuilder(EditorViewport.Pin()->GetCommandList(), FMultiBoxCustomization::None);
+
+		const FName ToolBarStyle = TEXT("EditorViewportToolBar");
+		ToolbarBuilder.SetStyle(&FAppStyle::Get(), ToolBarStyle);
+		ToolbarBuilder.SetLabelVisibility(EVisibility::Collapsed);
+
+		ToolbarBuilder.SetIsFocusable(false);
+
+		ToolbarBuilder.BeginSection("Transform");
+		{
+			ToolbarBuilder.BeginBlockGroup();
+
+			//static FName SelectModeName = FName(TEXT("SelectMode"));
+			//ToolbarBuilder.AddToolBarButton(FEditorViewportCommands::Get().SelectMode, NAME_None, TAttribute<FText>(), TAttribute<FText>(), TAttribute<FSlateIcon>(), SelectModeName);
+
+			static FName TranslateModeName = FName(TEXT("TranslateMode"));
+			ToolbarBuilder.AddToolBarButton(FEditorViewportCommands::Get().TranslateMode, NAME_None, TAttribute<FText>(), TAttribute<FText>(), TAttribute<FSlateIcon>(), TranslateModeName);
+
+			//static FName RotateModeName = FName(TEXT("RotateMode"));
+			//ToolbarBuilder.AddToolBarButton( FEditorViewportCommands::Get().RotateMode, NAME_None, TAttribute<FText>(), TAttribute<FText>(), TAttribute<FSlateIcon>(), RotateModeName );
+
+			static FName ScaleModeName = FName(TEXT("ScaleMode"));
+			ToolbarBuilder.AddToolBarButton( FEditorViewportCommands::Get().ScaleMode, NAME_None, TAttribute<FText>(), TAttribute<FText>(), TAttribute<FSlateIcon>(), ScaleModeName );
+
+			ToolbarBuilder.EndBlockGroup();
+		}
+
+		ToolbarBuilder.EndSection();
+
+		return ToolbarBuilder.MakeWidget();
+	}
+
 private:
 	/** Reference to the parent viewport */
 	TWeakPtr<SDisplayClusterLightCardEditorViewport> EditorViewport;
@@ -218,25 +250,79 @@ void SDisplayClusterLightCardEditorViewport::PopulateViewportOverlays(TSharedRef
 
 void SDisplayClusterLightCardEditorViewport::BindCommands()
 {
-	const FDisplayClusterLightCardEditorCommands& Commands = FDisplayClusterLightCardEditorCommands::Get();
-
-	CommandList->MapAction(
-		FDisplayClusterLightCardEditorCommands::Get().PerspectiveProjection,
-		FExecuteAction::CreateSP(this, &SDisplayClusterLightCardEditorViewport::SetProjectionMode, EDisplayClusterMeshProjectionType::Perspective),
-		FCanExecuteAction(),
-		FIsActionChecked::CreateSP(this, &SDisplayClusterLightCardEditorViewport::IsProjectionModeSelected, EDisplayClusterMeshProjectionType::Perspective));
-
-	CommandList->MapAction(
-		FDisplayClusterLightCardEditorCommands::Get().AzimuthalProjection,
-		FExecuteAction::CreateSP(this, &SDisplayClusterLightCardEditorViewport::SetProjectionMode, EDisplayClusterMeshProjectionType::Azimuthal),
-		FCanExecuteAction(),
-		FIsActionChecked::CreateSP(this, &SDisplayClusterLightCardEditorViewport::IsProjectionModeSelected, EDisplayClusterMeshProjectionType::Azimuthal));
-
-	CommandList->MapAction(
-		Commands.ResetCamera,
-		FExecuteAction::CreateSP(ViewportClient.Get(), &FDisplayClusterLightCardEditorViewportClient::ResetCamera, false));
-	
 	SEditorViewport::BindCommands();
+
+	// Override some of the editor viewport commands to work with the light card editor viewport
+	{
+		const FEditorViewportCommands& Commands = FEditorViewportCommands::Get();
+
+		CommandList->MapAction(
+			Commands.TranslateMode,
+			FExecuteAction::CreateSP(this, &SDisplayClusterLightCardEditorViewport::SetEditorWidgetMode, FDisplayClusterLightCardEditorWidget::EWidgetMode::WM_Translate),
+			FCanExecuteAction(),
+			FIsActionChecked::CreateSP(this, &SDisplayClusterLightCardEditorViewport::IsEditorWidgetModeSelected, FDisplayClusterLightCardEditorWidget::EWidgetMode::WM_Translate)
+		);
+
+		CommandList->MapAction(
+			Commands.ScaleMode,
+			FExecuteAction::CreateSP(this, &SDisplayClusterLightCardEditorViewport::SetEditorWidgetMode, FDisplayClusterLightCardEditorWidget::EWidgetMode::WM_Scale),
+			FCanExecuteAction(),
+			FIsActionChecked::CreateSP(this, &SDisplayClusterLightCardEditorViewport::IsEditorWidgetModeSelected, FDisplayClusterLightCardEditorWidget::EWidgetMode::WM_Scale)
+		);
+
+		CommandList->MapAction(
+			Commands.CycleTransformGizmos,
+			FExecuteAction::CreateSP(this, &SDisplayClusterLightCardEditorViewport::CycleEditorWidgetMode),
+			FCanExecuteAction()
+		);
+	}
+
+	{
+		const FDisplayClusterLightCardEditorCommands& Commands = FDisplayClusterLightCardEditorCommands::Get();
+
+		CommandList->MapAction(
+			FDisplayClusterLightCardEditorCommands::Get().PerspectiveProjection,
+			FExecuteAction::CreateSP(this, &SDisplayClusterLightCardEditorViewport::SetProjectionMode, EDisplayClusterMeshProjectionType::Perspective),
+			FCanExecuteAction(),
+			FIsActionChecked::CreateSP(this, &SDisplayClusterLightCardEditorViewport::IsProjectionModeSelected, EDisplayClusterMeshProjectionType::Perspective));
+
+		CommandList->MapAction(
+			FDisplayClusterLightCardEditorCommands::Get().AzimuthalProjection,
+			FExecuteAction::CreateSP(this, &SDisplayClusterLightCardEditorViewport::SetProjectionMode, EDisplayClusterMeshProjectionType::Azimuthal),
+			FCanExecuteAction(),
+			FIsActionChecked::CreateSP(this, &SDisplayClusterLightCardEditorViewport::IsProjectionModeSelected, EDisplayClusterMeshProjectionType::Azimuthal));
+
+		CommandList->MapAction(
+			Commands.ResetCamera,
+			FExecuteAction::CreateSP(ViewportClient.Get(), &FDisplayClusterLightCardEditorViewportClient::ResetCamera, false));
+	}
+}
+
+void SDisplayClusterLightCardEditorViewport::SetEditorWidgetMode(FDisplayClusterLightCardEditorWidget::EWidgetMode InWidgetMode)
+{
+	if (ViewportClient.IsValid())
+	{
+		ViewportClient->SetEditorWidgetMode(InWidgetMode);
+	}
+}
+
+bool SDisplayClusterLightCardEditorViewport::IsEditorWidgetModeSelected(FDisplayClusterLightCardEditorWidget::EWidgetMode InWidgetMode) const
+{
+	if (ViewportClient.IsValid())
+	{
+		return ViewportClient->GetEditorWidgetMode() == InWidgetMode;
+	}
+
+	return false;
+}
+
+void SDisplayClusterLightCardEditorViewport::CycleEditorWidgetMode()
+{
+	int32 WidgetModeAsInt = ViewportClient->GetEditorWidgetMode();
+
+	WidgetModeAsInt = (WidgetModeAsInt + 1) % FDisplayClusterLightCardEditorWidget::EWidgetMode::WM_Max;
+
+	SetEditorWidgetMode((FDisplayClusterLightCardEditorWidget::EWidgetMode)WidgetModeAsInt);
 }
 
 void SDisplayClusterLightCardEditorViewport::SetProjectionMode(EDisplayClusterMeshProjectionType InProjectionMode)
diff --git a/Engine/Plugins/Runtime/nDisplay/Source/DisplayClusterLightCardEditor/Private/Viewport/DisplayClusterLightcardEditorViewport.h b/Engine/Plugins/Runtime/nDisplay/Source/DisplayClusterLightCardEditor/Private/Viewport/DisplayClusterLightcardEditorViewport.h
--- a/Engine/Plugins/Runtime/nDisplay/Source/DisplayClusterLightCardEditor/Private/Viewport/DisplayClusterLightcardEditorViewport.h
+++ b/Engine/Plugins/Runtime/nDisplay/Source/DisplayClusterLightCardEditor/Private/Viewport/DisplayClusterLightcardEditorViewport.h
@@ -8,6 +8,7 @@
 #include "SEditorViewport.h"
 
 #include "DisplayClusterMeshProjectionRenderer.h"
+#include "DisplayClusterLightCardEditorWidget.h"
 
 class SDisplayClusterLightCardEditor;
 class ADisplayClusterRootActor;
@@ -45,6 +46,10 @@ private:
 	virtual void BindCommands() override;
 	// ~SEditorViewport
 
+	void SetEditorWidgetMode(FDisplayClusterLightCardEditorWidget::EWidgetMode InWidgetMode);
+	bool IsEditorWidgetModeSelected(FDisplayClusterLightCardEditorWidget::EWidgetMode InWidgetMode) const;
+	void CycleEditorWidgetMode();
+
 	void SetProjectionMode(EDisplayClusterMeshProjectionType InProjectionMode);
 	bool IsProjectionModeSelected(EDisplayClusterMeshProjectionType InProjectionMode) const;
 
