nDisplay: Fix crash on package reload. Fix possible crash when shutting down the editor and unregistering section categories.

diff --git a/Engine/Plugins/Runtime/nDisplay/Source/DisplayCluster/Private/Game/EngineClasses/Scene/DisplayClusterRootActorEditor.cpp b/Engine/Plugins/Runtime/nDisplay/Source/DisplayCluster/Private/Game/EngineClasses/Scene/DisplayClusterRootActorEditor.cpp
--- a/Engine/Plugins/Runtime/nDisplay/Source/DisplayCluster/Private/Game/EngineClasses/Scene/DisplayClusterRootActorEditor.cpp
+++ b/Engine/Plugins/Runtime/nDisplay/Source/DisplayCluster/Private/Game/EngineClasses/Scene/DisplayClusterRootActorEditor.cpp
@@ -96,12 +96,16 @@ void ADisplayClusterRootActor::Constructor_Editor()
 	PrimaryActorTick.bStartWithTickEnabled = true;
 
 	ResetPreviewInternals_Editor();
+
+	FCoreUObjectDelegates::OnPackageReloaded.AddUObject(this, &ADisplayClusterRootActor::HandleAssetReload);
 }
 
 void ADisplayClusterRootActor::Destructor_Editor()
 {
 	OnPreviewGenerated.Unbind();
 	OnPreviewDestroyed.Unbind();
+
+	FCoreUObjectDelegates::OnPackageReloaded.RemoveAll(this);
 }
 
 void ADisplayClusterRootActor::Tick_Editor(float DeltaSeconds)
@@ -779,6 +783,18 @@ void ADisplayClusterRootActor::PostEditMove(bool bFinished)
 	ResetPreviewComponents_Editor(false);
 }
 
+void ADisplayClusterRootActor::HandleAssetReload(const EPackageReloadPhase InPackageReloadPhase,
+	FPackageReloadedEvent* InPackageReloadedEvent)
+{
+	if (InPackageReloadPhase == EPackageReloadPhase::PrePackageFixup)
+	{
+		// Preview components need to be released here. During a package reload BeginDestroy will be called on the
+		// actor, but the preview components are already detached and never have DestroyComponent called on them.
+		// This causes the package to keep references around and cause a crash during the reload.
+		BeginDestroy_Editor();
+	}
+}
+
 void ADisplayClusterRootActor::ResetPreviewComponents_Editor(bool bInRestoreSceneMaterial)
 {
 	TArray<UDisplayClusterPreviewComponent*> AllPreviewComponents;
diff --git a/Engine/Plugins/Runtime/nDisplay/Source/DisplayCluster/Public/DisplayClusterRootActor.h b/Engine/Plugins/Runtime/nDisplay/Source/DisplayCluster/Public/DisplayClusterRootActor.h
--- a/Engine/Plugins/Runtime/nDisplay/Source/DisplayCluster/Public/DisplayClusterRootActor.h
+++ b/Engine/Plugins/Runtime/nDisplay/Source/DisplayCluster/Public/DisplayClusterRootActor.h
@@ -447,6 +447,9 @@ protected:
 	virtual void PostEditChangeProperty(FPropertyChangedEvent& PropertyChangedEvent) override;
 	virtual void PostEditMove(bool bFinished) override;
 
+	/** Called when the asset has been reloaded in the editor. */
+	void HandleAssetReload(const EPackageReloadPhase InPackageReloadPhase, FPackageReloadedEvent* InPackageReloadedEvent);
+	
 private:
 	bool bIsSelectedInEditor = false;
 	
diff --git a/Engine/Plugins/Runtime/nDisplay/Source/DisplayClusterConfigurator/Private/DisplayClusterConfiguratorModule.cpp b/Engine/Plugins/Runtime/nDisplay/Source/DisplayClusterConfigurator/Private/DisplayClusterConfiguratorModule.cpp
--- a/Engine/Plugins/Runtime/nDisplay/Source/DisplayClusterConfigurator/Private/DisplayClusterConfiguratorModule.cpp
+++ b/Engine/Plugins/Runtime/nDisplay/Source/DisplayClusterConfigurator/Private/DisplayClusterConfiguratorModule.cpp
@@ -285,17 +285,20 @@ void FDisplayClusterConfiguratorModule::RegisterSectionMappings()
 
 void FDisplayClusterConfiguratorModule::UnregisterSectionMappings()
 {
-	FPropertyEditorModule& PropertyModule = FModuleManager::LoadModuleChecked<FPropertyEditorModule>("PropertyEditor");
-	PropertyModule.RemoveSection(ADisplayClusterRootActor::StaticClass()->GetFName(), DisplayClusterConfigurationStrings::categories::ViewportsCategory);
-	PropertyModule.RemoveSection(ADisplayClusterRootActor::StaticClass()->GetFName(), DisplayClusterConfigurationStrings::categories::ICVFXCategory);
-	PropertyModule.RemoveSection(ADisplayClusterRootActor::StaticClass()->GetFName(), DisplayClusterConfigurationStrings::categories::ColorGradingCategory);
-	PropertyModule.RemoveSection(ADisplayClusterRootActor::StaticClass()->GetFName(), DisplayClusterConfigurationStrings::categories::OCIOCategory);
-	PropertyModule.RemoveSection(ADisplayClusterRootActor::StaticClass()->GetFName(), DisplayClusterConfigurationStrings::categories::LightcardCategory);
-
-	PropertyModule.RemoveSection(UDisplayClusterICVFXCameraComponent::StaticClass()->GetFName(), DisplayClusterConfigurationStrings::categories::ICVFXCategory);
-	PropertyModule.RemoveSection(UDisplayClusterICVFXCameraComponent::StaticClass()->GetFName(), DisplayClusterConfigurationStrings::categories::CameraColorGradingCategory);
-	PropertyModule.RemoveSection(UDisplayClusterICVFXCameraComponent::StaticClass()->GetFName(), DisplayClusterConfigurationStrings::categories::OCIOCategory);
-	PropertyModule.RemoveSection(UDisplayClusterICVFXCameraComponent::StaticClass()->GetFName(), DisplayClusterConfigurationStrings::categories::ChromaKeyCategory);
+	if (FModuleManager::Get().IsModuleLoaded("PropertyEditor") && FSlateApplication::IsInitialized())
+	{
+		FPropertyEditorModule& PropertyModule = FModuleManager::GetModuleChecked<FPropertyEditorModule>("PropertyEditor");
+		PropertyModule.RemoveSection(ADisplayClusterRootActor::StaticClass()->GetFName(), DisplayClusterConfigurationStrings::categories::ViewportsCategory);
+		PropertyModule.RemoveSection(ADisplayClusterRootActor::StaticClass()->GetFName(), DisplayClusterConfigurationStrings::categories::ICVFXCategory);
+		PropertyModule.RemoveSection(ADisplayClusterRootActor::StaticClass()->GetFName(), DisplayClusterConfigurationStrings::categories::ColorGradingCategory);
+		PropertyModule.RemoveSection(ADisplayClusterRootActor::StaticClass()->GetFName(), DisplayClusterConfigurationStrings::categories::OCIOCategory);
+		PropertyModule.RemoveSection(ADisplayClusterRootActor::StaticClass()->GetFName(), DisplayClusterConfigurationStrings::categories::LightcardCategory);
+
+		PropertyModule.RemoveSection(UDisplayClusterICVFXCameraComponent::StaticClass()->GetFName(), DisplayClusterConfigurationStrings::categories::ICVFXCategory);
+		PropertyModule.RemoveSection(UDisplayClusterICVFXCameraComponent::StaticClass()->GetFName(), DisplayClusterConfigurationStrings::categories::CameraColorGradingCategory);
+		PropertyModule.RemoveSection(UDisplayClusterICVFXCameraComponent::StaticClass()->GetFName(), DisplayClusterConfigurationStrings::categories::OCIOCategory);
+		PropertyModule.RemoveSection(UDisplayClusterICVFXCameraComponent::StaticClass()->GetFName(), DisplayClusterConfigurationStrings::categories::ChromaKeyCategory);
+	}
 }
 
 void FDisplayClusterConfiguratorModule::RegisterPlacementModeItems()
