nDisplay: Added rotation mode to Light Card Editor widget that allows users to change the spin of the selected light cards by clicking and dragging on the widget. Also fixed issue where axes of widget weren't rendering in the right direction in azimuthal projection mode.

diff --git a/Engine/Plugins/Runtime/nDisplay/Source/DisplayClusterLightCardEditor/Private/Viewport/DisplayClusterLightCardEditorViewportClient.cpp b/Engine/Plugins/Runtime/nDisplay/Source/DisplayClusterLightCardEditor/Private/Viewport/DisplayClusterLightCardEditorViewportClient.cpp
--- a/Engine/Plugins/Runtime/nDisplay/Source/DisplayClusterLightCardEditor/Private/Viewport/DisplayClusterLightCardEditorViewportClient.cpp
+++ b/Engine/Plugins/Runtime/nDisplay/Source/DisplayClusterLightCardEditor/Private/Viewport/DisplayClusterLightCardEditorViewportClient.cpp
@@ -501,14 +501,8 @@ void FDisplayClusterLightCardEditorViewportClient::Draw(const FSceneView* View,
 		// This needs to be computed on the render thread using the render thread's scene view, which will be behind the game thread's scene view
 		// by at least one frame
 
-		FTransform ProjectedEditorWidgetTransform(CachedEditorWidgetWorldTransform);
-
-		if (ProjectionMode != EDisplayClusterMeshProjectionType::Perspective)
-		{
-			ProjectedEditorWidgetTransform.SetTranslation(ProjectWorldPosition(CachedEditorWidgetWorldTransform.GetTranslation(), View->ViewMatrices));
-		}
-
-		EditorWidget->SetTransform(ProjectedEditorWidgetTransform);
+		EditorWidget->SetTransform(CachedEditorWidgetWorldTransform);
+		EditorWidget->SetProjectionTransform(FDisplayClusterMeshProjectionTransform(ProjectionMode, View->ViewMatrices.GetViewMatrix()));
 		EditorWidget->Draw(View, this, PDI);
 	}
 }
@@ -574,13 +568,19 @@ bool FDisplayClusterLightCardEditorViewportClient::InputWidgetDelta(FViewport* I
 	{
 		if (CurrentAxis != EAxisList::Type::None && SelectedLightCards.Num())
 		{
-			if (EditorWidget->GetWidgetMode() == FDisplayClusterLightCardEditorWidget::EWidgetMode::WM_Translate)
+			switch (EditorWidget->GetWidgetMode())
 			{
+			case FDisplayClusterLightCardEditorWidget::EWidgetMode::WM_Translate:
 				MoveSelectedLightCards(InViewport, CurrentAxis);
-			}
-			else if (EditorWidget->GetWidgetMode() == FDisplayClusterLightCardEditorWidget::EWidgetMode::WM_Scale)
-			{
+				break;
+
+			case FDisplayClusterLightCardEditorWidget::EWidgetMode::WM_RotateZ:
+				SpinSelectedLightCards(InViewport);
+				break;
+
+			case FDisplayClusterLightCardEditorWidget::EWidgetMode::WM_Scale:
 				ScaleSelectedLightCards(InViewport, CurrentAxis);
+				break;
 			}
 
 			bHandled = true;
@@ -1716,7 +1716,7 @@ FVector2D FDisplayClusterLightCardEditorViewportClient::GetLightCardScaleDelta(F
 	FVector2D ScaleDelta = FVector2D((CurrentAxis & EAxisList::X) * ScaleMagnitude, (CurrentAxis & EAxisList::Y) * ScaleMagnitude);
 
 	// If both axes are being scaled at the same time, preserve the aspect ratio of the scale delta
-	if (CurrentAxis == EAxisList::Type::XY)
+	if ((CurrentAxis & EAxisList::Type::X) && (CurrentAxis & EAxisList::Type::Y))
 	{
 		// Ensure the signs of the deltas remain the same, and avoid potential divide by zero
 		ScaleDelta.Y = ScaleDelta.X * FMath::Abs(LightCard->Scale.Y) / FMath::Max(0.001, FMath::Abs(LightCard->Scale.X));
@@ -1725,6 +1725,59 @@ FVector2D FDisplayClusterLightCardEditorViewportClient::GetLightCardScaleDelta(F
 	return ScaleDelta;
 }
 
+void FDisplayClusterLightCardEditorViewportClient::SpinSelectedLightCards(FViewport* InViewport)
+{
+	const TWeakObjectPtr<ADisplayClusterLightCardActor>& LastSelectedLightCard = SelectedLightCards.Last();
+	if (LastSelectedLightCard.IsValid())
+	{
+		const float DeltaSpin = GetLightCardSpinDelta(InViewport, LastSelectedLightCard.Get());
+		for (const TWeakObjectPtr<ADisplayClusterLightCardActor>& LightCard : SelectedLightCards)
+		{
+			if (!LightCard.IsValid())
+			{
+				continue;
+			}
+
+			LightCard->Spin += DeltaSpin;
+
+			PropagateLightCardTransform(LightCard.Get());
+		}
+	}
+}
+
+float FDisplayClusterLightCardEditorViewportClient::GetLightCardSpinDelta(FViewport* InViewport, ADisplayClusterLightCardActor* LightCard)
+{
+	FIntPoint MousePos;
+	InViewport->GetMousePos(MousePos);
+
+	FSceneViewFamilyContext ViewFamily(FSceneViewFamily::ConstructionValues(
+		InViewport,
+		GetScene(),
+		EngineShowFlags)
+		.SetRealtimeUpdate(IsRealtime()));
+	FSceneView* View = CalcSceneView(&ViewFamily);
+
+	const FVector WorldWidgetOrigin = CachedEditorWidgetWorldTransform.GetTranslation();
+	const FVector WorldWidgetXAxis = CachedEditorWidgetWorldTransform.GetRotation().RotateVector(FVector::XAxisVector);
+	const FVector WorldWidgetYAxis = CachedEditorWidgetWorldTransform.GetRotation().RotateVector(FVector::YAxisVector);
+
+	FVector2D ScreenOrigin = FVector2D::ZeroVector;
+	FVector2D ScreenXAxis = FVector2D::ZeroVector;
+	FVector2D ScreenYAxis = FVector2D::ZeroVector;
+
+	WorldToPixel(*View, WorldWidgetOrigin, ScreenOrigin);
+	WorldToScreenDirection(*View, WorldWidgetOrigin, WorldWidgetXAxis, ScreenXAxis);
+	WorldToScreenDirection(*View, WorldWidgetOrigin, WorldWidgetYAxis, ScreenYAxis);
+
+	FVector2D MousePosOffset = FVector2D(MousePos) - ScreenOrigin;
+	FVector2D LastMousePosOffset = FVector2D(LastWidgetMousePos) - ScreenOrigin;
+
+	float Theta = FMath::Atan2(MousePosOffset | ScreenYAxis, MousePosOffset | ScreenXAxis);
+	float LastTheta = FMath::Atan2(LastMousePosOffset | ScreenYAxis, LastMousePosOffset | ScreenXAxis);
+
+	return FMath::RadiansToDegrees(Theta - LastTheta);
+}
+
 FDisplayClusterLightCardEditorViewportClient::FSphericalCoordinates FDisplayClusterLightCardEditorViewportClient::GetLightCardCoordinates(ADisplayClusterLightCardActor* LightCard) const
 {
 	const FVector LightCardLocation = LightCard->GetLightCardTransform().GetTranslation() - LightCard->GetActorLocation();
@@ -1813,16 +1866,8 @@ ADisplayClusterLightCardActor* FDisplayClusterLightCardEditorViewportClient::Tra
 
 FVector FDisplayClusterLightCardEditorViewportClient::ProjectWorldPosition(const FVector& UnprojectedWorldPosition, const FViewMatrices& ViewMatrices) const
 {
-	FVector ProjectedPosition(UnprojectedWorldPosition);
-
-	if (ProjectionMode != EDisplayClusterMeshProjectionType::Perspective)
-	{
-		const FVector ViewPos = ViewMatrices.GetViewMatrix().TransformPosition(CachedEditorWidgetWorldTransform.GetTranslation());
-		const FVector ProjectedViewPos = FDisplayClusterMeshProjectionRenderer::ProjectViewPosition(ViewPos, ProjectionMode);
-		ProjectedPosition = ViewMatrices.GetInvViewMatrix().TransformPosition(ProjectedViewPos);
-	}
-
-	return ProjectedPosition;
+	FDisplayClusterMeshProjectionTransform Transform(ProjectionMode, ViewMatrices.GetViewMatrix());
+	return Transform.ProjectPosition(UnprojectedWorldPosition);
 }
 
 void FDisplayClusterLightCardEditorViewportClient::PixelToWorld(const FSceneView& View, const FIntPoint& PixelPos, FVector& OutOrigin, FVector& OutDirection)
diff --git a/Engine/Plugins/Runtime/nDisplay/Source/DisplayClusterLightCardEditor/Private/Viewport/DisplayClusterLightCardEditorViewportClient.h b/Engine/Plugins/Runtime/nDisplay/Source/DisplayClusterLightCardEditor/Private/Viewport/DisplayClusterLightCardEditorViewportClient.h
--- a/Engine/Plugins/Runtime/nDisplay/Source/DisplayClusterLightCardEditor/Private/Viewport/DisplayClusterLightCardEditorViewportClient.h
+++ b/Engine/Plugins/Runtime/nDisplay/Source/DisplayClusterLightCardEditor/Private/Viewport/DisplayClusterLightCardEditorViewportClient.h
@@ -238,6 +238,12 @@ private:
 	/** Determines the appropriate scale delta needed to scale the light card */
 	FVector2D GetLightCardScaleDelta(FViewport* InViewport, ADisplayClusterLightCardActor* LightCard, EAxisList::Type CurrentAxis);
 
+	/** Rotates the currently selected light cards around the light card's normal axis */
+	void SpinSelectedLightCards(FViewport* InViewport);
+
+	/** Determines the appropriate spin delta needed to rotate the light card */
+	float GetLightCardSpinDelta(FViewport* InViewport, ADisplayClusterLightCardActor* LightCard);
+
 	/** Gets the spherical coordinates of the specified light card */
 	FSphericalCoordinates GetLightCardCoordinates(ADisplayClusterLightCardActor* LightCard) const;
 
diff --git a/Engine/Plugins/Runtime/nDisplay/Source/DisplayClusterLightCardEditor/Private/Viewport/DisplayClusterLightCardEditorWidget.cpp b/Engine/Plugins/Runtime/nDisplay/Source/DisplayClusterLightCardEditor/Private/Viewport/DisplayClusterLightCardEditorWidget.cpp
--- a/Engine/Plugins/Runtime/nDisplay/Source/DisplayClusterLightCardEditor/Private/Viewport/DisplayClusterLightCardEditorWidget.cpp
+++ b/Engine/Plugins/Runtime/nDisplay/Source/DisplayClusterLightCardEditor/Private/Viewport/DisplayClusterLightCardEditorWidget.cpp
@@ -11,12 +11,21 @@ void FDisplayClusterLightCardEditorWidget::Draw(const FSceneView* View, const FE
 	const float SizeScalar = GetSizeScreenScalar(View, ViewportClient);
 	const float LengthScalar = GetLengthScreenScalar(View, ViewportClient, Transform.GetTranslation());
 
-	DrawAxis(View, PDI, EAxisList::Type::X, SizeScalar, LengthScalar);
-	DrawAxis(View, PDI, EAxisList::Type::Y, SizeScalar, LengthScalar);
-	DrawOrigin(PDI, SizeScalar);
+	if (WidgetMode == EWidgetMode::WM_RotateZ)
+	{
+		DrawAxis(PDI, EAxisList::Type::Y, SizeScalar, LengthScalar);
+		DrawCircle(PDI, EAxisList::Type::Z, SizeScalar, LengthScalar);
+		DrawOrigin(PDI, SizeScalar);
+	}
+	else
+	{
+		DrawAxis(PDI, EAxisList::Type::X, SizeScalar, LengthScalar);
+		DrawAxis(PDI, EAxisList::Type::Y, SizeScalar, LengthScalar);
+		DrawOrigin(PDI, SizeScalar);
+	}
 }
 
-void FDisplayClusterLightCardEditorWidget::DrawAxis(const FSceneView* View, FPrimitiveDrawInterface* PDI, EAxisList::Type Axis, float SizeScalar, float LengthScalar)
+void FDisplayClusterLightCardEditorWidget::DrawAxis(FPrimitiveDrawInterface* PDI, EAxisList::Type Axis, float SizeScalar, float LengthScalar)
 {
 	const FVector Origin = Transform.GetTranslation();
 
@@ -26,8 +35,8 @@ void FDisplayClusterLightCardEditorWidget::DrawAxis(const FSceneView* View, FPri
 
 	PDI->SetHitProxy(new HWidgetAxis(Axis));
 	{
-		const FVector AxisStart = Origin;
-		const FVector AxisEnd = Origin + TransformedAxis * AxisLength * WidgetScale * LengthScalar;
+		const FVector AxisStart = ProjectionTransform.ProjectPosition(Origin);
+		const FVector AxisEnd = ProjectionTransform.ProjectPosition(Origin + TransformedAxis * AxisLength * WidgetScale * LengthScalar);
 
 		PDI->DrawLine(AxisStart, AxisEnd, AxisColor, ESceneDepthPriorityGroup::SDPG_Foreground, AxisThickness * SizeScalar, 0.0, true);
 
@@ -41,14 +50,45 @@ void FDisplayClusterLightCardEditorWidget::DrawAxis(const FSceneView* View, FPri
 
 void FDisplayClusterLightCardEditorWidget::DrawOrigin(FPrimitiveDrawInterface* PDI, float SizeScalar)
 {
-	const FVector Origin = Transform.GetTranslation();
-	const FLinearColor Color = HighlightedAxis == EAxisList::Type::XY ? HighlightColor : FLinearColor::Black;
+	const FVector Origin = ProjectionTransform.ProjectPosition(Transform.GetTranslation());
+	const FLinearColor Color = HighlightedAxis == EAxisList::Type::XYZ ? HighlightColor : FLinearColor::Black;
 
-	PDI->SetHitProxy(new HWidgetAxis(EAxisList::Type::XY));
+	PDI->SetHitProxy(new HWidgetAxis(EAxisList::Type::XYZ));
 	PDI->DrawPoint(Origin, Color, OriginSize * SizeScalar, ESceneDepthPriorityGroup::SDPG_Foreground);
 	PDI->SetHitProxy(nullptr);
 }
 
+void FDisplayClusterLightCardEditorWidget::DrawCircle(FPrimitiveDrawInterface* PDI, EAxisList::Type Axis, float SizeScalar, float LengthScalar)
+{
+	const float ScaledRadius = CirlceRadius * LengthScalar;
+	const int32 NumSides = FMath::Max(ScaledRadius / 2, 24.0f);
+	const float AngleDelta = 2.0f * UE_PI / NumSides;
+
+	const FVector Origin = Transform.GetTranslation();
+
+	const FVector GlobalAxis = GetGlobalAxis(Axis);
+	const FVector Normal = Transform.TransformVector(GlobalAxis);
+	const FVector X = GlobalAxis == FVector::ZAxisVector ? Transform.TransformVector(FVector::XAxisVector) : Transform.TransformVector(FVector::ZAxisVector);
+	const FVector Y = Normal ^ X;
+
+	EAxisList::Type RotationPlane = (EAxisList::Type)(EAxisList::Type::XYZ & ~Axis);
+
+	const FLinearColor CircleColor = HighlightedAxis == RotationPlane ? HighlightColor : GetAxisColor(Axis);
+
+	PDI->SetHitProxy(new HWidgetAxis(RotationPlane));
+	{
+		FVector Vertex = ProjectionTransform.ProjectPosition(Origin + X * ScaledRadius);
+		for (int32 Index = 0; Index < NumSides; ++Index)
+		{
+			FVector NextVertex = ProjectionTransform.ProjectPosition(Origin + (X * FMath::Cos(AngleDelta * (Index + 1)) + Y * FMath::Sin(AngleDelta * (Index + 1))) * ScaledRadius);
+			PDI->DrawLine(Vertex, NextVertex, CircleColor, ESceneDepthPriorityGroup::SDPG_Foreground, CircleThickness * SizeScalar, 0.0f, true);
+
+			Vertex = NextVertex;
+		}
+	}
+	PDI->SetHitProxy(nullptr);
+}
+
 FVector FDisplayClusterLightCardEditorWidget::GetGlobalAxis(EAxisList::Type Axis) const
 {
 	switch (Axis)
diff --git a/Engine/Plugins/Runtime/nDisplay/Source/DisplayClusterLightCardEditor/Private/Viewport/DisplayClusterLightCardEditorWidget.h b/Engine/Plugins/Runtime/nDisplay/Source/DisplayClusterLightCardEditor/Private/Viewport/DisplayClusterLightCardEditorWidget.h
--- a/Engine/Plugins/Runtime/nDisplay/Source/DisplayClusterLightCardEditor/Private/Viewport/DisplayClusterLightCardEditorWidget.h
+++ b/Engine/Plugins/Runtime/nDisplay/Source/DisplayClusterLightCardEditor/Private/Viewport/DisplayClusterLightCardEditorWidget.h
@@ -5,6 +5,8 @@
 #include "CoreMinimal.h"
 #include "Math/Axis.h"
 
+#include "DisplayClusterMeshProjectionRenderer.h"
+
 class FSceneView;
 class FEditorViewportClient;
 class FPrimitiveDrawInterface;
@@ -28,6 +30,12 @@ public:
 	/** The size of the axes origin */
 	constexpr static float OriginSize = 35.f;
 
+	/** The size of the rotation circle, in pixels */
+	constexpr static float CirlceRadius = 50.0f;
+
+	/** The thickness of the rotation circle, in pixels */
+	constexpr static float CircleThickness = 2.5f;
+
 	constexpr static FLinearColor AxisColorX = FLinearColor(0.594f, 0.0197f, 0.0f);
 	constexpr static FLinearColor AxisColorY = FLinearColor(0.1349f, 0.3959f, 0.0f);
 	constexpr static FLinearColor AxisColorZ = FLinearColor(0.0251f, 0.207f, 0.85f);
@@ -36,6 +44,7 @@ public:
 	enum EWidgetMode
 	{
 		WM_Translate,
+		WM_RotateZ,
 		WM_Scale,
 
 		WM_Max
@@ -47,16 +56,19 @@ public:
 	void SetWidgetMode(EWidgetMode NewWidgetMode) { WidgetMode = NewWidgetMode; }
 
 	void SetTransform(const FTransform& NewTransform) { Transform = FTransform(NewTransform); }
+	void SetProjectionTransform(const FDisplayClusterMeshProjectionTransform& NewProjectionTransform) { ProjectionTransform = FDisplayClusterMeshProjectionTransform(NewProjectionTransform); }
 	void SetHighlightedAxis(EAxisList::Type InAxis) { HighlightedAxis = InAxis; }
 	void SetWidgetScale(float InWidgetScale) { WidgetScale = InWidgetScale; }
 
 private:
 	/** Draws the specified axis to the PDI, sizing and coloring it appropriately */
-	void DrawAxis(const FSceneView* View, FPrimitiveDrawInterface* PDI, EAxisList::Type Axis, float SizeScalar, float LengthScalar);
+	void DrawAxis(FPrimitiveDrawInterface* PDI, EAxisList::Type Axis, float SizeScalar, float LengthScalar);
 
 	/** Draws the widget origin to the PDI, sizing and coloring it appropriately */
 	void DrawOrigin(FPrimitiveDrawInterface* PDI, float SizeScalar);
 
+	void DrawCircle(FPrimitiveDrawInterface* PDI, EAxisList::Type Axis, float SizeScalar, float LengthScalar);
+
 	/** Gets the axis vector in global coordinates for the specified axis */
 	FVector GetGlobalAxis(EAxisList::Type Axis) const;
 
@@ -76,6 +88,9 @@ private:
 	/** The world transform to apply to the widget when rendering */
 	FTransform Transform;
 
+	/** The projection transfrom to apply to the widget when rendering */
+	FDisplayClusterMeshProjectionTransform ProjectionTransform;
+
 	/** The axis on the widget which should be highlighted */
 	EAxisList::Type HighlightedAxis = EAxisList::Type::None;
 
diff --git a/Engine/Plugins/Runtime/nDisplay/Source/DisplayClusterLightCardEditor/Private/Viewport/DisplayClusterLightcardEditorViewport.cpp b/Engine/Plugins/Runtime/nDisplay/Source/DisplayClusterLightCardEditor/Private/Viewport/DisplayClusterLightcardEditorViewport.cpp
--- a/Engine/Plugins/Runtime/nDisplay/Source/DisplayClusterLightCardEditor/Private/Viewport/DisplayClusterLightcardEditorViewport.cpp
+++ b/Engine/Plugins/Runtime/nDisplay/Source/DisplayClusterLightCardEditor/Private/Viewport/DisplayClusterLightcardEditorViewport.cpp
@@ -165,11 +165,11 @@ public:
 			static FName TranslateModeName = FName(TEXT("TranslateMode"));
 			ToolbarBuilder.AddToolBarButton(FEditorViewportCommands::Get().TranslateMode, NAME_None, TAttribute<FText>(), TAttribute<FText>(), TAttribute<FSlateIcon>(), TranslateModeName);
 
-			//static FName RotateModeName = FName(TEXT("RotateMode"));
-			//ToolbarBuilder.AddToolBarButton( FEditorViewportCommands::Get().RotateMode, NAME_None, TAttribute<FText>(), TAttribute<FText>(), TAttribute<FSlateIcon>(), RotateModeName );
+			static FName RotateModeName = FName(TEXT("RotateMode"));
+			ToolbarBuilder.AddToolBarButton(FEditorViewportCommands::Get().RotateMode, NAME_None, TAttribute<FText>(), TAttribute<FText>(), TAttribute<FSlateIcon>(), RotateModeName);
 
 			static FName ScaleModeName = FName(TEXT("ScaleMode"));
-			ToolbarBuilder.AddToolBarButton( FEditorViewportCommands::Get().ScaleMode, NAME_None, TAttribute<FText>(), TAttribute<FText>(), TAttribute<FSlateIcon>(), ScaleModeName );
+			ToolbarBuilder.AddToolBarButton(FEditorViewportCommands::Get().ScaleMode, NAME_None, TAttribute<FText>(), TAttribute<FText>(), TAttribute<FSlateIcon>(), ScaleModeName);
 
 			ToolbarBuilder.EndBlockGroup();
 		}
@@ -263,6 +263,14 @@ void SDisplayClusterLightCardEditorViewport::BindCommands()
 			FIsActionChecked::CreateSP(this, &SDisplayClusterLightCardEditorViewport::IsEditorWidgetModeSelected, FDisplayClusterLightCardEditorWidget::EWidgetMode::WM_Translate)
 		);
 
+		
+		CommandList->MapAction(
+			Commands.RotateMode,
+			FExecuteAction::CreateSP(this, &SDisplayClusterLightCardEditorViewport::SetEditorWidgetMode, FDisplayClusterLightCardEditorWidget::EWidgetMode::WM_RotateZ),
+			FCanExecuteAction(),
+			FIsActionChecked::CreateSP(this, &SDisplayClusterLightCardEditorViewport::IsEditorWidgetModeSelected, FDisplayClusterLightCardEditorWidget::EWidgetMode::WM_RotateZ)
+		);
+
 		CommandList->MapAction(
 			Commands.ScaleMode,
 			FExecuteAction::CreateSP(this, &SDisplayClusterLightCardEditorViewport::SetEditorWidgetMode, FDisplayClusterLightCardEditorWidget::EWidgetMode::WM_Scale),
diff --git a/Engine/Plugins/Runtime/nDisplay/Source/DisplayClusterLightCardEditorShaders/Private/DisplayClusterMeshProjectionRenderer.cpp b/Engine/Plugins/Runtime/nDisplay/Source/DisplayClusterLightCardEditorShaders/Private/DisplayClusterMeshProjectionRenderer.cpp
--- a/Engine/Plugins/Runtime/nDisplay/Source/DisplayClusterLightCardEditorShaders/Private/DisplayClusterMeshProjectionRenderer.cpp
+++ b/Engine/Plugins/Runtime/nDisplay/Source/DisplayClusterLightCardEditorShaders/Private/DisplayClusterMeshProjectionRenderer.cpp
@@ -692,6 +692,34 @@ bool FDisplayClusterMeshProjectionPrimitiveFilter::IsPrimitiveComponentFiltered(
 	return false;
 }
 
+FVector FDisplayClusterMeshProjectionTransform::ProjectPosition(const FVector& WorldPosition) const
+{
+	FVector ProjectedPosition(WorldPosition);
+
+	if (Projection != EDisplayClusterMeshProjectionType::Perspective)
+	{
+		const FVector ViewPos = ViewMatrix.TransformPosition(WorldPosition);
+		const FVector ProjectedViewPos = FDisplayClusterMeshProjectionRenderer::ProjectViewPosition(ViewPos, Projection);
+		ProjectedPosition = InvViewMatrix.TransformPosition(ProjectedViewPos);
+	}
+
+	return ProjectedPosition;
+}
+
+FVector FDisplayClusterMeshProjectionTransform::UnprojectPosition(const FVector& ProjectedPosition) const
+{
+	FVector UnprojectedPosition(ProjectedPosition);
+
+	if (Projection != EDisplayClusterMeshProjectionType::Perspective)
+	{
+		const FVector ViewPos = ViewMatrix.TransformPosition(ProjectedPosition);
+		const FVector ProjectedViewPos = FDisplayClusterMeshProjectionRenderer::UnprojectViewPosition(ViewPos, Projection);
+		UnprojectedPosition = InvViewMatrix.TransformPosition(ProjectedViewPos);
+	}
+
+	return ProjectedPosition;
+}
+
 FVector FDisplayClusterMeshProjectionRenderer::ProjectViewPosition(const FVector& ViewPosition, EDisplayClusterMeshProjectionType  ProjectionType)
 {
 	FVector ProjectedViewPosition(ViewPosition);
diff --git a/Engine/Plugins/Runtime/nDisplay/Source/DisplayClusterLightCardEditorShaders/Public/DisplayClusterMeshProjectionRenderer.h b/Engine/Plugins/Runtime/nDisplay/Source/DisplayClusterLightCardEditorShaders/Public/DisplayClusterMeshProjectionRenderer.h
--- a/Engine/Plugins/Runtime/nDisplay/Source/DisplayClusterLightCardEditorShaders/Public/DisplayClusterMeshProjectionRenderer.h
+++ b/Engine/Plugins/Runtime/nDisplay/Source/DisplayClusterLightCardEditorShaders/Public/DisplayClusterMeshProjectionRenderer.h
@@ -36,6 +36,29 @@ public:
 	bool IsPrimitiveComponentFiltered(const UPrimitiveComponent* InPrimitiveComponent) const;
 };
 
+/** A transform that can be passed around to project and unprojection positions for a specific projection type */
+class DISPLAYCLUSTERLIGHTCARDEDITORSHADERS_API FDisplayClusterMeshProjectionTransform
+{
+public:
+	FDisplayClusterMeshProjectionTransform()
+		: FDisplayClusterMeshProjectionTransform(EDisplayClusterMeshProjectionType::Perspective, FMatrix::Identity)
+	{ }
+
+	FDisplayClusterMeshProjectionTransform(EDisplayClusterMeshProjectionType InProjection, const FMatrix& InViewMatrix)
+		: Projection(InProjection)
+		, ViewMatrix(InViewMatrix)
+		, InvViewMatrix(InViewMatrix.Inverse())
+	{ }
+
+	FVector ProjectPosition(const FVector& WorldPosition) const;
+	FVector UnprojectPosition(const FVector& ProjectedPosition) const;
+
+private:
+	EDisplayClusterMeshProjectionType Projection;
+	FMatrix ViewMatrix;
+	FMatrix InvViewMatrix;
+};
+
 /** A renderer that projects meshes to screen space using non-linear projection methods */
 class DISPLAYCLUSTERLIGHTCARDEDITORSHADERS_API FDisplayClusterMeshProjectionRenderer
 {
