nDisplay: Lightcards can now be moved w/o apparent spin when using the widget center.

This is particularly useful when nudging the lightcards around the center of the ceiling.

This spin compensation is not activated when using the lat/long axes individually.

diff --git a/Engine/Plugins/Runtime/nDisplay/Source/DisplayCluster/Public/DisplayClusterLightCardActor.h b/Engine/Plugins/Runtime/nDisplay/Source/DisplayCluster/Public/DisplayClusterLightCardActor.h
--- a/Engine/Plugins/Runtime/nDisplay/Source/DisplayCluster/Public/DisplayClusterLightCardActor.h
+++ b/Engine/Plugins/Runtime/nDisplay/Source/DisplayCluster/Public/DisplayClusterLightCardActor.h
@@ -54,9 +54,11 @@ public:
 
 	/** Sets a new static mesh for the light card */
 	void SetStaticMesh(UStaticMesh* InStaticMesh);
-	
-protected:
+
+	/** Updates the Light Card transform based on its positional properties (Lat, Long, etc.) */
 	void UpdateLightCardTransform();
+
+protected:
 	void UpdateLightCardMaterialInstance();
 
 public:
diff --git a/Engine/Plugins/Runtime/nDisplay/Source/DisplayClusterLightCardEditor/Private/Viewport/DisplayClusterLightCardEditorViewportClient.cpp b/Engine/Plugins/Runtime/nDisplay/Source/DisplayClusterLightCardEditor/Private/Viewport/DisplayClusterLightCardEditorViewportClient.cpp
--- a/Engine/Plugins/Runtime/nDisplay/Source/DisplayClusterLightCardEditor/Private/Viewport/DisplayClusterLightCardEditorViewportClient.cpp
+++ b/Engine/Plugins/Runtime/nDisplay/Source/DisplayClusterLightCardEditor/Private/Viewport/DisplayClusterLightCardEditorViewportClient.cpp
@@ -1546,13 +1546,59 @@ void FDisplayClusterLightCardEditorViewportClient::MoveSelectedLightCards(FViewp
 				NewCoords.Radius = CurrentCoords.Radius + DeltaCoords.Radius;
 				NewCoords.Azimuth = CurrentCoords.Azimuth + DeltaCoords.Azimuth;
 				NewCoords.Inclination = CurrentCoords.Inclination + DeltaCoords.Inclination;
-
-				LightCard->DistanceFromCenter = NewCoords.Radius;
-				LightCard->Longitude = FRotator::ClampAxis(FMath::RadiansToDegrees(NewCoords.Azimuth) - 180);
-				LightCard->Latitude = 90.f - FMath::RadiansToDegrees(NewCoords.Inclination);
 			}
 
-			MoveLightCardTo(*LightCard.Get(), NewCoords);
+			// We will only adjust the spin (to maintain the apparent spin) when not using a lat/long axis to drag
+			if (CurrentAxis == EAxisList::Type::XYZ)
+			{
+				const FTransform Transform_A = LightCard->GetLightCardTransform(false /*bIgnoreSpinYawPitch*/);
+
+				MoveLightCardTo(*LightCard.Get(), NewCoords);
+
+				LightCard->UpdateLightCardTransform(); // We must call this for GetLightCardTransform to be valid
+
+				const FTransform Transform_B = LightCard->GetLightCardTransform(false /*bIgnoreSpinYawPitch*/);
+
+				// Calculate world delta translation of moving from A to B
+				const FVector WorldDelta = Transform_B.GetLocation() - Transform_A.GetLocation(); // X towards front of stage. Y towards right of stage. Z towards ceiling.
+
+				// Calculations are only valid if translation is not too small
+				if (WorldDelta.Length() > KINDA_SMALL_NUMBER)
+				{
+					// Calculate LC "Y" unit vector at A and B. ("X" is LC normal)
+					const FVector Y_A = Transform_A.Rotator().RotateVector(FVector::YAxisVector);
+					const FVector Y_B = Transform_B.Rotator().RotateVector(FVector::YAxisVector);
+
+					// Calculate card normal vector
+					const FVector CardNormal_A = Transform_A.Rotator().RotateVector(FVector::XAxisVector); // When card is on ceiling, expect around (0,0,-1).
+					const FVector CardNormal_B = Transform_B.Rotator().RotateVector(FVector::XAxisVector);
+
+					// Calculate projection of movement onto surface tangent plane at A and B
+					const FVector UnitWorldDeltaPlane_A = FVector::VectorPlaneProject(WorldDelta, CardNormal_A).GetSafeNormal();
+					const FVector UnitWorldDeltaPlane_B = FVector::VectorPlaneProject(WorldDelta, CardNormal_B).GetSafeNormal();
+
+					// Calculate relative spin angle at A, which is the angle between Y_A and UnitWorldDeltaPlane_A
+					const double SpinDotProduct_A = FVector::DotProduct(Y_A, UnitWorldDeltaPlane_A);
+					const FVector SpinCrossProduct_A = FVector::CrossProduct(Y_A, UnitWorldDeltaPlane_A);
+					const int32 SpinSign_A = FVector::DotProduct(CardNormal_A, SpinCrossProduct_A) > 0 ? -1 : 1;
+					const double RelativeSpinAngle_A = FMath::Acos(SpinDotProduct_A) * SpinSign_A; // radians
+
+					// Now we need to find the spin that keeps the same RelativeSpinAngle in B as it was in A
+					const double SpinDotProduct_B = FVector::DotProduct(Y_B, UnitWorldDeltaPlane_B);
+					const FVector SpinCrossProduct_B = FVector::CrossProduct(Y_B, UnitWorldDeltaPlane_B);
+					const int32 SpinSign_B = FVector::DotProduct(CardNormal_B, SpinCrossProduct_B) > 0 ? -1 : 1;
+					const double RelativeSpinAngle_B = FMath::Acos(SpinDotProduct_B) * SpinSign_B; // radians
+
+					const double DeltaSpin = RelativeSpinAngle_B - RelativeSpinAngle_A;
+
+					// Apply delta spin to lightcard
+					LightCard->Spin += FMath::RadiansToDegrees(DeltaSpin);
+				}
+			}
+			else
+			{
+				MoveLightCardTo(*LightCard.Get(), NewCoords);
+			}
 
 			PropagateLightCardTransform(LightCard.Get());
 		}
