nDisplay: Added opacity parameter to preview material. Disabled preview has now 50% transparency.

With this change, a disabled preview is no longer fully blocking the view.

diff --git a/Engine/Plugins/Runtime/nDisplay/Source/DisplayCluster/Private/Game/EngineClasses/Scene/DisplayClusterPreviewComponent.cpp b/Engine/Plugins/Runtime/nDisplay/Source/DisplayCluster/Private/Game/EngineClasses/Scene/DisplayClusterPreviewComponent.cpp
--- a/Engine/Plugins/Runtime/nDisplay/Source/DisplayCluster/Private/Game/EngineClasses/Scene/DisplayClusterPreviewComponent.cpp
+++ b/Engine/Plugins/Runtime/nDisplay/Source/DisplayCluster/Private/Game/EngineClasses/Scene/DisplayClusterPreviewComponent.cpp
@@ -246,6 +246,8 @@ void UDisplayClusterPreviewComponent::UpdatePreviewMaterial()
 {
 	if (PreviewMaterialInstance != nullptr)
 	{
+		PreviewMaterialInstance->SetScalarParameterValue(TEXT("Opacity"), 1.0);
+
 		if (OverrideTexture)
 		{
 			PreviewMaterialInstance->SetTextureParameterValue(TEXT("Preview"), OverrideTexture);
