nDisplay: improved global controls over color grading parameters for the viewports and incamera

diff --git a/Engine/Plugins/Runtime/nDisplay/Source/DisplayCluster/Private/Render/Viewport/Configuration/DisplayClusterViewportConfigurationHelpers_Postprocess.cpp b/Engine/Plugins/Runtime/nDisplay/Source/DisplayCluster/Private/Render/Viewport/Configuration/DisplayClusterViewportConfigurationHelpers_Postprocess.cpp
--- a/Engine/Plugins/Runtime/nDisplay/Source/DisplayCluster/Private/Render/Viewport/Configuration/DisplayClusterViewportConfigurationHelpers_Postprocess.cpp
+++ b/Engine/Plugins/Runtime/nDisplay/Source/DisplayCluster/Private/Render/Viewport/Configuration/DisplayClusterViewportConfigurationHelpers_Postprocess.cpp
@@ -261,7 +261,10 @@ bool FDisplayClusterViewportConfigurationHelpers_Postprocess::UpdateLightcardPos
 bool FDisplayClusterViewportConfigurationHelpers_Postprocess::ImplUpdateViewportColorGrading(FDisplayClusterViewport& DstViewport, ADisplayClusterRootActor& RootActor, const FString& InClusterViewportId)
 {
 	const FDisplayClusterConfigurationICVFX_StageSettings& StageSettings = RootActor.GetStageSettings();
-
+	if (StageSettings.EnableColorGrading == false)
+	{
+		return false;
+	}
 	
 	for (const FDisplayClusterConfigurationViewport_PerViewportColorGrading& ColorGradingProfileIt : StageSettings.PerViewportColorGrading)
 	{
@@ -284,12 +287,11 @@ bool FDisplayClusterViewportConfigurationHelpers_Postprocess::ImplUpdateViewport
 
 void FDisplayClusterViewportConfigurationHelpers_Postprocess::UpdateCameraPostProcessSettings(FDisplayClusterViewport& DstViewport, ADisplayClusterRootActor& RootActor, UDisplayClusterICVFXCameraComponent& InCameraComponent)
 {
-	const FDisplayClusterConfigurationICVFX_CameraSettings& CameraSettings = InCameraComponent.GetCameraSettingsICVFX();
-
 	FDisplayClusterConfigurationViewport_CustomPostprocessSettings CameraPPS;
 	CameraPPS.bIsOneFrame = true;
 	CameraPPS.BlendWeight = 1.f;
 
+	const FDisplayClusterConfigurationICVFX_CameraSettings& CameraSettings = InCameraComponent.GetCameraSettingsICVFX();
 	if (CameraSettings.RenderSettings.bUseCameraComponentPostprocess)
 	{
 		FMinimalViewInfo DesiredView;
@@ -316,12 +318,17 @@ void FDisplayClusterViewportConfigurationHelpers_Postprocess::UpdateCameraPostPr
 		CameraPPS.PostProcessSettings.bOverride_MotionBlurPerObjectSize = true;
 	}
 
-	ImplUpdateCustomPostprocess(DstViewport, CameraPPS.bIsEnabled, CameraPPS, IDisplayClusterViewport_CustomPostProcessSettings::ERenderPass::Override);
-
-	if (!ImplUpdateInnerFrustumColorGrading(DstViewport, RootActor, InCameraComponent))
+	const FDisplayClusterConfigurationICVFX_StageSettings& StageSettings = RootActor.GetStageSettings();
+	// check if frustum color grading is enabled	
+	if (StageSettings.EnableColorGrading && CameraSettings.EnableInnerFrustumColorGrading)
 	{
-		// This viewport doesn't use per-viewport PP
-		ImplRemoveCustomPostprocess(DstViewport, IDisplayClusterViewport_CustomPostProcessSettings::ERenderPass::FinalPerViewport);
+		ImplUpdateCustomPostprocess(DstViewport, CameraPPS.bIsEnabled, CameraPPS, IDisplayClusterViewport_CustomPostProcessSettings::ERenderPass::Override);
+
+		if (!ImplUpdateInnerFrustumColorGrading(DstViewport, RootActor, InCameraComponent))
+		{
+			// This viewport doesn't use per-viewport PP
+			ImplRemoveCustomPostprocess(DstViewport, IDisplayClusterViewport_CustomPostProcessSettings::ERenderPass::FinalPerViewport);
+		}
 	}
 }
 
diff --git a/Engine/Plugins/Runtime/nDisplay/Source/DisplayCluster/Public/Components/DisplayClusterICVFXCameraComponent.h b/Engine/Plugins/Runtime/nDisplay/Source/DisplayCluster/Public/Components/DisplayClusterICVFXCameraComponent.h
--- a/Engine/Plugins/Runtime/nDisplay/Source/DisplayCluster/Public/Components/DisplayClusterICVFXCameraComponent.h
+++ b/Engine/Plugins/Runtime/nDisplay/Source/DisplayCluster/Public/Components/DisplayClusterICVFXCameraComponent.h
@@ -142,6 +142,9 @@ private:
 	UPROPERTY(EditAnywhere, Transient, Category = OCIO, meta = (PropertyPath = "CameraSettings.PerNodeOCIOProfiles", EditConditionPath = "CameraSettings.AllNodesOCIOConfiguration.bIsEnabled"))
 	FDisplayClusterEditorPropertyReference PerNodeOCIOProfilesRef;
 
+	UPROPERTY(EditAnywhere, Transient, Category = "Inner Frustum Color Grading", meta = (PropertyPath = "CameraSettings.EnableInnerFrustumColorGrading", DisplayName = "Enable Inner Frustum Color Grading"))
+	FDisplayClusterEditorPropertyReference EnableInnerFrustumColorGrading;
+
 	UPROPERTY(EditAnywhere, Transient, Category = "Inner Frustum Color Grading", meta = (PropertyPath = "CameraSettings.AllNodesColorGrading", DisplayName = "All Nodes"))
 	FDisplayClusterEditorPropertyReference AllNodesColorGradingRef;
 
diff --git a/Engine/Plugins/Runtime/nDisplay/Source/DisplayCluster/Public/DisplayClusterRootActor.h b/Engine/Plugins/Runtime/nDisplay/Source/DisplayCluster/Public/DisplayClusterRootActor.h
--- a/Engine/Plugins/Runtime/nDisplay/Source/DisplayCluster/Public/DisplayClusterRootActor.h
+++ b/Engine/Plugins/Runtime/nDisplay/Source/DisplayCluster/Public/DisplayClusterRootActor.h
@@ -192,10 +192,10 @@ private:
 	UPROPERTY(EditAnywhere, Transient, Category = "In Camera VFX", meta = (PropertyPath = "CurrentConfigData.StageSettings.bEnableInnerFrustums"))
 	FDisplayClusterEditorPropertyReference EnableInnerFrustumsRef;
 
-	UPROPERTY(EditAnywhere, Transient, Category = "Color Grading", meta = (PropertyPath = "CurrentConfigData.StageSettings.EntireClusterColorGrading.bEnableEntireClusterColorGrading"))
-	FDisplayClusterEditorPropertyReference EnableClusterColorGradingRef;
+	UPROPERTY(EditAnywhere, Transient, Category = "Color Grading", meta = (PropertyPath = "CurrentConfigData.StageSettings.EnableColorGrading"))
+	FDisplayClusterEditorPropertyReference EnableColorGradingRef;
 
-	UPROPERTY(EditAnywhere, Transient, Category = "Color Grading", meta = (DisplayName = "Entire Cluster", PropertyPath = "CurrentConfigData.StageSettings.EntireClusterColorGrading.ColorGradingSettings", EditConditionPath = "CurrentConfigData.StageSettings.EntireClusterColorGrading.bEnableEntireClusterColorGrading"))
+	UPROPERTY(EditAnywhere, Transient, Category = "Color Grading", meta = (PropertyPath = "CurrentConfigData.StageSettings.EntireClusterColorGrading"))
 	FDisplayClusterEditorPropertyReference ClusterColorGradingRef;
 
 	UPROPERTY(EditAnywhere, Transient, Category = "Color Grading", meta = (PropertyPath = "CurrentConfigData.StageSettings.PerViewportColorGrading"))
diff --git a/Engine/Plugins/Runtime/nDisplay/Source/DisplayClusterConfiguration/Public/DisplayClusterConfigurationTypes_ICVFX.h b/Engine/Plugins/Runtime/nDisplay/Source/DisplayClusterConfiguration/Public/DisplayClusterConfigurationTypes_ICVFX.h
--- a/Engine/Plugins/Runtime/nDisplay/Source/DisplayClusterConfiguration/Public/DisplayClusterConfigurationTypes_ICVFX.h
+++ b/Engine/Plugins/Runtime/nDisplay/Source/DisplayClusterConfiguration/Public/DisplayClusterConfigurationTypes_ICVFX.h
@@ -514,6 +514,10 @@ public:
 	UPROPERTY(EditAnywhere, BlueprintReadWrite, Category = "OCIO", meta = (DisplayName = "Per-Node OCIO Overrides", ConfigurationMode = "ClusterNodes", EditCondition = "bEnable"))
 	TArray<FDisplayClusterConfigurationOCIOProfile> PerNodeOCIOProfiles;
 
+	/** Entire Cluster Color Grading */
+	UPROPERTY(EditAnywhere, BlueprintReadWrite, Category = "Inner Frustum Color Grading", meta = (DisplayName = "Enable Inner Frustum Color Grading"))
+	bool EnableInnerFrustumColorGrading = true;
+
 	/** All Nodes Color Grading */
 	UPROPERTY(EditAnywhere, BlueprintReadWrite, Category = "Inner Frustum Color Grading", meta = (DisplayName = "All Nodes Color Grading", EditCondition = "bEnable"))
 	FDisplayClusterConfigurationViewport_AllNodesColorGrading AllNodesColorGrading;
@@ -561,6 +565,10 @@ public:
 	UPROPERTY(EditAnywhere, BlueprintReadWrite, Category = Viewports, meta = (DisplayName = "Content Hidden from Viewports", ToolTip = "Content specified here will not appear in the nDisplay viewports, but can appear in the inner frustum.", Substitutions = "LayersTooltip = Layers hidden from the nDisplay viewports, ActorsTooltip = Actors hidden from the nDisplay viewports"))
 	FDisplayClusterConfigurationICVFX_VisibilityList OuterViewportHideList;
 
+	/** Viewport Color Grading */
+	UPROPERTY(EditAnywhere, BlueprintReadWrite, Category = "Viewport Color Grading", meta = (DisplayName = "Enable Color Grading"))
+	bool EnableColorGrading = true;
+
 	/** Entire Cluster Color Grading */
 	UPROPERTY(EditAnywhere, BlueprintReadWrite, Category = "Viewport Color Grading", meta = (DisplayName = "Entire Cluster"))
 	FDisplayClusterConfigurationViewport_EntireClusterColorGrading EntireClusterColorGrading;
