nDisplay: Added logic to the light card actor's tick that locks its world position and rotation to the stage actor that "owns" the light card.

diff --git a/Engine/Plugins/Runtime/nDisplay/Source/DisplayCluster/Private/Game/EngineClasses/Scene/DisplayClusterLightCardActor.cpp b/Engine/Plugins/Runtime/nDisplay/Source/DisplayCluster/Private/Game/EngineClasses/Scene/DisplayClusterLightCardActor.cpp
--- a/Engine/Plugins/Runtime/nDisplay/Source/DisplayCluster/Private/Game/EngineClasses/Scene/DisplayClusterLightCardActor.cpp
+++ b/Engine/Plugins/Runtime/nDisplay/Source/DisplayCluster/Private/Game/EngineClasses/Scene/DisplayClusterLightCardActor.cpp
@@ -21,7 +21,7 @@ ADisplayClusterLightCardActor::ADisplayClusterLightCardActor(const FObjectInitia
 	, Scale(FVector2D(1.f))
 	, Mask(EDisplayClusterLightCardMask::Circle)
 	, Texture(nullptr)
-	, Color(FLinearColor(1.f, 1.f, 1.f, 0.f))
+	, Color(FLinearColor(1.f, 1.f, 1.f, 1.f))
 	, Exposure(0.f)
 	, Opacity(1.f)
 	, Feathering(0.f)
diff --git a/Engine/Plugins/Runtime/nDisplay/Source/DisplayCluster/Private/Game/EngineClasses/Scene/DisplayClusterRootActor.cpp b/Engine/Plugins/Runtime/nDisplay/Source/DisplayCluster/Private/Game/EngineClasses/Scene/DisplayClusterRootActor.cpp
--- a/Engine/Plugins/Runtime/nDisplay/Source/DisplayCluster/Private/Game/EngineClasses/Scene/DisplayClusterRootActor.cpp
+++ b/Engine/Plugins/Runtime/nDisplay/Source/DisplayCluster/Private/Game/EngineClasses/Scene/DisplayClusterRootActor.cpp
@@ -14,6 +14,8 @@
 #include "CineCameraComponent.h"
 #include "ProceduralMeshComponent.h"
 
+#include "DisplayClusterLightCardActor.h"
+
 #include "Config/IPDisplayClusterConfigManager.h"
 #include "DisplayClusterConfigurationStrings.h"
 
@@ -628,6 +630,66 @@ bool ADisplayClusterRootActor::BuildHierarchy()
 	return true;
 }
 
+void ADisplayClusterRootActor::UpdateLightCardPositions()
+{
+	// Find a view origin to use as the transform "anchor" of each light card. At the moment, assume the first view origin component
+	// found is the correct view origin (the same assumption is made in the light card editor). If no view origin is found, use the root component
+	USceneComponent* ViewOriginComponent = GetRootComponent();
+
+	TArray<UDisplayClusterCameraComponent*> ViewOriginComponents;
+	GetComponents<UDisplayClusterCameraComponent>(ViewOriginComponents);
+
+	if (ViewOriginComponents.Num())
+	{
+		ViewOriginComponent = ViewOriginComponents[0];
+	}
+
+	const FVector Location = ViewOriginComponent ? ViewOriginComponent->GetComponentLocation() : GetActorLocation();
+	const FRotator Rotation = GetActorRotation();
+
+	// Iterate over all the light cards referenced in the root actor's config data and update their location and rotation to match the view origin
+	if (UDisplayClusterConfigurationData* CurrentData = GetConfigData())
+	{
+		FDisplayClusterConfigurationICVFX_VisibilityList& LightCards = CurrentData->StageSettings.Lightcard.ShowOnlyList;
+
+		for (const TSoftObjectPtr<AActor>& Actor : LightCards.Actors)
+		{
+			ADisplayClusterLightCardActor* LightCardActor = Actor.IsValid() ? Cast<ADisplayClusterLightCardActor>(Actor.Get()) : nullptr;
+			if (LightCardActor)
+			{
+				if (LightCardActor->bLockToOwningRootActor)
+				{
+					LightCardActor->SetActorLocation(Location);
+					LightCardActor->SetActorRotation(Rotation);
+				}
+			}
+		}
+
+		if (LightCards.ActorLayers.Num())
+		{
+			if (UWorld* World = GetWorld())
+			{
+				for (TActorIterator<ADisplayClusterLightCardActor> Iter(World); Iter; ++Iter)
+				{
+					ADisplayClusterLightCardActor* LightCardActor = *Iter;
+					for (const FActorLayer& ActorLayer : LightCards.ActorLayers)
+					{
+						if (LightCardActor->Layers.Contains(ActorLayer.Name))
+						{
+							if (LightCardActor->bLockToOwningRootActor)
+							{
+								LightCardActor->SetActorLocation(Location);
+								LightCardActor->SetActorRotation(Rotation);
+							}
+							break;
+						}
+					}
+				}
+			}
+		}
+	}
+}
+
 bool ADisplayClusterRootActor::IsBlueprint() const
 {
 	for (UClass* Class = GetClass(); Class; Class = Class->GetSuperClass())
@@ -732,6 +794,14 @@ void ADisplayClusterRootActor::Tick(float DeltaSeconds)
 	Tick_Editor(DeltaSeconds);
 #endif
 
+	if (UDisplayClusterConfigurationData* CurrentData = GetConfigData())
+	{
+		if (CurrentData->StageSettings.Lightcard.bEnable)
+		{
+			UpdateLightCardPositions();
+		}
+	}
+
 	Super::Tick(DeltaSeconds);
 }
 
diff --git a/Engine/Plugins/Runtime/nDisplay/Source/DisplayCluster/Public/DisplayClusterLightCardActor.h b/Engine/Plugins/Runtime/nDisplay/Source/DisplayCluster/Public/DisplayClusterLightCardActor.h
--- a/Engine/Plugins/Runtime/nDisplay/Source/DisplayCluster/Public/DisplayClusterLightCardActor.h
+++ b/Engine/Plugins/Runtime/nDisplay/Source/DisplayCluster/Public/DisplayClusterLightCardActor.h
@@ -101,6 +101,10 @@ public:
 	UPROPERTY(EditAnywhere, BlueprintReadWrite, Category = "Appearance", meta = (UIMin = 0, ClampMin = 0))
 	float Feathering;
 
+	/** A flag that controls wether the light card's location and rotation are locked to its "owning" root actor */
+	UPROPERTY()
+	bool bLockToOwningRootActor = true;
+
 	/** Used to flag this light card as a proxy of a "real" light card. Used by the LightCard Editor */
 	UPROPERTY(Transient)
 	bool bIsProxy = false;
diff --git a/Engine/Plugins/Runtime/nDisplay/Source/DisplayCluster/Public/DisplayClusterRootActor.h b/Engine/Plugins/Runtime/nDisplay/Source/DisplayCluster/Public/DisplayClusterRootActor.h
--- a/Engine/Plugins/Runtime/nDisplay/Source/DisplayCluster/Public/DisplayClusterRootActor.h
+++ b/Engine/Plugins/Runtime/nDisplay/Source/DisplayCluster/Public/DisplayClusterRootActor.h
@@ -115,6 +115,9 @@ protected:
 	// Creates all hierarchy objects declared in a config file
 	bool BuildHierarchy();
 
+	/** Updates the world position and rotation of each light card referenced by this root actor's light card list to match the default view origin */
+	void UpdateLightCardPositions();
+
 public:
 	UFUNCTION(BlueprintGetter)
 	UDisplayClusterCameraComponent* GetDefaultCamera() const;
diff --git a/Engine/Plugins/Runtime/nDisplay/Source/DisplayClusterLightCardEditor/Private/SDisplayClusterLightCardList.cpp b/Engine/Plugins/Runtime/nDisplay/Source/DisplayClusterLightCardEditor/Private/SDisplayClusterLightCardList.cpp
--- a/Engine/Plugins/Runtime/nDisplay/Source/DisplayClusterLightCardEditor/Private/SDisplayClusterLightCardList.cpp
+++ b/Engine/Plugins/Runtime/nDisplay/Source/DisplayClusterLightCardEditor/Private/SDisplayClusterLightCardList.cpp
@@ -373,18 +373,14 @@ void SDisplayClusterLightCardList::AddNewLightCard()
 	SpawnParameters.SpawnCollisionHandlingOverride = ESpawnActorCollisionHandlingMethod::AdjustIfPossibleButAlwaysSpawn;
 	SpawnParameters.Name = TEXT("LightCard");
 	SpawnParameters.NameMode = FActorSpawnParameters::ESpawnActorNameMode::Requested;
-	SpawnParameters.OverrideLevel = RootActor->GetLevel();
+	SpawnParameters.OverrideLevel = RootActor->GetWorld()->GetCurrentLevel();
 	
 	ADisplayClusterLightCardActor* NewActor = CastChecked<ADisplayClusterLightCardActor>(
 		RootActor->GetWorld()->SpawnActor(ADisplayClusterLightCardActor::StaticClass(),
 		&SpawnLocation, &SpawnRotation, MoveTemp(SpawnParameters)));
 
-	ensure(NewActor->GetLevel() == RootActor->GetLevel());
 	NewActor->SetActorLabel(NewActor->GetName());
 
-	// Parent to the root actor, since the typical intention is for the LC to move with the stage.
-	NewActor->AttachToActor(RootActor.Get(), FAttachmentTransformRules::KeepWorldTransform);
-
 	AddLightCardToActor(NewActor);
 
 	// When adding a new lightcard, usually the desired location is in the middle of the viewport
