nDisplay: Behavior of moving groups of cards is now more intuitive.

* Fixed sudden spins of lightcards when pushing gizmo near the north and south poles in the LC Editor.
* Fixed translation jumps of grouped LCs when primary one passes through poles in the LC Editor.
* Changed LC positional parameters from float to double

diff --git a/Engine/Plugins/Runtime/nDisplay/Source/DisplayCluster/Private/Game/EngineClasses/Scene/DisplayClusterLightCardActor.cpp b/Engine/Plugins/Runtime/nDisplay/Source/DisplayCluster/Private/Game/EngineClasses/Scene/DisplayClusterLightCardActor.cpp
--- a/Engine/Plugins/Runtime/nDisplay/Source/DisplayCluster/Private/Game/EngineClasses/Scene/DisplayClusterLightCardActor.cpp
+++ b/Engine/Plugins/Runtime/nDisplay/Source/DisplayCluster/Private/Game/EngineClasses/Scene/DisplayClusterLightCardActor.cpp
@@ -79,10 +79,10 @@ void ADisplayClusterLightCardActor::Tick(float DeltaSeconds)
 		// If latitude exceeds [-90, 90], mod it back into the appropriate range, and apply a shift of 180 degrees if
 		// needed to the longitude, to allow the latitude to be continuous (increasing latitude indefinitely should result in the LC 
 		// orbiting around a polar great circle)
-		float Parity = FMath::Fmod(FMath::Abs(Latitude) + 90, 360) - 180;
-		float DeltaLongitude = Parity > 1 ? 180.f : 0.f;
+		double Parity = FMath::Fmod(FMath::Abs(Latitude) + 90, 360) - 180;
+		double DeltaLongitude = Parity > 1 ? 180.f : 0.f;
 
-		float LatMod = FMath::Fmod(Latitude + 90.f, 180.f);
+		double LatMod = FMath::Fmod(Latitude + 90.f, 180.f);
 		if (LatMod < 0.f)
 		{
 			LatMod += 180.f;
@@ -204,3 +204,27 @@ void ADisplayClusterLightCardActor::SetStaticMesh(UStaticMesh* InStaticMesh)
 {
 	LightCardComponent->SetStaticMesh(InStaticMesh);
 }
+
+ADisplayClusterLightCardActor::PositionalParams ADisplayClusterLightCardActor::GetPositionalParams()
+{
+	PositionalParams Params;
+
+	Params.DistanceFromCenter = DistanceFromCenter;
+	Params.Latitude = Latitude;
+	Params.Longitude = Longitude;
+	Params.Pitch = Pitch;
+	Params.Spin = Spin;
+	Params.Yaw = Yaw;
+
+	return Params;
+}
+
+void ADisplayClusterLightCardActor::SetPositionalParams(const PositionalParams& Params)
+{
+	DistanceFromCenter = Params.DistanceFromCenter;
+	Latitude = Params.Latitude;
+	Longitude = Params.Longitude;
+	Pitch = Params.Pitch;
+	Spin = Params.Spin;
+	Yaw = Params.Yaw;
+}
diff --git a/Engine/Plugins/Runtime/nDisplay/Source/DisplayCluster/Public/DisplayClusterLightCardActor.h b/Engine/Plugins/Runtime/nDisplay/Source/DisplayCluster/Public/DisplayClusterLightCardActor.h
--- a/Engine/Plugins/Runtime/nDisplay/Source/DisplayCluster/Public/DisplayClusterLightCardActor.h
+++ b/Engine/Plugins/Runtime/nDisplay/Source/DisplayCluster/Public/DisplayClusterLightCardActor.h
@@ -25,6 +25,18 @@ class DISPLAYCLUSTER_API ADisplayClusterLightCardActor : public AActor
 {
 	GENERATED_BODY()
 
+public:
+
+	struct PositionalParams
+	{
+		double DistanceFromCenter;
+		double Longitude;
+		double Latitude;
+		double Spin;
+		double Pitch;
+		double Yaw;
+	};
+
 public:
 	/** The rotation used to orient the plane mesh used for the light card so that its normal points radially inwards */
 	static const FRotator PlaneMeshRotation;
@@ -58,27 +70,33 @@ public:
 	/** Updates the Light Card transform based on its positional properties (Lat, Long, etc.) */
 	void UpdateLightCardTransform();
 
+	/** Retrieves positional parameters */
+	PositionalParams GetPositionalParams();
+
+	/** Set positional parameters */
+	void SetPositionalParams(const PositionalParams& Params);
+
 protected:
 	void UpdateLightCardMaterialInstance();
 
 public:
 	UPROPERTY(EditAnywhere, BlueprintReadWrite, Category = "Orientation")
-	float DistanceFromCenter;
+	double DistanceFromCenter;
 
 	UPROPERTY(EditAnywhere, BlueprintReadWrite, Category = "Orientation", meta = (UIMin = 0, ClampMin = 0, UIMax = 360, ClampMax = 360))
-	float Longitude;
+	double Longitude;
 
 	UPROPERTY(EditAnywhere, BlueprintReadWrite, Category = "Orientation", meta = (UIMin = -90, ClampMin = -90, UIMax = 90, ClampMax = 90))
-	float Latitude;
+	double Latitude;
 
 	UPROPERTY(EditAnywhere, BlueprintReadWrite, Category = "Orientation")
-	float Spin;
+	double Spin;
 
 	UPROPERTY(EditAnywhere, BlueprintReadWrite, Category = "Orientation")
-	float Pitch;
+	double Pitch;
 
 	UPROPERTY(EditAnywhere, BlueprintReadWrite, Category = "Orientation")
-	float Yaw;
+	double Yaw;
 
 	UPROPERTY(EditAnywhere, BlueprintReadWrite, Category = "Orientation")
 	FVector2D Scale;
diff --git a/Engine/Plugins/Runtime/nDisplay/Source/DisplayClusterLightCardEditor/Private/Viewport/DisplayClusterLightCardEditorViewportClient.cpp b/Engine/Plugins/Runtime/nDisplay/Source/DisplayClusterLightCardEditor/Private/Viewport/DisplayClusterLightCardEditorViewportClient.cpp
--- a/Engine/Plugins/Runtime/nDisplay/Source/DisplayClusterLightCardEditor/Private/Viewport/DisplayClusterLightCardEditorViewportClient.cpp
+++ b/Engine/Plugins/Runtime/nDisplay/Source/DisplayClusterLightCardEditor/Private/Viewport/DisplayClusterLightCardEditorViewportClient.cpp
@@ -46,6 +46,133 @@
 #define LOCTEXT_NAMESPACE "DisplayClusterLightCardEditorViewportClient"
 
 
+FDisplayClusterLightCardEditorViewportClient::FSphericalCoordinates::FSphericalCoordinates(const FVector& CartesianPosition)
+{
+	Radius = CartesianPosition.Size();
+
+	if (Radius > UE_SMALL_NUMBER)
+	{
+		Inclination = FMath::Acos(CartesianPosition.Z / Radius);
+	}
+	else
+	{
+		Inclination = 0;
+	}
+
+	Azimuth = FMath::Atan2(CartesianPosition.Y, CartesianPosition.X);
+}
+
+FDisplayClusterLightCardEditorViewportClient::FSphericalCoordinates::FSphericalCoordinates()
+	: Radius(0)
+	, Inclination(0)
+	, Azimuth(0)
+{ 
+
+}
+
+FVector FDisplayClusterLightCardEditorViewportClient::FSphericalCoordinates::AsCartesian() const
+{
+	const double SinAzimuth = FMath::Sin(Azimuth);
+	const double CosAzimuth = FMath::Cos(Azimuth);
+
+	const double SinInclination = FMath::Sin(Inclination);
+	const double CosInclination = FMath::Cos(Inclination);
+
+	return FVector(
+		Radius * CosAzimuth * SinInclination,
+		Radius * SinAzimuth * SinInclination,
+		Radius * CosInclination
+	);
+}
+
+FDisplayClusterLightCardEditorViewportClient::FSphericalCoordinates 
+FDisplayClusterLightCardEditorViewportClient::FSphericalCoordinates::operator+(FSphericalCoordinates const& Other) const
+{
+	FSphericalCoordinates Result;
+
+	Result.Radius = Radius + Other.Radius;
+	Result.Inclination = Inclination + Other.Inclination;
+	Result.Azimuth = Azimuth + Other.Azimuth;
+
+	return Result;
+}
+
+FDisplayClusterLightCardEditorViewportClient::FSphericalCoordinates 
+FDisplayClusterLightCardEditorViewportClient::FSphericalCoordinates::operator-(FSphericalCoordinates const& Other) const
+{
+	FSphericalCoordinates Result;
+
+	Result.Radius = Radius - Other.Radius;
+	Result.Inclination = Inclination - Other.Inclination;
+	Result.Azimuth = Azimuth - Other.Azimuth;
+
+	return Result;
+}
+
+void FDisplayClusterLightCardEditorViewportClient::FSphericalCoordinates::Conform()
+{
+	if (Radius < 0)
+	{
+		Radius = -Radius;
+		Inclination += PI;
+	}
+
+	if (Inclination < 0 || Inclination > PI)
+	{
+		// -2PI to 2PI
+		Inclination = FMath::Fmod(Inclination, 2 * PI);
+
+		// 0 to 2PI
+		if (Inclination < 0)
+		{
+			Inclination += 2 * PI;
+		}
+
+		// 0 to PI
+		if (Inclination > PI)
+		{
+			Inclination = 2 * PI - Inclination;
+			Azimuth += PI;
+		}
+	}
+
+	if (Azimuth < -PI || Azimuth > PI)
+	{
+		// -2PI to 2PI
+		Azimuth = FMath::Fmod(Azimuth, 2 * PI);
+
+		// -PI to PI
+		if (Azimuth > PI)
+		{
+			Azimuth -= 2 * PI;
+		}
+		else if (Azimuth < -PI)
+		{
+			Azimuth += 2 * PI;
+		}
+	}
+
+	checkSlow(Radius >= 0);
+	checkSlow(Inclination >= 0 && Inclination <= PI);
+	checkSlow(Azimuth >= -PI && Azimuth <= PI);
+}
+
+FDisplayClusterLightCardEditorViewportClient::FSphericalCoordinates 
+FDisplayClusterLightCardEditorViewportClient::FSphericalCoordinates::GetConformed() const
+{
+	FSphericalCoordinates Result = *this;
+	Result.Conform();
+	return Result;
+}
+
+bool FDisplayClusterLightCardEditorViewportClient::FSphericalCoordinates::IsPointingAtPole(double Margin) const
+{
+	FSphericalCoordinates CoordsConformed = GetConformed();
+
+	return FMath::IsNearlyZero(CoordsConformed.Inclination, Margin)
+		|| FMath::IsNearlyEqual(CoordsConformed.Inclination, PI, Margin);
+}
+
 //////////////////////////////////////////////////////////////////////////
 // FNormalMap
 
@@ -1059,12 +1186,10 @@ void FDisplayClusterLightCardEditorViewportClient::MoveLightCardTo(ADisplayClust
 		NorthNormalMap.GetNormalAndDistanceAtPosition(LightCardPosition, DesiredNormal, DesiredDistance);
 	}
 
-	const FRotator Rotation = FRotationMatrix::MakeFromX(-DesiredNormal).Rotator();
-
+	SetLightCardCoordinates(&LightCard, SphericalCoords);
 	LightCard.DistanceFromCenter = FMath::Min(DesiredDistance, RootActorBoundingRadius) + LightCardFlushOffset;
 
-	LightCard.Longitude = FRotator::ClampAxis(FMath::RadiansToDegrees(SphericalCoords.Azimuth) - 180);
-	LightCard.Latitude = 90.f - FMath::RadiansToDegrees(SphericalCoords.Inclination);
+	const FRotator Rotation = FRotationMatrix::MakeFromX(-DesiredNormal).Rotator();
 
 	LightCard.Pitch = Rotation.Pitch;
 	LightCard.Yaw = Rotation.Yaw;
@@ -1474,6 +1599,25 @@ void FDisplayClusterLightCardEditorViewportClient::PropagateLightCardTransform(A
 	}
 }
 
+void FDisplayClusterLightCardEditorViewportClient::SetLightCardCoordinates(
+	ADisplayClusterLightCardActor* LightCard, const FSphericalCoordinates& SphericalCoords) const
+{
+	if (!LightCard)
+	{
+		return;
+	}
+
+	LightCard->DistanceFromCenter = SphericalCoords.Radius;
+	LightCard->Latitude = 90.f - FMath::RadiansToDegrees(SphericalCoords.Inclination);
+
+	// Keep the same longitude when pointing at the pole. This helps with continuity 
+	// and also mitigates sudden changes in apparent spin when moving around the poles
+	if (!SphericalCoords.IsPointingAtPole())
+	{
+		LightCard->Longitude = FRotator::ClampAxis(FMath::RadiansToDegrees(SphericalCoords.Azimuth) - 180);
+	}
+}
+
 void FDisplayClusterLightCardEditorViewportClient::VerifyAndFixLightCardOrigin(ADisplayClusterLightCardActor* LightCard) const
 {
 	// Center lightcard on the current view origin, let it keep its current world placement 
@@ -1504,9 +1648,7 @@ void FDisplayClusterLightCardEditorViewportClient::VerifyAndFixLightCardOrigin(A
 
 	const FSphericalCoordinates SphericalCoords(LightCardRelativeLocation);
 
-	LightCard->DistanceFromCenter = SphericalCoords.Radius;
-	LightCard->Longitude = FMath::RadiansToDegrees(SphericalCoords.Radius) - 180;
-	LightCard->Latitude = 90.f - FMath::RadiansToDegrees(SphericalCoords.Radius);
+	SetLightCardCoordinates(LightCard, SphericalCoords);
 }
 
 void FDisplayClusterLightCardEditorViewportClient::MoveSelectedLightCards(FViewport* InViewport, EAxisList::Type CurrentAxis)
@@ -1518,7 +1660,9 @@ void FDisplayClusterLightCardEditorViewportClient::MoveSelectedLightCards(FViewp
 		InViewport,
 		GetScene(),
 		EngineShowFlags)
-		.SetRealtimeUpdate(IsRealtime()));
+		.SetRealtimeUpdate(IsRealtime())
+	);
+
 	FSceneView* View = CalcSceneView(&ViewFamily);
 
 	FVector Origin;
@@ -1526,9 +1670,20 @@ void FDisplayClusterLightCardEditorViewportClient::MoveSelectedLightCards(FViewp
 	PixelToWorld(*View, MousePos, Origin, Direction);
 
 	const TWeakObjectPtr<ADisplayClusterLightCardActor>& LastSelectedLightCard = SelectedLightCards.Last();
+
 	if (LastSelectedLightCard.IsValid())
 	{
-		const FSphericalCoordinates DeltaCoords = GetLightCardTranslationDelta(InViewport, LastSelectedLightCard.Get(), CurrentAxis);
+		const bool bUseDeltaRotation = (CurrentAxis == EAxisList::Type::XYZ) || (CurrentAxis == EAxisList::Type::Y);
+
+		const FRotator DeltaRotation = 
+			bUseDeltaRotation ?
+			GetLightCardRotationDelta(InViewport, LastSelectedLightCard.Get(), CurrentAxis) 
+			: FRotator::ZeroRotator;
+
+		const FSphericalCoordinates DeltaCoords = 
+			bUseDeltaRotation ?
+			FSphericalCoordinates()
+			: GetLightCardTranslationDelta(InViewport, LastSelectedLightCard.Get(), CurrentAxis);
 
 		for (const TWeakObjectPtr<ADisplayClusterLightCardActor>& LightCard : SelectedLightCards)
 		{
@@ -1539,22 +1694,25 @@ void FDisplayClusterLightCardEditorViewportClient::MoveSelectedLightCards(FViewp
 
 			VerifyAndFixLightCardOrigin(LightCard.Get());
 
-			FSphericalCoordinates NewCoords;
+			// Note: GetLightCardCoordinates maintains last known Azimuth when looking at the poles
+			const FSphericalCoordinates CurrentCoords = GetLightCardCoordinates(LightCard.Get()); 
+
+			// We will adjust the spin (to maintain the apparent spin) when using center of gizmo
+			// or dragging longitudinally (this seems to provide an intuitive behavior)
+			if (CurrentAxis == EAxisList::Type::XYZ || CurrentAxis == EAxisList::Type::Y) // Dragging center of gizmo
 			{
-				const FSphericalCoordinates CurrentCoords = GetLightCardCoordinates(LightCard.Get());
+				// We might need this to put back the LightCard exactly as it was
+				const ADisplayClusterLightCardActor::PositionalParams OriginalPositionalParams = LightCard->GetPositionalParams();
 
-				NewCoords.Radius = CurrentCoords.Radius + DeltaCoords.Radius;
-				NewCoords.Azimuth = CurrentCoords.Azimuth + DeltaCoords.Azimuth;
-				NewCoords.Inclination = CurrentCoords.Inclination + DeltaCoords.Inclination;
-			}
+				const FVector CurrentPos = CurrentCoords.AsCartesian();
+				const FVector NewPos = DeltaRotation.RotateVector(CurrentPos);
+
+				// Calculations are only valid if translation is not too small
+				const FSphericalCoordinates NewCoords(NewPos);
 
-			// We will only adjust the spin (to maintain the apparent spin) when not using a lat/long axis to drag
-			if (CurrentAxis == EAxisList::Type::XYZ)
-			{
 				const FTransform Transform_A = LightCard->GetLightCardTransform(false /*bIgnoreSpinYawPitch*/);
 
 				MoveLightCardTo(*LightCard.Get(), NewCoords);
-
 				LightCard->UpdateLightCardTransform(); // We must call this for GetLightCardTransform to be valid
 
 				const FTransform Transform_B = LightCard->GetLightCardTransform(false /*bIgnoreSpinYawPitch*/);
@@ -1562,21 +1720,20 @@ void FDisplayClusterLightCardEditorViewportClient::MoveSelectedLightCards(FViewp
 				// Calculate world delta translation of moving from A to B
 				const FVector WorldDelta = Transform_B.GetLocation() - Transform_A.GetLocation(); // X towards front of stage. Y towards right of stage. Z towards ceiling.
 
-				// Calculations are only valid if translation is not too small
-				if (WorldDelta.Length() > KINDA_SMALL_NUMBER)
-				{
-					// Calculate LC "Y" unit vector at A and B. ("X" is LC normal)
-					const FVector Y_A = Transform_A.Rotator().RotateVector(FVector::YAxisVector);
-					const FVector Y_B = Transform_B.Rotator().RotateVector(FVector::YAxisVector);
+				// Calculate LC "Y" unit vector at A and B. ("X" is LC normal)
+				const FVector Y_A = Transform_A.Rotator().RotateVector(FVector::YAxisVector);
+				const FVector Y_B = Transform_B.Rotator().RotateVector(FVector::YAxisVector);
 
-					// Calculate card normal vector
-					const FVector CardNormal_A = Transform_A.Rotator().RotateVector(FVector::XAxisVector); // When card is on ceiling, expect around (0,0,-1).
-					const FVector CardNormal_B = Transform_B.Rotator().RotateVector(FVector::XAxisVector);
+				// Calculate card normal vector
+				const FVector CardNormal_A = Transform_A.Rotator().RotateVector(FVector::XAxisVector); // When card is on ceiling, expect around (0,0,-1).
+				const FVector CardNormal_B = Transform_B.Rotator().RotateVector(FVector::XAxisVector);
 
-					// Calculate projection of movement onto surface tangent plane at A and B
-					const FVector UnitWorldDeltaPlane_A = FVector::VectorPlaneProject(WorldDelta, CardNormal_A).GetSafeNormal();
-					const FVector UnitWorldDeltaPlane_B = FVector::VectorPlaneProject(WorldDelta, CardNormal_B).GetSafeNormal();
+				// Calculate projection of movement onto surface tangent plane at A and B
+				const FVector UnitWorldDeltaPlane_A = FVector::VectorPlaneProject(WorldDelta, CardNormal_A).GetSafeNormal();
+				const FVector UnitWorldDeltaPlane_B = FVector::VectorPlaneProject(WorldDelta, CardNormal_B).GetSafeNormal();
 
+				if (!FMath::IsNearlyZero(UnitWorldDeltaPlane_A.Length()) && !FMath::IsNearlyZero(UnitWorldDeltaPlane_B.Length()))
+				{
 					// Calculate relative spin angle at A, which is the angle between Y_A and UnitWorldDeltaPlane_A
 					const double SpinDotProduct_A = FVector::DotProduct(Y_A, UnitWorldDeltaPlane_A);
 					const FVector SpinCrossProduct_A = FVector::CrossProduct(Y_A, UnitWorldDeltaPlane_A);
@@ -1594,10 +1751,15 @@ void FDisplayClusterLightCardEditorViewportClient::MoveSelectedLightCards(FViewp
 					// Apply delta spin to lightcard
 					LightCard->Spin += FMath::RadiansToDegrees(DeltaSpin);
 				}
+				else
+				{
+					// Leave it where it was to avoid apparent spins even though motion would have been insignificant.
+					LightCard->SetPositionalParams(OriginalPositionalParams);
+				}
 			}
-			else
+			else // Dragging latitudinally
 			{
-				MoveLightCardTo(*LightCard.Get(), NewCoords);
+				MoveLightCardTo(*LightCard.Get(), CurrentCoords + DeltaCoords);
 			}
 
 			PropagateLightCardTransform(LightCard.Get());
@@ -1609,6 +1771,7 @@ void FDisplayClusterLightCardEditorViewportClient::MoveSelectedLightCards(FViewp
 
 		FVector LightCardProjectedLocation = ProjectWorldPosition(CachedEditorWidgetWorldTransform.GetTranslation(), ViewMatrices);
 		FVector2D ScreenPercentage;
+
 		if (IsLocationCloseToEdge(LightCardProjectedLocation, InViewport, View, &ScreenPercentage))
 		{
 			DesiredLookAtSpeed = FMath::Max(ScreenPercentage.X, ScreenPercentage.Y) * MaxDesiredLookAtSpeed;
@@ -1621,6 +1784,33 @@ void FDisplayClusterLightCardEditorViewportClient::MoveSelectedLightCards(FViewp
 	}
 }
 
+FRotator FDisplayClusterLightCardEditorViewportClient::GetLightCardRotationDelta(
+	FViewport* InViewport,
+	ADisplayClusterLightCardActor* LightCard,
+	EAxisList::Type CurrentAxis)
+{
+	check(LightCard);
+
+	const FSphericalCoordinates DeltaCoords = GetLightCardTranslationDelta(InViewport, LightCard, CurrentAxis);
+	const FSphericalCoordinates LightCardCoords = GetLightCardCoordinates(LightCard);
+	const FSphericalCoordinates NewCoords = LightCardCoords + DeltaCoords;
+
+	const FVector LightCardPos = LightCardCoords.AsCartesian();
+	const FVector NewPos = NewCoords.AsCartesian();
+
+	const FVector PosCrossProduct = FVector::CrossProduct(LightCardPos, NewPos);
+
+	if (FMath::IsNearlyZero(PosCrossProduct.Length()))
+	{
+		return FQuat(FVector::ForwardVector, 0).Rotator();
+	}
+
+	const FVector AxisOfRotation = PosCrossProduct.GetSafeNormal();
+	const double Angle = FMath::Acos(FVector::DotProduct(LightCardPos.GetSafeNormal(), NewPos.GetSafeNormal()));
+
+	return FQuat(AxisOfRotation, Angle).Rotator();
+}
+
 FDisplayClusterLightCardEditorViewportClient::FSphericalCoordinates FDisplayClusterLightCardEditorViewportClient::GetLightCardTranslationDelta(
 	FViewport* InViewport, 
 	ADisplayClusterLightCardActor* LightCard,
@@ -1658,13 +1848,10 @@ FDisplayClusterLightCardEditorViewportClient::FSphericalCoordinates FDisplayClus
 		NorthNormalMap.GetNormalAndDistanceAtPosition(LightCard->GetLightCardTransform().GetTranslation(), Normal, Distance);
 	}
 
-	FSphericalCoordinates LightCardCoords = GetLightCardCoordinates(LightCard);
-	FSphericalCoordinates RequestedCoords(LocalDirection * Distance);
-	
-	FSphericalCoordinates DeltaCoords;
-	DeltaCoords.Radius = RequestedCoords.Radius - LightCardCoords.Radius;
-	DeltaCoords.Azimuth = RequestedCoords.Azimuth - LightCardCoords.Azimuth;
-	DeltaCoords.Inclination = RequestedCoords.Inclination - LightCardCoords.Inclination;
+	const FSphericalCoordinates LightCardCoords = GetLightCardCoordinates(LightCard);
+	const FSphericalCoordinates RequestedCoords(LocalDirection * Distance);
+
+	FSphericalCoordinates DeltaCoords = RequestedCoords - LightCardCoords;
 
 	if (CurrentAxis == EAxisList::Type::X)
 	{
@@ -1674,10 +1861,13 @@ FDisplayClusterLightCardEditorViewportClient::FSphericalCoordinates FDisplayClus
 	{
 		// Convert the inclination to Cartesian coordinates, project it to the x-z plane, and convert back to spherical coordinates. This ensures that the motion in the inclination
 		// plane always lines up with the mouse's projected location along that plane
-		float FixedInclination = FMath::Abs(FMath::Atan2(FMath::Cos(DeltaCoords.Azimuth) * FMath::Sin(RequestedCoords.Inclination), FMath::Cos(RequestedCoords.Inclination)));
+		const double FixedInclination = FMath::Abs(FMath::Atan2(
+			FMath::Cos(DeltaCoords.Azimuth) * FMath::Sin(RequestedCoords.Inclination), 
+			FMath::Cos(RequestedCoords.Inclination))
+		);
 
 		// When translating along the inclination axis, the azimuth delta can only be intervals of pi
-		float FixedAzimuth = FMath::RoundToInt(DeltaCoords.Azimuth / PI) * PI;
+		const double FixedAzimuth = FMath::RoundToInt(DeltaCoords.Azimuth / PI) * PI;
 
 		DeltaCoords.Azimuth = FixedAzimuth;
 		DeltaCoords.Inclination = FixedInclination - LightCardCoords.Inclination;
@@ -1753,11 +1943,11 @@ FVector2D FDisplayClusterLightCardEditorViewportClient::GetLightCardScaleDelta(F
 	const bool bLocalSpace = true;
 	const FVector LightCardSize3D = LightCard->GetLightCardBounds(bLocalSpace).GetSize();
 	const FVector2D SizeToScale = FVector2D((CurrentAxis & EAxisList::X) * LightCardSize3D.Y, (CurrentAxis & EAxisList::Y) * LightCardSize3D.Z);
-	const float DistanceFromCamera = FMath::Max(FVector::Dist(WorldWidgetOrigin, View->ViewMatrices.GetViewOrigin()), 1.0f);
-	const float ScreenSize = View->ViewMatrices.GetScreenScale() * SizeToScale.Length() / DistanceFromCamera;
+	const double DistanceFromCamera = FMath::Max(FVector::Dist(WorldWidgetOrigin, View->ViewMatrices.GetViewOrigin()), 1.0f);
+	const double ScreenSize = View->ViewMatrices.GetScreenScale() * SizeToScale.Length() / DistanceFromCamera;
 
 	// Compute the scale delta as s' - s = 2d / h_0
-	const float ScaleMagnitude = 2.0f * (ScaleDir | DragDir) / ScreenSize;
+	const double ScaleMagnitude = 2.0f * (ScaleDir | DragDir) / ScreenSize;
 
 	FVector2D ScaleDelta = FVector2D((CurrentAxis & EAxisList::X) * ScaleMagnitude, (CurrentAxis & EAxisList::Y) * ScaleMagnitude);
 
@@ -1776,7 +1966,7 @@ void FDisplayClusterLightCardEditorViewportClient::SpinSelectedLightCards(FViewp
 	const TWeakObjectPtr<ADisplayClusterLightCardActor>& LastSelectedLightCard = SelectedLightCards.Last();
 	if (LastSelectedLightCard.IsValid())
 	{
-		const float DeltaSpin = GetLightCardSpinDelta(InViewport, LastSelectedLightCard.Get());
+		const double DeltaSpin = GetLightCardSpinDelta(InViewport, LastSelectedLightCard.Get());
 		for (const TWeakObjectPtr<ADisplayClusterLightCardActor>& LightCard : SelectedLightCards)
 		{
 			if (!LightCard.IsValid())
@@ -1791,7 +1981,7 @@ void FDisplayClusterLightCardEditorViewportClient::SpinSelectedLightCards(FViewp
 	}
 }
 
-float FDisplayClusterLightCardEditorViewportClient::GetLightCardSpinDelta(FViewport* InViewport, ADisplayClusterLightCardActor* LightCard)
+double FDisplayClusterLightCardEditorViewportClient::GetLightCardSpinDelta(FViewport* InViewport, ADisplayClusterLightCardActor* LightCard)
 {
 	FIntPoint MousePos;
 	InViewport->GetMousePos(MousePos);
@@ -1818,8 +2008,8 @@ float FDisplayClusterLightCardEditorViewportClient::GetLightCardSpinDelta(FViewp
 	FVector2D MousePosOffset = FVector2D(MousePos) - ScreenOrigin;
 	FVector2D LastMousePosOffset = FVector2D(LastWidgetMousePos) - ScreenOrigin;
 
-	float Theta = FMath::Atan2(MousePosOffset | ScreenYAxis, MousePosOffset | ScreenXAxis);
-	float LastTheta = FMath::Atan2(LastMousePosOffset | ScreenYAxis, LastMousePosOffset | ScreenXAxis);
+	double Theta = FMath::Atan2(MousePosOffset | ScreenYAxis, MousePosOffset | ScreenXAxis);
+	double LastTheta = FMath::Atan2(LastMousePosOffset | ScreenYAxis, LastMousePosOffset | ScreenXAxis);
 
 	return FMath::RadiansToDegrees(Theta - LastTheta);
 }
@@ -1830,10 +2020,11 @@ FDisplayClusterLightCardEditorViewportClient::FSphericalCoordinates FDisplayClus
 
 	FSphericalCoordinates LightCardCoords(LightCardLocation);
 
-	// If the light card inclination is 0 or 180, the spherical coordinates will have an
-	// "undefined" azimuth value. For continuity when dragging a light card positioned there, we can manually
-	// set the azimuthal value to match the light card's configured longitude
-	if (LightCardCoords.Inclination == 0.f || LightCardCoords.Inclination == PI)
+	// If the light card points at any of the poles, the spherical coordinates will have an "undefined" azimuth value. 
+	// For continuity when dragging a light card positioned there, 
+	// we can manually set the azimuthal value to match the light card's configured longitude
+
+	if (LightCardCoords.IsPointingAtPole())
 	{
 		LightCardCoords.Azimuth = FMath::DegreesToRadians(LightCard->Longitude + 180);
 	}
diff --git a/Engine/Plugins/Runtime/nDisplay/Source/DisplayClusterLightCardEditor/Private/Viewport/DisplayClusterLightCardEditorViewportClient.h b/Engine/Plugins/Runtime/nDisplay/Source/DisplayClusterLightCardEditor/Private/Viewport/DisplayClusterLightCardEditorViewportClient.h
--- a/Engine/Plugins/Runtime/nDisplay/Source/DisplayClusterLightCardEditor/Private/Viewport/DisplayClusterLightCardEditorViewportClient.h
+++ b/Engine/Plugins/Runtime/nDisplay/Source/DisplayClusterLightCardEditor/Private/Viewport/DisplayClusterLightCardEditorViewportClient.h
@@ -24,46 +24,32 @@ private:
 	struct FSphericalCoordinates
 	{
 	public:
-		FSphericalCoordinates(const FVector& CartesianPosition)
-		{
-			Radius = CartesianPosition.Size();
-
-			if (Radius > UE_SMALL_NUMBER)
-			{
-				Inclination = FMath::Acos(CartesianPosition.Z / Radius);
-			}
-			else
-			{
-				Inclination = 0;
-			}
-
-			Azimuth = FMath::Atan2(CartesianPosition.Y, CartesianPosition.X);
-		}
-
-		FSphericalCoordinates()
-			: Radius(0)
-			, Inclination(0)
-			, Azimuth(0)
-		{ }
 
-		FVector AsCartesian() const
-		{
-			const double SinAzimuth = FMath::Sin(Azimuth);
-			const double CosAzimuth = FMath::Cos(Azimuth);
+		/** Constructors */
+		FSphericalCoordinates(const FVector& CartesianPosition);
+		FSphericalCoordinates();
+
+		/** Return equivalent cartesian coordinates */
+		FVector AsCartesian() const;
+
+		/** Addition operator */
+		FSphericalCoordinates operator+(FSphericalCoordinates const& Other) const;
 
-			const double SinInclination = FMath::Sin(Inclination);
-			const double CosInclination = FMath::Cos(Inclination);
+		/** Subtraction operator */
+		FSphericalCoordinates operator-(FSphericalCoordinates const& Other) const;
 
-			return FVector(
-				Radius * CosAzimuth * SinInclination,
-				Radius * SinAzimuth * SinInclination,
-				Radius * CosInclination
-			);
-		}
+		/** Conform parameters to their normal ranges */
+		void Conform();
+		
+		/** Returns a conformed version of this struct without changing the current one */
+		FSphericalCoordinates GetConformed() const;
 
-		double Radius = 0;
-		double Inclination = 0;
-		double Azimuth = 0;
+		/** Returns true if the inclination is pointing at north or south poles, within the given margin (in radians) */
+		bool IsPointingAtPole(double Margin = 1e-6) const;
+
+		double Radius = 0;      // unitless   0+      (when conforming)
+		double Inclination = 0; // radians    0 to PI (when conforming)
+		double Azimuth = 0;     // radians  -PI to PI (when conforming)
 	};
 
 	/** Custom render target that stores the normal data for the stage */
@@ -229,6 +215,9 @@ private:
 	/** Ensures that the light card root component is at the same location as the projection/view origin */
 	void VerifyAndFixLightCardOrigin(ADisplayClusterLightCardActor* LightCard) const;
 
+	/** Determines the appropriate delta rotation needed to move the specified light card to the mouse's location */
+	FRotator GetLightCardRotationDelta(FViewport* InViewport, ADisplayClusterLightCardActor* LightCard, EAxisList::Type CurrentAxis);
+
 	/** Determines the appropriate delta in spherical coordinates needed to move the specified light card to the mouse's location */
 	FSphericalCoordinates GetLightCardTranslationDelta(FViewport* InViewport, ADisplayClusterLightCardActor* LightCard, EAxisList::Type CurrentAxis);
 
@@ -242,11 +231,14 @@ private:
 	void SpinSelectedLightCards(FViewport* InViewport);
 
 	/** Determines the appropriate spin delta needed to rotate the light card */
-	float GetLightCardSpinDelta(FViewport* InViewport, ADisplayClusterLightCardActor* LightCard);
+	double GetLightCardSpinDelta(FViewport* InViewport, ADisplayClusterLightCardActor* LightCard);
 
 	/** Gets the spherical coordinates of the specified light card */
 	FSphericalCoordinates GetLightCardCoordinates(ADisplayClusterLightCardActor* LightCard) const;
 
+	/** Sets the light card position to the given spherical coordinates */
+	void SetLightCardCoordinates(ADisplayClusterLightCardActor* LightCard, const FSphericalCoordinates& SphericalCoords) const;
+
 	/** Traces to find the light card corresponding to a click on a stage screen */
 	ADisplayClusterLightCardActor* TraceScreenForLightCard(const FSceneView& View, int32 HitX, int32 HitY);
 
