nDisplay: Fix LC Editor proxy DCRA trying to force origin location of LCs in the level.

diff --git a/Engine/Plugins/Runtime/nDisplay/Source/DisplayClusterLightCardEditor/Private/Viewport/DisplayClusterLightCardEditorViewportClient.cpp b/Engine/Plugins/Runtime/nDisplay/Source/DisplayClusterLightCardEditor/Private/Viewport/DisplayClusterLightCardEditorViewportClient.cpp
--- a/Engine/Plugins/Runtime/nDisplay/Source/DisplayClusterLightCardEditor/Private/Viewport/DisplayClusterLightCardEditorViewportClient.cpp
+++ b/Engine/Plugins/Runtime/nDisplay/Source/DisplayClusterLightCardEditor/Private/Viewport/DisplayClusterLightCardEditorViewportClient.cpp
@@ -950,6 +950,12 @@ void FDisplayClusterLightCardEditorViewportClient::UpdatePreviewActor(ADisplayCl
 				RootActorProxy->UpdatePreviewComponents();
 				RootActorProxy->EnableEditorRender(false);
 
+				if (UDisplayClusterConfigurationData* ProxyConfig = RootActorProxy->GetConfigData())
+				{
+					// Disable lightcards so that it doesn't try to update the ones in the level instance world.
+					ProxyConfig->StageSettings.Lightcard.bEnable = false;
+				}
+
 				RootActorBoundingRadius = 0.5f * RootActorProxy->GetComponentsBoundingBox().GetSize().Length();
 			}
 
