nDisplay: Re-enabled eye adaptation on previews.

This makes the LC Editor / Stage Operator Panel look correct.

diff --git a/Engine/Plugins/Runtime/nDisplay/Source/DisplayCluster/Private/Render/Viewport/DisplayClusterViewportManager.cpp b/Engine/Plugins/Runtime/nDisplay/Source/DisplayCluster/Private/Render/Viewport/DisplayClusterViewportManager.cpp
--- a/Engine/Plugins/Runtime/nDisplay/Source/DisplayCluster/Private/Render/Viewport/DisplayClusterViewportManager.cpp
+++ b/Engine/Plugins/Runtime/nDisplay/Source/DisplayCluster/Private/Render/Viewport/DisplayClusterViewportManager.cpp
@@ -618,9 +618,6 @@ FSceneViewFamily::ConstructionValues FDisplayClusterViewportManager::CreateViewF
 	{
 	case EDisplayClusterRenderFrameMode::PreviewInScene:
 
-		// Disable eye adaptation for preview
-		InEngineShowFlags.EyeAdaptation = 0;
-
 		if (RenderFrameSettings.bPreviewEnablePostProcess == false)
 		{
 			// Disable postprocess for preview
