nDisplay: Increased LC rotator circle radius and thickness.

Felt it needed to be a bit bigger for ease of use and rotation precision.

diff --git a/Engine/Plugins/Runtime/nDisplay/Source/DisplayClusterLightCardEditor/Private/Viewport/DisplayClusterLightCardEditorWidget.h b/Engine/Plugins/Runtime/nDisplay/Source/DisplayClusterLightCardEditor/Private/Viewport/DisplayClusterLightCardEditorWidget.h
--- a/Engine/Plugins/Runtime/nDisplay/Source/DisplayClusterLightCardEditor/Private/Viewport/DisplayClusterLightCardEditorWidget.h
+++ b/Engine/Plugins/Runtime/nDisplay/Source/DisplayClusterLightCardEditor/Private/Viewport/DisplayClusterLightCardEditorWidget.h
@@ -31,10 +31,10 @@ public:
 	constexpr static float OriginSize = 35.f;
 
 	/** The size of the rotation circle, in pixels */
-	constexpr static float CirlceRadius = 50.0f;
+	constexpr static float CirlceRadius = 100.0f;
 
 	/** The thickness of the rotation circle, in pixels */
-	constexpr static float CircleThickness = 2.5f;
+	constexpr static float CircleThickness = 5.f;
 
 	constexpr static FLinearColor AxisColorX = FLinearColor(0.594f, 0.0197f, 0.0f);
 	constexpr static FLinearColor AxisColorY = FLinearColor(0.1349f, 0.3959f, 0.0f);
