nDisplay: Added ability to draw lightcards on viewport using polygons.

To use, click on the viewport "Draw LC" button, then click on the viewport at the locations where the polygon vertices should exist, then end the polygon with right mouse button. Or abort by toggling the "Draw LC" button.

The polygon is saved in an array in the light card, and is used to generate an alpha texture. This alpha texture is independent of the color texture. This alpha texture is generated dynamically, so there is no need to save it as an asset.

The default generated alpha texture dimensions is 128x128, but can be changed with DC.LightCardPolygonTextureSize.

Currently there is visual indication of the drawing process (other than the mouse becoming a cross and the Draw LC button being highlighted).

diff --git a/Engine/Plugins/Runtime/nDisplay/Source/DisplayCluster/DisplayCluster.Build.cs b/Engine/Plugins/Runtime/nDisplay/Source/DisplayCluster/DisplayCluster.Build.cs
--- a/Engine/Plugins/Runtime/nDisplay/Source/DisplayCluster/DisplayCluster.Build.cs
+++ b/Engine/Plugins/Runtime/nDisplay/Source/DisplayCluster/DisplayCluster.Build.cs
@@ -31,6 +31,8 @@
 			"JsonUtilities",
 			"Networking",
 			"OpenColorIO",
+			"OpenCV",
+			"OpenCVHelper",
 			"ProceduralMeshComponent",
 			"Renderer",
 			"RenderCore",
diff --git a/Engine/Plugins/Runtime/nDisplay/Source/DisplayCluster/Private/Game/EngineClasses/Scene/DisplayClusterLightCardActor.cpp b/Engine/Plugins/Runtime/nDisplay/Source/DisplayCluster/Private/Game/EngineClasses/Scene/DisplayClusterLightCardActor.cpp
--- a/Engine/Plugins/Runtime/nDisplay/Source/DisplayCluster/Private/Game/EngineClasses/Scene/DisplayClusterLightCardActor.cpp
+++ b/Engine/Plugins/Runtime/nDisplay/Source/DisplayCluster/Private/Game/EngineClasses/Scene/DisplayClusterLightCardActor.cpp
@@ -8,6 +8,24 @@
 #include "Engine/StaticMesh.h"
 #include "UObject/ConstructorHelpers.h"
 
+#if WITH_OPENCV
+
+#include "OpenCVHelper.h"
+
+#include "PreOpenCVHeaders.h"
+#include "opencv2/imgproc.hpp"
+#include "PostOpenCVHeaders.h"
+
+#endif //WITH_OPENCV
+
+int32 GDisplayClusterLightCardPolygonTextureSize = 128;
+static FAutoConsoleVariableRef CVarDisplayClusterLightCardPolygonTextureSize(
+	TEXT("DC.LightCardPolygonTextureSize"),
+	GDisplayClusterLightCardPolygonTextureSize,
+	TEXT("Size of the textures generated by light cards that use a polygon to define its alpha mask. Use a power of 2."),
+	ECVF_Default
+);
+
 const FRotator ADisplayClusterLightCardActor::PlaneMeshRotation = FRotator(0.0f, -90.0f, 90.0f);
 
 ADisplayClusterLightCardActor::ADisplayClusterLightCardActor(const FObjectInitializer& ObjectInitializer)
@@ -55,7 +73,10 @@
 {
 	Super::OnConstruction(Transform);
 
+	UpdatePolygonTexture();
+
 	UMaterialInterface* Material = LightCardComponent->GetMaterial(0);
+
 	if (Material && !Material->IsA<UMaterialInstanceDynamic>())
 	{
 		UMaterialInstanceDynamic* LightCardMatInstance = UMaterialInstanceDynamic::Create(Material, this, TEXT("LightCardMID"));
@@ -113,7 +134,15 @@
 	}
 
 	if (PropertyChangedEvent.Property && (
+		PropertyChangedEvent.Property->GetFName() == GET_MEMBER_NAME_CHECKED(ADisplayClusterLightCardActor, Polygon) ||
+		PropertyChangedEvent.Property->GetFName() == GET_MEMBER_NAME_CHECKED(ADisplayClusterLightCardActor, Feathering)))
+	{
+		UpdatePolygonTexture();
+	}
+
+	if (PropertyChangedEvent.Property && (
 		PropertyChangedEvent.Property->GetFName() == GET_MEMBER_NAME_CHECKED(ADisplayClusterLightCardActor, Mask) ||
+		PropertyChangedEvent.Property->GetFName() == GET_MEMBER_NAME_CHECKED(ADisplayClusterLightCardActor, Polygon) ||
 		PropertyChangedEvent.Property->GetFName() == GET_MEMBER_NAME_CHECKED(ADisplayClusterLightCardActor, Texture) ||
 		PropertyChangedEvent.Property->GetFName() == GET_MEMBER_NAME_CHECKED(ADisplayClusterLightCardActor, Color) ||
 		PropertyChangedEvent.Property->GetFName() == GET_MEMBER_NAME_CHECKED(ADisplayClusterLightCardActor, Exposure) ||
@@ -184,17 +213,85 @@
 
 		LightCardMaterialInstance->ClearParameterValues();
 
-		LightCardMaterialInstance->SetTextureParameterValue(TEXT("Texture"), Texture);
-		LightCardMaterialInstance->SetScalarParameterValue(TEXT("UseMask"), Mask == EDisplayClusterLightCardMask::Square ? 0.f : 1.f);
-		LightCardMaterialInstance->SetScalarParameterValue(TEXT("UseTextureAlpha"), Mask == EDisplayClusterLightCardMask::UseTextureAlpha ? 1.f : 0.f);
 		LightCardMaterialInstance->SetVectorParameterValue(TEXT("CardColor"), Color);
 		LightCardMaterialInstance->SetScalarParameterValue(TEXT("Exposure"), Exposure);
 		LightCardMaterialInstance->SetScalarParameterValue(TEXT("EmissiveStrength"), Exposure);
 		LightCardMaterialInstance->SetScalarParameterValue(TEXT("Opacity"), bIsProxy ? ProxyOpacity : Opacity);
 		LightCardMaterialInstance->SetScalarParameterValue(TEXT("Feather"), Feathering);
+		LightCardMaterialInstance->SetTextureParameterValue(TEXT("Texture"), Texture);
+		LightCardMaterialInstance->SetTextureParameterValue(TEXT("AlphaTexture"), PolygonMask);
+
+		bool bUseMask = true; // Enable masking
+		bool bUseTextureAlpha = false; // Use the alpha channel of Texture as mask
+		bool bUseAlphaTexture = false; // Use the AlphaTexture 
+
+		switch (Mask)
+		{
+		case EDisplayClusterLightCardMask::Square:
+			bUseMask = false;
+			break;
+
+		case EDisplayClusterLightCardMask::UseTextureAlpha:
+			bUseTextureAlpha = true;
+			break;
+
+		case EDisplayClusterLightCardMask::Polygon:
+			bUseTextureAlpha = true;
+			bUseAlphaTexture = true;
+			break;
+		}
+
+		LightCardMaterialInstance->SetScalarParameterValue(TEXT("UseMask"), bUseMask);
+		LightCardMaterialInstance->SetScalarParameterValue(TEXT("UseTextureAlpha"), bUseTextureAlpha);
+		LightCardMaterialInstance->SetScalarParameterValue(TEXT("UseAlphaTexture"), bUseAlphaTexture);
 	}
 }
 
+void ADisplayClusterLightCardActor::UpdatePolygonTexture()
+{
+#if WITH_OPENCV
+
+	if (Polygon.Num() < 3)
+	{
+		PolygonMask = nullptr;
+		return;
+	}
+
+	if (GDisplayClusterLightCardPolygonTextureSize < 2)
+	{
+		return;
+	}
+
+	const int32 TextureWidth = GDisplayClusterLightCardPolygonTextureSize;
+	const int32 TextureHeight = GDisplayClusterLightCardPolygonTextureSize;
+
+	cv::Mat GrayMat(cv::Size(TextureWidth, TextureHeight), CV_8UC1, cv::Scalar(0));
+
+	std::vector<cv::Point2i> PolyVec; // polygon with LC shape that to be filled with cv::fillPoly
+	PolyVec.reserve(Polygon.Num());
+
+	for (const FVector2d& PolyPoint : Polygon)
+	{
+		PolyVec.push_back(cv::Point2i(
+			FMath::RoundToInt(TextureWidth  * PolyPoint.X), 
+			FMath::RoundToInt(TextureHeight * PolyPoint.Y))
+		);
+	}
+
+	// Fill the polygon pixels
+	cv::fillPoly(GrayMat, PolyVec, 255, cv::LINE_AA);
+
+	// Create a texture, all white and and the shape goes in the alpha
+	cv::Mat RGBAMat(cv::Size(TextureWidth, TextureHeight), CV_8UC4, cv::Scalar::all(255));
+	cv::insertChannel(GrayMat, RGBAMat, 3);
+
+	// Create UTexture
+	PolygonMask = FOpenCVHelper::TextureFromCvMat(RGBAMat);
+
+#endif // WITH_OPENCV
+}
+
+
 UStaticMesh* ADisplayClusterLightCardActor::GetStaticMesh() const
 {
 	return LightCardComponent->GetStaticMesh();
diff --git a/Engine/Plugins/Runtime/nDisplay/Source/DisplayCluster/Public/DisplayClusterLightCardActor.h b/Engine/Plugins/Runtime/nDisplay/Source/DisplayCluster/Public/DisplayClusterLightCardActor.h
--- a/Engine/Plugins/Runtime/nDisplay/Source/DisplayCluster/Public/DisplayClusterLightCardActor.h
+++ b/Engine/Plugins/Runtime/nDisplay/Source/DisplayCluster/Public/DisplayClusterLightCardActor.h
@@ -17,7 +17,8 @@
 {
 	Circle,
 	Square,
-	UseTextureAlpha
+	UseTextureAlpha,
+	Polygon,
 };
 
 UCLASS(Blueprintable)
@@ -76,9 +77,12 @@
 	/** Set positional parameters */
 	void SetPositionalParams(const PositionalParams& Params);
 
-protected:
+	/** Updates the card's material instance parameters */
 	void UpdateLightCardMaterialInstance();
 
+	/** Updates the polygon texture from the polygon points */
+	void UpdatePolygonTexture();
+
 public:
 	UPROPERTY(EditAnywhere, BlueprintReadWrite, Category = "Orientation")
 	double DistanceFromCenter;
@@ -127,6 +131,14 @@
 	UPROPERTY(Transient)
 	bool bIsProxy = false;
 
+	/** Polygon points when using this type of mask */
+	UPROPERTY(EditAnywhere, BlueprintReadWrite, Category = "Appearance")
+	TArray<FVector2D> Polygon;
+
+	/** Used to flag this light card as a proxy of a "real" light card. Used by the LightCard Editor */
+	UPROPERTY(Transient)
+	UTexture* PolygonMask = nullptr;
+
 protected:
 	UPROPERTY(EditDefaultsOnly, BlueprintReadOnly, Category = "Default")
 	TObjectPtr<USceneComponent> DefaultSceneRootComponent;
diff --git a/Engine/Plugins/Runtime/nDisplay/Source/DisplayClusterLightCardEditor/DisplayClusterLightCardEditor.Build.cs b/Engine/Plugins/Runtime/nDisplay/Source/DisplayClusterLightCardEditor/DisplayClusterLightCardEditor.Build.cs
--- a/Engine/Plugins/Runtime/nDisplay/Source/DisplayClusterLightCardEditor/DisplayClusterLightCardEditor.Build.cs
+++ b/Engine/Plugins/Runtime/nDisplay/Source/DisplayClusterLightCardEditor/DisplayClusterLightCardEditor.Build.cs
@@ -9,27 +9,29 @@
 	{
 		PrivateDependencyModuleNames.AddRange(
 			new string[] {
+
+				"AdvancedPreviewScene",
+				"AppFramework",
+				"Core",
+				"CoreUObject",
 				"DisplayCluster",
 				"DisplayClusterConfiguration",
 				"DisplayClusterOperator",
 				"DisplayClusterLightCardEditorShaders",
-				
-				"Core",
-				"CoreUObject",
+				"EditorStyle",
 				"Engine",
+				"InputCore",
+				"OpenCV",
+				"OpenCVHelper",
+				"PropertyEditor",
 				"RenderCore",
 				"Renderer",
 				"RHI",
-				"InputCore",
 				"Slate",
 				"SlateCore",
 				"UnrealEd",
-				"EditorStyle",
 				"WorkspaceMenuStructure",
-				"AdvancedPreviewScene",
 				"ToolWidgets",
-				"PropertyEditor",
-				"AppFramework"
 			});
 	}
 }
diff --git a/Engine/Plugins/Runtime/nDisplay/Source/DisplayClusterLightCardEditor/Private/DisplayClusterLightCardEditorCommands.cpp b/Engine/Plugins/Runtime/nDisplay/Source/DisplayClusterLightCardEditor/Private/DisplayClusterLightCardEditorCommands.cpp
--- a/Engine/Plugins/Runtime/nDisplay/Source/DisplayClusterLightCardEditor/Private/DisplayClusterLightCardEditorCommands.cpp
+++ b/Engine/Plugins/Runtime/nDisplay/Source/DisplayClusterLightCardEditor/Private/DisplayClusterLightCardEditorCommands.cpp
@@ -1,20 +1,22 @@
-// Copyright Epic Games, Inc. All Rights Reserved.
-
-#include "DisplayClusterLightCardEditorCommands.h"
-
-#define LOCTEXT_NAMESPACE "DisplayClusterLightCardEditorCommands"
-
-
-void FDisplayClusterLightCardEditorCommands::RegisterCommands()
-{
-	UI_COMMAND(ResetCamera, "Reset Camera", "Resets the camera to focus on the mesh", EUserInterfaceActionType::Button, FInputChord());
-
-	UI_COMMAND(PerspectiveProjection, "Perspective", "A perspective projection from the stage's view origin", EUserInterfaceActionType::Button, FInputChord());
-	UI_COMMAND(AzimuthalProjection, "Azimuthal", "An azimuthal hemisphere projection from the stage's center pointed upwards", EUserInterfaceActionType::Button, FInputChord());
-
-	UI_COMMAND(AddNewLightCard, "Add New Light Card", "Add and assign a new Light Card to the actor", EUserInterfaceActionType::Button, FInputChord());
-	UI_COMMAND(AddExistingLightCard, "Add Existing Light Card", "Add an existing Light Card to the actor", EUserInterfaceActionType::Button, FInputChord());
-	UI_COMMAND(RemoveLightCard, "Remove from Actor", "Remove the Light Card from the actor but do not delete it", EUserInterfaceActionType::Button, FInputChord());
-}
-
-#undef LOCTEXT_NAMESPACE
+// Copyright Epic Games, Inc. All Rights Reserved.
+
+#include "DisplayClusterLightCardEditorCommands.h"
+
+#define LOCTEXT_NAMESPACE "DisplayClusterLightCardEditorCommands"
+
+
+void FDisplayClusterLightCardEditorCommands::RegisterCommands()
+{
+	UI_COMMAND(ResetCamera, "Reset Camera", "Resets the camera to focus on the mesh", EUserInterfaceActionType::Button, FInputChord());
+
+	UI_COMMAND(PerspectiveProjection, "Perspective", "A perspective projection from the stage's view origin", EUserInterfaceActionType::Button, FInputChord());
+	UI_COMMAND(AzimuthalProjection, "Azimuthal", "An azimuthal hemisphere projection from the stage's center pointed upwards", EUserInterfaceActionType::Button, FInputChord());
+
+	UI_COMMAND(AddNewLightCard, "Add New Light Card", "Add and assign a new Light Card to the actor", EUserInterfaceActionType::Button, FInputChord());
+	UI_COMMAND(AddExistingLightCard, "Add Existing Light Card", "Add an existing Light Card to the actor", EUserInterfaceActionType::Button, FInputChord());
+	UI_COMMAND(RemoveLightCard, "Remove from Actor", "Remove the Light Card from the actor but do not delete it", EUserInterfaceActionType::Button, FInputChord());
+
+	UI_COMMAND(DrawLightCard, "Draw Light Card", "Draw polygon light card on viewport", EUserInterfaceActionType::ToggleButton, FInputChord());
+}
+
+#undef LOCTEXT_NAMESPACE
diff --git a/Engine/Plugins/Runtime/nDisplay/Source/DisplayClusterLightCardEditor/Private/DisplayClusterLightCardEditorCommands.h b/Engine/Plugins/Runtime/nDisplay/Source/DisplayClusterLightCardEditor/Private/DisplayClusterLightCardEditorCommands.h
--- a/Engine/Plugins/Runtime/nDisplay/Source/DisplayClusterLightCardEditor/Private/DisplayClusterLightCardEditorCommands.h
+++ b/Engine/Plugins/Runtime/nDisplay/Source/DisplayClusterLightCardEditor/Private/DisplayClusterLightCardEditorCommands.h
@@ -1,31 +1,33 @@
-// Copyright Epic Games, Inc. All Rights Reserved.
-
-#pragma once
-
-#include "CoreMinimal.h"
-#include "Framework/Commands/Commands.h"
-#include "Styling/AppStyle.h"
-
-
-class FDisplayClusterLightCardEditorCommands 
-	: public TCommands<FDisplayClusterLightCardEditorCommands>
-{
-public:
-	FDisplayClusterLightCardEditorCommands()
-		: TCommands<FDisplayClusterLightCardEditorCommands>(TEXT("DisplayClusterLightCardEditor"), 
-			NSLOCTEXT("Contexts", "DisplayClusterLightCardEditor", "Display Cluster LightCard Editor"), NAME_None, FAppStyle::GetAppStyleSetName())
-	{ }
-
-	virtual void RegisterCommands() override;
-
-public:
-	// Viewport commands
-	TSharedPtr<FUICommandInfo> ResetCamera;
-
-	TSharedPtr<FUICommandInfo> PerspectiveProjection;
-	TSharedPtr<FUICommandInfo> AzimuthalProjection;
-
-	TSharedPtr<FUICommandInfo> AddNewLightCard;
-	TSharedPtr<FUICommandInfo> AddExistingLightCard;
-	TSharedPtr<FUICommandInfo> RemoveLightCard;
-};
+// Copyright Epic Games, Inc. All Rights Reserved.
+
+#pragma once
+
+#include "CoreMinimal.h"
+#include "Framework/Commands/Commands.h"
+#include "Styling/AppStyle.h"
+
+
+class FDisplayClusterLightCardEditorCommands 
+	: public TCommands<FDisplayClusterLightCardEditorCommands>
+{
+public:
+	FDisplayClusterLightCardEditorCommands()
+		: TCommands<FDisplayClusterLightCardEditorCommands>(TEXT("DisplayClusterLightCardEditor"), 
+			NSLOCTEXT("Contexts", "DisplayClusterLightCardEditor", "Display Cluster LightCard Editor"), NAME_None, FAppStyle::GetAppStyleSetName())
+	{ }
+
+	virtual void RegisterCommands() override;
+
+public:
+	// Viewport commands
+	TSharedPtr<FUICommandInfo> ResetCamera;
+
+	TSharedPtr<FUICommandInfo> PerspectiveProjection;
+	TSharedPtr<FUICommandInfo> AzimuthalProjection;
+
+	TSharedPtr<FUICommandInfo> AddNewLightCard;
+	TSharedPtr<FUICommandInfo> AddExistingLightCard;
+	TSharedPtr<FUICommandInfo> RemoveLightCard;
+
+	TSharedPtr<FUICommandInfo> DrawLightCard;
+};
diff --git a/Engine/Plugins/Runtime/nDisplay/Source/DisplayClusterLightCardEditor/Private/SDisplayClusterLightCardEditor.cpp b/Engine/Plugins/Runtime/nDisplay/Source/DisplayClusterLightCardEditor/Private/SDisplayClusterLightCardEditor.cpp
--- a/Engine/Plugins/Runtime/nDisplay/Source/DisplayClusterLightCardEditor/Private/SDisplayClusterLightCardEditor.cpp
+++ b/Engine/Plugins/Runtime/nDisplay/Source/DisplayClusterLightCardEditor/Private/SDisplayClusterLightCardEditor.cpp
@@ -140,6 +140,12 @@
 	ViewportView->GetLightCardEditorViewportClient()->CenterLightCardInView(LightCard);
 }
 
+ADisplayClusterLightCardActor* SDisplayClusterLightCardEditor::AddNewLightCard()
+{
+	check(LightCardList);
+	return LightCardList->AddNewLightCard();
+}
+
 void SDisplayClusterLightCardEditor::OnActiveRootActorChanged(ADisplayClusterRootActor* NewRootActor)
 {
 	RemoveCompileDelegates();
diff --git a/Engine/Plugins/Runtime/nDisplay/Source/DisplayClusterLightCardEditor/Private/SDisplayClusterLightCardEditor.h b/Engine/Plugins/Runtime/nDisplay/Source/DisplayClusterLightCardEditor/Private/SDisplayClusterLightCardEditor.h
--- a/Engine/Plugins/Runtime/nDisplay/Source/DisplayClusterLightCardEditor/Private/SDisplayClusterLightCardEditor.h
+++ b/Engine/Plugins/Runtime/nDisplay/Source/DisplayClusterLightCardEditor/Private/SDisplayClusterLightCardEditor.h
@@ -60,6 +60,9 @@
 	/** Places the given light card in the middle of the viewport */
 	void CenterLightCardInView(ADisplayClusterLightCardActor& LightCard);
 
+	/** Spawns a new light card */
+	ADisplayClusterLightCardActor* AddNewLightCard();
+
 private:
 	/** Raised when the active Display cluster root actor has been changed in the operator panel */
 	void OnActiveRootActorChanged(ADisplayClusterRootActor* NewRootActor);
diff --git a/Engine/Plugins/Runtime/nDisplay/Source/DisplayClusterLightCardEditor/Private/SDisplayClusterLightCardList.cpp b/Engine/Plugins/Runtime/nDisplay/Source/DisplayClusterLightCardEditor/Private/SDisplayClusterLightCardList.cpp
--- a/Engine/Plugins/Runtime/nDisplay/Source/DisplayClusterLightCardEditor/Private/SDisplayClusterLightCardList.cpp
+++ b/Engine/Plugins/Runtime/nDisplay/Source/DisplayClusterLightCardEditor/Private/SDisplayClusterLightCardList.cpp
@@ -256,7 +256,7 @@
 	
 	CommandList->MapAction(
 		Commands.AddNewLightCard,
-		FExecuteAction::CreateSP(this, &SDisplayClusterLightCardList::AddNewLightCard),
+		FExecuteAction::CreateSP(this, &SDisplayClusterLightCardList::AddNewLightCardHandler),
 		FCanExecuteAction::CreateSP(this, &SDisplayClusterLightCardList::CanAddLightCard),
 		FCanExecuteAction());
 
@@ -358,12 +358,29 @@
 	return false;
 }
 
-void SDisplayClusterLightCardList::AddNewLightCard()
+void SDisplayClusterLightCardList::AddNewLightCardHandler()
 {
 	check(RootActor.IsValid());
-	
+
 	FScopedTransaction Transaction(LOCTEXT("AddNewLightCard", "Add New Light Card"));
-	
+
+	ADisplayClusterLightCardActor* NewLightCard = AddNewLightCard();
+
+	// When adding a new lightcard, usually the desired location is in the middle of the viewport
+	if (LightCardEditorPtr.IsValid())
+	{
+		LightCardEditorPtr.Pin()->CenterLightCardInView(*NewLightCard);
+	}
+}
+
+
+ADisplayClusterLightCardActor* SDisplayClusterLightCardList::AddNewLightCard()
+{
+	if (!RootActor.IsValid())
+	{
+		return nullptr;
+	}
+
 	const FVector SpawnLocation = RootActor->GetDefaultCamera()->GetComponentLocation();
 	FRotator SpawnRotation = RootActor->GetDefaultCamera()->GetComponentRotation();
 	SpawnRotation.Yaw -= 180.f;
@@ -375,19 +392,15 @@
 	SpawnParameters.NameMode = FActorSpawnParameters::ESpawnActorNameMode::Requested;
 	SpawnParameters.OverrideLevel = RootActor->GetWorld()->GetCurrentLevel();
 	
-	ADisplayClusterLightCardActor* NewActor = CastChecked<ADisplayClusterLightCardActor>(
+	ADisplayClusterLightCardActor* NewLightCard = CastChecked<ADisplayClusterLightCardActor>(
 		RootActor->GetWorld()->SpawnActor(ADisplayClusterLightCardActor::StaticClass(),
 		&SpawnLocation, &SpawnRotation, MoveTemp(SpawnParameters)));
 
-	NewActor->SetActorLabel(NewActor->GetName());
+	NewLightCard->SetActorLabel(NewLightCard->GetName());
 
-	AddLightCardToActor(NewActor);
+	AddLightCardToActor(NewLightCard);
 
-	// When adding a new lightcard, usually the desired location is in the middle of the viewport
-	if (LightCardEditorPtr.IsValid())
-	{
-		LightCardEditorPtr.Pin()->CenterLightCardInView(*NewActor);
-	}
+	return NewLightCard;
 }
 
 void SDisplayClusterLightCardList::AddExistingLightCard()
diff --git a/Engine/Plugins/Runtime/nDisplay/Source/DisplayClusterLightCardEditor/Private/SDisplayClusterLightCardList.h b/Engine/Plugins/Runtime/nDisplay/Source/DisplayClusterLightCardEditor/Private/SDisplayClusterLightCardList.h
--- a/Engine/Plugins/Runtime/nDisplay/Source/DisplayClusterLightCardEditor/Private/SDisplayClusterLightCardList.h
+++ b/Engine/Plugins/Runtime/nDisplay/Source/DisplayClusterLightCardEditor/Private/SDisplayClusterLightCardList.h
@@ -54,6 +54,9 @@
 	/** Refreshes the list */
 	void Refresh();
 
+	/** Spawn and add a new Light Card */
+	ADisplayClusterLightCardActor* AddNewLightCard();
+
 private:
 	void BindCommands();
 	
@@ -62,10 +65,10 @@
 	 * @return True if the list has been modified
 	 */
 	bool FillLightCardList();
+	
+	/** Handles the Add Light Card button. Spawn a new Light Card placed in the middle of the viewport. */
+	void AddNewLightCardHandler();
 
-	/** Spawn and add a new Light Card */
-	void AddNewLightCard();
-	
 	/** Select an existing Light Card from a menu */
 	void AddExistingLightCard();
 	
diff --git a/Engine/Plugins/Runtime/nDisplay/Source/DisplayClusterLightCardEditor/Private/Viewport/DisplayClusterLightCardEditorViewportClient.cpp b/Engine/Plugins/Runtime/nDisplay/Source/DisplayClusterLightCardEditor/Private/Viewport/DisplayClusterLightCardEditorViewportClient.cpp
--- a/Engine/Plugins/Runtime/nDisplay/Source/DisplayClusterLightCardEditor/Private/Viewport/DisplayClusterLightCardEditorViewportClient.cpp
+++ b/Engine/Plugins/Runtime/nDisplay/Source/DisplayClusterLightCardEditor/Private/Viewport/DisplayClusterLightCardEditorViewportClient.cpp
@@ -12,37 +12,47 @@
 #include "Components/DisplayClusterScreenComponent.h"
 
 #include "AudioDevice.h"
-#include "EngineUtils.h"
-#include "EditorModes.h"
-#include "EditorModeManager.h"
-#include "PreviewScene.h"
-#include "UnrealEdGlobals.h"
-#include "ScopedTransaction.h"
-#include "SDisplayClusterLightCardEditor.h"
-#include "UnrealWidget.h"
-#include "RayTracingDebugVisualizationMenuCommands.h"
+#include "CameraController.h"
 #include "Components/DirectionalLightComponent.h"
 #include "Components/LineBatchComponent.h"
 #include "Components/PostProcessComponent.h"
 #include "Components/SkyLightComponent.h"
 #include "Components/PrimitiveComponent.h"
-#include "Editor/UnrealEdEngine.h"
-#include "Kismet/GameplayStatics.h"
-#include "Slate/SceneViewport.h"
-#include "Widgets/Docking/SDockTab.h"
 #include "CustomEditorStaticScreenPercentage.h"
-#include "LegacyScreenPercentageDriver.h"
-#include "EngineModule.h"
 #include "Debug/DebugDrawService.h"
+#include "EngineUtils.h"
+#include "EditorModes.h"
+#include "Editor/UnrealEdEngine.h"
+#include "EditorModeManager.h"
+#include "EngineModule.h"
 #include "Engine/Canvas.h"
 #include "Engine/TextureRenderTarget2D.h"
-#include "CameraController.h"
 #include "ImageUtils.h"
 #include "Kismet/KismetMathLibrary.h"
-
-
+#include "Kismet/GameplayStatics.h"
+#include "LegacyScreenPercentageDriver.h"
+#include "Math/UnrealMathUtility.h"
+#include "PreviewScene.h"
+#include "RayTracingDebugVisualizationMenuCommands.h"
 #include "Renderer/Private/SceneRendering.h"
+#include "ScopedTransaction.h"
+#include "SDisplayClusterLightCardEditor.h"
+#include "Slate/SceneViewport.h"
+#include "UnrealEdGlobals.h"
+#include "UnrealWidget.h"
+#include "Widgets/Docking/SDockTab.h"
 
+
+#if WITH_OPENCV
+
+#include "OpenCVHelper.h"
+
+#include "PreOpenCVHeaders.h"
+#include "opencv2/imgproc.hpp"
+#include "PostOpenCVHeaders.h"
+
+#endif //WITH_OPENCV
+
 #define LOCTEXT_NAMESPACE "DisplayClusterLightCardEditorViewportClient"
 
 
@@ -313,7 +323,7 @@
 	MeshProjectionRenderer->ActorSelectedDelegate = FDisplayClusterMeshProjectionRenderer::FSelection::CreateRaw(this, &FDisplayClusterLightCardEditorViewportClient::IsLightCardSelected);
 	MeshProjectionRenderer->RenderSimpleElementsDelegate = FDisplayClusterMeshProjectionRenderer::FSimpleElementPass::CreateRaw(this, &FDisplayClusterLightCardEditorViewportClient::Draw);
 
-	bDraggingActor = false;
+	InputMode = EInputMode::Idle;
 	DragWidgetOffset = FVector::ZeroVector;
 	
 	EditorWidget = MakeShared<FDisplayClusterLightCardEditorWidget>();
@@ -406,12 +416,17 @@
 
 	// EditorViewportClient sets the cursor settings based on the state of the built in FWidget, which isn't being used here, so
 	// force a software cursor if we are dragging an actor so that the correct mouse cursor shows up
-	if (bDraggingActor)
+	switch (InputMode)
 	{
+	case EInputMode::DraggingActor:
 		SetRequiredCursor(false, true);
 		SetRequiredCursorOverride(true, EMouseCursor::CardinalCross);
+		break;
 
-		ApplyRequiredCursorVisibility(true);
+	case EInputMode::DrawingLightCard:
+		SetRequiredCursor(true, false);
+		SetRequiredCursorOverride(true, EMouseCursor::Crosshairs);
+		break;
 	}
 
 	if (DesiredLookAtLocation.IsSet())
@@ -667,14 +682,51 @@
 bool FDisplayClusterLightCardEditorViewportClient::InputKey(FViewport* InViewport, int32 ControllerId, FKey Key, EInputEvent Event,
                                             float AmountDepressed, bool bGamepad)
 {
-	if ((Key == EKeys::MouseScrollUp || Key == EKeys::MouseScrollDown) && Event == IE_Pressed)
+	switch (InputMode)
 	{
-		const int32 Sign = Key == EKeys::MouseScrollUp ? -1 : 1;
-		const float CurrentFOV = GetProjectionModeFOV(ProjectionMode);
-		const float NewFOV = FMath::Clamp(CurrentFOV + Sign * FOVScrollIncrement, CameraController->GetConfig().MinimumAllowedFOV, CameraController->GetConfig().MaximumAllowedFOV);
+	case EInputMode::DrawingLightCard:
+		
+		if ((Key == EKeys::LeftMouseButton) && (Event == IE_Released))
+		{
+			// Add a new light card polygon point
 
-		SetProjectionModeFOV(ProjectionMode, NewFOV);
+			FIntPoint MousePos;
+			InViewport->GetMousePos(MousePos);
+
+			if (!DrawnMousePositions.Num() || (DrawnMousePositions.Last() != MousePos))
+			{
+				DrawnMousePositions.Add(MousePos);
+			}
+		}
+		else if ((Key == EKeys::RightMouseButton) && (Event == IE_Released))
+		{
+			// Create the polygon light card.
+			CreateDrawnLightCard(DrawnMousePositions);
+
+			// Reset the list of polygon points used to create the new card
+			DrawnMousePositions.Empty();
+
+			// We go back to idle input mode
+			InputMode = EInputMode::Idle;
+		}
+
+		// Returning here (and not calling Super::InputKey) locks the viewport so that the user does not inadvertently 
+		// change the perspective while in the middle of drawing a light card, as that is not currently supported
 		return true;
+
+	default:
+
+		if ((Key == EKeys::MouseScrollUp || Key == EKeys::MouseScrollDown) && (Event == IE_Pressed))
+		{
+			const int32 Sign = Key == EKeys::MouseScrollUp ? -1 : 1;
+			const float CurrentFOV = GetProjectionModeFOV(ProjectionMode);
+			const float NewFOV = FMath::Clamp(CurrentFOV + Sign * FOVScrollIncrement, CameraController->GetConfig().MinimumAllowedFOV, CameraController->GetConfig().MaximumAllowedFOV);
+
+			SetProjectionModeFOV(ProjectionMode, NewFOV);
+			return true;
+		}
+
+		break;
 	}
 
 	return FEditorViewportClient::InputKey(InViewport, ControllerId, Key, Event, AmountDepressed, bGamepad);
@@ -687,6 +739,7 @@
 	FVector& Scale)
 {
 	bool bHandled = false;
+
 	if (FEditorViewportClient::InputWidgetDelta(InViewport, CurrentAxis, Drag, Rot, Scale))
 	{
 		bHandled = true;
@@ -719,11 +772,16 @@
 	return bHandled;
 }
 
-void FDisplayClusterLightCardEditorViewportClient::TrackingStarted(const FInputEventState& InInputState, bool bIsDraggingWidget,
+void FDisplayClusterLightCardEditorViewportClient::TrackingStarted(
+	const FInputEventState& InInputState, 
+	bool bIsDraggingWidget,
 	bool bNudge)
 {
-	if (!bDraggingActor && bIsDraggingWidget && InInputState.IsLeftMouseButtonPressed() && SelectedLightCards.Num())
+	if ((InputMode == EInputMode::Idle) && bIsDraggingWidget && InInputState.IsLeftMouseButtonPressed() && SelectedLightCards.Num())
 	{
+		// Start dragging actor
+		InputMode = EInputMode::DraggingActor;
+
 		GEditor->DisableDeltaModification(true);
 		{
 			// The pivot location won't update properly and the actor will rotate / move around the original selection origin
@@ -733,8 +791,7 @@
 		}
 
 		BeginTransaction(LOCTEXT("MoveLightCard", "Move Light Card"));
-		bDraggingActor = true;
-		
+
 		DesiredLookAtLocation.Reset();
 
 		// Compute and store the delta between the widget's origin and the place the user clicked on it,
@@ -747,6 +804,7 @@
 			GetScene(),
 			EngineShowFlags)
 			.SetRealtimeUpdate(IsRealtime()));
+
 		FSceneView* View = CalcSceneView(&ViewFamily);
 
 		FVector Origin;
@@ -762,20 +820,229 @@
 
 void FDisplayClusterLightCardEditorViewportClient::TrackingStopped()
 {
-	bDraggingActor = false;
-	DragWidgetOffset = FVector::ZeroVector;
-	EndTransaction();
+	if (InputMode == EInputMode::DraggingActor)
+	{
+		InputMode = EInputMode::Idle;
 
-	if (SelectedLightCards.Num())
-	{
-		GEditor->DisableDeltaModification(false);
+		DragWidgetOffset = FVector::ZeroVector;
+		EndTransaction();
+
+		if (SelectedLightCards.Num())
+		{
+			GEditor->DisableDeltaModification(false);
+		}
 	}
-	
+
 	FEditorViewportClient::TrackingStopped();
 }
 
+void FDisplayClusterLightCardEditorViewportClient::CreateDrawnLightCard(const TArray<FIntPoint>& MousePositions)
+{
+#if WITH_OPENCV
+	if (DrawnMousePositions.Num() < 3 || !Viewport)
+	{
+		return;
+	}
+
+	//
+	// Find direction of each mouse position
+	//
+
+	FSceneViewFamilyContext ViewFamily(FSceneViewFamily::ConstructionValues(
+		Viewport,
+		GetScene(),
+		EngineShowFlags
+	).SetRealtimeUpdate(IsRealtime()));
+
+	FSceneView* View = CalcSceneView(&ViewFamily);
+
+	FVector ViewOrigin;
+	TArray<FVector> MouseWorldDirections; // directions from view origin
+	MouseWorldDirections.AddUninitialized(MousePositions.Num());
+
+	for (int32 PointIdx = 0; PointIdx < MousePositions.Num(); ++PointIdx)
+	{
+		PixelToWorld(*View, MousePositions[PointIdx], ViewOrigin, MouseWorldDirections[PointIdx]);
+	}
+
+	//
+	// Find world position and normal of card using the mouse center in pixel space.
+	// 
+	// @todo A future improvement would be to find a center that doesn't make the LC bigger than it needs to, 
+	// since the pixel space center won't be the projected plane center
+
+	FVector LightCardDirection; // direction from origin to where pixel points to
+	FVector LightCardPlaneAxisX;    // Normal vector at the found enveloping surface point
+	FVector LightCardPlaneAxisY;
+	FVector LightCardPlaneAxisZ;
+	float LightCardDistance;    // Distance from origin to the found enveloping surface point
+	{
+		// Find center point in pixel space
+
+		int32 MinX = INT_MAX;
+		int32 MinY = INT_MAX;
+		int32 MaxX = INT_MIN;
+		int32 MaxY = INT_MIN;
+
+		for (const FIntPoint& Point : MousePositions)
+		{
+			MinX = FMath::Min(MinX, Point.X);
+			MinY = FMath::Min(MinY, Point.Y);
+			MaxX = FMath::Max(MaxX, Point.X);
+			MaxY = FMath::Max(MaxY, Point.Y);
+		}
+
+		FIntPoint CenterPoint((MaxX + MinX)/2, (MaxY + MinY) / 2);
+
+		PixelToWorld(*View, CenterPoint, ViewOrigin, LightCardDirection);
+
+		const FVector Position = ViewOrigin + LightCardDirection * 100.; // We fabricate a position in the right direction
+
+		FVector RelativeNormal;
+
+		// We find the normal and distance from origin
+		if (Position.Z < ViewOrigin.Z)
+		{
+			SouthNormalMap.GetNormalAndDistanceAtPosition(Position, RelativeNormal, LightCardDistance);
+		}
+		else
+		{
+			NorthNormalMap.GetNormalAndDistanceAtPosition(Position, RelativeNormal, LightCardDistance);
+		}
+
+		LightCardDistance = CalculateFinalLightCardDistance(LightCardDistance);
+
+		// Calculate light card normal from light card direction and relative normal
+
+		const FMatrix RotEffector = FRotationMatrix::MakeFromX(-RelativeNormal);
+		const FMatrix RotArm = FRotationMatrix::MakeFromX(LightCardDirection);
+		const FRotator TotalRotator = (RotEffector * RotArm).Rotator();
+
+		LightCardPlaneAxisX = TotalRotator.RotateVector(-FVector::XAxisVector); // Same as Normal
+		LightCardPlaneAxisY = TotalRotator.RotateVector(-FVector::YAxisVector);
+		LightCardPlaneAxisZ = TotalRotator.RotateVector(FVector::ZAxisVector);
+	}
+
+	const FVector LightCardLocation = ViewOrigin + LightCardDistance * LightCardDirection;
+
+	const FSphericalCoordinates LightCardCoords(LightCardLocation - ViewOrigin);
+
+	//
+	// Project mouse positions to light card plane.
+	//
+
+	// @todo A future improvement would be to find a spin that minimizes the lightcard size
+
+	TArray<FVector3d> ProjectedMousePositions;
+	ProjectedMousePositions.AddUninitialized(MouseWorldDirections.Num());
+	const FPlane LightCardPlane = FPlane(LightCardLocation, LightCardPlaneAxisX);
+
+	for (int32 PointIdx = 0; PointIdx < MouseWorldDirections.Num(); ++PointIdx)
+	{
+		// skip the point if the intersection doesn't exist, which happens when the normal and the 
+		// ray direction are perpendicular (nearly zero dot product)
+		if (FMath::IsNearlyZero(FVector::DotProduct(LightCardPlaneAxisX, MouseWorldDirections[PointIdx])))
+		{
+			continue;
+		}
+
+		ProjectedMousePositions[PointIdx] = FMath::RayPlaneIntersection(ViewOrigin, MouseWorldDirections[PointIdx], LightCardPlane);
+	}
+
+	// Convert projected mouse positions in world space to light card plane space (2d)
+
+	TArray<FVector2d> PointsInLightCardPlane;
+	PointsInLightCardPlane.AddUninitialized(ProjectedMousePositions.Num());
+
+	for (int32 PointIdx = 0; PointIdx < ProjectedMousePositions.Num(); ++PointIdx)
+	{
+		const FVector ProjMousePosLocal = ProjectedMousePositions[PointIdx] - LightCardLocation; 
+
+		PointsInLightCardPlane[PointIdx].X = -FVector::DotProduct(ProjMousePosLocal, LightCardPlaneAxisY);
+		PointsInLightCardPlane[PointIdx].Y =  FVector::DotProduct(ProjMousePosLocal, LightCardPlaneAxisZ);
+	}
+
+	// Calculate LC x/y bounds
+
+	double LightCardWidth  = 0;
+	double LightCardHeight = 0;
+
+	for (const FVector2d& Point : PointsInLightCardPlane)
+	{
+		LightCardWidth  = FMath::Max(LightCardWidth , 2 * abs(Point.X));
+		LightCardHeight = FMath::Max(LightCardHeight, 2 * abs(Point.Y));
+	}
+
+	if (FMath::IsNearlyZero(LightCardWidth) || FMath::IsNearlyZero(LightCardHeight))
+	{
+		return;
+	}
+
+	// Create lightcard
+
+	if (!LightCardEditorPtr.IsValid())
+	{
+		return;
+	}
+
+	FScopedTransaction Transaction(LOCTEXT("AddNewLightCard", "Add New Light Card"));
+
+	ADisplayClusterLightCardActor* LightCard = LightCardEditorPtr.Pin()->AddNewLightCard();
+
+	if (!LightCard)
+	{
+		return;
+	}
+
+	// Assign polygon mask
+
+	LightCard->Mask = EDisplayClusterLightCardMask::Polygon;
+
+	LightCard->Polygon.Empty(PointsInLightCardPlane.Num());
+
+	for (const FVector2d& PlanePoint : PointsInLightCardPlane)
+	{
+		LightCard->Polygon.Add(FVector2D(
+			PlanePoint.X / LightCardWidth  + 0.5,
+			PlanePoint.Y / LightCardHeight + 0.5
+		));
+	}
+
+	LightCard->UpdatePolygonTexture();
+	LightCard->UpdateLightCardMaterialInstance();
+
+	// Update scale to match the desired size
+	{
+		// The default card plane is a 100x100 square.
+		constexpr double LightCardPlaneWidth = 100;
+		constexpr double LightCardPlaneHeight = 100;
+
+		LightCard->Scale.X = LightCardWidth / LightCardPlaneWidth;
+		LightCard->Scale.Y = LightCardHeight / LightCardPlaneHeight;
+	}
+
+	// Update position
+	MoveLightCardTo(*LightCard, LightCardCoords);
+
+#endif // WITH_OPENCV
+}
+
+double FDisplayClusterLightCardEditorViewportClient::CalculateFinalLightCardDistance(double FlushDistance, double DesiredOffsetFromFlush) const
+{
+	double Distance = FMath::Min(FlushDistance, RootActorBoundingRadius) + DesiredOffsetFromFlush;
+
+	return FMath::Max(Distance, 0);
+}
+
 void FDisplayClusterLightCardEditorViewportClient::ProcessClick(FSceneView& View, HHitProxy* HitProxy, FKey Key, EInputEvent Event, uint32 HitX, uint32 HitY)
 {
+	// Don't select light cards while drawing a new light card
+	if (InputMode == EInputMode::DrawingLightCard)
+	{
+		FEditorViewportClient::ProcessClick(View, HitProxy, Key, Event, HitX, HitY);
+		return;
+	}
+
 	UWorld* PreviewWorld = PreviewScene->GetWorld();
 	check(PreviewWorld);
 
@@ -984,6 +1251,7 @@
 				
 					LightCardProxy->SetActorLocation(LightCard->GetActorLocation() - RootActor->GetActorLocation());
 					LightCardProxy->SetActorRotation(LightCard->GetActorRotation() - RootActor->GetActorRotation());
+					LightCardProxy->PolygonMask = LightCard->PolygonMask;
 					LightCardProxy->bIsProxy = true;
 
 					// Change mesh to proxy mesh with more vertices
@@ -1193,7 +1461,7 @@
 	}
 
 	SetLightCardCoordinates(&LightCard, SphericalCoords);
-	LightCard.DistanceFromCenter = FMath::Min(DesiredDistance, RootActorBoundingRadius) + LightCardFlushOffset;
+	LightCard.DistanceFromCenter = CalculateFinalLightCardDistance(DesiredDistance);
 
 	const FRotator Rotation = FRotationMatrix::MakeFromX(-DesiredNormal).Rotator();
 
@@ -1830,6 +2098,7 @@
 		GetScene(),
 		EngineShowFlags)
 		.SetRealtimeUpdate(IsRealtime()));
+
 	FSceneView* View = CalcSceneView(&ViewFamily);
 
 	FVector Origin;
@@ -2317,4 +2586,23 @@
 	ProjectionFOVs[static_cast<int32>(EDisplayClusterMeshProjectionType::Azimuthal)] = 130.0f;
 }
 
+void FDisplayClusterLightCardEditorViewportClient::EnterDrawingLightCardMode()
+{
+	if (InputMode == EInputMode::Idle)
+	{
+		InputMode = EInputMode::DrawingLightCard;
+		SelectLightCard(nullptr);
+	}
+}
+
+void FDisplayClusterLightCardEditorViewportClient::ExitDrawingLightCardMode()
+{
+	if (InputMode == EInputMode::DrawingLightCard)
+	{
+		InputMode = EInputMode::Idle;
+		DrawnMousePositions.Empty();
+	}
+}
+
+
 #undef LOCTEXT_NAMESPACE
diff --git a/Engine/Plugins/Runtime/nDisplay/Source/DisplayClusterLightCardEditor/Private/Viewport/DisplayClusterLightCardEditorViewportClient.h b/Engine/Plugins/Runtime/nDisplay/Source/DisplayClusterLightCardEditor/Private/Viewport/DisplayClusterLightCardEditorViewportClient.h
--- a/Engine/Plugins/Runtime/nDisplay/Source/DisplayClusterLightCardEditor/Private/Viewport/DisplayClusterLightCardEditorViewportClient.h
+++ b/Engine/Plugins/Runtime/nDisplay/Source/DisplayClusterLightCardEditor/Private/Viewport/DisplayClusterLightCardEditorViewportClient.h
@@ -20,6 +20,22 @@
 /** Viewport Client for the preview viewport */
 class FDisplayClusterLightCardEditorViewportClient : public FEditorViewportClient, public TSharedFromThis<FDisplayClusterLightCardEditorViewportClient>
 {
+	using Super = FEditorViewportClient;
+
+public:
+
+	/** State machine that helps keep track of the context to which user inputs apply (e.g. mouse buttons, key presses) */
+	enum class EInputMode
+	{
+		Idle,
+
+		/** Indicates that the user is dragging an actor in the viewport */
+		DraggingActor,
+
+		/** Indicates that the user is drawing a light card in the viewport */
+		DrawingLightCard,
+	};
+
 private:
 	struct FSphericalCoordinates
 	{
@@ -171,7 +187,16 @@
 
 	/** Places the given light card in the middle of the current viewport */
 	void CenterLightCardInView(ADisplayClusterLightCardActor& LightCard);
-	
+
+	/** Returns the current input mode */
+	EInputMode GetInputMode() { return InputMode; }
+
+	/** Requests that we enter light card drawing input mode */
+	void EnterDrawingLightCardMode();
+
+	/** Requests that we exit light card drawing input mode (and go back to idle/normal) */
+	void ExitDrawingLightCardMode();
+
 private:
 	/** Initiates a transaction. */
 	void BeginTransaction(const FText& Description);
@@ -269,6 +294,12 @@
 	/** Resets the camera FOVs */
 	void ResetFOVs();
 
+	/** Creates a new light card using a polygon alpha mask as defined by the given mouse positions on the viewport */
+	void CreateDrawnLightCard(const TArray<FIntPoint>& MousePositions);
+
+	/** Calculates the final distance from the origin of a light card, given its flush distance and a desired offset */
+	double CalculateFinalLightCardDistance(double FlushDistance, double DesiredOffsetFromFlush = 0.) const;
+
 private:
 	TWeakPtr<FSceneViewport> SceneViewportPtr;
 	TWeakPtr<SDisplayClusterLightCardEditor> LightCardEditorPtr;
@@ -297,9 +328,6 @@
 	/** The renderer for the viewport, which can render the meshes with a variety of projection types */
 	TSharedPtr<FDisplayClusterMeshProjectionRenderer> MeshProjectionRenderer;
 	
-	/** Indicates that the user is dragging an actor in the viewport */
-	bool bDraggingActor = false;
-
 	/** The LC editor widget used to manipulate light cards */
 	TSharedPtr<FDisplayClusterLightCardEditorWidget> EditorWidget;
 
@@ -347,4 +375,10 @@
 
 	/** Indicates if the normal map should be displayed to the screen */
 	bool bDisplayNormalMapVisualization = false;
+
+	/** Current input mode */
+	EInputMode InputMode = EInputMode::Idle;
+
+	/** Array of mouse positions that will be used to spawn a new light card with the shape of the drawn polygon */
+	TArray<FIntPoint> DrawnMousePositions;
 };
diff --git a/Engine/Plugins/Runtime/nDisplay/Source/DisplayClusterLightCardEditor/Private/Viewport/DisplayClusterLightcardEditorViewport.cpp b/Engine/Plugins/Runtime/nDisplay/Source/DisplayClusterLightCardEditor/Private/Viewport/DisplayClusterLightcardEditorViewport.cpp
--- a/Engine/Plugins/Runtime/nDisplay/Source/DisplayClusterLightCardEditor/Private/Viewport/DisplayClusterLightcardEditorViewport.cpp
+++ b/Engine/Plugins/Runtime/nDisplay/Source/DisplayClusterLightCardEditor/Private/Viewport/DisplayClusterLightcardEditorViewport.cpp
@@ -56,6 +56,12 @@
 					.OnGetMenuContent(this, &SDisplayClusterLightCardEditorViewportToolBar::GenerateProjectionMenu)
 				]
 				+ SHorizontalBox::Slot()
+				.Padding(3.0f, 1.0f)
+				.HAlign(HAlign_Right)
+				[
+					MakeDrawingToolBar()
+				]
+				+ SHorizontalBox::Slot()
 				.Padding( 3.0f, 1.0f )
 				.HAlign( HAlign_Right )
 				[
@@ -179,6 +185,40 @@
 		return ToolbarBuilder.MakeWidget();
 	}
 
+	TSharedRef<SWidget> MakeDrawingToolBar()
+	{
+		FSlimHorizontalToolBarBuilder ToolbarBuilder(EditorViewport.Pin()->GetCommandList(), FMultiBoxCustomization::None);
+
+		const FName ToolBarStyle = TEXT("EditorViewportToolBar");
+		ToolbarBuilder.SetStyle(&FAppStyle::Get(), ToolBarStyle);
+		ToolbarBuilder.SetLabelVisibility(EVisibility::Visible);
+
+		ToolbarBuilder.SetIsFocusable(false);
+
+		ToolbarBuilder.BeginSection("Drawing");
+		{
+			ToolbarBuilder.BeginBlockGroup();
+
+			static FName DrawLightCardName = FName(TEXT("DrawLightCard"));
+			ToolbarBuilder.AddToolBarButton(
+				FDisplayClusterLightCardEditorCommands::Get().DrawLightCard,
+				NAME_None,
+				LOCTEXT("DrawLC", "Draw Light Card    "),
+				LOCTEXT("DrawLCTooltip", 
+					"Draw a new custom light card by adding polygon points clicking the left mouse button on the viewport. "
+					"End the shape with right mouse button, or abort by pressing this button again."),
+				TAttribute<FSlateIcon>(),
+				DrawLightCardName
+			);
+
+			ToolbarBuilder.EndBlockGroup();
+		}
+
+		ToolbarBuilder.EndSection();
+
+		return ToolbarBuilder.MakeWidget();
+	}
+
 private:
 	/** Reference to the parent viewport */
 	TWeakPtr<SDisplayClusterLightCardEditorViewport> EditorViewport;
@@ -303,6 +343,12 @@
 		CommandList->MapAction(
 			Commands.ResetCamera,
 			FExecuteAction::CreateSP(ViewportClient.Get(), &FDisplayClusterLightCardEditorViewportClient::ResetCamera, false));
+
+		CommandList->MapAction(
+			FDisplayClusterLightCardEditorCommands::Get().DrawLightCard,
+			FExecuteAction::CreateSP(this, &SDisplayClusterLightCardEditorViewport::DrawLightCard),
+			FCanExecuteAction(),
+			FIsActionChecked::CreateSP(this, &SDisplayClusterLightCardEditorViewport::IsDrawingLightCard));
 	}
 }
 
@@ -324,6 +370,24 @@
 	return false;
 }
 
+void SDisplayClusterLightCardEditorViewport::DrawLightCard()
+{
+	if (!ViewportClient.IsValid())
+	{
+		return;
+	}
+
+	if (IsDrawingLightCard())
+	{
+		ViewportClient->ExitDrawingLightCardMode();
+	}
+	else
+	{
+		ViewportClient->EnterDrawingLightCardMode();
+	}
+}
+
+
 void SDisplayClusterLightCardEditorViewport::CycleEditorWidgetMode()
 {
 	int32 WidgetModeAsInt = ViewportClient->GetEditorWidgetMode();
@@ -351,4 +415,15 @@
 	return false;
 }
 
+bool SDisplayClusterLightCardEditorViewport::IsDrawingLightCard() const
+{
+	if (ViewportClient.IsValid())
+	{
+		return ViewportClient->GetInputMode() == FDisplayClusterLightCardEditorViewportClient::EInputMode::DrawingLightCard;
+	}
+
+	return false;
+}
+
+
 #undef LOCTEXT_NAMESPACE
diff --git a/Engine/Plugins/Runtime/nDisplay/Source/DisplayClusterLightCardEditor/Private/Viewport/DisplayClusterLightcardEditorViewport.h b/Engine/Plugins/Runtime/nDisplay/Source/DisplayClusterLightCardEditor/Private/Viewport/DisplayClusterLightcardEditorViewport.h
--- a/Engine/Plugins/Runtime/nDisplay/Source/DisplayClusterLightCardEditor/Private/Viewport/DisplayClusterLightcardEditorViewport.h
+++ b/Engine/Plugins/Runtime/nDisplay/Source/DisplayClusterLightCardEditor/Private/Viewport/DisplayClusterLightcardEditorViewport.h
@@ -53,6 +53,13 @@
 	void SetProjectionMode(EDisplayClusterMeshProjectionType InProjectionMode);
 	bool IsProjectionModeSelected(EDisplayClusterMeshProjectionType InProjectionMode) const;
 
+	/** Handles the Draw LC button being pressed. */
+	void DrawLightCard();
+
+	/** Returns true if we are in light card drawing mode */
+	bool IsDrawingLightCard() const;
+
+
 private:
 	/** Preview Scene - uses advanced preview settings */
 	TSharedPtr<FAdvancedPreviewScene> AdvancedPreviewScene;
diff --git a/Engine/Plugins/Runtime/nDisplay/nDisplay.uplugin b/Engine/Plugins/Runtime/nDisplay/nDisplay.uplugin
--- a/Engine/Plugins/Runtime/nDisplay/nDisplay.uplugin
+++ b/Engine/Plugins/Runtime/nDisplay/nDisplay.uplugin
@@ -149,22 +149,22 @@
 			"Enabled": true
 		},
 		{
-			"Name": "ProceduralMeshComponent",
-			"Enabled": true
-		},
-		{
-			"Name": "ProceduralMeshComponent",
-			"Enabled": true
-		},
-		{
-			"Name": "RemoteControlInterception",
-			"Enabled": true
-		},
-		{
 			"Name": "ConcertSyncClient",
 			"Enabled": true
 		},
 		{
+			"Name": "OpenCV",
+			"Enabled": true
+		},
+		{
+			"Name": "ProceduralMeshComponent",
+			"Enabled": true
+		},
+		{
+			"Name": "RemoteControlInterception",
+			"Enabled": true
+		},
+		{
 			"Name": "VirtualProductionUtilities",
 			"Enabled": true
 		}
