nDisplay: Optimized spin value to minimize area of light card for a given polygon.

This improves the visual quality of the light card. And you end up with smaller light cards.

diff --git a/Engine/Plugins/Runtime/nDisplay/Source/DisplayClusterLightCardEditor/Private/DisplayClusterLightCardEditorStyle.h b/Engine/Plugins/Runtime/nDisplay/Source/DisplayClusterLightCardEditor/Private/DisplayClusterLightCardEditorStyle.h
new file mode 100644
--- /dev/null
+++ b/Engine/Plugins/Runtime/nDisplay/Source/DisplayClusterLightCardEditor/Private/DisplayClusterLightCardEditorStyle.h
@@ -0,0 +1,40 @@
+// Copyright Epic Games, Inc. All Rights Reserved.
+
+#pragma once
+
+#include "CoreMinimal.h"
+#include "Styling/SlateStyleRegistry.h"
+#include "Styling/SlateStyleMacros.h"
+#include "Styling/SlateStyle.h"
+
+
+/**
+ * Implements the visual style of the operator panel
+ */
+class FDisplayClusterLightCardEditorStyle final : public FSlateStyleSet
+{
+public:
+
+	FDisplayClusterLightCardEditorStyle()
+		: FSlateStyleSet("DisplayClusterLightCardEditorStyle")
+	{
+		const FVector2D Icon40x40(40.0f, 40.0f);
+
+		// Set miscellaneous icons
+		SetContentRoot(FPaths::EnginePluginsDir() / TEXT("Runtime/nDisplay/Content/Icons/OperatorPanel/"));
+		Set("DrawPoly", new IMAGE_BRUSH("PolyPath_40x", Icon40x40));
+
+		FSlateStyleRegistry::RegisterSlateStyle(*this);
+	}
+
+	virtual ~FDisplayClusterLightCardEditorStyle()
+	{
+		FSlateStyleRegistry::UnRegisterSlateStyle(*this);
+	}
+
+	static FDisplayClusterLightCardEditorStyle& Get()
+	{
+		static FDisplayClusterLightCardEditorStyle Inst;
+		return Inst;
+	}
+};
diff --git a/Engine/Plugins/Runtime/nDisplay/Source/DisplayClusterLightCardEditor/Private/Viewport/DisplayClusterLightCardEditorViewportClient.cpp b/Engine/Plugins/Runtime/nDisplay/Source/DisplayClusterLightCardEditor/Private/Viewport/DisplayClusterLightCardEditorViewportClient.cpp
--- a/Engine/Plugins/Runtime/nDisplay/Source/DisplayClusterLightCardEditor/Private/Viewport/DisplayClusterLightCardEditorViewportClient.cpp
+++ b/Engine/Plugins/Runtime/nDisplay/Source/DisplayClusterLightCardEditor/Private/Viewport/DisplayClusterLightCardEditorViewportClient.cpp
@@ -871,11 +871,11 @@ void FDisplayClusterLightCardEditorViewportClient::CreateDrawnLightCard(const TA
 	// @todo A future improvement would be to find a center that doesn't make the LC bigger than it needs to, 
 	// since the pixel space center won't be the projected plane center
 
-	FVector LightCardDirection; // direction from origin to where pixel points to
-	FVector LightCardPlaneAxisX;    // Normal vector at the found enveloping surface point
+	FVector LightCardDirection;  // direction from origin to where pixel points to
+	FVector LightCardPlaneAxisX; // Normal vector at the found enveloping surface point
 	FVector LightCardPlaneAxisY;
 	FVector LightCardPlaneAxisZ;
-	float LightCardDistance;    // Distance from origin to the found enveloping surface point
+	float LightCardDistance;     // Distance from origin to the found enveloping surface point
 	{
 		// Find center point in pixel space
 
@@ -892,7 +892,7 @@ void FDisplayClusterLightCardEditorViewportClient::CreateDrawnLightCard(const TA
 			MaxY = FMath::Max(MaxY, Point.Y);
 		}
 
-		FIntPoint CenterPoint((MaxX + MinX)/2, (MaxY + MinY) / 2);
+		FIntPoint CenterPoint((MaxX + MinX) / 2, (MaxY + MinY) / 2);
 
 		PixelToWorld(*View, CenterPoint, ViewOrigin, LightCardDirection);
 
@@ -931,8 +931,6 @@ void FDisplayClusterLightCardEditorViewportClient::CreateDrawnLightCard(const TA
 	// Project mouse positions to light card plane.
 	//
 
-	// @todo A future improvement would be to find a spin that minimizes the lightcard size
-
 	TArray<FVector3d> ProjectedMousePositions;
 	ProjectedMousePositions.AddUninitialized(MouseWorldDirections.Num());
 	const FPlane LightCardPlane = FPlane(LightCardLocation, LightCardPlaneAxisX);
@@ -962,15 +960,55 @@ void FDisplayClusterLightCardEditorViewportClient::CreateDrawnLightCard(const TA
 		PointsInLightCardPlane[PointIdx].Y =  FVector::DotProduct(ProjMousePosLocal, LightCardPlaneAxisZ);
 	}
 
-	// Calculate LC x/y bounds
+	// Find the spin that minimizes the area. 
+	// For now, simply a fixed size linear search but a better optimization algorithm could be used.
 
-	double LightCardWidth  = 0;
+	double Spin = 0;
+	double LightCardWidth = 0;
 	double LightCardHeight = 0;
 
-	for (const FVector2d& Point : PointsInLightCardPlane)
 	{
-		LightCardWidth  = FMath::Max(LightCardWidth , 2 * abs(Point.X));
-		LightCardHeight = FMath::Max(LightCardHeight, 2 * abs(Point.Y));
+		double Area = DBL_MAX;
+		double constexpr SpinStepSize = 15;
+
+		for (double SpinTest = 0; SpinTest < 180; SpinTest += SpinStepSize)
+		{
+			double LightCardWidthTest = 0;
+			double LightCardHeightTest = 0;
+
+			const FRotator Rotator(0, -SpinTest, 0); // using yaw for spin for convenience
+
+			for (const FVector2d& Point : PointsInLightCardPlane)
+			{
+				FVector RotatedPoint = Rotator.RotateVector(FVector(Point.X, Point.Y, 0));
+
+				LightCardWidthTest = FMath::Max(LightCardWidthTest, 2 * abs(RotatedPoint.X));
+				LightCardHeightTest = FMath::Max(LightCardHeightTest, 2 * abs(RotatedPoint.Y));
+			}
+
+			const double AreaTest = LightCardWidthTest * LightCardHeightTest;
+
+			if (AreaTest < Area)
+			{
+				Spin = SpinTest;
+				Area = AreaTest;
+				LightCardWidth = LightCardWidthTest;
+				LightCardHeight = LightCardHeightTest;
+			}
+		}
+	}
+
+	// Update the points with the rotated ones, since we're going to spin the card.
+	{
+		const FRotator Rotator(0, -Spin, 0);
+
+		for (FVector2d& Point : PointsInLightCardPlane)
+		{
+			FVector RotatedPoint = Rotator.RotateVector(FVector(Point.X, Point.Y, 0));
+
+			Point.X = RotatedPoint.X;
+			Point.Y = RotatedPoint.Y;
+		}
 	}
 
 	if (FMath::IsNearlyZero(LightCardWidth) || FMath::IsNearlyZero(LightCardHeight))
@@ -994,6 +1032,8 @@ void FDisplayClusterLightCardEditorViewportClient::CreateDrawnLightCard(const TA
 		return;
 	}
 
+	LightCard->Spin = Spin;
+
 	// Assign polygon mask
 
 	LightCard->Mask = EDisplayClusterLightCardMask::Polygon;
diff --git a/Engine/Plugins/Runtime/nDisplay/Source/DisplayClusterLightCardEditor/Private/Viewport/DisplayClusterLightcardEditorViewport.cpp b/Engine/Plugins/Runtime/nDisplay/Source/DisplayClusterLightCardEditor/Private/Viewport/DisplayClusterLightcardEditorViewport.cpp
--- a/Engine/Plugins/Runtime/nDisplay/Source/DisplayClusterLightCardEditor/Private/Viewport/DisplayClusterLightcardEditorViewport.cpp
+++ b/Engine/Plugins/Runtime/nDisplay/Source/DisplayClusterLightCardEditor/Private/Viewport/DisplayClusterLightcardEditorViewport.cpp
@@ -2,6 +2,7 @@
 
 #include "DisplayClusterLightcardEditorViewport.h"
 
+#include "DisplayClusterLightCardEditorStyle.h"
 #include "DisplayClusterLightCardEditorViewportClient.h"
 #include "DisplayClusterLightCardEditorCommands.h"
 
@@ -191,7 +192,7 @@ public:
 
 		const FName ToolBarStyle = TEXT("EditorViewportToolBar");
 		ToolbarBuilder.SetStyle(&FAppStyle::Get(), ToolBarStyle);
-		ToolbarBuilder.SetLabelVisibility(EVisibility::Visible);
+		ToolbarBuilder.SetLabelVisibility(EVisibility::Collapsed);
 
 		ToolbarBuilder.SetIsFocusable(false);
 
@@ -207,7 +208,7 @@ public:
 				LOCTEXT("DrawLCTooltip", 
 					"Draw a new custom light card by adding polygon points clicking the left mouse button on the viewport. "
 					"End the shape with right mouse button, or abort by pressing this button again."),
-				TAttribute<FSlateIcon>(),
+				FSlateIcon(FSlateIcon(FDisplayClusterLightCardEditorStyle::Get().GetStyleSetName(), "DrawPoly")),
 				DrawLightCardName
 			);
 
