nDisplay: Light card world placement when drawn on the viewport was improved to be based on world positions instead of screen space positions.

This makes it independent of the map projection.

diff --git a/Engine/Plugins/Runtime/nDisplay/Source/DisplayClusterLightCardEditor/Private/DisplayClusterLightCardEditorLog.cpp b/Engine/Plugins/Runtime/nDisplay/Source/DisplayClusterLightCardEditor/Private/DisplayClusterLightCardEditorLog.cpp
new file mode 100644
--- /dev/null
+++ b/Engine/Plugins/Runtime/nDisplay/Source/DisplayClusterLightCardEditor/Private/DisplayClusterLightCardEditorLog.cpp
@@ -0,0 +1,7 @@
+// Copyright Epic Games, Inc. All Rights Reserved.
+
+#include "DisplayClusterLightCardEditorLog.h"
+
+
+// Module-wide log categories
+DEFINE_LOG_CATEGORY(DisplayClusterLightCardEditorLog);
diff --git a/Engine/Plugins/Runtime/nDisplay/Source/DisplayClusterLightCardEditor/Private/DisplayClusterLightCardEditorLog.h b/Engine/Plugins/Runtime/nDisplay/Source/DisplayClusterLightCardEditor/Private/DisplayClusterLightCardEditorLog.h
new file mode 100644
--- /dev/null
+++ b/Engine/Plugins/Runtime/nDisplay/Source/DisplayClusterLightCardEditor/Private/DisplayClusterLightCardEditorLog.h
@@ -0,0 +1,8 @@
+// Copyright Epic Games, Inc. All Rights Reserved.
+
+#pragma once
+
+#include "CoreMinimal.h"
+
+// Module-wide log categories
+DECLARE_LOG_CATEGORY_EXTERN(DisplayClusterLightCardEditorLog, Log, All);
diff --git a/Engine/Plugins/Runtime/nDisplay/Source/DisplayClusterLightCardEditor/Private/Viewport/DisplayClusterLightCardEditorViewportClient.cpp b/Engine/Plugins/Runtime/nDisplay/Source/DisplayClusterLightCardEditor/Private/Viewport/DisplayClusterLightCardEditorViewportClient.cpp
--- a/Engine/Plugins/Runtime/nDisplay/Source/DisplayClusterLightCardEditor/Private/Viewport/DisplayClusterLightCardEditorViewportClient.cpp
+++ b/Engine/Plugins/Runtime/nDisplay/Source/DisplayClusterLightCardEditor/Private/Viewport/DisplayClusterLightCardEditorViewportClient.cpp
@@ -2,14 +2,15 @@
 
 #include "DisplayClusterLightCardEditorViewportClient.h"
 
+#include "Components/DisplayClusterPreviewComponent.h"
+#include "Components/DisplayClusterCameraComponent.h"
+#include "Components/DisplayClusterScreenComponent.h"
 #include "DisplayClusterLightCardEditorWidget.h"
 #include "DisplayClusterConfigurationTypes.h"
 #include "DisplayClusterProjectionStrings.h"
 #include "DisplayClusterRootActor.h"
 #include "DisplayClusterLightCardActor.h"
-#include "Components/DisplayClusterPreviewComponent.h"
-#include "Components/DisplayClusterCameraComponent.h"
-#include "Components/DisplayClusterScreenComponent.h"
+#include "DisplayClusterLightCardEditorLog.h"
 
 #include "AudioDevice.h"
 #include "CameraController.h"
@@ -42,7 +43,6 @@
 #include "UnrealWidget.h"
 #include "Widgets/Docking/SDockTab.h"
 
-
 #if WITH_OPENCV
 
 #include "OpenCVHelper.h"
@@ -836,6 +836,34 @@
 	FEditorViewportClient::TrackingStopped();
 }
 
+void FDisplayClusterLightCardEditorViewportClient::CalculateNormalAndPositionInDirection(
+	const FVector& InViewOrigin, 
+	const FVector& InDirection, 
+	FVector& OutWorldPosition, 
+	FVector& OutRelativeNormal, 
+	double InDesiredDistanceFromFlush) const
+{
+	const FVector Position = InViewOrigin + InDirection; // We fabricate a position in the right direction
+
+	float Distance;
+
+	// We find the normal and distance from origin
+	if (Position.Z < InViewOrigin.Z)
+	{
+		SouthNormalMap.GetNormalAndDistanceAtPosition(Position, OutRelativeNormal, Distance);
+	}
+	else
+	{
+		NorthNormalMap.GetNormalAndDistanceAtPosition(Position, OutRelativeNormal, Distance);
+	}
+
+	Distance = CalculateFinalLightCardDistance(Distance, InDesiredDistanceFromFlush);
+
+	// Calculate world position
+	OutWorldPosition = InViewOrigin + Distance * InDirection;
+};
+
+
 void FDisplayClusterLightCardEditorViewportClient::CreateDrawnLightCard(const TArray<FIntPoint>& MousePositions)
 {
 #if WITH_OPENCV
@@ -866,65 +894,74 @@
 	}
 
 	//
-	// Find world position and normal of card using the mouse center in pixel space.
+	// Find light card direction based on center of bounds of world positions of mouse positions
+	//
+
+	FVector LightCardDirection; // direction from origin to where pixel points to
+
+	{
+		double MinX =  DBL_MAX;
+		double MinY =  DBL_MAX;
+		double MinZ =  DBL_MAX;
+
+		double MaxX = -DBL_MAX;
+		double MaxY = -DBL_MAX;
+		double MaxZ = -DBL_MAX;
+
+		for (const FVector& MouseWorldDirection : MouseWorldDirections)
+		{
+			FVector MouseWorldPosition;
+			FVector MouseRelativeNormal;
+
+			CalculateNormalAndPositionInDirection(ViewOrigin, MouseWorldDirection, MouseWorldPosition, MouseRelativeNormal);
+
+			MinX = FMath::Min(MinX, MouseWorldPosition.X);
+			MinY = FMath::Min(MinY, MouseWorldPosition.Y);
+			MinZ = FMath::Min(MinZ, MouseWorldPosition.Z);
+
+			MaxX = FMath::Max(MaxX, MouseWorldPosition.X);
+			MaxY = FMath::Max(MaxY, MouseWorldPosition.Y);
+			MaxZ = FMath::Max(MaxZ, MouseWorldPosition.Z);
+		}
+
+		const FVector MouseWorldBoundsCenter(
+			(MaxX + MinX) / 2,
+			(MaxY + MinY) / 2,
+			(MaxZ + MinZ) / 2
+		);
+
+		LightCardDirection = (MouseWorldBoundsCenter - ViewOrigin).GetSafeNormal();
+
+		if (LightCardDirection == FVector::ZeroVector)
+		{
+			UE_LOG(DisplayClusterLightCardEditorLog, Warning, TEXT("Could not determine a LightCardDirection for the given polygon"));
+			return;
+		}
+	}
+
+	//
+	// Find world position and normal of card, and its plane axes
 	// 
-	// @todo A future improvement would be to find a center that doesn't make the LC bigger than it needs to, 
-	// since the pixel space center won't be the projected plane center
 
-	FVector LightCardDirection;  // direction from origin to where pixel points to
+	FVector LightCardLocation; // Wolrd location of light card
 	FVector LightCardPlaneAxisX; // Normal vector at the found enveloping surface point
 	FVector LightCardPlaneAxisY;
 	FVector LightCardPlaneAxisZ;
-	float LightCardDistance;     // Distance from origin to the found enveloping surface point
+
 	{
-		// Find center point in pixel space
+		FVector LightCardRelativeNormal;
 
-		int32 MinX = INT_MAX;
-		int32 MinY = INT_MAX;
-		int32 MaxX = INT_MIN;
-		int32 MaxY = INT_MIN;
+		CalculateNormalAndPositionInDirection(ViewOrigin, LightCardDirection, LightCardLocation, LightCardRelativeNormal);
 
-		for (const FIntPoint& Point : MousePositions)
-		{
-			MinX = FMath::Min(MinX, Point.X);
-			MinY = FMath::Min(MinY, Point.Y);
-			MaxX = FMath::Max(MaxX, Point.X);
-			MaxY = FMath::Max(MaxY, Point.Y);
-		}
-
-		FIntPoint CenterPoint((MaxX + MinX) / 2, (MaxY + MinY) / 2);
-
-		PixelToWorld(*View, CenterPoint, ViewOrigin, LightCardDirection);
-
-		const FVector Position = ViewOrigin + LightCardDirection * 100.; // We fabricate a position in the right direction
-
-		FVector RelativeNormal;
-
-		// We find the normal and distance from origin
-		if (Position.Z < ViewOrigin.Z)
-		{
-			SouthNormalMap.GetNormalAndDistanceAtPosition(Position, RelativeNormal, LightCardDistance);
-		}
-		else
-		{
-			NorthNormalMap.GetNormalAndDistanceAtPosition(Position, RelativeNormal, LightCardDistance);
-		}
-
-		LightCardDistance = CalculateFinalLightCardDistance(LightCardDistance);
-
-		// Calculate light card normal from light card direction and relative normal
-
-		const FMatrix RotEffector = FRotationMatrix::MakeFromX(-RelativeNormal);
+		const FMatrix RotEffector = FRotationMatrix::MakeFromX(-LightCardRelativeNormal);
 		const FMatrix RotArm = FRotationMatrix::MakeFromX(LightCardDirection);
 		const FRotator TotalRotator = (RotEffector * RotArm).Rotator();
 
 		LightCardPlaneAxisX = TotalRotator.RotateVector(-FVector::XAxisVector); // Same as Normal
 		LightCardPlaneAxisY = TotalRotator.RotateVector(-FVector::YAxisVector);
-		LightCardPlaneAxisZ = TotalRotator.RotateVector(FVector::ZAxisVector);
+		LightCardPlaneAxisZ = TotalRotator.RotateVector( FVector::ZAxisVector);
 	}
 
-	const FVector LightCardLocation = ViewOrigin + LightCardDistance * LightCardDirection;
-
 	const FSphericalCoordinates LightCardCoords(LightCardLocation - ViewOrigin);
 
 	//
@@ -932,7 +969,7 @@
 	//
 
 	TArray<FVector3d> ProjectedMousePositions;
-	ProjectedMousePositions.AddUninitialized(MouseWorldDirections.Num());
+	ProjectedMousePositions.Reserve(MouseWorldDirections.Num());
 	const FPlane LightCardPlane = FPlane(LightCardLocation, LightCardPlaneAxisX);
 
 	for (int32 PointIdx = 0; PointIdx < MouseWorldDirections.Num(); ++PointIdx)
@@ -944,20 +981,22 @@
 			continue;
 		}
 
-		ProjectedMousePositions[PointIdx] = FMath::RayPlaneIntersection(ViewOrigin, MouseWorldDirections[PointIdx], LightCardPlane);
+		ProjectedMousePositions.Add(FMath::RayPlaneIntersection(ViewOrigin, MouseWorldDirections[PointIdx], LightCardPlane));
 	}
 
 	// Convert projected mouse positions in world space to light card plane space (2d)
 
 	TArray<FVector2d> PointsInLightCardPlane;
-	PointsInLightCardPlane.AddUninitialized(ProjectedMousePositions.Num());
+	PointsInLightCardPlane.Reserve(ProjectedMousePositions.Num());
 
-	for (int32 PointIdx = 0; PointIdx < ProjectedMousePositions.Num(); ++PointIdx)
+	for (const FVector& ProjectedMousePosition: ProjectedMousePositions)
 	{
-		const FVector ProjMousePosLocal = ProjectedMousePositions[PointIdx] - LightCardLocation; 
+		const FVector ProjMousePosLocal = ProjectedMousePosition - LightCardLocation; 
 
-		PointsInLightCardPlane[PointIdx].X = -FVector::DotProduct(ProjMousePosLocal, LightCardPlaneAxisY);
-		PointsInLightCardPlane[PointIdx].Y =  FVector::DotProduct(ProjMousePosLocal, LightCardPlaneAxisZ);
+		PointsInLightCardPlane.Add(FVector2D(
+			-FVector::DotProduct(ProjMousePosLocal, LightCardPlaneAxisY),
+			 FVector::DotProduct(ProjMousePosLocal, LightCardPlaneAxisZ))
+		);
 	}
 
 	// Find the spin that minimizes the area. 
@@ -1013,6 +1052,7 @@
 
 	if (FMath::IsNearlyZero(LightCardWidth) || FMath::IsNearlyZero(LightCardHeight))
 	{
+		UE_LOG(DisplayClusterLightCardEditorLog, Warning, TEXT("Could not create new light card because one or more of its dimensions was zero."));
 		return;
 	}
 
@@ -1020,6 +1060,7 @@
 
 	if (!LightCardEditorPtr.IsValid())
 	{
+		UE_LOG(DisplayClusterLightCardEditorLog, Warning, TEXT("Could not create new light card because LightCardEditorPtr was not valid."));
 		return;
 	}
 
@@ -1029,6 +1070,7 @@
 
 	if (!LightCard)
 	{
+		UE_LOG(DisplayClusterLightCardEditorLog, Warning, TEXT("Could not create new light card because AddNewLightCard failed."));
 		return;
 	}
 
diff --git a/Engine/Plugins/Runtime/nDisplay/Source/DisplayClusterLightCardEditor/Private/Viewport/DisplayClusterLightCardEditorViewportClient.h b/Engine/Plugins/Runtime/nDisplay/Source/DisplayClusterLightCardEditor/Private/Viewport/DisplayClusterLightCardEditorViewportClient.h
--- a/Engine/Plugins/Runtime/nDisplay/Source/DisplayClusterLightCardEditor/Private/Viewport/DisplayClusterLightCardEditorViewportClient.h
+++ b/Engine/Plugins/Runtime/nDisplay/Source/DisplayClusterLightCardEditor/Private/Viewport/DisplayClusterLightCardEditorViewportClient.h
@@ -300,6 +300,9 @@
 	/** Calculates the final distance from the origin of a light card, given its flush distance and a desired offset */
 	double CalculateFinalLightCardDistance(double FlushDistance, double DesiredOffsetFromFlush = 0.) const;
 
+	/** Calculates the relative normal vector and world position in the specified direction from the given view origin */
+	void CalculateNormalAndPositionInDirection(const FVector& InViewOrigin, const FVector& InDirection, FVector& OutWorldLocation, FVector& OutRelativeNormal, double InDesiredDistanceFromFlush = 0.) const;
+
 private:
 	TWeakPtr<FSceneViewport> SceneViewportPtr;
 	TWeakPtr<SDisplayClusterLightCardEditor> LightCardEditorPtr;
