nDisplay: Added context menus to the light card list view and the 3D viewport, which contain common editing commands like copy and paste; users can now copy and paste light cards in either view, as well as use common keyboard shortcuts. Also replaced all light card types with specific ADisplayClusterLightCardActor references, as we were implicitly making the assumption all light cards were of that type, but rarely enforcing it

diff --git a/Engine/Plugins/Runtime/nDisplay/Source/DisplayClusterLightCardEditor/DisplayClusterLightCardEditor.Build.cs b/Engine/Plugins/Runtime/nDisplay/Source/DisplayClusterLightCardEditor/DisplayClusterLightCardEditor.Build.cs
--- a/Engine/Plugins/Runtime/nDisplay/Source/DisplayClusterLightCardEditor/DisplayClusterLightCardEditor.Build.cs
+++ b/Engine/Plugins/Runtime/nDisplay/Source/DisplayClusterLightCardEditor/DisplayClusterLightCardEditor.Build.cs
@@ -11,6 +11,7 @@
 			new string[] {
 
 				"AdvancedPreviewScene",
+				"ApplicationCore",
 				"AppFramework",
 				"Core",
 				"CoreUObject",
diff --git a/Engine/Plugins/Runtime/nDisplay/Source/DisplayClusterLightCardEditor/Private/DisplayClusterLightCardEditorCommands.cpp b/Engine/Plugins/Runtime/nDisplay/Source/DisplayClusterLightCardEditor/Private/DisplayClusterLightCardEditorCommands.cpp
--- a/Engine/Plugins/Runtime/nDisplay/Source/DisplayClusterLightCardEditor/Private/DisplayClusterLightCardEditorCommands.cpp
+++ b/Engine/Plugins/Runtime/nDisplay/Source/DisplayClusterLightCardEditor/Private/DisplayClusterLightCardEditorCommands.cpp
@@ -1,22 +1,23 @@
-// Copyright Epic Games, Inc. All Rights Reserved.
-
-#include "DisplayClusterLightCardEditorCommands.h"
-
-#define LOCTEXT_NAMESPACE "DisplayClusterLightCardEditorCommands"
-
-
-void FDisplayClusterLightCardEditorCommands::RegisterCommands()
-{
-	UI_COMMAND(ResetCamera, "Reset Camera", "Resets the camera to focus on the mesh", EUserInterfaceActionType::Button, FInputChord());
-
-	UI_COMMAND(PerspectiveProjection, "Perspective", "A perspective projection from the stage's view origin", EUserInterfaceActionType::Button, FInputChord());
-	UI_COMMAND(AzimuthalProjection, "Azimuthal", "An azimuthal hemisphere projection from the stage's center pointed upwards", EUserInterfaceActionType::Button, FInputChord());
-
-	UI_COMMAND(AddNewLightCard, "Add New Light Card", "Add and assign a new Light Card to the actor", EUserInterfaceActionType::Button, FInputChord());
-	UI_COMMAND(AddExistingLightCard, "Add Existing Light Card", "Add an existing Light Card to the actor", EUserInterfaceActionType::Button, FInputChord());
-	UI_COMMAND(RemoveLightCard, "Remove from Actor", "Remove the Light Card from the actor but do not delete it", EUserInterfaceActionType::Button, FInputChord());
-
-	UI_COMMAND(DrawLightCard, "Draw Light Card", "Draw polygon light card on viewport", EUserInterfaceActionType::ToggleButton, FInputChord());
-}
-
-#undef LOCTEXT_NAMESPACE
+// Copyright Epic Games, Inc. All Rights Reserved.
+
+#include "DisplayClusterLightCardEditorCommands.h"
+
+#define LOCTEXT_NAMESPACE "DisplayClusterLightCardEditorCommands"
+
+
+void FDisplayClusterLightCardEditorCommands::RegisterCommands()
+{
+	UI_COMMAND(ResetCamera, "Reset Camera", "Resets the camera to focus on the mesh", EUserInterfaceActionType::Button, FInputChord());
+
+	UI_COMMAND(PerspectiveProjection, "Perspective", "A perspective projection from the stage's view origin", EUserInterfaceActionType::Button, FInputChord());
+	UI_COMMAND(AzimuthalProjection, "Azimuthal", "An azimuthal hemisphere projection from the stage's center pointed upwards", EUserInterfaceActionType::Button, FInputChord());
+
+	UI_COMMAND(AddNewLightCard, "Add New Light Card", "Add and assign a new Light Card to the actor", EUserInterfaceActionType::Button, FInputChord());
+	UI_COMMAND(AddExistingLightCard, "Add Existing Light Card", "Add an existing Light Card to the actor", EUserInterfaceActionType::Button, FInputChord());
+	UI_COMMAND(RemoveLightCard, "Remove from Actor", "Remove the Light Card from the actor but do not delete it", EUserInterfaceActionType::Button, FInputChord());
+	UI_COMMAND(PasteHere, "Paste Here", "Paste clipboard contents at the click location", EUserInterfaceActionType::Button, FInputChord());
+
+	UI_COMMAND(DrawLightCard, "Draw Light Card", "Draw polygon light card on viewport", EUserInterfaceActionType::ToggleButton, FInputChord());
+}
+
+#undef LOCTEXT_NAMESPACE
diff --git a/Engine/Plugins/Runtime/nDisplay/Source/DisplayClusterLightCardEditor/Private/DisplayClusterLightCardEditorCommands.h b/Engine/Plugins/Runtime/nDisplay/Source/DisplayClusterLightCardEditor/Private/DisplayClusterLightCardEditorCommands.h
--- a/Engine/Plugins/Runtime/nDisplay/Source/DisplayClusterLightCardEditor/Private/DisplayClusterLightCardEditorCommands.h
+++ b/Engine/Plugins/Runtime/nDisplay/Source/DisplayClusterLightCardEditor/Private/DisplayClusterLightCardEditorCommands.h
@@ -1,33 +1,37 @@
-// Copyright Epic Games, Inc. All Rights Reserved.
-
-#pragma once
-
-#include "CoreMinimal.h"
-#include "Framework/Commands/Commands.h"
-#include "Styling/AppStyle.h"
-
-
-class FDisplayClusterLightCardEditorCommands 
-	: public TCommands<FDisplayClusterLightCardEditorCommands>
-{
-public:
-	FDisplayClusterLightCardEditorCommands()
-		: TCommands<FDisplayClusterLightCardEditorCommands>(TEXT("DisplayClusterLightCardEditor"), 
-			NSLOCTEXT("Contexts", "DisplayClusterLightCardEditor", "Display Cluster LightCard Editor"), NAME_None, FAppStyle::GetAppStyleSetName())
-	{ }
-
-	virtual void RegisterCommands() override;
-
-public:
-	// Viewport commands
-	TSharedPtr<FUICommandInfo> ResetCamera;
-
-	TSharedPtr<FUICommandInfo> PerspectiveProjection;
-	TSharedPtr<FUICommandInfo> AzimuthalProjection;
-
-	TSharedPtr<FUICommandInfo> AddNewLightCard;
-	TSharedPtr<FUICommandInfo> AddExistingLightCard;
-	TSharedPtr<FUICommandInfo> RemoveLightCard;
-
-	TSharedPtr<FUICommandInfo> DrawLightCard;
-};
+// Copyright Epic Games, Inc. All Rights Reserved.
+
+#pragma once
+
+#include "CoreMinimal.h"
+#include "Framework/Commands/Commands.h"
+
+#include "DisplayClusterLightCardEditorStyle.h"
+
+
+class FDisplayClusterLightCardEditorCommands 
+	: public TCommands<FDisplayClusterLightCardEditorCommands>
+{
+public:
+	FDisplayClusterLightCardEditorCommands()
+		: TCommands<FDisplayClusterLightCardEditorCommands>(TEXT("DisplayClusterLightCardEditor"), 
+			NSLOCTEXT("Contexts", "DisplayClusterLightCardEditor", "Display Cluster LightCard Editor"),
+			NAME_None,
+			FDisplayClusterLightCardEditorStyle::Get().GetStyleSetName())
+	{ }
+
+	virtual void RegisterCommands() override;
+
+public:
+	// Viewport commands
+	TSharedPtr<FUICommandInfo> ResetCamera;
+
+	TSharedPtr<FUICommandInfo> PerspectiveProjection;
+	TSharedPtr<FUICommandInfo> AzimuthalProjection;
+
+	TSharedPtr<FUICommandInfo> AddNewLightCard;
+	TSharedPtr<FUICommandInfo> AddExistingLightCard;
+	TSharedPtr<FUICommandInfo> RemoveLightCard;
+	TSharedPtr<FUICommandInfo> PasteHere;
+
+	TSharedPtr<FUICommandInfo> DrawLightCard;
+};
diff --git a/Engine/Plugins/Runtime/nDisplay/Source/DisplayClusterLightCardEditor/Private/DisplayClusterLightCardEditorStyle.h b/Engine/Plugins/Runtime/nDisplay/Source/DisplayClusterLightCardEditor/Private/DisplayClusterLightCardEditorStyle.h
--- a/Engine/Plugins/Runtime/nDisplay/Source/DisplayClusterLightCardEditor/Private/DisplayClusterLightCardEditorStyle.h
+++ b/Engine/Plugins/Runtime/nDisplay/Source/DisplayClusterLightCardEditor/Private/DisplayClusterLightCardEditorStyle.h
@@ -19,11 +19,16 @@
 		: FSlateStyleSet("DisplayClusterLightCardEditorStyle")
 	{
 		const FVector2D Icon40x40(40.0f, 40.0f);
+		const FVector2D Icon16x16(16.0f, 16.0f);
 
 		// Set miscellaneous icons
 		SetContentRoot(FPaths::EnginePluginsDir() / TEXT("Runtime/nDisplay/Content/Icons/OperatorPanel/"));
+		SetCoreContentRoot(FPaths::EngineContentDir() / TEXT("Editor/Slate"));
+
 		Set("DrawPoly", new IMAGE_BRUSH("PolyPath_40x", Icon40x40));
 
+		Set("DisplayClusterLightCardEditor.PasteHere", new CORE_IMAGE_BRUSH_SVG("Starship/Actors/paste-here", Icon16x16));
+
 		FSlateStyleRegistry::RegisterSlateStyle(*this);
 	}
 
diff --git a/Engine/Plugins/Runtime/nDisplay/Source/DisplayClusterLightCardEditor/Private/SDisplayClusterLightCardEditor.cpp b/Engine/Plugins/Runtime/nDisplay/Source/DisplayClusterLightCardEditor/Private/SDisplayClusterLightCardEditor.cpp
--- a/Engine/Plugins/Runtime/nDisplay/Source/DisplayClusterLightCardEditor/Private/SDisplayClusterLightCardEditor.cpp
+++ b/Engine/Plugins/Runtime/nDisplay/Source/DisplayClusterLightCardEditor/Private/SDisplayClusterLightCardEditor.cpp
@@ -2,6 +2,7 @@
 
 #include "SDisplayClusterLightCardEditor.h"
 
+#include "DisplayClusterLightCardEditorCommands.h"
 #include "SDisplayClusterLightCardList.h"
 
 #include "IDisplayClusterOperator.h"
@@ -9,14 +10,21 @@
 #include "Viewport/DisplayClusterLightCardEditorViewportClient.h"
 
 #include "DisplayClusterRootActor.h"
+#include "DisplayClusterConfigurationTypes.h"
+#include "Components/DisplayClusterCameraComponent.h"
 
+#include "Framework/Commands/GenericCommands.h"
 #include "Framework/Docking/TabManager.h"
 #include "Framework/Docking/LayoutExtender.h"
 #include "Framework/MultiBox/MultiBoxBuilder.h"
 #include "Framework/MultiBox/MultiBoxExtender.h"
+#include "PropertyCustomizationHelpers.h"
+#include "ScopedTransaction.h"
+#include "Selection.h"
 #include "Toolkits/AssetEditorToolkit.h"
 #include "Widgets/Docking/SDockTab.h"
 #include "Widgets/Layout/SSplitter.h"
+#include "Widgets/Workflow/SWizard.h"
 
 #define LOCTEXT_NAMESPACE "SDisplayClusterLightCardEditor"
 
@@ -92,6 +100,8 @@
 	
 	OnObjectTransactedHandle = FCoreUObjectDelegates::OnObjectTransacted.AddSP(this, &SDisplayClusterLightCardEditor::OnObjectTransacted);
 
+	BindCommands();
+
 	ChildSlot                     
 	[
 		SNew(SSplitter)
@@ -103,7 +113,7 @@
 			// Vertical box for the left hand panel of the editor. Add new slots here as needed for any editor UI controls
 			SNew(SVerticalBox)
 			+ SVerticalBox::Slot()
-			.AutoHeight()
+			.FillHeight(1.0f)
 			[
 				CreateLightCardListWidget()
 			]
@@ -122,13 +132,19 @@
 	BindCompileDelegates();
 }
 
-void SDisplayClusterLightCardEditor::SelectLightCards(const TArray<AActor*>& LightCardsToSelect)
+void SDisplayClusterLightCardEditor::SelectLightCards(const TArray<ADisplayClusterLightCardActor*>& LightCardsToSelect)
 {
 	check(LightCardList);
 	LightCardList->SelectLightCards(LightCardsToSelect);
 }
 
-void SDisplayClusterLightCardEditor::SelectLightCardProxies(const TArray<AActor*>& LightCardsToSelect)
+void SDisplayClusterLightCardEditor::GetSelectedLightCards(TArray<ADisplayClusterLightCardActor*>& OutSelectedLightCards)
+{
+	check(LightCardList);
+	LightCardList->GetSelectedLightCards(OutSelectedLightCards);
+}
+
+void SDisplayClusterLightCardEditor::SelectLightCardProxies(const TArray<ADisplayClusterLightCardActor*>& LightCardsToSelect)
 {
 	check(ViewportView);
 	ViewportView->GetLightCardEditorViewportClient()->SelectLightCards(LightCardsToSelect);
@@ -140,12 +156,354 @@
 	ViewportView->GetLightCardEditorViewportClient()->CenterLightCardInView(LightCard);
 }
 
-ADisplayClusterLightCardActor* SDisplayClusterLightCardEditor::AddNewLightCard()
+ADisplayClusterLightCardActor* SDisplayClusterLightCardEditor::SpawnLightCard()
 {
-	check(LightCardList);
-	return LightCardList->AddNewLightCard();
+	if (!ActiveRootActor.IsValid())
+	{
+		return nullptr;
+	}
+
+	const FVector SpawnLocation = ActiveRootActor->GetDefaultCamera()->GetComponentLocation();
+	FRotator SpawnRotation = ActiveRootActor->GetDefaultCamera()->GetComponentRotation();
+	SpawnRotation.Yaw -= 180.f;
+
+	FActorSpawnParameters SpawnParameters;
+	SpawnParameters.bNoFail = true;
+	SpawnParameters.SpawnCollisionHandlingOverride = ESpawnActorCollisionHandlingMethod::AdjustIfPossibleButAlwaysSpawn;
+	SpawnParameters.Name = TEXT("LightCard");
+	SpawnParameters.NameMode = FActorSpawnParameters::ESpawnActorNameMode::Requested;
+	SpawnParameters.OverrideLevel = ActiveRootActor->GetWorld()->GetCurrentLevel();
+
+	ADisplayClusterLightCardActor* NewLightCard = CastChecked<ADisplayClusterLightCardActor>(
+		ActiveRootActor->GetWorld()->SpawnActor(ADisplayClusterLightCardActor::StaticClass(),
+			&SpawnLocation, &SpawnRotation, MoveTemp(SpawnParameters)));
+
+	NewLightCard->SetActorLabel(NewLightCard->GetName());
+
+	TArray<ADisplayClusterLightCardActor*> LightCards { NewLightCard } ;
+	AddLightCardsToActor(LightCards);
+
+	return NewLightCard;
 }
 
+void SDisplayClusterLightCardEditor::AddNewLightCard()
+{
+	check(ActiveRootActor.IsValid());
+
+	FScopedTransaction Transaction(LOCTEXT("AddNewLightCardTransactionMessage", "Add New Light Card"));
+
+	ADisplayClusterLightCardActor* NewLightCard = SpawnLightCard();
+
+	// When adding a new lightcard, usually the desired location is in the middle of the viewport
+	if (NewLightCard)
+	{
+		CenterLightCardInView(*NewLightCard);
+	}
+}
+
+void SDisplayClusterLightCardEditor::AddExistingLightCard()
+{
+	TSharedPtr<SWindow> PickerWindow;
+	TWeakObjectPtr<ADisplayClusterLightCardActor> SelectedActorPtr;
+	bool bFinished = false;
+	
+	const TSharedRef<SWidget> ActorPicker = PropertyCustomizationHelpers::MakeActorPickerWithMenu(
+		nullptr,
+		false,
+		FOnShouldFilterActor::CreateLambda([&](const AActor* const InActor) -> bool // ActorFilter
+		{
+			const bool IsAllowed = InActor != nullptr && !InActor->IsChildActor() && InActor->IsA<ADisplayClusterLightCardActor>() &&
+				!InActor->GetClass()->HasAnyClassFlags(CLASS_Interface)	&& !InActor->IsA<ADisplayClusterRootActor>();
+			
+			return IsAllowed;
+		}),
+		FOnActorSelected::CreateLambda([&](AActor* InActor) -> void // OnSet
+		{
+			SelectedActorPtr = Cast<ADisplayClusterLightCardActor>(InActor);
+		}),
+		FSimpleDelegate::CreateLambda([&]() -> void // OnClose
+		{
+		}),
+		FSimpleDelegate::CreateLambda([&]() -> void // OnUseSelected
+		{
+			if (ADisplayClusterLightCardActor* Selection = Cast<ADisplayClusterLightCardActor>(GEditor->GetSelectedActors()->GetTop(ADisplayClusterLightCardActor::StaticClass())))
+			{
+				SelectedActorPtr = Selection;
+			}
+		}));
+	
+	PickerWindow = SNew(SWindow)
+	.Title(LOCTEXT("AddExistingLightCard", "Select an existing Light Card actor"))
+	.ClientSize(FVector2D(500, 525))
+	.SupportsMinimize(false) .SupportsMaximize(false)
+	[
+		SNew(SBorder)
+		.BorderImage(FAppStyle::GetBrush("Menu.Background"))
+		[
+			SNew(SWizard)
+			.FinishButtonText(LOCTEXT("FinishAddingExistingLightCard", "Add Actor"))
+			.OnCanceled(FSimpleDelegate::CreateLambda([&]()
+			{
+				if (PickerWindow.IsValid())
+				{
+					PickerWindow->RequestDestroyWindow();
+				}
+			}))
+			.OnFinished(FSimpleDelegate::CreateLambda([&]()
+			{
+				bFinished = true;
+				if (PickerWindow.IsValid())
+				{
+					PickerWindow->RequestDestroyWindow();
+				}
+			}))
+			.CanFinish(TAttribute<bool>::CreateLambda([&]()
+			{
+				return SelectedActorPtr.IsValid();
+			}))
+			.ShowPageList(false)
+			+SWizard::Page()
+			.CanShow(true)
+			[
+				SNew(SBorder)
+				.VAlign(VAlign_Fill)
+				.BorderImage(FAppStyle::GetBrush("Menu.Background"))
+				[
+					SNew(SVerticalBox)
+					+SVerticalBox::Slot()
+					.AutoHeight()
+					.VAlign(VAlign_Fill)
+					.HAlign(HAlign_Fill)
+					[
+						ActorPicker
+					]
+					+SVerticalBox::Slot()
+					.AutoHeight()
+					.VAlign(VAlign_Bottom)
+					.Padding(0.f, 8.f)
+					[
+						SNew(STextBlock)
+						.TextStyle(FAppStyle::Get(), "NormalText.Important")
+						.Text_Lambda([&]
+						{
+							const FString Result = FString::Printf(TEXT("Selected Actor: %s"),
+								SelectedActorPtr.IsValid() ? *SelectedActorPtr->GetActorLabel() : TEXT(""));
+							return FText::FromString(Result);
+						})
+					]
+				]
+			]
+		]
+	];
+
+	GEditor->EditorAddModalWindow(PickerWindow.ToSharedRef());
+	if (bFinished && SelectedActorPtr.IsValid())
+	{
+		FScopedTransaction Transaction(LOCTEXT("AddExistingLightCardTransactionMessage", "Add Existing Light Card"));
+
+		TArray<ADisplayClusterLightCardActor*> LightCards { SelectedActorPtr.Get() };
+		AddLightCardsToActor(LightCards);
+	}
+
+	PickerWindow.Reset();
+	SelectedActorPtr.Reset();
+}
+
+void SDisplayClusterLightCardEditor::AddLightCardsToActor(TArray<ADisplayClusterLightCardActor*> LightCards)
+{
+	if (ActiveRootActor.IsValid())
+	{
+		UDisplayClusterConfigurationData* ConfigData = ActiveRootActor->GetConfigData();
+		ConfigData->Modify();
+		FDisplayClusterConfigurationICVFX_VisibilityList& RootActorLightCards = ConfigData->StageSettings.Lightcard.ShowOnlyList;
+
+		for (ADisplayClusterLightCardActor* LightCard : LightCards)
+		{
+			check(LightCard);
+
+			if (!RootActorLightCards.Actors.ContainsByPredicate([&](const TSoftObjectPtr<AActor>& Actor)
+				{
+					// Don't add if a loaded actor is already present.
+					return Actor.Get() == LightCard;
+				}))
+			{
+				const TSoftObjectPtr<AActor> LightCardSoftObject(LightCard);
+
+				// Remove any exact paths to this actor. It's possible invalid actors are present if a light card
+				// was force deleted from a level.
+				RootActorLightCards.Actors.RemoveAll([&](const TSoftObjectPtr<AActor>& Actor)
+					{
+						return Actor == LightCardSoftObject;
+					});
+
+				RootActorLightCards.Actors.Add(LightCard);
+			}
+		}
+
+		RefreshPreviewActors(EDisplayClusterLightCardEditorProxyType::LightCards);
+	}
+}
+
+bool SDisplayClusterLightCardEditor::CanAddLightCard() const
+{
+	return ActiveRootActor.IsValid() && ActiveRootActor->GetWorld() != nullptr;
+}
+
+void SDisplayClusterLightCardEditor::CutLightCards()
+{
+	CopyLightCards();
+	RemoveLightCards(true);
+}
+
+bool SDisplayClusterLightCardEditor::CanCutLightCards()
+{
+	return CanCopyLightCards() && CanRemoveLightCards();
+}
+
+void SDisplayClusterLightCardEditor::CopyLightCards()
+{
+	TArray<ADisplayClusterLightCardActor*> SelectedLightCards;
+	LightCardList->GetSelectedLightCards(SelectedLightCards);
+
+	UWorld* EditorWorld = GEditor->GetEditorWorldContext().World();
+
+	const bool bNoteSelectionChange = false;
+	const bool bDeselectBSPSurfs = true;
+	const bool bWarnAboutManyActors = false;
+	GEditor->SelectNone(bNoteSelectionChange, bDeselectBSPSurfs, bWarnAboutManyActors);
+
+	for (AActor* LightCard : SelectedLightCards)
+	{
+		const bool bInSelected = true;
+		const bool bNotify = false;
+		const bool bSelectEvenIfHidden = true;
+		GEditor->SelectActor(LightCard, true, bNotify, bSelectEvenIfHidden);
+	}
+
+	const bool bShouldCut = false;
+	const bool bIsMove = false;
+	const bool bWarnAboutReferences = false;
+	GEditor->CopySelectedActorsToClipboard(EditorWorld, bShouldCut, bIsMove, bWarnAboutReferences);
+}
+
+bool SDisplayClusterLightCardEditor::CanCopyLightCards() const
+{
+	TArray<ADisplayClusterLightCardActor*> SelectedLightCards;
+	LightCardList->GetSelectedLightCards(SelectedLightCards);
+
+	return SelectedLightCards.Num() > 0;
+}
+
+void SDisplayClusterLightCardEditor::PasteLightCards(bool bOffsetLightCardPosition)
+{
+	UWorld* EditorWorld = GEditor->GetEditorWorldContext().World();
+	GEditor->PasteSelectedActorsFromClipboard(EditorWorld, LOCTEXT("PasteLightCardsTransactionMessage", "Paste Light Cards"), EPasteTo::PT_OriginalLocation);
+
+	TArray<ADisplayClusterLightCardActor*> PastedLightCards;
+	for (FSelectionIterator It(GEditor->GetSelectedActorIterator()); It; ++It)
+	{
+		if (ADisplayClusterLightCardActor* LightCard = Cast<ADisplayClusterLightCardActor>(*It))
+		{
+			PastedLightCards.Add(LightCard);
+		}
+	}
+
+	for (ADisplayClusterLightCardActor* LightCard : PastedLightCards)
+	{
+		if (bOffsetLightCardPosition)
+		{
+			// If the light card should be offset from its pasted location, offset its longitude and latitude by a number of
+			// degrees equal to an arc length of 10 units (arc length = angle in radians * radius)
+			const float AngleOffset = FMath::RadiansToDegrees(10.0f / FMath::Max(LightCard->DistanceFromCenter, 1.0f));
+			LightCard->Latitude -= AngleOffset;
+			LightCard->Longitude += AngleOffset;
+		}
+	}
+
+	AddLightCardsToActor(PastedLightCards);
+
+	SelectLightCards(PastedLightCards);
+}
+
+bool SDisplayClusterLightCardEditor::CanPasteLightCards() const
+{
+	UWorld* EditorWorld = GEditor->GetEditorWorldContext().World();
+	return GEditor->CanPasteSelectedActorsFromClipboard(EditorWorld);
+}
+
+void SDisplayClusterLightCardEditor::DuplicateLightCards()
+{
+	CopyLightCards();
+
+	const bool bOffsetLightCardPosition = true;
+	PasteLightCards(bOffsetLightCardPosition);
+}
+
+bool SDisplayClusterLightCardEditor::CanDuplicateLightCards() const
+{
+	return CanCopyLightCards();
+}
+
+void SDisplayClusterLightCardEditor::RemoveLightCards(bool bDeleteLightCardActor)
+{
+	TArray<ADisplayClusterLightCardActor*> SelectedLightCards;
+	LightCardList->GetSelectedLightCards(SelectedLightCards);
+
+	FScopedTransaction Transaction(LOCTEXT("RemoveLightCardTransactionMessage", "Remove Light Card(s)"));
+
+	USelection* EdSelectionManager = GEditor->GetSelectedActors();
+	UWorld* WorldToUse = nullptr;
+
+	if (bDeleteLightCardActor)
+	{
+		EdSelectionManager->BeginBatchSelectOperation();
+		EdSelectionManager->Modify();
+		EdSelectionManager->DeselectAll();
+	}
+
+	if (ActiveRootActor.IsValid())
+	{
+		UDisplayClusterConfigurationData* ConfigData = ActiveRootActor->GetConfigData();
+		ConfigData->Modify();
+
+		FDisplayClusterConfigurationICVFX_VisibilityList& RootActorLightCards = ConfigData->StageSettings.Lightcard.ShowOnlyList;
+
+		for (AActor* LightCard : SelectedLightCards)
+		{
+			RootActorLightCards.Actors.RemoveAll([&](const TSoftObjectPtr<AActor>& Actor)
+				{
+					return Actor.Get() == LightCard;
+				});
+
+			if (bDeleteLightCardActor)
+			{
+				WorldToUse = LightCard->GetWorld();
+				GEditor->SelectActor(LightCard, /*bSelect =*/true, /*bNotifyForActor =*/false, /*bSelectEvenIfHidden =*/true);
+			}
+		}
+	}
+
+	if (bDeleteLightCardActor)
+	{
+		EdSelectionManager->EndBatchSelectOperation();
+
+		if (WorldToUse)
+		{
+			GEditor->edactDeleteSelected(WorldToUse);
+		}
+	}
+
+	RefreshPreviewActors(EDisplayClusterLightCardEditorProxyType::LightCards);
+}
+
+bool SDisplayClusterLightCardEditor::CanRemoveLightCards() const
+{
+	TArray<ADisplayClusterLightCardActor*> SelectedLightCards;
+	LightCardList->GetSelectedLightCards(SelectedLightCards);
+
+	return SelectedLightCards.Num() > 0;
+}
+
 void SDisplayClusterLightCardEditor::OnActiveRootActorChanged(ADisplayClusterRootActor* NewRootActor)
 {
 	RemoveCompileDelegates();
@@ -163,15 +521,59 @@
 
 TSharedRef<SWidget> SDisplayClusterLightCardEditor::CreateLightCardListWidget()
 {
-	return SAssignNew(LightCardList, SDisplayClusterLightCardList, SharedThis(this))
-	.OnLightCardListChanged(this, &SDisplayClusterLightCardEditor::OnLightCardListChanged);
+	return SAssignNew(LightCardList, SDisplayClusterLightCardList, SharedThis(this), CommandList);
 }
 
 TSharedRef<SWidget> SDisplayClusterLightCardEditor::CreateViewportWidget()
 {
-	return SAssignNew(ViewportView, SDisplayClusterLightCardEditorViewport, SharedThis(this));
+	return SAssignNew(ViewportView, SDisplayClusterLightCardEditorViewport, SharedThis(this), CommandList);
 }
 
+void SDisplayClusterLightCardEditor::BindCommands()
+{
+	CommandList = MakeShareable(new FUICommandList);
+
+	CommandList->MapAction(
+		FDisplayClusterLightCardEditorCommands::Get().AddNewLightCard,
+		FExecuteAction::CreateSP(this, &SDisplayClusterLightCardEditor::AddNewLightCard),
+		FCanExecuteAction::CreateSP(this, &SDisplayClusterLightCardEditor::CanAddLightCard));
+
+	CommandList->MapAction(
+		FDisplayClusterLightCardEditorCommands::Get().AddExistingLightCard,
+		FExecuteAction::CreateSP(this, &SDisplayClusterLightCardEditor::AddExistingLightCard),
+		FCanExecuteAction::CreateSP(this, &SDisplayClusterLightCardEditor::CanAddLightCard));
+
+	CommandList->MapAction(
+		FDisplayClusterLightCardEditorCommands::Get().RemoveLightCard,
+		FExecuteAction::CreateSP(this, &SDisplayClusterLightCardEditor::RemoveLightCards, false),
+		FCanExecuteAction::CreateSP(this, &SDisplayClusterLightCardEditor::CanRemoveLightCards));
+
+	CommandList->MapAction(
+		FGenericCommands::Get().Cut,
+		FExecuteAction::CreateSP(this, &SDisplayClusterLightCardEditor::CutLightCards),
+		FCanExecuteAction::CreateSP(this, &SDisplayClusterLightCardEditor::CanCutLightCards));
+
+	CommandList->MapAction(
+		FGenericCommands::Get().Copy,
+		FExecuteAction::CreateSP(this, &SDisplayClusterLightCardEditor::CopyLightCards),
+		FCanExecuteAction::CreateSP(this, &SDisplayClusterLightCardEditor::CanCopyLightCards));
+
+	CommandList->MapAction(
+		FGenericCommands::Get().Paste,
+		FExecuteAction::CreateSP(this, &SDisplayClusterLightCardEditor::PasteLightCards, false),
+		FCanExecuteAction::CreateSP(this, &SDisplayClusterLightCardEditor::CanPasteLightCards));
+
+	CommandList->MapAction(
+		FGenericCommands::Get().Duplicate,
+		FExecuteAction::CreateSP(this, &SDisplayClusterLightCardEditor::DuplicateLightCards),
+		FCanExecuteAction::CreateSP(this, &SDisplayClusterLightCardEditor::CanDuplicateLightCards));
+
+	CommandList->MapAction(
+		FGenericCommands::Get().Delete,
+		FExecuteAction::CreateSP(this, &SDisplayClusterLightCardEditor::RemoveLightCards, true),
+		FCanExecuteAction::CreateSP(this, &SDisplayClusterLightCardEditor::CanRemoveLightCards));
+}
+
 void SDisplayClusterLightCardEditor::RefreshPreviewActors(EDisplayClusterLightCardEditorProxyType ProxyType)
 {
 	RemoveCompileDelegates();
diff --git a/Engine/Plugins/Runtime/nDisplay/Source/DisplayClusterLightCardEditor/Private/SDisplayClusterLightCardEditor.h b/Engine/Plugins/Runtime/nDisplay/Source/DisplayClusterLightCardEditor/Private/SDisplayClusterLightCardEditor.h
--- a/Engine/Plugins/Runtime/nDisplay/Source/DisplayClusterLightCardEditor/Private/SDisplayClusterLightCardEditor.h
+++ b/Engine/Plugins/Runtime/nDisplay/Source/DisplayClusterLightCardEditor/Private/SDisplayClusterLightCardEditor.h
@@ -12,6 +12,7 @@
 class FLayoutExtender;
 class FSpawnTabArgs;
 class FToolBarBuilder;
+class FUICommandList;
 class SDockTab;
 class SDisplayClusterLightCardList;
 class SDisplayClusterLightCardEditorViewport;
@@ -52,16 +53,64 @@
 	TWeakObjectPtr<ADisplayClusterRootActor> GetActiveRootActor() const { return ActiveRootActor; }
 
 	/** Selects the specified light cards in the light card list and details panel */
-	void SelectLightCards(const TArray<AActor*>& LightCardsToSelect);
+	void SelectLightCards(const TArray<ADisplayClusterLightCardActor*>& LightCardsToSelect);
 
+	/** Gets the light cards that are selected in the light card list */
+	void GetSelectedLightCards(TArray<ADisplayClusterLightCardActor*>& OutSelectedLightCards);
+
 	/** Selects the light card proxies that correspond to the specified light cards */
-	void SelectLightCardProxies(const TArray<AActor*>& LightCardsToSelect);
+	void SelectLightCardProxies(const TArray<ADisplayClusterLightCardActor*>& LightCardsToSelect);
 
 	/** Places the given light card in the middle of the viewport */
 	void CenterLightCardInView(ADisplayClusterLightCardActor& LightCard);
 
-	/** Spawns a new light card */
-	ADisplayClusterLightCardActor* AddNewLightCard();
+	/** Spawns a new light card and adds it to the root actor */
+	ADisplayClusterLightCardActor* SpawnLightCard();
+
+	/** Adds a new light card to the root actor and centers it in the viewport */
+	void AddNewLightCard();
+
+	/** Select an existing Light Card from a menu */
+	void AddExistingLightCard();
+
+	/** Adds the given Light Card to the root actor */
+	void AddLightCardsToActor(TArray<ADisplayClusterLightCardActor*> LightCards);
+
+	/** If a Light Card can currently be added */
+	bool CanAddLightCard() const;
+
+	/** Copies any selected light cards to the clipboard, and then deletes them */
+	void CutLightCards();
+
+	/** Determines if there are selected light cards that can be cut */
+	bool CanCutLightCards();
+
+	/** Copies any selected light cards to the clipboard */
+	void CopyLightCards();
+
+	/** Determines if there are selected light cards that can be copied */
+	bool CanCopyLightCards() const;
+
+	/** Pastes any light cards in the clipboard to the current root actor */
+	void PasteLightCards(bool bOffsetLightCardPosition);
+
+	/** Determines if there are any light cards that can be pasted from the clipboard */
+	bool CanPasteLightCards() const;
+
+	/** Copies any selected light cards to the clipboard and then pastes them */
+	void DuplicateLightCards();
+
+	/** Determines if there are selected light cards that can be duplicated */
+	bool CanDuplicateLightCards() const;
+
+	/**
+	 * Remove the light card from the actor
+	 *@param bDeleteLightCardActor Delete the actor from the level
+	 */
+	void RemoveLightCards(bool bDeleteLightCardActor);
+
+	/** If the selected Light Card can be removed */
+	bool CanRemoveLightCards() const;
 
 private:
 	/** Raised when the active Display cluster root actor has been changed in the operator panel */
@@ -73,6 +122,9 @@
 	/** Create the 3d viewport widget */
 	TSharedRef<SWidget> CreateViewportWidget();
 
+	/** Creates the editor's command list and binds commands to it */
+	void BindCommands();
+
 	/** Refresh all preview actors */
 	void RefreshPreviewActors(EDisplayClusterLightCardEditorProxyType ProxyType = EDisplayClusterLightCardEditorProxyType::All);
 
@@ -112,6 +164,12 @@
 	/** The 3d viewport */
 	TSharedPtr<SDisplayClusterLightCardEditorViewport> ViewportView;
 
+	/** The command list for editor commands */
+	TSharedPtr<FUICommandList> CommandList;
+
+	/** Stores the mouse position when the context menu was opened */
+	TOptional<FIntPoint> CachedContextMenuMousePos;
+
 	/** A reference to the root actor that is currently being operated on */
 	TWeakObjectPtr<ADisplayClusterRootActor> ActiveRootActor;
 
diff --git a/Engine/Plugins/Runtime/nDisplay/Source/DisplayClusterLightCardEditor/Private/SDisplayClusterLightCardList.cpp b/Engine/Plugins/Runtime/nDisplay/Source/DisplayClusterLightCardEditor/Private/SDisplayClusterLightCardList.cpp
--- a/Engine/Plugins/Runtime/nDisplay/Source/DisplayClusterLightCardEditor/Private/SDisplayClusterLightCardList.cpp
+++ b/Engine/Plugins/Runtime/nDisplay/Source/DisplayClusterLightCardEditor/Private/SDisplayClusterLightCardList.cpp
@@ -136,12 +136,10 @@ private:
 };
 
 
-void SDisplayClusterLightCardList::Construct(const FArguments& InArgs, TSharedPtr<SDisplayClusterLightCardEditor> InLightCardEditor)
+void SDisplayClusterLightCardList::Construct(const FArguments& InArgs, TSharedPtr<SDisplayClusterLightCardEditor> InLightCardEditor, TSharedPtr<class FUICommandList> InCommandList)
 {
 	LightCardEditorPtr = InLightCardEditor;
-	OnLightCardChanged = InArgs._OnLightCardListChanged;
-	
-	BindCommands();
+	CommandList = InCommandList;
 
 	ChildSlot
 	[
@@ -214,7 +212,7 @@ void SDisplayClusterLightCardList::SetRootActor(ADisplayClusterRootActor* NewRoo
 	// Select previously selected light cards. This fixes an issue where the details panel may clear in some cases
 	// and also supports maintaining the selection of shared light cards across different root actors.
 	{
-		TArray<AActor*> LightCardActorsToSelect;
+		TArray<ADisplayClusterLightCardActor*> LightCardActorsToSelect;
 		for (const TSharedPtr<FLightCardTreeItem>& LightCard : PreviouslySelectedLightCards)
 		{
 			if (LightCard.IsValid() && LightCard->LightCardActor.IsValid())
@@ -227,7 +225,21 @@ void SDisplayClusterLightCardList::SetRootActor(ADisplayClusterRootActor* NewRoo
 	}
 }
 
-void SDisplayClusterLightCardList::SelectLightCards(const TArray<AActor*>& LightCardsToSelect)
+void SDisplayClusterLightCardList::GetSelectedLightCards(TArray<ADisplayClusterLightCardActor*>& OutSelectedLightCards)
+{
+	TArray<TSharedPtr<FLightCardTreeItem>> SelectedTreeItems;
+	LightCardTreeView->GetSelectedItems(SelectedTreeItems);
+
+	for (const TSharedPtr<FLightCardTreeItem>& SelectedItem : SelectedTreeItems)
+	{
+		if (SelectedItem.IsValid() && SelectedItem->LightCardActor.IsValid())
+		{
+			OutSelectedLightCards.Add(SelectedItem->LightCardActor.Get());
+		}
+	}
+}
+
+void SDisplayClusterLightCardList::SelectLightCards(const TArray<ADisplayClusterLightCardActor*>& LightCardsToSelect)
 {
 	LightCardTreeView->ClearSelection();
 
@@ -248,37 +260,6 @@ void SDisplayClusterLightCardList::Refresh()
 	SetRootActor(RootActor.Get());
 }
 
-void SDisplayClusterLightCardList::BindCommands()
-{
-	const FDisplayClusterLightCardEditorCommands& Commands = FDisplayClusterLightCardEditorCommands::Get();
-
-	CommandList = MakeShareable(new FUICommandList);
-	
-	CommandList->MapAction(
-		Commands.AddNewLightCard,
-		FExecuteAction::CreateSP(this, &SDisplayClusterLightCardList::AddNewLightCardHandler),
-		FCanExecuteAction::CreateSP(this, &SDisplayClusterLightCardList::CanAddLightCard),
-		FCanExecuteAction());
-
-	CommandList->MapAction(
-		Commands.AddExistingLightCard,
-		FExecuteAction::CreateSP(this, &SDisplayClusterLightCardList::AddExistingLightCard),
-		FCanExecuteAction::CreateSP(this, &SDisplayClusterLightCardList::CanAddLightCard),
-		FCanExecuteAction());
-	
-	CommandList->MapAction(
-		Commands.RemoveLightCard,
-		FExecuteAction::CreateSP(this, &SDisplayClusterLightCardList::RemoveLightCard, false),
-		FCanExecuteAction::CreateSP(this, &SDisplayClusterLightCardList::CanRemoveLightCard),
-		FCanExecuteAction());
-
-	CommandList->MapAction(
-		FGenericCommands::Get().Delete,
-		FExecuteAction::CreateSP(this, &SDisplayClusterLightCardList::RemoveLightCard, true),
-		FCanExecuteAction::CreateSP(this, &SDisplayClusterLightCardList::CanRemoveLightCard),
-		FCanExecuteAction());
-}
-
 bool SDisplayClusterLightCardList::FillLightCardList()
 {
 	TArray<TSharedPtr<FLightCardTreeItem>> OriginalTree = MoveTemp(LightCardTree);
@@ -297,10 +278,10 @@ bool SDisplayClusterLightCardList::FillLightCardList()
 
 		for (const TSoftObjectPtr<AActor>& LightCardActor : RootActorLightCards.Actors)
 		{
-			if (LightCardActor.IsValid())
+			if (LightCardActor.IsValid() && LightCardActor->IsA<ADisplayClusterLightCardActor>())
 			{
 				TSharedPtr<FLightCardTreeItem> LightCardTreeItem = MakeShared<FLightCardTreeItem>();
-				LightCardTreeItem->LightCardActor = LightCardActor.Get();
+				LightCardTreeItem->LightCardActor = CastChecked<ADisplayClusterLightCardActor>(LightCardActor.Get());
 				LightCardTree.Add(LightCardTreeItem);
 				LightCardActors.Add(LightCardTreeItem);
 			}
@@ -313,7 +294,7 @@ bool SDisplayClusterLightCardList::FillLightCardList()
 		{
 			if (UWorld* World = RootActor->GetWorld())
 			{
-				for (const TWeakObjectPtr<AActor> WeakActor : FActorRange(World))
+				for (const TWeakObjectPtr<ADisplayClusterLightCardActor> WeakActor : TActorRange<ADisplayClusterLightCardActor>(World))
 				{
 					if (WeakActor.IsValid())
 					{
@@ -358,253 +339,6 @@ bool SDisplayClusterLightCardList::FillLightCardList()
 	return false;
 }
 
-void SDisplayClusterLightCardList::AddNewLightCardHandler()
-{
-	check(RootActor.IsValid());
-
-	FScopedTransaction Transaction(LOCTEXT("AddNewLightCard", "Add New Light Card"));
-
-	ADisplayClusterLightCardActor* NewLightCard = AddNewLightCard();
-
-	// When adding a new lightcard, usually the desired location is in the middle of the viewport
-	if (LightCardEditorPtr.IsValid())
-	{
-		LightCardEditorPtr.Pin()->CenterLightCardInView(*NewLightCard);
-	}
-}
-
-
-ADisplayClusterLightCardActor* SDisplayClusterLightCardList::AddNewLightCard()
-{
-	if (!RootActor.IsValid())
-	{
-		return nullptr;
-	}
-
-	const FVector SpawnLocation = RootActor->GetDefaultCamera()->GetComponentLocation();
-	FRotator SpawnRotation = RootActor->GetDefaultCamera()->GetComponentRotation();
-	SpawnRotation.Yaw -= 180.f;
-	
-	FActorSpawnParameters SpawnParameters;
-	SpawnParameters.bNoFail = true;
-	SpawnParameters.SpawnCollisionHandlingOverride = ESpawnActorCollisionHandlingMethod::AdjustIfPossibleButAlwaysSpawn;
-	SpawnParameters.Name = TEXT("LightCard");
-	SpawnParameters.NameMode = FActorSpawnParameters::ESpawnActorNameMode::Requested;
-	SpawnParameters.OverrideLevel = RootActor->GetWorld()->GetCurrentLevel();
-	
-	ADisplayClusterLightCardActor* NewLightCard = CastChecked<ADisplayClusterLightCardActor>(
-		RootActor->GetWorld()->SpawnActor(ADisplayClusterLightCardActor::StaticClass(),
-		&SpawnLocation, &SpawnRotation, MoveTemp(SpawnParameters)));
-
-	NewLightCard->SetActorLabel(NewLightCard->GetName());
-
-	AddLightCardToActor(NewLightCard);
-
-	return NewLightCard;
-}
-
-void SDisplayClusterLightCardList::AddExistingLightCard()
-{
-	TSharedPtr<SWindow> PickerWindow;
-	TWeakObjectPtr<AActor> SelectedActorPtr;
-	bool bFinished = false;
-	
-	const TSharedRef<SWidget> ActorPicker = PropertyCustomizationHelpers::MakeActorPickerWithMenu(
-		nullptr,
-		false,
-		FOnShouldFilterActor::CreateLambda([&](const AActor* const InActor) -> bool // ActorFilter
-		{
-			const bool IsAllowed = InActor != nullptr && !InActor->IsChildActor() && InActor->IsA<AActor>() &&
-				!InActor->GetClass()->HasAnyClassFlags(CLASS_Interface)	&& !InActor->IsA<ADisplayClusterRootActor>();
-			
-			return IsAllowed;
-		}),
-		FOnActorSelected::CreateLambda([&](AActor* InActor) -> void // OnSet
-		{
-			SelectedActorPtr = InActor;
-		}),
-		FSimpleDelegate::CreateLambda([&]() -> void // OnClose
-		{
-		}),
-		FSimpleDelegate::CreateLambda([&]() -> void // OnUseSelected
-		{
-			if (AActor* Selection = Cast<AActor>(GEditor->GetSelectedActors()->GetTop(AActor::StaticClass())))
-			{
-				SelectedActorPtr = Selection;
-			}
-		}));
-	
-	PickerWindow = SNew(SWindow)
-	.Title(LOCTEXT("AddExistingLightCard", "Select an existing Light Card actor"))
-	.ClientSize(FVector2D(500, 525))
-	.SupportsMinimize(false) .SupportsMaximize(false)
-	[
-		SNew(SBorder)
-		.BorderImage(FAppStyle::GetBrush("Menu.Background"))
-		[
-			SNew(SWizard)
-			.FinishButtonText(LOCTEXT("FinishAddingExistingLightCard", "Add Actor"))
-			.OnCanceled(FSimpleDelegate::CreateLambda([&]()
-			{
-				if (PickerWindow.IsValid())
-				{
-					PickerWindow->RequestDestroyWindow();
-				}
-			}))
-			.OnFinished(FSimpleDelegate::CreateLambda([&]()
-			{
-				bFinished = true;
-				if (PickerWindow.IsValid())
-				{
-					PickerWindow->RequestDestroyWindow();
-				}
-			}))
-			.CanFinish(TAttribute<bool>::CreateLambda([&]()
-			{
-				return SelectedActorPtr.IsValid();
-			}))
-			.ShowPageList(false)
-			+SWizard::Page()
-			.CanShow(true)
-			[
-				SNew(SBorder)
-				.VAlign(VAlign_Fill)
-				.BorderImage(FAppStyle::GetBrush("Menu.Background"))
-				[
-					SNew(SVerticalBox)
-					+SVerticalBox::Slot()
-					.AutoHeight()
-					.VAlign(VAlign_Fill)
-					.HAlign(HAlign_Fill)
-					[
-						ActorPicker
-					]
-					+SVerticalBox::Slot()
-					.AutoHeight()
-					.VAlign(VAlign_Bottom)
-					.Padding(0.f, 8.f)
-					[
-						SNew(STextBlock)
-						.TextStyle(FAppStyle::Get(), "NormalText.Important")
-						.Text_Lambda([&]
-						{
-							const FString Result = FString::Printf(TEXT("Selected Actor: %s"),
-								SelectedActorPtr.IsValid() ? *SelectedActorPtr->GetActorLabel() : TEXT(""));
-							return FText::FromString(Result);
-						})
-					]
-				]
-			]
-		]
-	];
-
-	GEditor->EditorAddModalWindow(PickerWindow.ToSharedRef());
-	if (bFinished && SelectedActorPtr.IsValid())
-	{
-		FScopedTransaction Transaction(LOCTEXT("AddExistingLightCardTransactionMessage", "Add Existing Light Card"));
-		AddLightCardToActor(SelectedActorPtr.Get());
-	}
-
-	PickerWindow.Reset();
-	SelectedActorPtr.Reset();
-}
-
-void SDisplayClusterLightCardList::AddLightCardToActor(AActor* LightCard)
-{
-	check(LightCard);
-	
-	UDisplayClusterConfigurationData* ConfigData = RootActor->GetConfigData();
-	ConfigData->Modify();
-	FDisplayClusterConfigurationICVFX_VisibilityList& RootActorLightCards = ConfigData->StageSettings.Lightcard.ShowOnlyList;
-
-	if (!RootActorLightCards.Actors.ContainsByPredicate([&](const TSoftObjectPtr<AActor>& Actor)
-	{
-		// Don't add if a loaded actor is already present.
-		return Actor.Get() == LightCard;
-	}))
-	{
-		const TSoftObjectPtr<AActor> LightCardSoftObject(LightCard);
-		
-		// Remove any exact paths to this actor. It's possible invalid actors are present if a light card
-		// was force deleted from a level.
-		RootActorLightCards.Actors.RemoveAll([&] (const TSoftObjectPtr<AActor>& Actor)
-		{
-			return Actor == LightCardSoftObject;
-		});
-	
-		RootActorLightCards.Actors.Add(LightCard);
-	}
-	Refresh();
-	OnLightCardChanged.ExecuteIfBound();
-}
-
-bool SDisplayClusterLightCardList::CanAddLightCard() const
-{
-	return RootActor.IsValid() && RootActor->GetWorld() != nullptr;
-}
-
-void SDisplayClusterLightCardList::RemoveLightCard(bool bDeleteLightCardActor)
-{
-	TArray<TSharedPtr<FLightCardTreeItem>> SelectedTreeItems;
-	LightCardTreeView->GetSelectedItems(SelectedTreeItems);
-
-	FScopedTransaction Transaction(LOCTEXT("RemoveLightCard", "Remove Light Card(s)"));
-
-	USelection* EdSelectionManager = GEditor->GetSelectedActors();
-	UWorld* WorldToUse = nullptr;
-	
-	if (bDeleteLightCardActor)
-	{
-		EdSelectionManager->BeginBatchSelectOperation();
-		EdSelectionManager->Modify();
-		EdSelectionManager->DeselectAll();
-	}
-	
-	for (const TSharedPtr<FLightCardTreeItem>& Item : SelectedTreeItems)
-	{
-		if (Item.IsValid() && Item->LightCardActor.IsValid())
-		{
-			if (RootActor.IsValid())
-			{
-				UDisplayClusterConfigurationData* ConfigData = RootActor->GetConfigData();
-				ConfigData->Modify();
-				
-				FDisplayClusterConfigurationICVFX_VisibilityList& RootActorLightCards = ConfigData->StageSettings.Lightcard.ShowOnlyList;
-				RootActorLightCards.Actors.RemoveAll([&](const TSoftObjectPtr<AActor>& Actor)
-				{
-					return Actor.Get() == Item->LightCardActor.Get();
-				});
-			}
-
-			if (bDeleteLightCardActor)
-			{
-				WorldToUse = Item->LightCardActor->GetWorld();
-				GEditor->SelectActor(Item->LightCardActor.Get(), /*bSelect =*/true, /*bNotifyForActor =*/false, /*bSelectEvenIfHidden =*/true);
-			}
-		}
-	}
-
-	if (bDeleteLightCardActor)
-	{
-		EdSelectionManager->EndBatchSelectOperation();
-
-		if (WorldToUse)
-		{
-			GEditor->edactDeleteSelected(WorldToUse);
-		}
-	}
-
-	Refresh();
-	OnLightCardChanged.ExecuteIfBound();
-}
-
-bool SDisplayClusterLightCardList::CanRemoveLightCard() const
-{
-	TArray<TSharedPtr<FLightCardTreeItem>> SelectedTreeItems;
-	LightCardTreeView->GetSelectedItems(SelectedTreeItems);
-	return SelectedTreeItems.Num() > 0;
-}
-
 TSharedRef<ITableRow> SDisplayClusterLightCardList::GenerateTreeItemRow(TSharedPtr<FLightCardTreeItem> Item, const TSharedRef<STableViewBase>& OwnerTable)
 {
 	return SNew(SLightCardTreeItemRow, OwnerTable, Item);
@@ -632,7 +366,7 @@ void SDisplayClusterLightCardList::OnTreeItemSelected(TSharedPtr<FLightCardTreeI
 	TArray<TSharedPtr<FLightCardTreeItem>> SelectedTreeItems;
 	LightCardTreeView->GetSelectedItems(SelectedTreeItems);
 
-	TArray<AActor*> SelectedLightCards;
+	TArray<ADisplayClusterLightCardActor*> SelectedLightCards;
 	for (const TSharedPtr<FLightCardTreeItem>& SelectedTreeItem : SelectedTreeItems)
 	{
 		if (SelectedTreeItem->LightCardActor.IsValid())
@@ -672,12 +406,19 @@ TSharedRef<SWidget> SDisplayClusterLightCardList::CreateAddNewMenuContent()
 
 TSharedPtr<SWidget> SDisplayClusterLightCardList::CreateContextMenu()
 {
-	const FDisplayClusterLightCardEditorCommands& Commands = FDisplayClusterLightCardEditorCommands::Get();
-	
 	const bool bCloseAfterSelection = true;
 	FMenuBuilder MenuBuilder(bCloseAfterSelection, CommandList, Extenders);
-	MenuBuilder.AddMenuEntry(Commands.RemoveLightCard);
-	MenuBuilder.AddMenuEntry(FGenericCommands::Get().Delete);
+	MenuBuilder.AddMenuEntry(FDisplayClusterLightCardEditorCommands::Get().RemoveLightCard);
+
+	MenuBuilder.BeginSection("Edit", LOCTEXT("EditSection", "Edit"));
+	{
+		MenuBuilder.AddMenuEntry(FGenericCommands::Get().Cut);
+		MenuBuilder.AddMenuEntry(FGenericCommands::Get().Copy);
+		MenuBuilder.AddMenuEntry(FGenericCommands::Get().Paste);
+		MenuBuilder.AddMenuEntry(FGenericCommands::Get().Duplicate);
+		MenuBuilder.AddMenuEntry(FGenericCommands::Get().Delete);
+	}
+	MenuBuilder.EndSection();
 
 	return MenuBuilder.MakeWidget();
 }
diff --git a/Engine/Plugins/Runtime/nDisplay/Source/DisplayClusterLightCardEditor/Private/SDisplayClusterLightCardList.h b/Engine/Plugins/Runtime/nDisplay/Source/DisplayClusterLightCardEditor/Private/SDisplayClusterLightCardList.h
--- a/Engine/Plugins/Runtime/nDisplay/Source/DisplayClusterLightCardEditor/Private/SDisplayClusterLightCardList.h
+++ b/Engine/Plugins/Runtime/nDisplay/Source/DisplayClusterLightCardEditor/Private/SDisplayClusterLightCardList.h
@@ -25,7 +25,7 @@
 	
 	struct FLightCardTreeItem
 	{
-		TWeakObjectPtr<AActor> LightCardActor;
+		TWeakObjectPtr<ADisplayClusterLightCardActor> LightCardActor;
 		FName ActorLayer;
 
 		FLightCardTreeItem() :
@@ -40,17 +40,19 @@
 public:
 	SLATE_BEGIN_ARGS(SDisplayClusterLightCardList)
 	{}
-	SLATE_EVENT(FOnLightCardListChanged, OnLightCardListChanged)
 	SLATE_END_ARGS()
 
-	void Construct(const FArguments& InArgs, TSharedPtr<SDisplayClusterLightCardEditor> InLightCardEditor);
+	void Construct(const FArguments& InArgs, TSharedPtr<SDisplayClusterLightCardEditor> InLightCardEditor, TSharedPtr<class FUICommandList> InCommandList);
 
 	void SetRootActor(ADisplayClusterRootActor* NewRootActor);
 
 	const TArray<TSharedPtr<FLightCardTreeItem>>& GetLightCardActors() const { return LightCardActors; }
 
-	void SelectLightCards(const TArray<AActor*>& LightCardsToSelect);
+	/** Gets a list of light card actors that are currently selected in the list */
+	void GetSelectedLightCards(TArray<ADisplayClusterLightCardActor*>& OutSelectedLightCards);
 
+	void SelectLightCards(const TArray<ADisplayClusterLightCardActor*>& LightCardsToSelect);
+
 	/** Refreshes the list */
 	void Refresh();
 
@@ -58,8 +60,6 @@
 	ADisplayClusterLightCardActor* AddNewLightCard();
 
 private:
-	void BindCommands();
-	
 	/**
 	 * Fill the LightCard list with available LightCards
 	 * @return True if the list has been modified
@@ -69,24 +69,6 @@
 	/** Handles the Add Light Card button. Spawn a new Light Card placed in the middle of the viewport. */
 	void AddNewLightCardHandler();
 
-	/** Select an existing Light Card from a menu */
-	void AddExistingLightCard();
-	
-	/** Adds the given Light Card to the root actor */
-	void AddLightCardToActor(AActor* LightCard);
-	
-	/** If a Light Card can currently be added */
-	bool CanAddLightCard() const;
-
-	/**
-	 * Remove the light card from the actor
-	 *@param bDeleteLightCardActor Delete the actor from the level
-	 */
-	void RemoveLightCard(bool bDeleteLightCardActor);
-	
-	/** If the selected Light Card can be removed */
-	bool CanRemoveLightCard() const;
-
 	TSharedRef<ITableRow> GenerateTreeItemRow(TSharedPtr<FLightCardTreeItem> Item, const TSharedRef<STableViewBase>& OwnerTable);
 	void GetChildrenForTreeItem(TSharedPtr<FLightCardTreeItem> InItem, TArray<TSharedPtr<FLightCardTreeItem>>& OutChildren);
 	void OnTreeItemSelected(TSharedPtr<FLightCardTreeItem> InItem, ESelectInfo::Type SelectInfo);
@@ -117,7 +99,4 @@
 
 	/** Mapped commands for this list */
 	TSharedPtr<FUICommandList> CommandList;
-
-	/** When a user adds or removes a light card */
-	FOnLightCardListChanged OnLightCardChanged;
 };
diff --git a/Engine/Plugins/Runtime/nDisplay/Source/DisplayClusterLightCardEditor/Private/Viewport/DisplayClusterLightCardEditorViewportClient.cpp b/Engine/Plugins/Runtime/nDisplay/Source/DisplayClusterLightCardEditor/Private/Viewport/DisplayClusterLightCardEditorViewportClient.cpp
--- a/Engine/Plugins/Runtime/nDisplay/Source/DisplayClusterLightCardEditor/Private/Viewport/DisplayClusterLightCardEditorViewportClient.cpp
+++ b/Engine/Plugins/Runtime/nDisplay/Source/DisplayClusterLightCardEditor/Private/Viewport/DisplayClusterLightCardEditorViewportClient.cpp
@@ -2,10 +2,12 @@
 
 #include "DisplayClusterLightCardEditorViewportClient.h"
 
+#include "DisplayClusterLightcardEditorViewport.h"
+#include "DisplayClusterLightCardEditorWidget.h"
+
 #include "Components/DisplayClusterPreviewComponent.h"
 #include "Components/DisplayClusterCameraComponent.h"
 #include "Components/DisplayClusterScreenComponent.h"
-#include "DisplayClusterLightCardEditorWidget.h"
 #include "DisplayClusterConfigurationTypes.h"
 #include "DisplayClusterProjectionStrings.h"
 #include "DisplayClusterRootActor.h"
@@ -311,13 +313,14 @@ UTexture2D* FDisplayClusterLightCardEditorViewportClient::FNormalMap::GenerateNo
 // FDisplayClusterLightCardEditorViewportClient
 
 FDisplayClusterLightCardEditorViewportClient::FDisplayClusterLightCardEditorViewportClient(FAdvancedPreviewScene& InPreviewScene,
-                                                           const TWeakPtr<SEditorViewport>& InEditorViewportWidget,
-                                                           TWeakPtr<SDisplayClusterLightCardEditor> InLightCardEditor) :
-	FEditorViewportClient(nullptr, &InPreviewScene, InEditorViewportWidget)
+	const TWeakPtr<SDisplayClusterLightCardEditorViewport>& InEditorViewportWidget)
+	: FEditorViewportClient(nullptr, &InPreviewScene, InEditorViewportWidget)
+	, LightCardEditorViewportPtr(InEditorViewportWidget)
 {
-	check (InLightCardEditor.IsValid());
+	check(InEditorViewportWidget.IsValid());
 	
-	LightCardEditorPtr = InLightCardEditor;
+	LightCardEditorPtr = InEditorViewportWidget.Pin()->GetLightCardEditor();
+	check(LightCardEditorPtr.IsValid())
 	
 	MeshProjectionRenderer = MakeShared<FDisplayClusterMeshProjectionRenderer>();
 	MeshProjectionRenderer->ActorSelectedDelegate = FDisplayClusterMeshProjectionRenderer::FSelection::CreateRaw(this, &FDisplayClusterLightCardEditorViewportClient::IsLightCardSelected);
@@ -1066,7 +1069,7 @@ void FDisplayClusterLightCardEditorViewportClient::CreateDrawnLightCard(const TA
 
 	FScopedTransaction Transaction(LOCTEXT("AddNewLightCard", "Add New Light Card"));
 
-	ADisplayClusterLightCardActor* LightCard = LightCardEditorPtr.Pin()->AddNewLightCard();
+	ADisplayClusterLightCardActor* LightCard = LightCardEditorPtr.Pin()->SpawnLightCard();
 
 	if (!LightCard)
 	{
@@ -1129,8 +1132,8 @@ void FDisplayClusterLightCardEditorViewportClient::ProcessClick(FSceneView& View
 	check(PreviewWorld);
 
 	const bool bIsCtrlKeyDown = Viewport->KeyState(EKeys::LeftControl) || Viewport->KeyState(EKeys::RightControl);
-
 	const bool bMultiSelect = Key == EKeys::LeftMouseButton && bIsCtrlKeyDown;
+	const bool bIsRightClickSelection = Key == EKeys::RightMouseButton && !bIsCtrlKeyDown && !Viewport->KeyState(EKeys::LeftMouseButton);
 
 	if (HitProxy)
 	{
@@ -1162,6 +1165,11 @@ void FDisplayClusterLightCardEditorViewportClient::ProcessClick(FSceneView& View
 	
 	PropagateLightCardSelection();
 
+	if (bIsRightClickSelection)
+	{
+		LightCardEditorViewportPtr.Pin()->SummonContextMenu();
+	}
+
 	FEditorViewportClient::ProcessClick(View, HitProxy, Key, Event, HitX, HitY);
 }
 
@@ -1242,6 +1250,9 @@ void FDisplayClusterLightCardEditorViewportClient::UpdatePreviewActor(ADisplayCl
 		Viewport->InvalidateHitProxy();
 		bShouldCheckHitProxy = true;
 		InvalidateNormalMap();
+
+		OnNextSceneRefreshDelegate.Broadcast();
+		OnNextSceneRefreshDelegate.Clear();
 	};
 	
 	if (RootActor == nullptr)
@@ -1257,22 +1268,7 @@ void FDisplayClusterLightCardEditorViewportClient::UpdatePreviewActor(ADisplayCl
 		// Schedule for the next tick so CDO changes get propagated first in the event of config editor skeleton
 		// regeneration & compiles. nDisplay's custom propagation may have issues if the archetype isn't correct.
 		PreviewWorld->GetTimerManager().SetTimerForNextTick([=]()
-		{
-			TSet<AActor*> LastSelectedLightCardLevelInstances;
-			for (TWeakObjectPtr<ADisplayClusterLightCardActor>& SelectedLightCard : SelectedLightCards)
-			{
-				if (SelectedLightCard.IsValid())
-				{
-					if (FLightCardProxy* FoundProxy = LightCardProxies.FindByKey(SelectedLightCard.Get()))
-					{
-						if (FoundProxy->LevelInstance.IsValid())
-						{
-							LastSelectedLightCardLevelInstances.Add(FoundProxy->LevelInstance.Get());
-						}
-					}
-				}
-			}
-			
+		{			
 			DestroyProxies(ProxyType);
 			RootActor->SubscribeToPostProcessRenderTarget(reinterpret_cast<uint8*>(this));
 			RootActorLevelInstance = RootActor;
@@ -1355,12 +1351,16 @@ void FDisplayClusterLightCardEditorViewportClient::UpdatePreviewActor(ADisplayCl
 					}
 
 					LightCardProxies.Add(FLightCardProxy(LightCard.Get(), LightCardProxy));
-					
-					if (LastSelectedLightCardLevelInstances.Contains(LightCard.Get()))
-					{
-						SelectLightCard(LightCardProxy, true);
-					}
 				}
+
+				// Update the selected light card proxies to match the currently selected light cards in the light card list
+				TArray<ADisplayClusterLightCardActor*> CurrentlySelectedLightCards;
+				if (LightCardEditorPtr.IsValid())
+				{
+					LightCardEditorPtr.Pin()->GetSelectedLightCards(CurrentlySelectedLightCards);
+				}
+
+				SelectLightCards(CurrentlySelectedLightCards);
 			}
 
 			for (const FLightCardProxy& LightCardProxy : LightCardProxies)
@@ -1434,7 +1434,7 @@ EDisplayClusterLightCardEditorProxyType ProxyType)
 	}
 }
 
-void FDisplayClusterLightCardEditorViewportClient::SelectLightCards(const TArray<AActor*>& LightCardsToSelect)
+void FDisplayClusterLightCardEditorViewportClient::SelectLightCards(const TArray<ADisplayClusterLightCardActor*>& LightCardsToSelect)
 {
 	SelectLightCard(nullptr);
 	for (AActor* LightCard : LightCardsToSelect)
@@ -1563,6 +1563,60 @@ void FDisplayClusterLightCardEditorViewportClient::CenterLightCardInView(ADispla
 	}
 }
 
+void FDisplayClusterLightCardEditorViewportClient::MoveSelectedLightCardsToPixel(const FIntPoint& PixelPos)
+{
+	// First, find the average position of all selected light cards
+	FSphericalCoordinates AverageCoords = FSphericalCoordinates();
+	int32 NumLightCards = 0;
+
+	for (const TWeakObjectPtr<ADisplayClusterLightCardActor>& LightCard : SelectedLightCards)
+	{
+		if (LightCard.IsValid())
+		{
+			const FSphericalCoordinates LightCardCoords = GetLightCardCoordinates(LightCard.Get());
+
+			AverageCoords = AverageCoords + LightCardCoords;
+			++NumLightCards;
+		}
+	}
+
+	AverageCoords.Radius /= NumLightCards;
+	AverageCoords.Azimuth /= NumLightCards;
+	AverageCoords.Inclination /= NumLightCards;
+	AverageCoords.Conform();
+
+	// Now, find the desired coordinates based on the mouse position
+	FSceneViewFamilyContext ViewFamily(FSceneViewFamily::ConstructionValues(
+		Viewport,
+		GetScene(),
+		EngineShowFlags)
+		.SetRealtimeUpdate(IsRealtime())
+	);
+
+	FSceneView* View = CalcSceneView(&ViewFamily);
+
+	FVector Origin;
+	FVector Direction;
+	PixelToWorld(*View, PixelPos, Origin, Direction);
+ 
+	// Compute desired coordinates (radius doesn't matter here since we will use the flush constraint on the light cards after moving them)
+	const FSphericalCoordinates DesiredCoords(Direction * 100.0f);
+	const FSphericalCoordinates DeltaCoords = DesiredCoords - AverageCoords;
+
+	// Update each light card with the delta coordinates; the flush constraint is applied by MoveLightCardTo, ensuring the light card is always flush to screens
+	for (const TWeakObjectPtr<ADisplayClusterLightCardActor>& LightCard : SelectedLightCards)
+	{
+		if (LightCard.IsValid())
+		{
+			const FSphericalCoordinates LightCardCoords = GetLightCardCoordinates(LightCard.Get());
+			const FSphericalCoordinates NewCoords = LightCardCoords + DeltaCoords;
+
+			MoveLightCardTo(*LightCard.Get(), NewCoords);
+			PropagateLightCardTransform(LightCard.Get());
+		}
+	}
+}
+
 void FDisplayClusterLightCardEditorViewportClient::BeginTransaction(const FText& Description)
 {
 	GEditor->BeginTransaction(Description);
@@ -1886,7 +1940,7 @@ void FDisplayClusterLightCardEditorViewportClient::SelectLightCard(ADisplayClust
 
 void FDisplayClusterLightCardEditorViewportClient::PropagateLightCardSelection()
 {
-	TArray<AActor*> SelectedLevelInstances;
+	TArray<ADisplayClusterLightCardActor*> SelectedLevelInstances;
 	for (const TWeakObjectPtr<ADisplayClusterLightCardActor>& SelectedLightCard : SelectedLightCards)
 	{
 		if (FLightCardProxy* FoundProxy = LightCardProxies.FindByKey(SelectedLightCard.Get()))
diff --git a/Engine/Plugins/Runtime/nDisplay/Source/DisplayClusterLightCardEditor/Private/Viewport/DisplayClusterLightCardEditorViewportClient.h b/Engine/Plugins/Runtime/nDisplay/Source/DisplayClusterLightCardEditor/Private/Viewport/DisplayClusterLightCardEditorViewportClient.h
--- a/Engine/Plugins/Runtime/nDisplay/Source/DisplayClusterLightCardEditor/Private/Viewport/DisplayClusterLightCardEditorViewportClient.h
+++ b/Engine/Plugins/Runtime/nDisplay/Source/DisplayClusterLightCardEditor/Private/Viewport/DisplayClusterLightCardEditorViewportClient.h
@@ -14,6 +14,7 @@
 
 class ADisplayClusterRootActor;
 class SDisplayClusterLightCardEditor;
+class SDisplayClusterLightCardEditorViewport;
 class FScopedTransaction;
 class UDisplayClusterConfigurationViewport;
 
@@ -37,6 +38,8 @@ public:
 	};
 
 private:
+	DECLARE_MULTICAST_DELEGATE(FOnNextSceneRefresh);
+
 	struct FSphericalCoordinates
 	{
 	public:
@@ -121,8 +124,8 @@ private:
 	static constexpr float LightCardFlushOffset = -0.5f;
 
 public:
-	FDisplayClusterLightCardEditorViewportClient(FAdvancedPreviewScene& InPreviewScene, const TWeakPtr<SEditorViewport>& InEditorViewportWidget,
-		TWeakPtr<SDisplayClusterLightCardEditor> InLightCardEditor);
+	FDisplayClusterLightCardEditorViewportClient(FAdvancedPreviewScene& InPreviewScene,
+		const TWeakPtr<SDisplayClusterLightCardEditorViewport>& InEditorViewportWidget);
 	virtual ~FDisplayClusterLightCardEditorViewportClient() override;
 	
 	// FEditorViewportClient
@@ -144,6 +147,9 @@ public:
 	virtual ELevelViewportType GetViewportType() const override { return LVT_Perspective; }
 	// ~FEditorViewportClient
 
+	/** Returns a delegate that is invoked on the next scene refresh. The delegate is cleared afterwards */
+	FOnNextSceneRefresh& GetOnNextSceneRefresh() { return OnNextSceneRefreshDelegate; }
+
 	/**
 	 * Update the spawned preview actor from a root actor in the level
 	 * 
@@ -161,7 +167,7 @@ public:
 	void DestroyProxies(EDisplayClusterLightCardEditorProxyType ProxyType);
 	
 	/** Selects the light card proxies that correspond to the specified light cards */
-	void SelectLightCards(const TArray<AActor*>& LightCardsToSelect);
+	void SelectLightCards(const TArray<ADisplayClusterLightCardActor*>& LightCardsToSelect);
 
 	FDisplayClusterLightCardEditorWidget::EWidgetMode GetEditorWidgetMode() const { return EditorWidget->GetWidgetMode(); }
 	void SetEditorWidgetMode(FDisplayClusterLightCardEditorWidget::EWidgetMode InWidgetMode) { EditorWidget->SetWidgetMode(InWidgetMode); }
@@ -188,6 +194,9 @@ public:
 	/** Places the given light card in the middle of the current viewport */
 	void CenterLightCardInView(ADisplayClusterLightCardActor& LightCard);
 
+	/** Moves all selected light cards to the specified pixel position */
+	void MoveSelectedLightCardsToPixel(const FIntPoint& PixelPos);
+
 	/** Returns the current input mode */
 	EInputMode GetInputMode() { return InputMode; }
 
@@ -305,6 +314,7 @@ private:
 
 private:
 	TWeakPtr<FSceneViewport> SceneViewportPtr;
+	TWeakPtr<SDisplayClusterLightCardEditorViewport> LightCardEditorViewportPtr;
 	TWeakPtr<SDisplayClusterLightCardEditor> LightCardEditorPtr;
 	TWeakObjectPtr<ADisplayClusterRootActor> RootActorProxy;
 	TWeakObjectPtr<ADisplayClusterRootActor> RootActorLevelInstance;
@@ -384,4 +394,7 @@ private:
 
 	/** Array of mouse positions that will be used to spawn a new light card with the shape of the drawn polygon */
 	TArray<FIntPoint> DrawnMousePositions;
+
+	/** Multicast delegate that stores callbacks to be invoked the next time the scene is refreshed */
+	FOnNextSceneRefresh OnNextSceneRefreshDelegate;
 };
\ No newline at end of file
diff --git a/Engine/Plugins/Runtime/nDisplay/Source/DisplayClusterLightCardEditor/Private/Viewport/DisplayClusterLightcardEditorViewport.cpp b/Engine/Plugins/Runtime/nDisplay/Source/DisplayClusterLightCardEditor/Private/Viewport/DisplayClusterLightcardEditorViewport.cpp
--- a/Engine/Plugins/Runtime/nDisplay/Source/DisplayClusterLightCardEditor/Private/Viewport/DisplayClusterLightcardEditorViewport.cpp
+++ b/Engine/Plugins/Runtime/nDisplay/Source/DisplayClusterLightCardEditor/Private/Viewport/DisplayClusterLightcardEditorViewport.cpp
@@ -5,8 +5,11 @@
 #include "DisplayClusterLightCardEditorStyle.h"
 #include "DisplayClusterLightCardEditorViewportClient.h"
 #include "DisplayClusterLightCardEditorCommands.h"
+#include "SDisplayClusterLightCardEditor.h"
 
 #include "EditorViewportCommands.h"
+#include "Framework/Commands/GenericCommands.h"
+#include "HAL/PlatformApplicationMisc.h"
 #include "SEditorViewportToolBarMenu.h"
 #include "STransformViewportToolbar.h"
 #include "Slate/SceneViewport.h"
@@ -225,7 +228,7 @@
 	TWeakPtr<SDisplayClusterLightCardEditorViewport> EditorViewport;
 };
 
-void SDisplayClusterLightCardEditorViewport::Construct(const FArguments& InArgs, TSharedPtr<SDisplayClusterLightCardEditor> InLightCardEditor)
+void SDisplayClusterLightCardEditorViewport::Construct(const FArguments& InArgs, TSharedPtr<SDisplayClusterLightCardEditor> InLightCardEditor, TSharedPtr<class FUICommandList> InCommandList)
 {
 	LightCardEditorPtr = InLightCardEditor;
 	
@@ -233,6 +236,11 @@
 	AdvancedPreviewScene->SetFloorVisibility(true);
 		
 	SEditorViewport::Construct(SEditorViewport::FArguments());
+
+	if (InCommandList.IsValid())
+	{
+		CommandList->Append(InCommandList.ToSharedRef());
+	}
 }
 
 SDisplayClusterLightCardEditorViewport::~SDisplayClusterLightCardEditorViewport()
@@ -259,6 +267,16 @@
 {
 }
 
+FReply SDisplayClusterLightCardEditorViewport::OnKeyDown(const FGeometry& MyGeometry, const FKeyEvent& InKeyEvent)
+{
+	// Cache the current mouse position in case the key event invokes the paste here command
+	const FVector2D MousePos = FSlateApplication::Get().GetCursorPos();
+	const float DPIScale = FPlatformApplicationMisc::GetDPIScaleFactorAtPoint(MousePos.X, MousePos.Y);
+	PasteHerePos = MyGeometry.AbsoluteToLocal(MousePos) * DPIScale;
+
+	return SEditorViewport::OnKeyDown(MyGeometry, InKeyEvent);
+}
+
 void SDisplayClusterLightCardEditorViewport::SetRootActor(ADisplayClusterRootActor* NewRootActor)
 {
 	if (ViewportClient.IsValid())
@@ -267,12 +285,27 @@
 	}
 }
 
+void SDisplayClusterLightCardEditorViewport::SummonContextMenu()
+{
+	const FVector2D MousePos = FSlateApplication::Get().GetCursorPos();
+	const float DPIScale = FPlatformApplicationMisc::GetDPIScaleFactorAtPoint(MousePos.X, MousePos.Y);
+	PasteHerePos = GetTickSpaceGeometry().AbsoluteToLocal(FSlateApplication::Get().GetCursorPos()) * DPIScale;
+
+	TSharedRef<SWidget> MenuContents = MakeContextMenu();
+	FSlateApplication::Get().PushMenu(
+		AsShared(),
+		FWidgetPath(),
+		MenuContents,
+		MousePos,
+		FPopupTransitionEffect(FPopupTransitionEffect::ContextMenu));
+}
+
 TSharedRef<FEditorViewportClient> SDisplayClusterLightCardEditorViewport::MakeEditorViewportClient()
 {
 	check(AdvancedPreviewScene.IsValid());
 	
 	ViewportClient = MakeShareable(new FDisplayClusterLightCardEditorViewportClient(*AdvancedPreviewScene.Get(),
-		SharedThis(this), LightCardEditorPtr));
+		SharedThis(this)));
 	return ViewportClient.ToSharedRef();
 }
 
@@ -350,9 +383,35 @@
 			FExecuteAction::CreateSP(this, &SDisplayClusterLightCardEditorViewport::DrawLightCard),
 			FCanExecuteAction(),
 			FIsActionChecked::CreateSP(this, &SDisplayClusterLightCardEditorViewport::IsDrawingLightCard));
+
+		CommandList->MapAction(
+			Commands.PasteHere,
+			FExecuteAction::CreateSP(this, &SDisplayClusterLightCardEditorViewport::PasteLightCardsHere),
+			FCanExecuteAction::CreateSP(this, &SDisplayClusterLightCardEditorViewport::CanPasteLightCardsHere));
 	}
 }
 
+TSharedRef<SWidget> SDisplayClusterLightCardEditorViewport::MakeContextMenu()
+{
+	const bool bInShouldCloseWindowAfterMenuSelection = true;
+	FMenuBuilder MenuBuilder(bInShouldCloseWindowAfterMenuSelection, CommandList);
+
+	MenuBuilder.AddMenuEntry(FDisplayClusterLightCardEditorCommands::Get().RemoveLightCard);
+
+	MenuBuilder.BeginSection("Edit", LOCTEXT("EditSection", "Edit"));
+	{
+		MenuBuilder.AddMenuEntry(FGenericCommands::Get().Cut);
+		MenuBuilder.AddMenuEntry(FGenericCommands::Get().Copy);
+		MenuBuilder.AddMenuEntry(FGenericCommands::Get().Paste);
+		MenuBuilder.AddMenuEntry(FDisplayClusterLightCardEditorCommands::Get().PasteHere);
+		MenuBuilder.AddMenuEntry(FGenericCommands::Get().Duplicate);
+		MenuBuilder.AddMenuEntry(FGenericCommands::Get().Delete);
+	}
+	MenuBuilder.EndSection();
+
+	return MenuBuilder.MakeWidget();
+}
+
 void SDisplayClusterLightCardEditorViewport::SetEditorWidgetMode(FDisplayClusterLightCardEditorWidget::EWidgetMode InWidgetMode)
 {
 	if (ViewportClient.IsValid())
@@ -426,5 +485,33 @@
 	return false;
 }
 
+void SDisplayClusterLightCardEditorViewport::PasteLightCardsHere()
+{
+	if (LightCardEditorPtr.IsValid())
+	{
+		const bool bOffsetLightCardPositions = false;
+		LightCardEditorPtr.Pin()->PasteLightCards(bOffsetLightCardPositions);
 
+		if (ViewportClient.IsValid())
+		{
+			// Perform the positioning of the pasted light cards on the next scene refresh, since the light card proxies will have not been 
+			// regenerated until then
+			ViewportClient->GetOnNextSceneRefresh().AddLambda([this]()
+			{
+				ViewportClient->MoveSelectedLightCardsToPixel(FIntPoint(PasteHerePos.X, PasteHerePos.Y));
+			});
+		}
+	}
+}
+
+bool SDisplayClusterLightCardEditorViewport::CanPasteLightCardsHere() const
+{
+	if (LightCardEditorPtr.IsValid())
+	{
+		return LightCardEditorPtr.Pin()->CanPasteLightCards();
+	}
+
+	return false;
+}
+
 #undef LOCTEXT_NAMESPACE
diff --git a/Engine/Plugins/Runtime/nDisplay/Source/DisplayClusterLightCardEditor/Private/Viewport/DisplayClusterLightcardEditorViewport.h b/Engine/Plugins/Runtime/nDisplay/Source/DisplayClusterLightCardEditor/Private/Viewport/DisplayClusterLightcardEditorViewport.h
--- a/Engine/Plugins/Runtime/nDisplay/Source/DisplayClusterLightCardEditor/Private/Viewport/DisplayClusterLightcardEditorViewport.h
+++ b/Engine/Plugins/Runtime/nDisplay/Source/DisplayClusterLightCardEditor/Private/Viewport/DisplayClusterLightcardEditorViewport.h
@@ -23,7 +23,7 @@
 	SLATE_BEGIN_ARGS(SDisplayClusterLightCardEditorViewport) {}
 	SLATE_END_ARGS()
 
-	void Construct(const FArguments& InArgs, TSharedPtr<SDisplayClusterLightCardEditor> InLightCardEditor);
+	void Construct(const FArguments& InArgs, TSharedPtr<SDisplayClusterLightCardEditor> InLightCardEditor, TSharedPtr<class FUICommandList> InCommandList);
 	~SDisplayClusterLightCardEditorViewport();
 
 	// ICommonEditorViewportToolbarInfoProvider
@@ -32,12 +32,19 @@
 	virtual void OnFloatingButtonClicked() override;
 	// ~ICommonEditorViewportToolbarInfoProvider
 
+	// SWidget interface
+	virtual FReply OnKeyDown(const FGeometry& MyGeometry, const FKeyEvent& InKeyEvent) override;
+	// ~SWidget
+
 	TWeakPtr<SDisplayClusterLightCardEditor> GetLightCardEditor() { return LightCardEditorPtr;}
 
 	void SetRootActor(ADisplayClusterRootActor* NewRootActor);
 	
 	TSharedRef<FDisplayClusterLightCardEditorViewportClient> GetLightCardEditorViewportClient() const { return ViewportClient.ToSharedRef(); }
 
+	/** Summons the context menu within this viewport at the mouse's current position */
+	void SummonContextMenu();
+
 private:
 	// SEditorViewport
 	virtual TSharedRef<FEditorViewportClient> MakeEditorViewportClient() override;
@@ -46,6 +53,9 @@
 	virtual void BindCommands() override;
 	// ~SEditorViewport
 
+	/** Create the popup context menu */
+	TSharedRef<SWidget> MakeContextMenu();
+
 	void SetEditorWidgetMode(FDisplayClusterLightCardEditorWidget::EWidgetMode InWidgetMode);
 	bool IsEditorWidgetModeSelected(FDisplayClusterLightCardEditorWidget::EWidgetMode InWidgetMode) const;
 	void CycleEditorWidgetMode();
@@ -59,6 +69,11 @@
 	/** Returns true if we are in light card drawing mode */
 	bool IsDrawingLightCard() const;
 
+	/** Pastes any light cards in the clipboard and moves them to the cached mouse position */
+	void PasteLightCardsHere();
+
+	/** Determines if there are any light cards in the clipboard that can be pasted */
+	bool CanPasteLightCardsHere() const;
 
 private:
 	/** Preview Scene - uses advanced preview settings */
@@ -68,4 +83,6 @@
 	TSharedPtr<FDisplayClusterLightCardEditorViewportClient> ViewportClient;
 	TWeakPtr<SDisplayClusterLightCardEditor> LightCardEditorPtr;
 
+	/** The cached mouse position used for the Paste Here command */
+	FVector2D PasteHerePos;
 };
