nDisplay: Added some support for feathering when using custom polygon light cards.

The current method is erosion followed by blur of the inverted alpha.

diff --git a/Engine/Plugins/Runtime/nDisplay/Source/DisplayCluster/Private/Game/EngineClasses/Scene/DisplayClusterLightCardActor.cpp b/Engine/Plugins/Runtime/nDisplay/Source/DisplayCluster/Private/Game/EngineClasses/Scene/DisplayClusterLightCardActor.cpp
--- a/Engine/Plugins/Runtime/nDisplay/Source/DisplayCluster/Private/Game/EngineClasses/Scene/DisplayClusterLightCardActor.cpp
+++ b/Engine/Plugins/Runtime/nDisplay/Source/DisplayCluster/Private/Game/EngineClasses/Scene/DisplayClusterLightCardActor.cpp
@@ -217,13 +217,13 @@
 		LightCardMaterialInstance->SetScalarParameterValue(TEXT("Exposure"), Exposure);
 		LightCardMaterialInstance->SetScalarParameterValue(TEXT("EmissiveStrength"), Exposure);
 		LightCardMaterialInstance->SetScalarParameterValue(TEXT("Opacity"), bIsProxy ? ProxyOpacity : Opacity);
-		LightCardMaterialInstance->SetScalarParameterValue(TEXT("Feather"), Feathering);
 		LightCardMaterialInstance->SetTextureParameterValue(TEXT("Texture"), Texture);
 		LightCardMaterialInstance->SetTextureParameterValue(TEXT("AlphaTexture"), PolygonMask);
 
 		bool bUseMask = true; // Enable masking
 		bool bUseTextureAlpha = false; // Use the alpha channel of Texture as mask
-		bool bUseAlphaTexture = false; // Use the AlphaTexture 
+		bool bUseAlphaTexture = false; // Use the AlphaTexture
+		float FeatheringValue = Feathering;
 
 		switch (Mask)
 		{
@@ -238,9 +238,11 @@
 		case EDisplayClusterLightCardMask::Polygon:
 			bUseTextureAlpha = true;
 			bUseAlphaTexture = true;
+			FeatheringValue = 0.;
 			break;
 		}
 
+		LightCardMaterialInstance->SetScalarParameterValue(TEXT("Feather"), FeatheringValue);
 		LightCardMaterialInstance->SetScalarParameterValue(TEXT("UseMask"), bUseMask);
 		LightCardMaterialInstance->SetScalarParameterValue(TEXT("UseTextureAlpha"), bUseTextureAlpha);
 		LightCardMaterialInstance->SetScalarParameterValue(TEXT("UseAlphaTexture"), bUseAlphaTexture);
@@ -265,7 +267,7 @@
 	const int32 TextureWidth = GDisplayClusterLightCardPolygonTextureSize;
 	const int32 TextureHeight = GDisplayClusterLightCardPolygonTextureSize;
 
-	cv::Mat GrayMat(cv::Size(TextureWidth, TextureHeight), CV_8UC1, cv::Scalar(0));
+	cv::Mat GrayMat(cv::Size(TextureHeight, TextureWidth), CV_8UC1, cv::Scalar(0));
 
 	std::vector<cv::Point2i> PolyVec; // polygon with LC shape that to be filled with cv::fillPoly
 	PolyVec.reserve(Polygon.Num());
@@ -279,8 +281,52 @@
 	}
 
 	// Fill the polygon pixels
-	cv::fillPoly(GrayMat, PolyVec, 255, cv::LINE_AA);
+	cv::fillPoly(GrayMat, PolyVec, 255, cv::LINE_AA); // Polygon inside is 255, outside is 0
 
+	// Apply feathering effect
+	if (Feathering > 0)
+	{
+		// Overscan it so that the effects near the edges are as desired
+		const int32 Overscan = (TextureWidth + TextureHeight) / 16;
+		const cv::Size OverscanMatSize(GrayMat.rows + 2 * Overscan, GrayMat.cols + 2 * Overscan);
+		cv::Mat OverscanMat(OverscanMatSize, GrayMat.type(), cv::Scalar(0));
+
+		const cv::Rect OverscanActiveRect(Overscan, Overscan, GrayMat.cols, GrayMat.rows);
+
+		GrayMat.copyTo(OverscanMat(OverscanActiveRect));
+
+		// Erode to compensate for the growth that the blur causes
+		{
+			const int32 KernelSize = int32((Feathering * TextureWidth) / 6) * 2 + 1;
+
+			const cv::Mat ErodeKernel = cv::getStructuringElement(
+				cv::MORPH_ERODE,
+				cv::Size(KernelSize, KernelSize)
+			);
+
+			cv::erode(OverscanMat, OverscanMat, ErodeKernel);
+		}
+
+		// Blur
+		{
+			const int32 KernelSize = int32((Feathering * TextureWidth) / 2) * 2 + 1;
+
+			// We are going to blur the outside
+			OverscanMat = 255 - OverscanMat;
+
+			const float SigmaX = 0.2 * ((KernelSize - 1) * 0.5 - 1) + 0.8;
+			const float SigmaY = SigmaX;
+
+			cv::GaussianBlur(OverscanMat, OverscanMat, cv::Size(KernelSize, KernelSize), SigmaX, SigmaY);
+
+			// We recover the non-inverted alpha now
+			OverscanMat = 255 - OverscanMat;
+		}
+
+		// Put back to Gray
+		OverscanMat(OverscanActiveRect).copyTo(GrayMat);
+	}
+
 	// Create a texture, all white and and the shape goes in the alpha
 	cv::Mat RGBAMat(cv::Size(TextureWidth, TextureHeight), CV_8UC4, cv::Scalar::all(255));
 	cv::insertChannel(GrayMat, RGBAMat, 3);
