nDisplay: Added Tint and Temperature to light cards.

diff --git a/Engine/Plugins/Runtime/nDisplay/Source/DisplayCluster/Private/Game/EngineClasses/Scene/DisplayClusterLightCardActor.cpp b/Engine/Plugins/Runtime/nDisplay/Source/DisplayCluster/Private/Game/EngineClasses/Scene/DisplayClusterLightCardActor.cpp
--- a/Engine/Plugins/Runtime/nDisplay/Source/DisplayCluster/Private/Game/EngineClasses/Scene/DisplayClusterLightCardActor.cpp
+++ b/Engine/Plugins/Runtime/nDisplay/Source/DisplayCluster/Private/Game/EngineClasses/Scene/DisplayClusterLightCardActor.cpp
@@ -40,6 +40,8 @@ ADisplayClusterLightCardActor::ADisplayClusterLightCardActor(const FObjectInitia
 	, Mask(EDisplayClusterLightCardMask::Circle)
 	, Texture(nullptr)
 	, Color(FLinearColor(1.f, 1.f, 1.f, 1.f))
+	, Temperature(6500)
+	, Tint(0)
 	, Exposure(0.f)
 	, Opacity(1.f)
 	, Feathering(0.f)
@@ -214,6 +216,8 @@ void ADisplayClusterLightCardActor::UpdateLightCardMaterialInstance()
 		LightCardMaterialInstance->ClearParameterValues();
 
 		LightCardMaterialInstance->SetVectorParameterValue(TEXT("CardColor"), Color);
+		LightCardMaterialInstance->SetScalarParameterValue(TEXT("Temperature"), Temperature);
+		LightCardMaterialInstance->SetScalarParameterValue(TEXT("Tint"), Tint);
 		LightCardMaterialInstance->SetScalarParameterValue(TEXT("Exposure"), Exposure);
 		LightCardMaterialInstance->SetScalarParameterValue(TEXT("EmissiveStrength"), Exposure);
 		LightCardMaterialInstance->SetScalarParameterValue(TEXT("Opacity"), bIsProxy ? ProxyOpacity : Opacity);
diff --git a/Engine/Plugins/Runtime/nDisplay/Source/DisplayCluster/Public/DisplayClusterLightCardActor.h b/Engine/Plugins/Runtime/nDisplay/Source/DisplayCluster/Public/DisplayClusterLightCardActor.h
--- a/Engine/Plugins/Runtime/nDisplay/Source/DisplayCluster/Public/DisplayClusterLightCardActor.h
+++ b/Engine/Plugins/Runtime/nDisplay/Source/DisplayCluster/Public/DisplayClusterLightCardActor.h
@@ -114,6 +114,12 @@
 	UPROPERTY(EditAnywhere, BlueprintReadWrite, Category = "Appearance")
 	FLinearColor Color;
 
+	UPROPERTY(EditAnywhere, BlueprintReadWrite, Category = "Appearance", meta = (UIMin = 0, ClampMin = 0, UIMax = 10000, ClampMax = 10000))
+	float Temperature;
+
+	UPROPERTY(EditAnywhere, BlueprintReadWrite, Category = "Appearance", meta = (UIMin = -1, ClampMin = -1, UIMax = 1, ClampMax = 1))
+	float Tint;
+
 	UPROPERTY(EditAnywhere, BlueprintReadWrite, Category = "Appearance")
 	float Exposure;
 
