nDisplay: Added alpha gradient with angle and start/end values to light cards.

diff --git a/Engine/Plugins/Runtime/nDisplay/Source/DisplayCluster/Private/Game/EngineClasses/Scene/DisplayClusterLightCardActor.cpp b/Engine/Plugins/Runtime/nDisplay/Source/DisplayCluster/Private/Game/EngineClasses/Scene/DisplayClusterLightCardActor.cpp
--- a/Engine/Plugins/Runtime/nDisplay/Source/DisplayCluster/Private/Game/EngineClasses/Scene/DisplayClusterLightCardActor.cpp
+++ b/Engine/Plugins/Runtime/nDisplay/Source/DisplayCluster/Private/Game/EngineClasses/Scene/DisplayClusterLightCardActor.cpp
@@ -143,13 +143,15 @@ void ADisplayClusterLightCardActor::PostEditChangeProperty(FPropertyChangedEvent
 	}
 
 	if (PropertyChangedEvent.Property && (
-		PropertyChangedEvent.Property->GetFName() == GET_MEMBER_NAME_CHECKED(ADisplayClusterLightCardActor, Mask) ||
-		PropertyChangedEvent.Property->GetFName() == GET_MEMBER_NAME_CHECKED(ADisplayClusterLightCardActor, Polygon) ||
-		PropertyChangedEvent.Property->GetFName() == GET_MEMBER_NAME_CHECKED(ADisplayClusterLightCardActor, Texture) ||
-		PropertyChangedEvent.Property->GetFName() == GET_MEMBER_NAME_CHECKED(ADisplayClusterLightCardActor, Color) ||
-		PropertyChangedEvent.Property->GetFName() == GET_MEMBER_NAME_CHECKED(ADisplayClusterLightCardActor, Exposure) ||
-		PropertyChangedEvent.Property->GetFName() == GET_MEMBER_NAME_CHECKED(ADisplayClusterLightCardActor, Opacity) ||
-		PropertyChangedEvent.Property->GetFName() == GET_MEMBER_NAME_CHECKED(ADisplayClusterLightCardActor, Feathering)))
+		   PropertyChangedEvent.Property->GetFName() == GET_MEMBER_NAME_CHECKED(ADisplayClusterLightCardActor, Mask)
+		|| PropertyChangedEvent.Property->GetFName() == GET_MEMBER_NAME_CHECKED(ADisplayClusterLightCardActor, Polygon) 
+		|| PropertyChangedEvent.Property->GetFName() == GET_MEMBER_NAME_CHECKED(ADisplayClusterLightCardActor, Texture) 
+		|| PropertyChangedEvent.Property->GetFName() == GET_MEMBER_NAME_CHECKED(ADisplayClusterLightCardActor, Color)
+		|| PropertyChangedEvent.Property->GetFName() == GET_MEMBER_NAME_CHECKED(ADisplayClusterLightCardActor, Exposure)
+		|| PropertyChangedEvent.Property->GetFName() == GET_MEMBER_NAME_CHECKED(ADisplayClusterLightCardActor, Opacity)
+		|| PropertyChangedEvent.Property->GetFName() == GET_MEMBER_NAME_CHECKED(ADisplayClusterLightCardActor, Feathering)
+		|| PropertyChangedEvent.Property->GetFName() == GET_MEMBER_NAME_CHECKED(ADisplayClusterLightCardActor, AlphaGradient)
+		))
 	{
 		UpdateLightCardMaterialInstance();
 	}
@@ -224,6 +226,11 @@ void ADisplayClusterLightCardActor::UpdateLightCardMaterialInstance()
 		LightCardMaterialInstance->SetTextureParameterValue(TEXT("Texture"), Texture);
 		LightCardMaterialInstance->SetTextureParameterValue(TEXT("AlphaTexture"), PolygonMask);
 
+		LightCardMaterialInstance->SetScalarParameterValue(TEXT("AlphaGradient"), AlphaGradient.bEnableAlphaGradient);
+		LightCardMaterialInstance->SetScalarParameterValue(TEXT("AlphaGradientStartingAlpha"), AlphaGradient.StartingAlpha);
+		LightCardMaterialInstance->SetScalarParameterValue(TEXT("AlphaGradientEndingAlpha"), AlphaGradient.EndingAlpha);
+		LightCardMaterialInstance->SetScalarParameterValue(TEXT("AlphaGradientAngle"), AlphaGradient.Angle);
+
 		bool bUseMask = true; // Enable masking
 		bool bUseTextureAlpha = false; // Use the alpha channel of Texture as mask
 		bool bUseAlphaTexture = false; // Use the AlphaTexture
diff --git a/Engine/Plugins/Runtime/nDisplay/Source/DisplayCluster/Public/DisplayClusterLightCardActor.h b/Engine/Plugins/Runtime/nDisplay/Source/DisplayCluster/Public/DisplayClusterLightCardActor.h
--- a/Engine/Plugins/Runtime/nDisplay/Source/DisplayCluster/Public/DisplayClusterLightCardActor.h
+++ b/Engine/Plugins/Runtime/nDisplay/Source/DisplayCluster/Public/DisplayClusterLightCardActor.h
@@ -21,6 +21,28 @@ enum class EDisplayClusterLightCardMask : uint8
 	Polygon,
 };
 
+USTRUCT(Blueprintable)
+struct FLightCardAlphaGradientSettings
+{
+	GENERATED_BODY()
+
+	/** Enables/disables alpha gradient effect */
+	UPROPERTY(EditAnywhere, BlueprintReadWrite, Category = "Appearance")
+	bool bEnableAlphaGradient = false;
+
+	/** Starting alpha value in the gradient */
+	UPROPERTY(EditAnywhere, BlueprintReadWrite, Category = "Appearance", meta = (EditCondition = "bEnableAlphaGradient"))
+	float StartingAlpha = 0;
+
+	/** Ending alpha value in the gradient */
+	UPROPERTY(EditAnywhere, BlueprintReadWrite, Category = "Appearance", meta = (EditCondition = "bEnableAlphaGradient"))
+	float EndingAlpha = 1;
+
+	/** The angle (degrees) determines the gradient direction. */
+	UPROPERTY(EditAnywhere, BlueprintReadWrite, Category = "Appearance", meta = (EditCondition = "bEnableAlphaGradient"))
+	float Angle = 0;
+};
+
 UCLASS(Blueprintable)
 class DISPLAYCLUSTER_API ADisplayClusterLightCardActor : public AActor
 {
@@ -129,6 +151,10 @@ public:
 	UPROPERTY(EditAnywhere, BlueprintReadWrite, Category = "Appearance", meta = (UIMin = 0, ClampMin = 0))
 	float Feathering;
 
+	/** Settings related to an alpha gradient effect */
+	UPROPERTY(EditAnywhere, BlueprintReadWrite, Category = "Appearance")
+	FLightCardAlphaGradientSettings AlphaGradient;
+
 	/** A flag that controls wether the light card's location and rotation are locked to its "owning" root actor */
 	UPROPERTY()
 	bool bLockToOwningRootActor = true;
