nDisplay: Added Gain to light card parameters. Exposure is 2^n while Gain a is linear multiplier to RGB.

diff --git a/Engine/Plugins/Runtime/nDisplay/Source/DisplayCluster/Private/Game/EngineClasses/Scene/DisplayClusterLightCardActor.cpp b/Engine/Plugins/Runtime/nDisplay/Source/DisplayCluster/Private/Game/EngineClasses/Scene/DisplayClusterLightCardActor.cpp
--- a/Engine/Plugins/Runtime/nDisplay/Source/DisplayCluster/Private/Game/EngineClasses/Scene/DisplayClusterLightCardActor.cpp
+++ b/Engine/Plugins/Runtime/nDisplay/Source/DisplayCluster/Private/Game/EngineClasses/Scene/DisplayClusterLightCardActor.cpp
@@ -41,8 +41,9 @@ ADisplayClusterLightCardActor::ADisplayClusterLightCardActor(const FObjectInitia
 	, Texture(nullptr)
 	, Color(FLinearColor(1.f, 1.f, 1.f, 1.f))
 	, Temperature(6500)
-	, Tint(0)
+	, Tint(0.f)
 	, Exposure(0.f)
+	, Gain(1.f)
 	, Opacity(1.f)
 	, Feathering(0.f)
 {
@@ -148,6 +149,7 @@ void ADisplayClusterLightCardActor::PostEditChangeProperty(FPropertyChangedEvent
 		|| PropertyChangedEvent.Property->GetFName() == GET_MEMBER_NAME_CHECKED(ADisplayClusterLightCardActor, Texture) 
 		|| PropertyChangedEvent.Property->GetFName() == GET_MEMBER_NAME_CHECKED(ADisplayClusterLightCardActor, Color)
 		|| PropertyChangedEvent.Property->GetFName() == GET_MEMBER_NAME_CHECKED(ADisplayClusterLightCardActor, Exposure)
+		|| PropertyChangedEvent.Property->GetFName() == GET_MEMBER_NAME_CHECKED(ADisplayClusterLightCardActor, Gain)
 		|| PropertyChangedEvent.Property->GetFName() == GET_MEMBER_NAME_CHECKED(ADisplayClusterLightCardActor, Opacity)
 		|| PropertyChangedEvent.Property->GetFName() == GET_MEMBER_NAME_CHECKED(ADisplayClusterLightCardActor, Feathering)
 		|| PropertyChangedEvent.Property->GetFName() == GET_MEMBER_NAME_CHECKED(ADisplayClusterLightCardActor, AlphaGradient)
@@ -221,7 +223,7 @@ void ADisplayClusterLightCardActor::UpdateLightCardMaterialInstance()
 		LightCardMaterialInstance->SetScalarParameterValue(TEXT("Temperature"), Temperature);
 		LightCardMaterialInstance->SetScalarParameterValue(TEXT("Tint"), Tint);
 		LightCardMaterialInstance->SetScalarParameterValue(TEXT("Exposure"), Exposure);
-		LightCardMaterialInstance->SetScalarParameterValue(TEXT("EmissiveStrength"), Exposure);
+		LightCardMaterialInstance->SetScalarParameterValue(TEXT("Gain"), Gain);
 		LightCardMaterialInstance->SetScalarParameterValue(TEXT("Opacity"), bIsProxy ? ProxyOpacity : Opacity);
 		LightCardMaterialInstance->SetTextureParameterValue(TEXT("Texture"), Texture);
 		LightCardMaterialInstance->SetTextureParameterValue(TEXT("AlphaTexture"), PolygonMask);
diff --git a/Engine/Plugins/Runtime/nDisplay/Source/DisplayCluster/Public/DisplayClusterLightCardActor.h b/Engine/Plugins/Runtime/nDisplay/Source/DisplayCluster/Public/DisplayClusterLightCardActor.h
--- a/Engine/Plugins/Runtime/nDisplay/Source/DisplayCluster/Public/DisplayClusterLightCardActor.h
+++ b/Engine/Plugins/Runtime/nDisplay/Source/DisplayCluster/Public/DisplayClusterLightCardActor.h
@@ -106,15 +106,20 @@ public:
 	void UpdatePolygonTexture();
 
 public:
+
+	/** Radius of light card polar coordinates */
 	UPROPERTY(EditAnywhere, BlueprintReadWrite, Category = "Orientation")
 	double DistanceFromCenter;
 
+	/** Related to the Azimuth of light card polar coordinates */
 	UPROPERTY(EditAnywhere, BlueprintReadWrite, Category = "Orientation", meta = (UIMin = 0, ClampMin = 0, UIMax = 360, ClampMax = 360))
 	double Longitude;
 
+	/** Related to the Elevation of light card polar coordinates */
 	UPROPERTY(EditAnywhere, BlueprintReadWrite, Category = "Orientation", meta = (UIMin = -90, ClampMin = -90, UIMax = 90, ClampMax = 90))
 	double Latitude;
 
+	/** Roll rotation of light card around its plane axis */
 	UPROPERTY(EditAnywhere, BlueprintReadWrite, Category = "Orientation")
 	double Spin;
 
@@ -133,6 +138,7 @@ public:
 	UPROPERTY(EditAnywhere, BlueprintReadWrite, Category = "Appearance")
 	UTexture* Texture;
 
+	/** Light card color, before any modifier is applied */
 	UPROPERTY(EditAnywhere, BlueprintReadWrite, Category = "Appearance")
 	FLinearColor Color;
 
@@ -142,12 +148,19 @@ public:
 	UPROPERTY(EditAnywhere, BlueprintReadWrite, Category = "Appearance", meta = (UIMin = -1, ClampMin = -1, UIMax = 1, ClampMax = 1))
 	float Tint;
 
+	/** 2^Exposure color value multiplier */
 	UPROPERTY(EditAnywhere, BlueprintReadWrite, Category = "Appearance")
 	float Exposure;
 
+	/** Linear color value multiplier */
+	UPROPERTY(EditAnywhere, BlueprintReadWrite, Category = "Appearance", meta = (UIMin = 0, ClampMin = 0))
+	float Gain;
+
+	/** Linear alpha multiplier */
 	UPROPERTY(EditAnywhere, BlueprintReadWrite, Category = "Appearance", meta = (UIMin = 0, ClampMin = 0, UIMax = 1, ClampMax = 1))
 	float Opacity;
 
+	/** Feathers in the alpha from the edges */
 	UPROPERTY(EditAnywhere, BlueprintReadWrite, Category = "Appearance", meta = (UIMin = 0, ClampMin = 0))
 	float Feathering;
 
