nDisplay: Made the distance from flush a settable property of the light cards called "RadialOffset".

diff --git a/Engine/Plugins/Runtime/nDisplay/Source/DisplayCluster/Private/Game/EngineClasses/Scene/DisplayClusterLightCardActor.cpp b/Engine/Plugins/Runtime/nDisplay/Source/DisplayCluster/Private/Game/EngineClasses/Scene/DisplayClusterLightCardActor.cpp
--- a/Engine/Plugins/Runtime/nDisplay/Source/DisplayCluster/Private/Game/EngineClasses/Scene/DisplayClusterLightCardActor.cpp
+++ b/Engine/Plugins/Runtime/nDisplay/Source/DisplayCluster/Private/Game/EngineClasses/Scene/DisplayClusterLightCardActor.cpp
@@ -37,6 +37,7 @@ ADisplayClusterLightCardActor::ADisplayClusterLightCardActor(const FObjectInitia
 	, Pitch(0.f)
 	, Yaw(0.f)
 	, Scale(FVector2D(1.f))
+	, RadialOffset(-0.5)
 	, Mask(EDisplayClusterLightCardMask::Circle)
 	, Texture(nullptr)
 	, Color(FLinearColor(1.f, 1.f, 1.f, 1.f))
@@ -131,6 +132,7 @@ void ADisplayClusterLightCardActor::PostEditChangeProperty(FPropertyChangedEvent
 		PropertyChangedEvent.Property->GetFName() == GET_MEMBER_NAME_CHECKED(ADisplayClusterLightCardActor, Longitude) ||
 		PropertyChangedEvent.Property->GetFName() == GET_MEMBER_NAME_CHECKED(ADisplayClusterLightCardActor, Latitude) ||
 		PropertyChangedEvent.Property->GetFName() == GET_MEMBER_NAME_CHECKED(ADisplayClusterLightCardActor, Spin) ||
+		PropertyChangedEvent.Property->GetFName() == GET_MEMBER_NAME_CHECKED(ADisplayClusterLightCardActor, RadialOffset) ||
 		PropertyChangedEvent.Property->GetFName() == GET_MEMBER_NAME_CHECKED(ADisplayClusterLightCardActor, Scale)))
 	{
 		UpdateLightCardTransform();
@@ -201,7 +203,7 @@ FBox ADisplayClusterLightCardActor::GetLightCardBounds(bool bLocalSpace) const
 
 void ADisplayClusterLightCardActor::UpdateLightCardTransform()
 {
-	MainSpringArmComponent->TargetArmLength = DistanceFromCenter;
+	MainSpringArmComponent->TargetArmLength = DistanceFromCenter + RadialOffset;
 	MainSpringArmComponent->SetRelativeRotation(FRotator(-Latitude, Longitude, 0.0));
 
 	FRotator LightCardOrientation = FRotator(-Pitch, Yaw, Spin);
@@ -371,6 +373,7 @@ ADisplayClusterLightCardActor::PositionalParams ADisplayClusterLightCardActor::G
 	Params.Pitch = Pitch;
 	Params.Spin = Spin;
 	Params.Yaw = Yaw;
+	Params.RadialOffset = RadialOffset;
 
 	return Params;
 }
@@ -383,4 +386,5 @@ void ADisplayClusterLightCardActor::SetPositionalParams(const PositionalParams&
 	Pitch = Params.Pitch;
 	Spin = Params.Spin;
 	Yaw = Params.Yaw;
+	RadialOffset = Params.RadialOffset;
 }
diff --git a/Engine/Plugins/Runtime/nDisplay/Source/DisplayCluster/Public/DisplayClusterLightCardActor.h b/Engine/Plugins/Runtime/nDisplay/Source/DisplayCluster/Public/DisplayClusterLightCardActor.h
--- a/Engine/Plugins/Runtime/nDisplay/Source/DisplayCluster/Public/DisplayClusterLightCardActor.h
+++ b/Engine/Plugins/Runtime/nDisplay/Source/DisplayCluster/Public/DisplayClusterLightCardActor.h
@@ -58,6 +58,7 @@ public:
 		double Spin;
 		double Pitch;
 		double Yaw;
+		double RadialOffset;
 	};
 
 public:
@@ -107,7 +108,7 @@ public:
 
 public:
 
-	/** Radius of light card polar coordinates */
+	/** Radius of light card polar coordinates. Does not include the effect of RadialOffset */
 	UPROPERTY(EditAnywhere, BlueprintReadWrite, Category = "Orientation")
 	double DistanceFromCenter;
 
@@ -132,6 +133,10 @@ public:
 	UPROPERTY(EditAnywhere, BlueprintReadWrite, Category = "Orientation")
 	FVector2D Scale;
 
+	/** Used by the flush constraint to offset the location of the light card form the wall */
+	UPROPERTY(EditAnywhere, BlueprintReadWrite, Category = "Orientation")
+	double RadialOffset;
+
 	UPROPERTY(EditAnywhere, BlueprintReadWrite, Category = "Appearance")
 	EDisplayClusterLightCardMask Mask;
 
diff --git a/Engine/Plugins/Runtime/nDisplay/Source/DisplayClusterLightCardEditor/Private/Viewport/DisplayClusterLightCardEditorViewportClient.h b/Engine/Plugins/Runtime/nDisplay/Source/DisplayClusterLightCardEditor/Private/Viewport/DisplayClusterLightCardEditorViewportClient.h
--- a/Engine/Plugins/Runtime/nDisplay/Source/DisplayClusterLightCardEditor/Private/Viewport/DisplayClusterLightCardEditorViewportClient.h
+++ b/Engine/Plugins/Runtime/nDisplay/Source/DisplayClusterLightCardEditor/Private/Viewport/DisplayClusterLightCardEditorViewportClient.h
@@ -120,9 +120,6 @@ private:
 		uint32 SizeY = 0;
 	};
 
-	/** The amount to offset a light card's DistanceFromCenter when making it flush with a screen */
-	static constexpr float LightCardFlushOffset = -0.5f;
-
 public:
 	FDisplayClusterLightCardEditorViewportClient(FAdvancedPreviewScene& InPreviewScene,
 		const TWeakPtr<SDisplayClusterLightCardEditorViewport>& InEditorViewportWidget);
@@ -184,10 +181,11 @@ public:
 	/** Resets the camera to the initial rotation / position */
 	void ResetCamera(bool bLocationOnly = false);
 
-	/** Moves specified card to desired coordinates 
+	/** Moves specified card to desired coordinates. Actual radius will be based on flush constraint and LightCard's RadialOffset.
 	 *
 	 * @param LightCard The light card that we are moving
-	 * @param SphericalCoords specifies desired location of light card in spherical coordinates with respect to view origin
+	 * @param SphericalCoords specifies desired location of light card in spherical coordinates with respect to view origin.
+	 * 
 	*/
 	void MoveLightCardTo(ADisplayClusterLightCardActor& LightCard, const FSphericalCoordinates& SphericalCoords) const;
 
