nDisplay: Fixed crash when the light card editor is open and user changes nDisplay Preview's "Viewports per frame"

diff --git a/Engine/Plugins/Runtime/nDisplay/Source/DisplayCluster/Private/Render/Viewport/DisplayClusterViewportManager.cpp b/Engine/Plugins/Runtime/nDisplay/Source/DisplayCluster/Private/Render/Viewport/DisplayClusterViewportManager.cpp
--- a/Engine/Plugins/Runtime/nDisplay/Source/DisplayCluster/Private/Render/Viewport/DisplayClusterViewportManager.cpp
+++ b/Engine/Plugins/Runtime/nDisplay/Source/DisplayCluster/Private/Render/Viewport/DisplayClusterViewportManager.cpp
@@ -58,10 +58,13 @@ FDisplayClusterViewportManager::FDisplayClusterViewportManager()
 	Configuration      = MakeUnique<FDisplayClusterViewportConfiguration>(*this);
 	RenderFrameManager = MakeUnique<FDisplayClusterRenderFrameManager>();
 
-	RenderTargetManager = MakeShared<FDisplayClusterRenderTargetManager, ESPMode::ThreadSafe>();
+	ViewportManagerProxy = new FDisplayClusterViewportManagerProxy();
+
+	RenderTargetManager = MakeShared<FDisplayClusterRenderTargetManager, ESPMode::ThreadSafe>(ViewportManagerProxy);
 	PostProcessManager  = MakeShared<FDisplayClusterViewportPostProcessManager, ESPMode::ThreadSafe>(*this);
 
-	ViewportManagerProxy = new FDisplayClusterViewportManagerProxy(*this);
+	// initialize proxy
+	ViewportManagerProxy->Initialize(*this);
 
 	// Always reset RTT when root actor re-created
 	ResetSceneRenderTargetSize();
diff --git a/Engine/Plugins/Runtime/nDisplay/Source/DisplayCluster/Private/Render/Viewport/DisplayClusterViewportManagerProxy.cpp b/Engine/Plugins/Runtime/nDisplay/Source/DisplayCluster/Private/Render/Viewport/DisplayClusterViewportManagerProxy.cpp
--- a/Engine/Plugins/Runtime/nDisplay/Source/DisplayCluster/Private/Render/Viewport/DisplayClusterViewportManagerProxy.cpp
+++ b/Engine/Plugins/Runtime/nDisplay/Source/DisplayCluster/Private/Render/Viewport/DisplayClusterViewportManagerProxy.cpp
@@ -54,12 +54,8 @@ static TAutoConsoleVariable<int32> CVarClearFrameRTTEnabled(
 ///////////////////////////////////////////////////////////////////////////////////////
 //          FDisplayClusterViewportManagerProxy
 ///////////////////////////////////////////////////////////////////////////////////////
-
-FDisplayClusterViewportManagerProxy::FDisplayClusterViewportManagerProxy(FDisplayClusterViewportManager& InViewportManager)
-{
-	RenderTargetManager = InViewportManager.RenderTargetManager;
-	PostProcessManager = InViewportManager.PostProcessManager;
-}
+FDisplayClusterViewportManagerProxy::FDisplayClusterViewportManagerProxy()
+{ }
 
 FDisplayClusterViewportManagerProxy::~FDisplayClusterViewportManagerProxy()
 {
@@ -76,6 +72,30 @@ FDisplayClusterViewportManagerProxy::~FDisplayClusterViewportManagerProxy()
 	ImplUpdateClusterNodeViewportProxies();
 }
 
+void FDisplayClusterViewportManagerProxy::Initialize(FDisplayClusterViewportManager& InViewportManager)
+{
+	RenderTargetManager = InViewportManager.RenderTargetManager;
+	PostProcessManager = InViewportManager.PostProcessManager;
+}
+
+void FDisplayClusterViewportManagerProxy::DeleteResource_RenderThread(FDisplayClusterViewportResource* InDeletedResourcePtr)
+{
+	if (InDeletedResourcePtr)
+	{
+		// Handle resource refs must be removed from all viewports
+		for (FDisplayClusterViewportProxy* ViewportProxyIt : ViewportProxies)
+		{
+			if (ViewportProxyIt)
+			{
+				ViewportProxyIt->HandleResourceDelete_RenderThread(InDeletedResourcePtr);
+			}
+		}
+
+		InDeletedResourcePtr->ReleaseResource();
+		delete InDeletedResourcePtr;
+	}
+}
+
 void FDisplayClusterViewportManagerProxy::ImplUpdateClusterNodeViewportProxies()
 {
 	ClusterNodeViewportProxies.Empty();
diff --git a/Engine/Plugins/Runtime/nDisplay/Source/DisplayCluster/Private/Render/Viewport/DisplayClusterViewportManagerProxy.h b/Engine/Plugins/Runtime/nDisplay/Source/DisplayCluster/Private/Render/Viewport/DisplayClusterViewportManagerProxy.h
--- a/Engine/Plugins/Runtime/nDisplay/Source/DisplayCluster/Private/Render/Viewport/DisplayClusterViewportManagerProxy.h
+++ b/Engine/Plugins/Runtime/nDisplay/Source/DisplayCluster/Private/Render/Viewport/DisplayClusterViewportManagerProxy.h
@@ -22,7 +22,7 @@ class FDisplayClusterViewportManagerProxy
 	: public IDisplayClusterViewportManagerProxy
 {
 public:
-	FDisplayClusterViewportManagerProxy(FDisplayClusterViewportManager& InViewportManager);
+	FDisplayClusterViewportManagerProxy();
 	virtual ~FDisplayClusterViewportManagerProxy();
 
 public:
@@ -47,6 +47,9 @@ public:
 	virtual bool ResolveFrameTargetToBackBuffer_RenderThread(FRHICommandListImmediate& RHICmdList, const uint32 InContextNum, const int32 DestArrayIndex, FRHITexture2D* DestTexture, FVector2D WindowSize) const override;
 
 	// internal use only
+	void DeleteResource_RenderThread(FDisplayClusterViewportResource* InDeletedResourcePtr);
+	void Initialize(FDisplayClusterViewportManager& InViewportManager);
+
 	const FDisplayClusterRenderFrameSettings& GetRenderFrameSettings_RenderThread() const
 	{
 		check(IsInRenderingThread());
diff --git a/Engine/Plugins/Runtime/nDisplay/Source/DisplayCluster/Private/Render/Viewport/DisplayClusterViewportProxy.cpp b/Engine/Plugins/Runtime/nDisplay/Source/DisplayCluster/Private/Render/Viewport/DisplayClusterViewportProxy.cpp
--- a/Engine/Plugins/Runtime/nDisplay/Source/DisplayCluster/Private/Render/Viewport/DisplayClusterViewportProxy.cpp
+++ b/Engine/Plugins/Runtime/nDisplay/Source/DisplayCluster/Private/Render/Viewport/DisplayClusterViewportProxy.cpp
@@ -44,85 +44,145 @@ static TAutoConsoleVariable<int32> CVarDisplayClusterRenderOverscanResolve(
 	ECVF_RenderThreadSafe
 );
 
-// The viewport override has the maximum depth. This protects against a link cycle
-static const int32 DisplayClusterViewportProxyResourcesOverrideRecursionDepthMax = 4;
-
-FDisplayClusterViewportProxy::FDisplayClusterViewportProxy(const FDisplayClusterViewportManagerProxy& InOwner, const FDisplayClusterViewport& RenderViewport)
-	: ViewportId(RenderViewport.ViewportId)
-	, ClusterNodeId(RenderViewport.ClusterNodeId)
-	, RenderSettings(RenderViewport.RenderSettings)
-	, ProjectionPolicy(RenderViewport.UninitializedProjectionPolicy)
-	, Owner(InOwner)
-	, ShadersAPI(IDisplayClusterShaders::Get())
+///////////////////////////////////////////////////////////////////////////////////////
+namespace DisplayClusterViewportProxyHelpers
 {
-	check(ProjectionPolicy.IsValid());
-}
+	// The viewport override has the maximum depth. This protects against a link cycle
+	static const int32 DisplayClusterViewportProxyResourcesOverrideRecursionDepthMax = 4;
+
+	/**
+	* Get viewport RHI resources
+	*
+	* @param InResources - Array with DC resources (for mono num=1, for stereo num=2)
+	* @param OutResources - Array with RHI resources
+	*
+	* @return true, if all resources are valid
+	*/
+	static bool GetViewportRHIResourcesImpl_RenderThread(const TArray<FDisplayClusterViewportTextureResource*>& InResources, TArray<FRHITexture2D*>& OutResources)
+	{
+		check(OutResources.Num() == 0);
 
-FDisplayClusterViewportProxy::~FDisplayClusterViewportProxy()
-{
-}
+		if (InResources.Num() > 0)
+		{
+			for (int32 ResourceIndex = 0; ResourceIndex < InResources.Num(); ResourceIndex++)
+			{
+				if (InResources[ResourceIndex] != nullptr)
+				{
+					OutResources.Add(InResources[ResourceIndex]->GetViewportResourceRHI());
+				}
 
-const IDisplayClusterViewportManagerProxy& FDisplayClusterViewportProxy::GetOwner_RenderThread() const
-{
-	check(IsInRenderingThread());
+			}
 
-	return Owner;
-}
+			if (OutResources.Num() == InResources.Num())
+			{
+				return true;
+			}
 
-bool ImplGetTextureResources_RenderThread(const TArray<FDisplayClusterViewportTextureResource*>& InResources, TArray<FRHITexture2D*>& OutResources)
-{
-	check(OutResources.Num() == 0);
+			// Some resources lost
+			OutResources.Empty();
+		}
+
+		return false;
+	}
 
-	if (InResources.Num() > 0)
+	/**
+	* Get viewport RHI resources
+	*
+	* @param InResources - Array with DC resources (for mono num=1, for stereo num=2)
+	* @param OutResources - Array with RHI resources
+	*
+	* @return true, if all resources are valid
+	*/
+	static bool GetViewportRHIResourcesImpl_RenderThread(const TArray<FTextureRenderTargetResource*>& InResources, TArray<FRHITexture2D*>& OutResources)
 	{
-		for (int32 ResourceIndex = 0; ResourceIndex < InResources.Num(); ResourceIndex++)
+		check(OutResources.Num() == 0);
+
+		if (InResources.Num() > 0)
 		{
-			if(InResources[ResourceIndex] != nullptr)
+			for (int32 ResourceIndex = 0; ResourceIndex < InResources.Num(); ResourceIndex++)
 			{
-				OutResources.Add(InResources[ResourceIndex]->GetViewportResourceRHI());
+				if (InResources[ResourceIndex] != nullptr)
+				{
+					OutResources.Add(InResources[ResourceIndex]->GetTexture2DRHI());
+				}
 			}
-			
-		}
 
-		if (OutResources.Num() == InResources.Num())
-		{
-			return true;
+			if (OutResources.Num() == InResources.Num())
+			{
+				return true;
+			}
+
+			// Some resources lost
+			OutResources.Empty();
 		}
 
-		// Some resources lost
-		OutResources.Empty();
+		return false;
 	}
 
-	return false;
-}
-
-
-bool ImplGetTextureResources_RenderThread(const TArray<FTextureRenderTargetResource*>& InResources, TArray<FRHITexture2D*>& OutResources)
-{
-	check(OutResources.Num() == 0);
-
-	if (InResources.Num() > 0)
+	// Reset reference to deleted resource
+	static void HandleResourceDeleteImpl_RenderThread(TArray<FDisplayClusterViewportTextureResource*>& InResources, const FDisplayClusterViewportResource* InDeletedResourcePtr)
 	{
-		for (int32 ResourceIndex = 0; ResourceIndex < InResources.Num(); ResourceIndex++)
+		for (int32 Index = 0; Index < InResources.Num(); Index++)
 		{
-			if (InResources[ResourceIndex] != nullptr)
+			if (InResources[Index] == InDeletedResourcePtr)
 			{
-				OutResources.Add(InResources[ResourceIndex]->GetTexture2DRHI());
+				// remove references to a deleted resource
+				InResources[Index] = nullptr;
 			}
 		}
+	}
 
-		if (OutResources.Num() == InResources.Num())
+	// Reset reference to deleted resource
+	static void HandleResourceDeleteImpl_RenderThread(TArray<FDisplayClusterViewportRenderTargetResource*>& InResources, const FDisplayClusterViewportResource* InDeletedResourcePtr)
+	{
+		for (int32 Index = 0; Index < InResources.Num(); Index++)
 		{
-			return true;
+			if (InResources[Index] == InDeletedResourcePtr)
+			{
+				// remove references to a deleted resource
+				InResources[Index] = nullptr;
+			}
 		}
-
-		// Some resources lost
-		OutResources.Empty();
 	}
+};
+using namespace DisplayClusterViewportProxyHelpers;
 
-	return false;
+///////////////////////////////////////////////////////////////////////////////////////
+FDisplayClusterViewportProxy::FDisplayClusterViewportProxy(const FDisplayClusterViewportManagerProxy& InOwner, const FDisplayClusterViewport& RenderViewport)
+	: ViewportId(RenderViewport.ViewportId)
+	, ClusterNodeId(RenderViewport.ClusterNodeId)
+	, RenderSettings(RenderViewport.RenderSettings)
+	, ProjectionPolicy(RenderViewport.UninitializedProjectionPolicy)
+	, Owner(InOwner)
+	, ShadersAPI(IDisplayClusterShaders::Get())
+{
+	check(ProjectionPolicy.IsValid());
+}
+
+FDisplayClusterViewportProxy::~FDisplayClusterViewportProxy()
+{
+}
+
+const IDisplayClusterViewportManagerProxy& FDisplayClusterViewportProxy::GetOwner_RenderThread() const
+{
+	check(IsInRenderingThread());
+
+	return Owner;
 }
 
+void FDisplayClusterViewportProxy::HandleResourceDelete_RenderThread(FDisplayClusterViewportResource* InDeletedResourcePtr)
+{
+	// Reset all references to deleted resource
+	HandleResourceDeleteImpl_RenderThread(RenderTargets, InDeletedResourcePtr);
+	HandleResourceDeleteImpl_RenderThread(OutputFrameTargetableResources, InDeletedResourcePtr);
+	HandleResourceDeleteImpl_RenderThread(AdditionalFrameTargetableResources, InDeletedResourcePtr);
+
+	HandleResourceDeleteImpl_RenderThread(InputShaderResources, InDeletedResourcePtr);
+	HandleResourceDeleteImpl_RenderThread(AdditionalTargetableResources, InDeletedResourcePtr);
+	HandleResourceDeleteImpl_RenderThread(MipsShaderResources, InDeletedResourcePtr);
+}
+
+
 bool FDisplayClusterViewportProxy::IsShouldOverrideViewportResource(const EDisplayClusterViewportResourceType InResourceType) const
 {
 	// Override resources from other viewport
@@ -206,19 +266,19 @@ bool FDisplayClusterViewportProxy::ImplGetResources_RenderThread(const EDisplayC
 	}
 
 	case EDisplayClusterViewportResourceType::InputShaderResource:
-		return ImplGetTextureResources_RenderThread(InputShaderResources, OutResources);
+		return GetViewportRHIResourcesImpl_RenderThread(InputShaderResources, OutResources);
 
 	case EDisplayClusterViewportResourceType::AdditionalTargetableResource:
-		return ImplGetTextureResources_RenderThread(AdditionalTargetableResources, OutResources);
+		return GetViewportRHIResourcesImpl_RenderThread(AdditionalTargetableResources, OutResources);
 
 	case EDisplayClusterViewportResourceType::MipsShaderResource:
-		return ImplGetTextureResources_RenderThread(MipsShaderResources, OutResources);
+		return GetViewportRHIResourcesImpl_RenderThread(MipsShaderResources, OutResources);
 
 	case EDisplayClusterViewportResourceType::OutputFrameTargetableResource:
-		return ImplGetTextureResources_RenderThread(OutputFrameTargetableResources, OutResources);
+		return GetViewportRHIResourcesImpl_RenderThread(OutputFrameTargetableResources, OutResources);
 
 	case EDisplayClusterViewportResourceType::AdditionalFrameTargetableResource:
-		return ImplGetTextureResources_RenderThread(AdditionalFrameTargetableResources, OutResources);
+		return GetViewportRHIResourcesImpl_RenderThread(AdditionalFrameTargetableResources, OutResources);
 
 #if WITH_EDITOR
 		// Support preview:
diff --git a/Engine/Plugins/Runtime/nDisplay/Source/DisplayCluster/Private/Render/Viewport/DisplayClusterViewportProxy.h b/Engine/Plugins/Runtime/nDisplay/Source/DisplayCluster/Private/Render/Viewport/DisplayClusterViewportProxy.h
--- a/Engine/Plugins/Runtime/nDisplay/Source/DisplayCluster/Private/Render/Viewport/DisplayClusterViewportProxy.h
+++ b/Engine/Plugins/Runtime/nDisplay/Source/DisplayCluster/Private/Render/Viewport/DisplayClusterViewportProxy.h
@@ -101,6 +101,9 @@ public:
 	// ~IDisplayClusterViewportProxy
 	///////////////////////////////
 
+	// Release internal resource refs
+	void HandleResourceDelete_RenderThread(class FDisplayClusterViewportResource* InDeletedResourcePtr);
+
 	void PostResolveViewport_RenderThread(FRHICommandListImmediate& RHICmdList) const;
 
 #if WITH_EDITOR
diff --git a/Engine/Plugins/Runtime/nDisplay/Source/DisplayCluster/Private/Render/Viewport/RenderTarget/DisplayClusterRenderTargetManager.cpp b/Engine/Plugins/Runtime/nDisplay/Source/DisplayCluster/Private/Render/Viewport/RenderTarget/DisplayClusterRenderTargetManager.cpp
--- a/Engine/Plugins/Runtime/nDisplay/Source/DisplayCluster/Private/Render/Viewport/RenderTarget/DisplayClusterRenderTargetManager.cpp
+++ b/Engine/Plugins/Runtime/nDisplay/Source/DisplayCluster/Private/Render/Viewport/RenderTarget/DisplayClusterRenderTargetManager.cpp
@@ -48,9 +48,9 @@ using namespace DisplayClusterRenderTargetManager;
 ////////////////////////////////////////////////////////////////////////////////////////////////////////////
 /// FDisplayClusterRenderTargetManager
 ////////////////////////////////////////////////////////////////////////////////////////////////////////////
-FDisplayClusterRenderTargetManager::FDisplayClusterRenderTargetManager()
+FDisplayClusterRenderTargetManager::FDisplayClusterRenderTargetManager(FDisplayClusterViewportManagerProxy* InViewportManagerProxy)
 {
-	ResourcesPool = MakeUnique<FDisplayClusterRenderTargetResourcesPool>();
+	ResourcesPool = MakeUnique<FDisplayClusterRenderTargetResourcesPool>(InViewportManagerProxy);
 }
 
 FDisplayClusterRenderTargetManager::~FDisplayClusterRenderTargetManager()
diff --git a/Engine/Plugins/Runtime/nDisplay/Source/DisplayCluster/Private/Render/Viewport/RenderTarget/DisplayClusterRenderTargetManager.h b/Engine/Plugins/Runtime/nDisplay/Source/DisplayCluster/Private/Render/Viewport/RenderTarget/DisplayClusterRenderTargetManager.h
--- a/Engine/Plugins/Runtime/nDisplay/Source/DisplayCluster/Private/Render/Viewport/RenderTarget/DisplayClusterRenderTargetManager.h
+++ b/Engine/Plugins/Runtime/nDisplay/Source/DisplayCluster/Private/Render/Viewport/RenderTarget/DisplayClusterRenderTargetManager.h
@@ -8,13 +8,14 @@ class FViewport;
 class FDisplayClusterViewport;
 class FDisplayClusterRenderFrame;
 class FDisplayClusterRenderTargetResourcesPool;
+class FDisplayClusterViewportManagerProxy;
 
 struct FDisplayClusterRenderFrameSettings;
 
 class FDisplayClusterRenderTargetManager
 {
 public:
-	FDisplayClusterRenderTargetManager();
+	FDisplayClusterRenderTargetManager(FDisplayClusterViewportManagerProxy* InViewportManagerProxy);
 	~FDisplayClusterRenderTargetManager();
 
 public:
diff --git a/Engine/Plugins/Runtime/nDisplay/Source/DisplayCluster/Private/Render/Viewport/RenderTarget/DisplayClusterRenderTargetResourcesPool.cpp b/Engine/Plugins/Runtime/nDisplay/Source/DisplayCluster/Private/Render/Viewport/RenderTarget/DisplayClusterRenderTargetResourcesPool.cpp
--- a/Engine/Plugins/Runtime/nDisplay/Source/DisplayCluster/Private/Render/Viewport/RenderTarget/DisplayClusterRenderTargetResourcesPool.cpp
+++ b/Engine/Plugins/Runtime/nDisplay/Source/DisplayCluster/Private/Render/Viewport/RenderTarget/DisplayClusterRenderTargetResourcesPool.cpp
@@ -3,6 +3,7 @@
 #include "DisplayClusterRenderTargetResourcesPool.h"
 #include "Render/Viewport/RenderFrame/DisplayClusterRenderFrameSettings.h"
 #include "Render/Viewport/DisplayClusterViewportHelpers.h"
+#include "Render/Viewport/DisplayClusterViewportManagerProxy.h"
 
 #include "RHICommandList.h"
 #include "Engine/RendererSettings.h"
@@ -25,17 +26,19 @@ namespace DisplayClusterRenderTargetResourcesPool
 		});
 	}
 
-	static void ImplReleaseViewportResourcesRHI(const TArrayView<FDisplayClusterViewportResource*>& InResources)
+	static void ImplReleaseViewportResourcesRHI(FDisplayClusterViewportManagerProxy* InViewportManagerProxy, const TArrayView<FDisplayClusterViewportResource*>& InResources)
 	{
 		TArray<FDisplayClusterViewportResource*> ResourcesForRenderThread(InResources);
 
 		ENQUEUE_RENDER_COMMAND(DisplayCluster_ReleaseViewportResourcesRHI)(
-			[ReleasedResources = std::move(ResourcesForRenderThread)](FRHICommandListImmediate& RHICmdList)
+			[InViewportManagerProxy, ReleasedResources = std::move(ResourcesForRenderThread)](FRHICommandListImmediate& RHICmdList)
 		{
-			for (FDisplayClusterViewportResource* ResourceIt : ReleasedResources)
+			for (FDisplayClusterViewportResource* DeletedResourcePtrIt : ReleasedResources)
 			{
-				ResourceIt->ReleaseResource();
-				delete ResourceIt;
+				if (InViewportManagerProxy)
+				{
+					InViewportManagerProxy->DeleteResource_RenderThread(DeletedResourcePtrIt);
+				}
 			}
 		});
 	}
@@ -75,9 +78,9 @@ using namespace DisplayClusterRenderTargetResourcesPool;
 ////////////////////////////////////////////////////////////////////////////////////////////////////////////
 /// FDisplayClusterRenderTargetResourcesPool
 ////////////////////////////////////////////////////////////////////////////////////////////////////////////
-FDisplayClusterRenderTargetResourcesPool::FDisplayClusterRenderTargetResourcesPool()
-{
-}
+FDisplayClusterRenderTargetResourcesPool::FDisplayClusterRenderTargetResourcesPool(FDisplayClusterViewportManagerProxy* InViewportManagerProxy)
+	: ViewportManagerProxy(InViewportManagerProxy)
+{ }
 
 FDisplayClusterRenderTargetResourcesPool::~FDisplayClusterRenderTargetResourcesPool()
 {
@@ -160,7 +163,7 @@ void FDisplayClusterRenderTargetResourcesPool::ImplFinishReallocateResources(TAr
 		}
 
 		// Send to released resources to rendering thread
-		ImplReleaseViewportResourcesRHI(TArrayView<FDisplayClusterViewportResource*>((FDisplayClusterViewportResource**)(UnusedResources.GetData()), UnusedResources.Num()));
+		ImplReleaseViewportResourcesRHI(ViewportManagerProxy, TArrayView<FDisplayClusterViewportResource*>((FDisplayClusterViewportResource**)(UnusedResources.GetData()), UnusedResources.Num()));
 	}
 }
 
@@ -168,7 +171,7 @@ template <typename TViewportResourceType>
 void FDisplayClusterRenderTargetResourcesPool::ImplReleaseResources(TArray<TViewportResourceType*>& InOutViewportResources)
 {
 	// Send all resources to rendering thread
-	ImplReleaseViewportResourcesRHI(TArray<FDisplayClusterViewportResource*>((FDisplayClusterViewportResource**)(InOutViewportResources.GetData()), InOutViewportResources.Num()));
+	ImplReleaseViewportResourcesRHI(ViewportManagerProxy, TArray<FDisplayClusterViewportResource*>((FDisplayClusterViewportResource**)(InOutViewportResources.GetData()), InOutViewportResources.Num()));
 
 	// Reset refs on game thread too
 	InOutViewportResources.Empty();
diff --git a/Engine/Plugins/Runtime/nDisplay/Source/DisplayCluster/Private/Render/Viewport/RenderTarget/DisplayClusterRenderTargetResourcesPool.h b/Engine/Plugins/Runtime/nDisplay/Source/DisplayCluster/Private/Render/Viewport/RenderTarget/DisplayClusterRenderTargetResourcesPool.h
--- a/Engine/Plugins/Runtime/nDisplay/Source/DisplayCluster/Private/Render/Viewport/RenderTarget/DisplayClusterRenderTargetResourcesPool.h
+++ b/Engine/Plugins/Runtime/nDisplay/Source/DisplayCluster/Private/Render/Viewport/RenderTarget/DisplayClusterRenderTargetResourcesPool.h
@@ -5,12 +5,13 @@
 #include "CoreMinimal.h"
 #include "DisplayClusterRenderTargetResource.h"
 
+class FDisplayClusterViewportManagerProxy;
 struct FDisplayClusterRenderFrameSettings;
 
 class FDisplayClusterRenderTargetResourcesPool
 {
 public:
-	FDisplayClusterRenderTargetResourcesPool();
+	FDisplayClusterRenderTargetResourcesPool(FDisplayClusterViewportManagerProxy* InViewportManagerProxy);
 	~FDisplayClusterRenderTargetResourcesPool();
 
 public:
@@ -40,4 +41,6 @@ private:
 	// Viewport render resources
 	TArray<FDisplayClusterViewportRenderTargetResource*> RenderTargetResources;
 	TArray<FDisplayClusterViewportTextureResource*>      TextureResources;
+
+	FDisplayClusterViewportManagerProxy* ViewportManagerProxy = nullptr;
 };
