[nDisplay] Split nDisplay scene preview rendering into its own module

diff --git a/Engine/Plugins/Runtime/nDisplay/Source/DisplayCluster/Private/Blueprints/DisplayClusterBlueprintLib.cpp b/Engine/Plugins/Runtime/nDisplay/Source/DisplayCluster/Private/Blueprints/DisplayClusterBlueprintLib.cpp
--- a/Engine/Plugins/Runtime/nDisplay/Source/DisplayCluster/Private/Blueprints/DisplayClusterBlueprintLib.cpp
+++ b/Engine/Plugins/Runtime/nDisplay/Source/DisplayCluster/Private/Blueprints/DisplayClusterBlueprintLib.cpp
@@ -69,4 +69,52 @@ ADisplayClusterLightCardActor* UDisplayClusterBlueprintLib::CreateLightCard(ADis
 	return NewActor;
 }
 
+void UDisplayClusterBlueprintLib::FindLightCardsForRootActor(ADisplayClusterRootActor* RootActor, TSet<ADisplayClusterLightCardActor*>& OutLightCards)
+{
+	if (!RootActor)
+	{
+		return;
+	}
+
+	FDisplayClusterConfigurationICVFX_VisibilityList& RootActorLightCards = RootActor->GetConfigData()->StageSettings.Lightcard.ShowOnlyList;
+
+	for (const TSoftObjectPtr<AActor>& LightCardActor : RootActorLightCards.Actors)
+	{
+		if (!LightCardActor.IsValid() || !LightCardActor->IsA<ADisplayClusterLightCardActor>())
+		{
+			continue;
+		}
+
+		OutLightCards.Add(Cast<ADisplayClusterLightCardActor>(LightCardActor.Get()));
+	}
+
+	// If there are any layers that are specified as light card layers, iterate over all actors in the world and 
+	// add any that are members of any of the light card layers to the list. Only add an actor once, even if it is
+	// in multiple layers
+	if (RootActorLightCards.ActorLayers.IsEmpty())
+	{
+		return;
+	}
+
+	if (UWorld* World = RootActor->GetWorld())
+	{
+		for (TActorIterator<ADisplayClusterLightCardActor> ActorIt(World); ActorIt; ++ActorIt)
+		{
+			if (!IsValid(*ActorIt))
+			{
+				continue;
+			}
+
+			for (const FActorLayer& ActorLayer : RootActorLightCards.ActorLayers)
+			{
+				if (ActorIt->Layers.Contains(ActorLayer.Name))
+				{
+					OutLightCards.Add(*ActorIt);
+					break;
+				}
+			}
+		}
+	}
+}
+
 #undef LOCTEXT_NAMESPACE
\ No newline at end of file
diff --git a/Engine/Plugins/Runtime/nDisplay/Source/DisplayCluster/Public/Blueprints/DisplayClusterBlueprintLib.h b/Engine/Plugins/Runtime/nDisplay/Source/DisplayCluster/Public/Blueprints/DisplayClusterBlueprintLib.h
--- a/Engine/Plugins/Runtime/nDisplay/Source/DisplayCluster/Public/Blueprints/DisplayClusterBlueprintLib.h
+++ b/Engine/Plugins/Runtime/nDisplay/Source/DisplayCluster/Public/Blueprints/DisplayClusterBlueprintLib.h
@@ -27,4 +27,8 @@ public:
 	/** Create a new light card parented to the given nDisplay root actor. */
 	UFUNCTION(BlueprintPure, meta = (DisplayName = "DisplayCluster Module API"), Category = "NDisplay")
 	static ADisplayClusterLightCardActor* CreateLightCard(ADisplayClusterRootActor* RootActor);
+
+	/** Gets a list of all light card actors on the level linked to the specified root actor. */
+	UFUNCTION(BlueprintPure, meta = (DisplayName = "DisplayCluster Module API"), Category = "NDisplay")
+	static DISPLAYCLUSTER_API void FindLightCardsForRootActor(ADisplayClusterRootActor* RootActor, TSet<ADisplayClusterLightCardActor*>& OutLightCards);
 };
diff --git a/Engine/Plugins/Runtime/nDisplay/Source/DisplayClusterLightCardEditor/DisplayClusterLightCardEditor.Build.cs b/Engine/Plugins/Runtime/nDisplay/Source/DisplayClusterLightCardEditor/DisplayClusterLightCardEditor.Build.cs
--- a/Engine/Plugins/Runtime/nDisplay/Source/DisplayClusterLightCardEditor/DisplayClusterLightCardEditor.Build.cs
+++ b/Engine/Plugins/Runtime/nDisplay/Source/DisplayClusterLightCardEditor/DisplayClusterLightCardEditor.Build.cs
@@ -19,6 +19,7 @@ public class DisplayClusterLightCardEditor : ModuleRules
 				"DisplayClusterConfiguration",
 				"DisplayClusterOperator",
 				"DisplayClusterLightCardEditorShaders",
+				"DisplayClusterScenePreview",
 				"EditorStyle",
 				"Engine",
 				"InputCore",
diff --git a/Engine/Plugins/Runtime/nDisplay/Source/DisplayClusterLightCardEditor/Private/Viewport/DisplayClusterLightCardEditorViewportClient.cpp b/Engine/Plugins/Runtime/nDisplay/Source/DisplayClusterLightCardEditor/Private/Viewport/DisplayClusterLightCardEditorViewportClient.cpp
--- a/Engine/Plugins/Runtime/nDisplay/Source/DisplayClusterLightCardEditor/Private/Viewport/DisplayClusterLightCardEditorViewportClient.cpp
+++ b/Engine/Plugins/Runtime/nDisplay/Source/DisplayClusterLightCardEditor/Private/Viewport/DisplayClusterLightCardEditorViewportClient.cpp
@@ -5,6 +5,7 @@
 #include "DisplayClusterLightcardEditorViewport.h"
 #include "DisplayClusterLightCardEditorWidget.h"
 
+#include "Blueprints/DisplayClusterBlueprintLib.h"
 #include "Components/DisplayClusterPreviewComponent.h"
 #include "Components/DisplayClusterCameraComponent.h"
 #include "Components/DisplayClusterScreenComponent.h"
@@ -30,6 +31,7 @@
 #include "EngineModule.h"
 #include "Engine/Canvas.h"
 #include "Engine/TextureRenderTarget2D.h"
+#include "IDisplayClusterScenePreview.h"
 #include "ImageUtils.h"
 #include "Kismet/KismetMathLibrary.h"
 #include "Kismet/GameplayStatics.h"
@@ -321,10 +323,17 @@ FDisplayClusterLightCardEditorViewportClient::FDisplayClusterLightCardEditorView
 	
 	LightCardEditorPtr = InEditorViewportWidget.Pin()->GetLightCardEditor();
 	check(LightCardEditorPtr.IsValid())
-	
-	MeshProjectionRenderer = MakeShared<FDisplayClusterMeshProjectionRenderer>();
-	MeshProjectionRenderer->ActorSelectedDelegate = FDisplayClusterMeshProjectionRenderer::FSelection::CreateRaw(this, &FDisplayClusterLightCardEditorViewportClient::IsLightCardSelected);
-	MeshProjectionRenderer->RenderSimpleElementsDelegate = FDisplayClusterMeshProjectionRenderer::FSimpleElementPass::CreateRaw(this, &FDisplayClusterLightCardEditorViewportClient::Draw);
+
+	IDisplayClusterScenePreview& PreviewModule = IDisplayClusterScenePreview::Get();
+	PreviewRendererId = PreviewModule.CreateRenderer();
+	PreviewModule.SetRendererActorSelectedDelegate(
+		PreviewRendererId,
+		FDisplayClusterMeshProjectionRenderer::FSelection::CreateRaw(this, &FDisplayClusterLightCardEditorViewportClient::IsLightCardSelected)
+	);
+	PreviewModule.SetRendererRenderSimpleElementsDelegate(
+		PreviewRendererId,
+		FDisplayClusterMeshProjectionRenderer::FSimpleElementPass::CreateRaw(this, &FDisplayClusterLightCardEditorViewportClient::Draw)
+	);
 
 	InputMode = EInputMode::Idle;
 	DragWidgetOffset = FVector::ZeroVector;
@@ -361,6 +370,8 @@ FDisplayClusterLightCardEditorViewportClient::FDisplayClusterLightCardEditorView
 
 FDisplayClusterLightCardEditorViewportClient::~FDisplayClusterLightCardEditorViewportClient()
 {
+	IDisplayClusterScenePreview::Get().DestroyRenderer(PreviewRendererId);
+
 	EndTransaction();
 	if (RootActorLevelInstance.IsValid())
 	{
@@ -557,10 +568,13 @@ void FDisplayClusterLightCardEditorViewportClient::Draw(FViewport* InViewport, F
 
 	Canvas->Clear(FLinearColor::Black);
 
-	FSceneViewInitOptions ViewInitOptions;
-	GetSceneViewInitOptions(ViewInitOptions);
+	FDisplayClusterScenePreviewRenderSettings RenderSettings;
+	RenderSettings.RenderType = FDisplayClusterScenePreviewRenderSettings::ERenderType::Color;
+	RenderSettings.EngineShowFlags = EngineShowFlags;
+	RenderSettings.ProjectionType = ProjectionMode;
+	GetSceneViewInitOptions(RenderSettings.ViewInitOptions);
 
-	MeshProjectionRenderer->Render(Canvas, GetScene(), ViewInitOptions, EngineShowFlags, ProjectionMode);
+	IDisplayClusterScenePreview::Get().Render(PreviewRendererId, RenderSettings, *Canvas);
 
 	if (View)
 	{
@@ -1305,22 +1319,19 @@ void FDisplayClusterLightCardEditorViewportClient::UpdatePreviewActor(ADisplayCl
 			}
 
 			// Filter out any primitives hidden in game except screen components
-			MeshProjectionRenderer->AddActor(RootActorProxy.Get(), [](const UPrimitiveComponent* PrimitiveComponent)
-			{
-				return !PrimitiveComponent->bHiddenInGame || PrimitiveComponent->IsA<UDisplayClusterScreenComponent>();
-			});
+			IDisplayClusterScenePreview::Get().SetRendererRootActor(PreviewRendererId, RootActorProxy.Get());
 
 			if (ProxyType == EDisplayClusterLightCardEditorProxyType::All ||
 				ProxyType == EDisplayClusterLightCardEditorProxyType::LightCards)
 			{
-				TArray<TWeakObjectPtr<ADisplayClusterLightCardActor>> LightCards;
-				FindLightCardsForRootActor(RootActor, LightCards);
+				TSet<ADisplayClusterLightCardActor*> LightCards;
+				UDisplayClusterBlueprintLib::FindLightCardsForRootActor(RootActor, LightCards);
 
 				SelectLightCard(nullptr);
 				
-				for (const TWeakObjectPtr<ADisplayClusterLightCardActor>& LightCard : LightCards)
+				for (ADisplayClusterLightCardActor* LightCard : LightCards)
 				{
-					FObjectDuplicationParameters DupeActorParameters(LightCard.Get(), PreviewWorld->GetCurrentLevel());
+					FObjectDuplicationParameters DupeActorParameters(LightCard, PreviewWorld->GetCurrentLevel());
 					DupeActorParameters.FlagMask = RF_AllFlags & ~(RF_ArchetypeObject | RF_Transactional);
 					DupeActorParameters.PortFlags = PPF_DuplicateVerbatim;
 				
@@ -1350,7 +1361,7 @@ void FDisplayClusterLightCardEditorViewportClient::UpdatePreviewActor(ADisplayCl
 						}
 					}
 
-					LightCardProxies.Add(FLightCardProxy(LightCard.Get(), LightCardProxy));
+					LightCardProxies.Add(FLightCardProxy(LightCard, LightCardProxy));
 				}
 
 				// Update the selected light card proxies to match the currently selected light cards in the light card list
@@ -1365,7 +1376,7 @@ void FDisplayClusterLightCardEditorViewportClient::UpdatePreviewActor(ADisplayCl
 
 			for (const FLightCardProxy& LightCardProxy : LightCardProxies)
 			{
-				MeshProjectionRenderer->AddActor(LightCardProxy.Proxy.Get());
+				IDisplayClusterScenePreview::Get().AddActorToRenderer(PreviewRendererId, LightCardProxy.Proxy.Get());
 			}
 
 			Finalize();
@@ -1398,7 +1409,7 @@ void FDisplayClusterLightCardEditorViewportClient::UpdateProxyTransforms()
 void FDisplayClusterLightCardEditorViewportClient::DestroyProxies(
 EDisplayClusterLightCardEditorProxyType ProxyType)
 {
-	MeshProjectionRenderer->ClearScene();
+	IDisplayClusterScenePreview::Get().ClearRendererScene(PreviewRendererId);
 
 	UWorld* PreviewWorld = PreviewScene->GetWorld();
 	check(PreviewWorld);
@@ -1864,46 +1875,6 @@ USceneComponent* FDisplayClusterLightCardEditorViewportClient::FindProjectionOri
 	return InRootActor->GetRootComponent();
 }
 
-void FDisplayClusterLightCardEditorViewportClient::FindLightCardsForRootActor(ADisplayClusterRootActor* RootActor, TArray<TWeakObjectPtr<ADisplayClusterLightCardActor>>& OutLightCards)
-{
-	if (RootActor)
-	{
-		FDisplayClusterConfigurationICVFX_VisibilityList& RootActorLightCards = RootActor->GetConfigData()->StageSettings.Lightcard.ShowOnlyList;
-
-		for (const TSoftObjectPtr<AActor>& LightCardActor : RootActorLightCards.Actors)
-		{
-			if (LightCardActor.IsValid() && LightCardActor->IsA<ADisplayClusterLightCardActor>())
-			{
-				OutLightCards.Add(Cast<ADisplayClusterLightCardActor>(LightCardActor.Get()));
-			}
-		}
-
-		// If there are any layers that are specified as light card layers, iterate over all actors in the world and 
-		// add any that are members of any of the light card layers to the list. Only add an actor once, even if it is
-		// in multiple layers
-		if (RootActorLightCards.ActorLayers.Num())
-		{
-			if (UWorld* World = RootActor->GetWorld())
-			{
-				for (const TWeakObjectPtr<AActor> WeakActor : FActorRange(World))
-				{
-					if (WeakActor.IsValid() && WeakActor->IsA<ADisplayClusterLightCardActor>())
-					{
-						for (const FActorLayer& ActorLayer : RootActorLightCards.ActorLayers)
-						{
-							if (WeakActor->Layers.Contains(ActorLayer.Name))
-							{
-								OutLightCards.Add(Cast<ADisplayClusterLightCardActor>(WeakActor));
-								break;
-							}
-						}
-					}
-				}
-			}
-		}
-	}
-}
-
 bool FDisplayClusterLightCardEditorViewportClient::IsLightCardSelected(const AActor* Actor)
 {
 	return SelectedLightCards.Contains(Actor);
@@ -2620,11 +2591,6 @@ bool FDisplayClusterLightCardEditorViewportClient::CalcEditorWidgetTransform(FTr
 
 void FDisplayClusterLightCardEditorViewportClient::RenderNormalMap(FNormalMap& NormalMap, const FVector& NormalMapDirection)
 {
-	FSceneViewInitOptions ViewInitOptions;
-	GetNormalMapSceneViewInitOptions(FIntPoint(FNormalMap::NormalMapSize), FNormalMap::NormalMapFOV, NormalMapDirection, ViewInitOptions);
-
-	NormalMap.Init(ViewInitOptions);
-
 	// Only render primitive components from the stage actor for the normal map
 	FDisplayClusterMeshProjectionPrimitiveFilter PrimitiveFilter;
 	PrimitiveFilter.PrimitiveFilterDelegate = FDisplayClusterMeshProjectionPrimitiveFilter::FPrimitiveFilter::CreateLambda([this](const UPrimitiveComponent* PrimitiveComponent)
@@ -2632,11 +2598,22 @@ void FDisplayClusterLightCardEditorViewportClient::RenderNormalMap(FNormalMap& N
 		return PrimitiveComponent->GetOwner() == RootActorProxy;
 	});
 
+	FDisplayClusterScenePreviewRenderSettings RenderSettings;
+	RenderSettings.RenderType = FDisplayClusterScenePreviewRenderSettings::ERenderType::Normals;
+	RenderSettings.EngineShowFlags = EngineShowFlags;
+	RenderSettings.ProjectionType = EDisplayClusterMeshProjectionType::Azimuthal;
+	RenderSettings.PrimitiveFilter = PrimitiveFilter;
+	
+	FSceneViewInitOptions ViewInitOptions;
+	GetNormalMapSceneViewInitOptions(FIntPoint(FNormalMap::NormalMapSize), FNormalMap::NormalMapFOV, NormalMapDirection, RenderSettings.ViewInitOptions);
+
+	NormalMap.Init(RenderSettings.ViewInitOptions);
+
 	FCanvas Canvas(&NormalMap, nullptr, GetWorld(), GetScene()->GetFeatureLevel(), FCanvas::CDM_DeferDrawing, 1.0f);
 	{
 		Canvas.Clear(FLinearColor::Black);
 
-		MeshProjectionRenderer->RenderNormals(&Canvas, GetScene(), ViewInitOptions, EngineShowFlags, EDisplayClusterMeshProjectionType::Azimuthal, &PrimitiveFilter);
+		IDisplayClusterScenePreview::Get().Render(PreviewRendererId, RenderSettings, Canvas);
 	}
 	Canvas.Flush_GameThread();
 
diff --git a/Engine/Plugins/Runtime/nDisplay/Source/DisplayClusterLightCardEditor/Private/Viewport/DisplayClusterLightCardEditorViewportClient.h b/Engine/Plugins/Runtime/nDisplay/Source/DisplayClusterLightCardEditor/Private/Viewport/DisplayClusterLightCardEditorViewportClient.h
--- a/Engine/Plugins/Runtime/nDisplay/Source/DisplayClusterLightCardEditor/Private/Viewport/DisplayClusterLightCardEditorViewportClient.h
+++ b/Engine/Plugins/Runtime/nDisplay/Source/DisplayClusterLightCardEditor/Private/Viewport/DisplayClusterLightCardEditorViewportClient.h
@@ -226,9 +226,6 @@ private:
 	/** Finds a suitable primitive component on the stage actor to use as a projection origin */
 	USceneComponent* FindProjectionOriginComponent(const ADisplayClusterRootActor* InRootActor) const;
 
-	/** Gets a list of all light card actors on the level linked to the specified root actor */
-	void FindLightCardsForRootActor(ADisplayClusterRootActor* RootActor, TArray<TWeakObjectPtr<ADisplayClusterLightCardActor>>& OutLightCards);
-
 	/** Callback to check if an light card actor is among the list of selected light card actors */
 	bool IsLightCardSelected(const AActor* Actor);
 
@@ -336,8 +333,8 @@ private:
 	TArray<FLightCardProxy> LightCardProxies;
 	TArray<TWeakObjectPtr<ADisplayClusterLightCardActor>> SelectedLightCards;
 	
-	/** The renderer for the viewport, which can render the meshes with a variety of projection types */
-	TSharedPtr<FDisplayClusterMeshProjectionRenderer> MeshProjectionRenderer;
+	/** The index of the scene preview renderer returned from IDisplayClusterScenePreview */
+	int32 PreviewRendererId = -1;
 	
 	/** The LC editor widget used to manipulate light cards */
 	TSharedPtr<FDisplayClusterLightCardEditorWidget> EditorWidget;
diff --git a/Engine/Plugins/Runtime/nDisplay/Source/DisplayClusterLightCardEditorShaders/Private/DisplayClusterMeshProjectionRenderer.cpp b/Engine/Plugins/Runtime/nDisplay/Source/DisplayClusterLightCardEditorShaders/Private/DisplayClusterMeshProjectionRenderer.cpp
--- a/Engine/Plugins/Runtime/nDisplay/Source/DisplayClusterLightCardEditorShaders/Private/DisplayClusterMeshProjectionRenderer.cpp
+++ b/Engine/Plugins/Runtime/nDisplay/Source/DisplayClusterLightCardEditorShaders/Private/DisplayClusterMeshProjectionRenderer.cpp
@@ -916,13 +916,17 @@ void FDisplayClusterMeshProjectionRenderer::Render(FCanvas* Canvas, FSceneInterf
 				AddTranslucencyRenderPass(GraphBuilder, View, ProjectionType, ColorRenderTargetBinding, DepthStencilBinding);
 
 #if WITH_EDITOR
-				const FRDGTextureDesc SelectionDepthDesc = FRDGTextureDesc::Create2D(OutputTexture->Desc.Extent, PF_DepthStencil, FClearValueBinding::DepthFar, TexCreate_DepthStencilTargetable | TexCreate_ShaderResource);
-				FRDGTextureRef SelectionDepthTexture = GraphBuilder.CreateTexture(SelectionDepthDesc, TEXT("DisplayClusterMeshProjection.SelectionDepthTexture"));
-				FDepthStencilBinding SelectionDepthStencilBinding(SelectionDepthTexture, ERenderTargetLoadAction::EClear, ERenderTargetLoadAction::EClear, FExclusiveDepthStencil::DepthWrite_StencilWrite);
+				if (EngineShowFlags.SelectionOutline)
+				{
+					const FRDGTextureDesc SelectionDepthDesc = FRDGTextureDesc::Create2D(OutputTexture->Desc.Extent, PF_DepthStencil, FClearValueBinding::DepthFar, TexCreate_DepthStencilTargetable | TexCreate_ShaderResource);
+					FRDGTextureRef SelectionDepthTexture = GraphBuilder.CreateTexture(SelectionDepthDesc, TEXT("DisplayClusterMeshProjection.SelectionDepthTexture"));
+					FDepthStencilBinding SelectionDepthStencilBinding(SelectionDepthTexture, ERenderTargetLoadAction::EClear, ERenderTargetLoadAction::EClear, FExclusiveDepthStencil::DepthWrite_StencilWrite);
 
-				AddSelectionDepthRenderPass(GraphBuilder, View, ProjectionType, SelectionDepthStencilBinding);
-				AddSelectionOutlineScreenPass(GraphBuilder, View, OutputRenderTargetBinding, ColorTexture, DepthTexture, SelectionDepthTexture);
-#else
+					AddSelectionDepthRenderPass(GraphBuilder, View, ProjectionType, SelectionDepthStencilBinding);
+					AddSelectionOutlineScreenPass(GraphBuilder, View, OutputRenderTargetBinding, ColorTexture, DepthTexture, SelectionDepthTexture);
+				}
+				else
+#endif
 				// Copy the scene color to the output render target
 				{
 					FCopyRectPS::FParameters* ScreenPassParameters = GraphBuilder.AllocParameters<FCopyRectPS::FParameters>();
@@ -942,20 +946,19 @@ void FDisplayClusterMeshProjectionRenderer::Render(FCanvas* Canvas, FSceneInterf
 						ScreenPassParameters,
 						ERDGPassFlags::Raster,
 						[View, ScreenPassVS, CopyPixelShader, RegionViewport, ScreenPassParameters, DefaultBlendState](FRHICommandList& RHICmdList)
-					{
-						DrawScreenPass(
-							RHICmdList,
-							*View,
-							RegionViewport,
-							RegionViewport,
-							FScreenPassPipelineState(ScreenPassVS, CopyPixelShader, DefaultBlendState),
-							[&](FRHICommandList&)
 						{
-							SetShaderParameters(RHICmdList, CopyPixelShader, CopyPixelShader.GetPixelShader(), *ScreenPassParameters);
+							DrawScreenPass(
+								RHICmdList,
+								*View,
+								RegionViewport,
+								RegionViewport,
+								FScreenPassPipelineState(ScreenPassVS, CopyPixelShader, DefaultBlendState),
+								[&](FRHICommandList&)
+								{
+									SetShaderParameters(RHICmdList, CopyPixelShader, CopyPixelShader.GetPixelShader(), *ScreenPassParameters);
 						});
 					});
 				}
-#endif
 			}
 			
 			FMeshProjectionElementCollector ElementCollector(HitProxyConsumer);
diff --git a/Engine/Plugins/Runtime/nDisplay/Source/DisplayClusterScenePreview/DisplayClusterScenePreview.Build.cs b/Engine/Plugins/Runtime/nDisplay/Source/DisplayClusterScenePreview/DisplayClusterScenePreview.Build.cs
new file mode 100644
--- /dev/null
+++ b/Engine/Plugins/Runtime/nDisplay/Source/DisplayClusterScenePreview/DisplayClusterScenePreview.Build.cs
@@ -0,0 +1,26 @@
+// Copyright Epic Games, Inc. All Rights Reserved.
+
+using UnrealBuildTool;
+using System.IO;
+
+public class DisplayClusterScenePreview : ModuleRules
+{
+	public DisplayClusterScenePreview(ReadOnlyTargetRules ROTargetRules) : base(ROTargetRules)
+	{
+		PublicDependencyModuleNames.AddRange(
+			new string[] {
+				"DisplayClusterLightCardEditorShaders",
+			});
+
+		PrivateDependencyModuleNames.AddRange(
+			new string[] {
+				"DisplayCluster",
+				"DisplayClusterConfiguration",
+
+				"Core",
+				"CoreUObject",
+				"Engine",
+				"UnrealEd"
+			});
+	}
+}
diff --git a/Engine/Plugins/Runtime/nDisplay/Source/DisplayClusterScenePreview/Private/DisplayClusterScenePreviewModule.cpp b/Engine/Plugins/Runtime/nDisplay/Source/DisplayClusterScenePreview/Private/DisplayClusterScenePreviewModule.cpp
new file mode 100644
--- /dev/null
+++ b/Engine/Plugins/Runtime/nDisplay/Source/DisplayClusterScenePreview/Private/DisplayClusterScenePreviewModule.cpp
@@ -0,0 +1,430 @@
+// Copyright Epic Games, Inc. All Rights Reserved.
+
+#include "DisplayClusterScenePreviewModule.h"
+
+#include "Blueprints/DisplayClusterBlueprintLib.h"
+#include "CanvasTypes.h"
+#include "Components/DisplayClusterScreenComponent.h"
+#include "DisplayClusterLightCardActor.h"
+#include "DisplayClusterRootActor.h"
+#include "Engine/TextureRenderTarget2D.h"
+
+#define LOCTEXT_NAMESPACE "DisplayClusterScenePreview"
+
+static TAutoConsoleVariable<float> CVarDisplayClusterScenePreviewRenderTickDelay(
+	TEXT("nDisplay.ScenePreview.RenderTickDelay"),
+	0.1f,
+	TEXT("The number of seconds to wait between processing queued renders.")
+);
+
+void FDisplayClusterScenePreviewModule::StartupModule()
+{
+}
+
+void FDisplayClusterScenePreviewModule::ShutdownModule()
+{
+	FTSTicker::GetCoreTicker().RemoveTicker(RenderTickerHandle);
+	RenderTickerHandle.Reset();
+}
+
+int32 FDisplayClusterScenePreviewModule::CreateRenderer()
+{
+	FRendererConfig& Config = RendererConfigs.Add(NextRendererId);
+
+	Config.Renderer = MakeShared<FDisplayClusterMeshProjectionRenderer>();
+
+	return NextRendererId++;
+}
+
+bool FDisplayClusterScenePreviewModule::DestroyRenderer(int32 RendererId)
+{
+	if (FRendererConfig* Config = RendererConfigs.Find(RendererId))
+	{
+		RegisterRootActorEvents(Config->RootActor.Get(), false);
+		RendererConfigs.Remove(RendererId);
+
+		RegisterOrUnregisterGlobalActorEvents();
+		return true;
+	}
+
+	return false;
+}
+
+bool FDisplayClusterScenePreviewModule::SetRendererRootActorPath(int32 RendererId, const FString& ActorPath, bool bAutoUpdateLightcards)
+{
+	if (FRendererConfig* Config = RendererConfigs.Find(RendererId))
+	{
+		Config->RootActorPath = ActorPath;
+
+		ADisplayClusterRootActor* RootActor = FindObject<ADisplayClusterRootActor>(nullptr, *ActorPath);
+		InternalSetRendererRootActor(*Config, RootActor, bAutoUpdateLightcards);
+
+		return true;
+	}
+
+	return false;
+}
+
+bool FDisplayClusterScenePreviewModule::SetRendererRootActor(int32 RendererId, ADisplayClusterRootActor* Actor, bool bAutoUpdateLightcards)
+{
+	if (FRendererConfig* Config = RendererConfigs.Find(RendererId))
+	{
+		Config->RootActorPath.Empty();
+		InternalSetRendererRootActor(*Config, Actor, bAutoUpdateLightcards);
+
+		return true;
+	}
+
+	return false;
+}
+
+ADisplayClusterRootActor* FDisplayClusterScenePreviewModule::GetRendererRootActor(int32 RendererId)
+{
+	if (FRendererConfig* Config = RendererConfigs.Find(RendererId))
+	{
+		return InternalGetRendererRootActor(*Config);
+	}
+
+	return nullptr;
+}
+
+bool FDisplayClusterScenePreviewModule::AddActorToRenderer(int32 RendererId, AActor* Actor)
+{
+	if (FRendererConfig* Config = RendererConfigs.Find(RendererId))
+	{
+		Config->Renderer->AddActor(Actor);
+		return true;
+	}
+
+	return false;
+}
+
+bool FDisplayClusterScenePreviewModule::ClearRendererScene(int32 RendererId)
+{
+	if (FRendererConfig* Config = RendererConfigs.Find(RendererId))
+	{
+		Config->Renderer->ClearScene();
+		return true;
+	}
+
+	return false;
+}
+
+bool FDisplayClusterScenePreviewModule::SetRendererActorSelectedDelegate(int32 RendererId, FDisplayClusterMeshProjectionRenderer::FSelection ActorSelectedDelegate)
+{
+	if (FRendererConfig* Config = RendererConfigs.Find(RendererId))
+	{
+		Config->Renderer->ActorSelectedDelegate = ActorSelectedDelegate;
+		return true;
+	}
+
+	return false;
+}
+
+bool FDisplayClusterScenePreviewModule::SetRendererRenderSimpleElementsDelegate(int32 RendererId, FDisplayClusterMeshProjectionRenderer::FSimpleElementPass RenderSimpleElementsDelegate)
+{
+	if (FRendererConfig* Config = RendererConfigs.Find(RendererId))
+	{
+		Config->Renderer->RenderSimpleElementsDelegate = RenderSimpleElementsDelegate;
+		return true;
+	}
+
+	return false;
+}
+
+bool FDisplayClusterScenePreviewModule::Render(int32 RendererId, FDisplayClusterScenePreviewRenderSettings& RenderSettings, FCanvas& Canvas)
+{
+	if (FRendererConfig* Config = RendererConfigs.Find(RendererId))
+	{
+		return InternalRenderImmediate(*Config, RenderSettings, Canvas);
+	}
+
+	return false;
+}
+
+bool FDisplayClusterScenePreviewModule::RenderQueued(int32 RendererId, FDisplayClusterScenePreviewRenderSettings& RenderSettings, const FIntPoint& Size, FRenderResultDelegate ResultDelegate)
+{
+	if (FRendererConfig* Config = RendererConfigs.Find(RendererId))
+	{
+		RenderQueue.Enqueue(FPreviewRenderJob(RendererId, RenderSettings, Size, ResultDelegate));
+
+		if (!RenderTickerHandle.IsValid())
+		{
+			RenderTickerHandle = FTSTicker::GetCoreTicker().AddTicker(
+				FTickerDelegate::CreateRaw(this, &FDisplayClusterScenePreviewModule::OnTick),
+				CVarDisplayClusterScenePreviewRenderTickDelay.GetValueOnGameThread()
+			);
+		}
+
+		return true;
+	}
+
+	return false;
+}
+
+ADisplayClusterRootActor* FDisplayClusterScenePreviewModule::InternalGetRendererRootActor(FRendererConfig& RendererConfig)
+{
+	if (!RendererConfig.RootActor.IsValid() && !RendererConfig.RootActorPath.IsEmpty())
+	{
+		// Try to find the actor by its path
+		ADisplayClusterRootActor* RootActor = FindObject<ADisplayClusterRootActor>(nullptr, *RendererConfig.RootActorPath);
+		if (RootActor)
+		{
+			InternalSetRendererRootActor(RendererConfig, RootActor, RendererConfig.bAutoUpdateLightcards);
+		}
+	}
+
+	return RendererConfig.RootActor.Get();
+}
+
+void FDisplayClusterScenePreviewModule::InternalSetRendererRootActor(FRendererConfig& RendererConfig, ADisplayClusterRootActor* Actor, bool bAutoUpdateLightcards)
+{
+	// Determine these values before we update the config's RootActor/bAutoUpdateLightcards
+	const bool bRootChanged = RendererConfig.RootActor != Actor;
+	const bool bStartAutoUpdating = bAutoUpdateLightcards && (!RendererConfig.bAutoUpdateLightcards || bRootChanged);
+
+	if (bRootChanged)
+	{
+		// Unregister events for the previous cluster
+		if (RendererConfig.bAutoUpdateLightcards && RendererConfig.RootActor.IsValid())
+		{
+			RegisterRootActorEvents(RendererConfig.RootActor.Get(), false);
+		}
+
+		RendererConfig.RootActor = Actor;
+		RendererConfig.bAutoUpdateLightcards = bAutoUpdateLightcards;
+		AutoPopulateScene(RendererConfig);
+	}
+
+	if (bStartAutoUpdating)
+	{
+		RegisterRootActorEvents(Actor, true);
+	}
+
+	RegisterOrUnregisterGlobalActorEvents();
+}
+
+bool FDisplayClusterScenePreviewModule::InternalRenderImmediate(FRendererConfig& RendererConfig, FDisplayClusterScenePreviewRenderSettings& RenderSettings, FCanvas& Canvas)
+{
+	UWorld* World = RendererConfig.RootActor.IsValid() ? RendererConfig.RootActor->GetWorld() : nullptr;
+	if (!World)
+	{
+		return false;
+	}
+
+	if (RendererConfig.bAutoUpdateLightcards && RendererConfig.bIsSceneDirty)
+	{
+		AutoPopulateScene(RendererConfig);
+	}
+
+	switch (RenderSettings.RenderType)
+	{
+	case FDisplayClusterScenePreviewRenderSettings::ERenderType::Color:
+		RendererConfig.Renderer->Render(
+			&Canvas,
+			World->Scene,
+			RenderSettings.ViewInitOptions,
+			RenderSettings.EngineShowFlags,
+			RenderSettings.ProjectionType
+		);
+		return true;
+
+	case FDisplayClusterScenePreviewRenderSettings::ERenderType::Normals:
+		RendererConfig.Renderer->RenderNormals(
+			&Canvas,
+			World->Scene,
+			RenderSettings.ViewInitOptions,
+			RenderSettings.EngineShowFlags,
+			RenderSettings.ProjectionType,
+			&RenderSettings.PrimitiveFilter
+		);
+		return true;
+	}
+
+	return false;
+}
+
+void FDisplayClusterScenePreviewModule::RegisterOrUnregisterGlobalActorEvents()
+{
+	// Check whether any of our configs need actor events
+	bool bShouldBeRegistered = false;
+	for (const TPair<int32, FRendererConfig>& ConfigPair : RendererConfigs)
+	{
+		if (ConfigPair.Value.bAutoUpdateLightcards)
+		{
+			bShouldBeRegistered = true;
+			break;
+		}
+	}
+
+	if (bShouldBeRegistered && !bIsRegisteredForActorEvents)
+	{
+		// Register for events
+		FCoreUObjectDelegates::OnObjectPropertyChanged.AddRaw(this, &FDisplayClusterScenePreviewModule::OnActorPropertyChanged);
+		FCoreUObjectDelegates::OnObjectTransacted.AddRaw(this, &FDisplayClusterScenePreviewModule::OnObjectTransacted);
+
+		if (GEngine != nullptr)
+		{
+			GEngine->OnLevelActorDeleted().AddRaw(this, &FDisplayClusterScenePreviewModule::OnLevelActorDeleted);
+		}
+	}
+	else if (!bShouldBeRegistered && bIsRegisteredForActorEvents)
+	{
+		// Unregister for events
+		FCoreUObjectDelegates::OnObjectPropertyChanged.RemoveAll(this);
+		FCoreUObjectDelegates::OnObjectTransacted.RemoveAll(this);
+
+		if (GEngine != nullptr)
+		{
+			GEngine->OnLevelActorDeleted().RemoveAll(this);
+		}
+	}
+}
+
+void FDisplayClusterScenePreviewModule::RegisterRootActorEvents(AActor* Actor, bool bShouldRegister)
+{
+	if (!Actor)
+	{
+		return;
+	}
+
+	if (UBlueprint* Blueprint = UBlueprint::GetBlueprintFromClass(Actor->GetClass()))
+	{
+		if (bShouldRegister)
+		{
+			Blueprint->OnCompiled().AddRaw(this, &FDisplayClusterScenePreviewModule::OnBlueprintCompiled);
+		}
+		else
+		{
+			Blueprint->OnCompiled().RemoveAll(this);
+		}
+	}
+}
+
+void FDisplayClusterScenePreviewModule::AutoPopulateScene(FRendererConfig& RendererConfig)
+{
+	RendererConfig.Renderer->ClearScene();
+
+	if (ADisplayClusterRootActor* RootActor = InternalGetRendererRootActor(RendererConfig))
+	{
+		RendererConfig.Renderer->AddActor(RootActor, [](const UPrimitiveComponent* PrimitiveComponent)
+		{
+			return !PrimitiveComponent->bHiddenInGame || PrimitiveComponent->IsA<UDisplayClusterScreenComponent>();
+		});
+
+		if (RendererConfig.bAutoUpdateLightcards)
+		{
+			// Automatically add the lightcards found on this actor
+			TSet<ADisplayClusterLightCardActor*> LightCards;
+			UDisplayClusterBlueprintLib::FindLightCardsForRootActor(RootActor, LightCards);
+
+			for (ADisplayClusterLightCardActor* LightCard : LightCards)
+			{
+				RendererConfig.Renderer->AddActor(LightCard);
+				RendererConfig.AutoLightcards.Add(LightCard);
+			}
+		}
+	}
+
+	RendererConfig.bIsSceneDirty = false;
+}
+
+bool FDisplayClusterScenePreviewModule::OnTick(float DeltaTime)
+{
+	FPreviewRenderJob Job;
+	if (!RenderQueue.Dequeue(Job))
+	{
+		RenderTickerHandle.Reset();
+		return false;
+	}
+
+	ensure(Job.ResultDelegate.IsBound());
+
+	if (FRendererConfig* Config = RendererConfigs.Find(Job.RendererId))
+	{
+		if (UWorld* World = Config->RootActor.IsValid() ? Config->RootActor->GetWorld() : nullptr)
+		{
+			UTextureRenderTarget2D* RenderTarget = Config->RenderTarget.Get();
+
+			if (!RenderTarget)
+			{
+				// Create a new render target
+				RenderTarget = NewObject<UTextureRenderTarget2D>();
+				RenderTarget->InitCustomFormat(Job.Size.X, Job.Size.Y, PF_B8G8R8A8, false);
+
+				Config->RenderTarget = TStrongObjectPtr<UTextureRenderTarget2D>(RenderTarget);
+			}
+			else if (RenderTarget->SizeX != Job.Size.X || RenderTarget->SizeY != Job.Size.Y)
+			{
+				// Resize to match the new size
+				RenderTarget->ResizeTarget(Job.Size.X, Job.Size.Y);
+			}
+
+			FTextureRenderTargetResource* RenderTargetResource = RenderTarget->GameThread_GetRenderTargetResource();
+			FCanvas Canvas(RenderTargetResource, nullptr, FGameTime::GetTimeSinceAppStart(), World->Scene->GetFeatureLevel());
+
+			InternalRenderImmediate(*Config, Job.Settings, Canvas);
+
+			Job.ResultDelegate.Execute(*RenderTargetResource);
+		}
+	}
+
+	if (RenderQueue.IsEmpty())
+	{
+		RenderTickerHandle.Reset();
+		return false;
+	}
+
+	return true;
+}
+
+void FDisplayClusterScenePreviewModule::OnActorPropertyChanged(UObject* ObjectBeingModified, FPropertyChangedEvent& PropertyChangedEvent)
+{
+	for (TPair<int32, FRendererConfig>& ConfigPair : RendererConfigs)
+	{
+		FRendererConfig& Config = ConfigPair.Value;
+		if (Config.bAutoUpdateLightcards && Config.RootActor == ObjectBeingModified)
+		{
+			Config.bIsSceneDirty = true;
+		}
+	}
+}
+
+void FDisplayClusterScenePreviewModule::OnLevelActorDeleted(AActor* Actor)
+{
+	for (TPair<int32, FRendererConfig>& ConfigPair : RendererConfigs)
+	{
+		FRendererConfig& Config = ConfigPair.Value;
+		if (Config.AutoLightcards.Contains(Actor))
+		{
+			Config.bIsSceneDirty = true;
+		}
+	}
+}
+
+void FDisplayClusterScenePreviewModule::OnBlueprintCompiled(UBlueprint* Blueprint)
+{
+	for (TPair<int32, FRendererConfig>& ConfigPair : RendererConfigs)
+	{
+		FRendererConfig& Config = ConfigPair.Value;
+		if (Config.RootActor.IsValid() && Blueprint == UBlueprint::GetBlueprintFromClass(Config.RootActor->GetClass()))
+		{
+			Config.bIsSceneDirty = true;
+		}
+	}
+}
+
+void FDisplayClusterScenePreviewModule::OnObjectTransacted(UObject* Object, const FTransactionObjectEvent& TransactionObjectEvent)
+{
+	if (TransactionObjectEvent.GetEventType() == ETransactionObjectEventType::UndoRedo)
+	{
+		for (TPair<int32, FRendererConfig>& ConfigPair : RendererConfigs)
+		{
+			ConfigPair.Value.bIsSceneDirty = true;
+		}
+	}
+}
+
+IMPLEMENT_MODULE(FDisplayClusterScenePreviewModule, DisplayClusterScenePreview);
+
+#undef LOCTEXT_NAMESPACE
diff --git a/Engine/Plugins/Runtime/nDisplay/Source/DisplayClusterScenePreview/Private/DisplayClusterScenePreviewModule.h b/Engine/Plugins/Runtime/nDisplay/Source/DisplayClusterScenePreview/Private/DisplayClusterScenePreviewModule.h
new file mode 100644
--- /dev/null
+++ b/Engine/Plugins/Runtime/nDisplay/Source/DisplayClusterScenePreview/Private/DisplayClusterScenePreviewModule.h
@@ -0,0 +1,139 @@
+// Copyright Epic Games, Inc. All Rights Reserved.
+
+#pragma once
+
+#include "IDisplayClusterScenePreview.h"
+
+#include "Containers/Ticker.h"
+#include "DisplayClusterMeshProjectionRenderer.h"
+#include "UObject/StrongObjectPtr.h"
+
+class ADisplayClusterRootActor;
+class ADisplayClusterLightCardActor;
+class UTextureRenderTarget2D;
+
+/**
+ * Module containing tools for rendering nDisplay scene previews.
+ */
+class FDisplayClusterScenePreviewModule :
+	public IDisplayClusterScenePreview
+{
+public:
+	//~ IModuleInterface interface
+	virtual void StartupModule() override;
+	virtual void ShutdownModule() override;
+
+	//~ Begin IDisplayClusterScenePreview Interface
+	virtual int32 CreateRenderer() override;
+	virtual bool DestroyRenderer(int32 RendererId) override;
+	virtual bool SetRendererRootActorPath(int32 RendererId, const FString& ActorPath, bool bAutoUpdateLightcards = false) override;
+	virtual bool SetRendererRootActor(int32 RendererId, ADisplayClusterRootActor* Actor, bool bAutoUpdateLightcards = false) override;
+	virtual ADisplayClusterRootActor* GetRendererRootActor(int32 RendererId) override;
+	virtual bool AddActorToRenderer(int32 RendererId, AActor* Actor) override;
+	virtual bool ClearRendererScene(int32 RendererId) override;
+	virtual bool SetRendererActorSelectedDelegate(int32 RendererId, FDisplayClusterMeshProjectionRenderer::FSelection ActorSelectedDelegate) override;
+	virtual bool SetRendererRenderSimpleElementsDelegate(int32 RendererId, FDisplayClusterMeshProjectionRenderer::FSimpleElementPass RenderSimpleElementsDelegate) override;
+	virtual bool Render(int32 RendererId, FDisplayClusterScenePreviewRenderSettings& RenderSettings, FCanvas& Canvas) override;
+	virtual bool RenderQueued(int32 RendererId, FDisplayClusterScenePreviewRenderSettings& RenderSettings, const FIntPoint& Size, FRenderResultDelegate ResultDelegate) override;
+	//~ End IDisplayClusterScenePreview Interface
+
+private:
+	/** Holds information about an active renderer created by this module. */
+	struct FRendererConfig
+	{
+		/** The renderer itself. */
+		TSharedPtr<FDisplayClusterMeshProjectionRenderer> Renderer;
+
+		/** The root of the display cluster that this renderer is previewing. */
+		TWeakObjectPtr<ADisplayClusterRootActor> RootActor;
+
+		/** The path of the root actor that this renderer is previewing. If this is not empty and the root actor becomes invalid, we will attempt to find it again using this path. */
+		FString RootActorPath;
+
+		/** Lightcards that have been automatically added to the scene. */
+		TArray<TWeakObjectPtr<ADisplayClusterLightCardActor>> AutoLightcards;
+
+		/** If true, automatically update the renderer with lightcards belonging to the root actor. */
+		bool bAutoUpdateLightcards = false;
+
+		/** If true, the scene needs to be updated before the next render. This is only relevant if bAutoUpdateLightcards is true. */
+		bool bIsSceneDirty = true;
+
+		/** The render target to use for queued renders. */
+		TStrongObjectPtr<UTextureRenderTarget2D> RenderTarget = nullptr;
+	};
+
+	/** Holds information about a preview render that was queued to be completed later. */
+	struct FPreviewRenderJob
+	{
+		FPreviewRenderJob() {}
+
+		FPreviewRenderJob(int32 RendererId, const FDisplayClusterScenePreviewRenderSettings& Settings, const FIntPoint& Size,
+			FRenderResultDelegate ResultDelegate)
+			: RendererId(RendererId), Settings(Settings), Size(Size), ResultDelegate(ResultDelegate)
+		{
+		}
+
+		/** The ID of the renderer to use. */
+		int32 RendererId;
+
+		/** The settings to use for the render. */
+		FDisplayClusterScenePreviewRenderSettings Settings;
+
+		/** The size of the image to render. */
+		FIntPoint Size;
+
+		/** The delegate to call when the render is completed. */
+		FRenderResultDelegate ResultDelegate;
+	};
+
+	/** Get the root actor for a config. If the root actor pointer is invalid but we have a path to the actor, try to reacquire a pointer using the path first. */
+	ADisplayClusterRootActor* InternalGetRendererRootActor(FRendererConfig& RendererConfig);
+
+	/** Set the root actor for a config, update its scene, and register events accordingly. */
+	void InternalSetRendererRootActor(FRendererConfig& RendererConfig, ADisplayClusterRootActor* Actor, bool bAutoUpdateLightcards);
+
+	/** Immediately render with the given renderer config and settings to the given canvas. */
+	bool InternalRenderImmediate(FRendererConfig& RendererConfig, FDisplayClusterScenePreviewRenderSettings& RenderSettings, FCanvas& Canvas);
+
+	/** Check if any of the tracked root actors are set to auto-update their lightcards and register/unregister event listeners accordingly. */
+	void RegisterOrUnregisterGlobalActorEvents();
+
+	/** Register/unregister to events affecting a cluster root actor. */
+	void RegisterRootActorEvents(AActor* Actor, bool bShouldRegister);
+
+	/** Clear and re-populate a renderer's scene with the root actor and lightcards if applicable. */
+	void AutoPopulateScene(FRendererConfig& RendererConfig);
+
+	/** Called on tick to process the queued renders. */
+	bool OnTick(float DeltaTime);
+
+	/** Called when a property on a root DisplayCluster actor has changed. */
+	void OnActorPropertyChanged(UObject* ObjectBeingModified, FPropertyChangedEvent& PropertyChangedEvent);
+
+	/** Called when the user deletes an actor from the level. */
+	void OnLevelActorDeleted(AActor* Actor);
+
+	/** Called when a blueprint for an actor we care about is compiled. */
+	void OnBlueprintCompiled(UBlueprint* Blueprint);
+
+	/** Called when any object is transacted. */
+	void OnObjectTransacted(UObject* Object, const FTransactionObjectEvent& TransactionObjectEvent);
+
+private:
+
+	/** Map from renderer ID to configuration data for that renderer. */
+	TMap<int32, FRendererConfig> RendererConfigs;
+
+	/** Queue of render jobs pending completion. */
+	TQueue<FPreviewRenderJob> RenderQueue;
+
+	/** Handle for the render ticker. */
+	FTSTicker::FDelegateHandle RenderTickerHandle;
+
+	/** The ID to use for the next created renderer. */
+	int32 NextRendererId = 0;
+
+	/** Whether this is currently registered for actor update events. */
+	bool bIsRegisteredForActorEvents = false;
+};
\ No newline at end of file
diff --git a/Engine/Plugins/Runtime/nDisplay/Source/DisplayClusterScenePreview/Public/IDisplayClusterScenePreview.h b/Engine/Plugins/Runtime/nDisplay/Source/DisplayClusterScenePreview/Public/IDisplayClusterScenePreview.h
new file mode 100644
--- /dev/null
+++ b/Engine/Plugins/Runtime/nDisplay/Source/DisplayClusterScenePreview/Public/IDisplayClusterScenePreview.h
@@ -0,0 +1,180 @@
+// Copyright Epic Games, Inc. All Rights Reserved.
+
+#pragma once
+
+#include "CoreMinimal.h"
+#include "Modules/ModuleInterface.h"
+#include "Modules/ModuleManager.h"
+
+#include "SceneView.h"
+#include "ShowFlags.h"
+#include "DisplayClusterMeshProjectionRenderer.h"
+
+class ADisplayClusterRootActor;
+struct FSceneViewInitOptions;
+
+DECLARE_DELEGATE_OneParam(FRenderResultDelegate, FRenderTarget&);
+
+/** Settings for producing a single render. */
+struct FDisplayClusterScenePreviewRenderSettings
+{
+	FDisplayClusterScenePreviewRenderSettings()
+		: EngineShowFlags(ESFIM_Editor)
+	{}
+
+	/** Possible render types. */
+	enum ERenderType
+	{
+		/** Render the scene's colors. */
+		Color,
+
+		/** Render the scene's normals. */
+		Normals
+	};
+
+	/** The type of render to perform. */
+	ERenderType RenderType = ERenderType::Color;
+
+	/** Camera setup options for the render. */
+	FSceneViewInitOptions ViewInitOptions;
+
+	/**
+	 * Flags controlling renderer features to enable/disable for this preview.
+	 * Only flags supported by FDisplayClusterMeshProjectionRenderer will have an impact on the render output.
+	 */
+	FEngineShowFlags EngineShowFlags;
+
+	/** Type of projection to use for the renderer. */
+	EDisplayClusterMeshProjectionType ProjectionType = EDisplayClusterMeshProjectionType::Azimuthal;
+
+	/**
+	 * Optional filter to prevent specific primitives from being rendered.
+	 * This only applies when RenderType is ERenderType::Normals.
+	 */
+	FDisplayClusterMeshProjectionPrimitiveFilter PrimitiveFilter;
+};
+
+/**
+ * Interface for module containing tools for generation Display Cluster scene previews.
+ */
+class IDisplayClusterScenePreview : public IModuleInterface
+{
+public:
+	static constexpr const TCHAR* ModuleName = TEXT("DisplayClusterScenePreview");
+
+public:
+	virtual ~IDisplayClusterScenePreview() = default;
+
+	/**
+	 * Singleton-like access to this module's interface. This is just for convenience!
+	 * Beware of calling this during the shutdown phase, though. Your module might have been unloaded already.
+	 *
+	 * @return Returns singleton instance, loading the module on demand if needed
+	 */
+	static inline IDisplayClusterScenePreview& Get()
+	{
+		return FModuleManager::GetModuleChecked<IDisplayClusterScenePreview>(ModuleName);
+	}
+
+	/**
+	 * Checks to see if this module is loaded and ready. It is only valid to call Get() if IsAvailable() returns true.
+	 *
+	 * @return True if the module is loaded and ready to use
+	 */
+	static inline bool IsAvailable()
+	{
+		return FModuleManager::Get().IsModuleLoaded(ModuleName);
+	}
+
+	/**
+	 * Create a preview renderer. Returns an index used to refer to this renderer in future API calls.
+	 */
+	virtual int32 CreateRenderer() = 0;
+
+	/** 
+     * Destroy a renderer, cleaning up all of its resources.
+	 * 
+	 * @param RendererId The ID of the renderer as returned from CreateRenderer.
+	 */
+	virtual bool DestroyRenderer(int32 RendererId) = 0;
+
+	/**
+	 * Set the path of the ADisplayClusterRootActor actor for a renderer, which will be added to the preview scene and used to determine the render world.
+	 * Whenever a render is performed with this renderer, it will try to find the actor at that path and use it as the root.
+	 * 
+	 * @param RendererId The ID of the renderer as returned from CreateRenderer.
+	 * @param ActorPath The path of the ADisplayClusterRootActor actor to preview.
+	 * @param bAutoUpdateLightcards If true, the renderer scene will also be automatically populated with the actor's associated lightcards, including future changes to them.
+	 */
+	virtual bool SetRendererRootActorPath(int32 RendererId, const FString& ActorPath, bool bAutoUpdateLightcards = false) = 0;
+
+	/**
+	 * Set the root DisplayCluster actor for the renderer with the given ID, which will be added to the preview scene and used to determine the render world.
+	 * Note that if the actor is destroyed (including by a recompile, reload, etc.), you will need to call this again to update the actor.
+	 * This will also clear the path set by SetRendererRootActorPath.
+	 * 
+	 * @param RendererId The ID of the renderer as returned from CreateRenderer.
+	 * @param Actor The actor to preview.
+	 * @param bAutoUpdateLightcards If true, the renderer scene will also be automatically populated with the actor's associated lightcards, including future changes to them.
+	 */
+	virtual bool SetRendererRootActor(int32 RendererId, ADisplayClusterRootActor* Actor, bool bAutoUpdateLightcards = false) = 0;
+
+	/**
+	 * Get the root actor of a renderer. 
+	 * 
+	 * @param RendererId The ID of the renderer as returned from CreateRenderer.
+	 */
+	virtual ADisplayClusterRootActor* GetRendererRootActor(int32 RendererId) = 0;
+
+	/**
+	 * Add an actor to the preview scene for a renderer.
+	 * 
+	 * @param RendererId The ID of the renderer as returned from CreateRenderer.
+	 * @param Actor The actor to add to the renderer's scene.
+	 */
+	virtual bool AddActorToRenderer(int32 RendererId, AActor* Actor) = 0;
+
+	/**
+	 * Clear the scene of a renderer.
+	 * 
+	 * @param RendererId The ID of the renderer as returned from CreateRenderer.
+	 */
+	virtual bool ClearRendererScene(int32 RendererId) = 0;
+
+	/**
+	 * Set the ActorSelectedDelegate of a renderer.
+	 * See FDisplayClusterMeshProjectionRenderer::ActorSelectedDelegate for more info.
+	 * 
+	 * @param RendererId The ID of the renderer as returned from CreateRenderer.
+	 * @param ActorSelectedDelegate The delegate to set on the renderer.
+	 */
+	virtual bool SetRendererActorSelectedDelegate(int32 RendererId, FDisplayClusterMeshProjectionRenderer::FSelection ActorSelectedDelegate) = 0;
+
+	/**
+	 * Set the RenderSimpleElementsDelegate of a renderer.
+	 * See FDisplayClusterMeshProjectionRenderer::RenderSimpleElementsDelegate for more info.
+	 * 
+	 * @param RendererId The ID of the renderer as returned from CreateRenderer.
+	 * @param RenderSimpleElementsDelegate The delegate to set on the renderer.
+	 */
+	virtual bool SetRendererRenderSimpleElementsDelegate(int32 RendererId, FDisplayClusterMeshProjectionRenderer::FSimpleElementPass RenderSimpleElementsDelegate) = 0;
+
+	/**
+	 * Immediately render a preview.
+	 * 
+	 * @param RendererId The ID of the renderer as returned from CreateRenderer.
+	 * @param RenderSettings Settings controlling how the scene will be rendered.
+	 * @param Canvas The canvas to render to.
+	 */
+	virtual bool Render(int32 RendererId, FDisplayClusterScenePreviewRenderSettings& RenderSettings, FCanvas& Canvas) = 0;
+
+	/**
+	 * Queue a preview to be rendered in the future.
+	 *
+	 * @param RendererId The ID of the renderer as returned from CreateRenderer.
+	 * @param RenderSettings Settings controlling how the scene will be rendered.
+	 * @param Size The size of the image to produce. Note that whenever this changes for a given renderer, the underlying RenderTarget will be resized, which has a performance cost.
+	 * @param ResultDelegate The delegate to call when the render is complete. It will be passed a FRenderTarget containing the rendered preview.
+	 */
+	virtual bool RenderQueued(int32 RendererId, FDisplayClusterScenePreviewRenderSettings& RenderSettings, const FIntPoint& Size, FRenderResultDelegate ResultDelegate) = 0;
+};
\ No newline at end of file
diff --git a/Engine/Plugins/Runtime/nDisplay/nDisplay.uplugin b/Engine/Plugins/Runtime/nDisplay/nDisplay.uplugin
--- a/Engine/Plugins/Runtime/nDisplay/nDisplay.uplugin
+++ b/Engine/Plugins/Runtime/nDisplay/nDisplay.uplugin
@@ -109,6 +109,15 @@
 			]
 		},
 		{
+			"Name": "DisplayClusterScenePreview",
+			"Type": "Editor",
+			"LoadingPhase": "Default",
+			"PlatformAllowList": [
+				"Win64",
+				"Linux"
+			]
+		},
+		{
 			"Name": "DisplayClusterStageMonitoring",
 			"Type": "Runtime",
 			"LoadingPhase": "Default",
