nDisplay: Fixed case in an include.

diff --git a/Engine/Plugins/Runtime/WarpUtils/Shaders/Private/PFMExporterShaders.usf b/Engine/Plugins/Runtime/WarpUtils/Shaders/Private/PFMExporterShaders.usf
--- a/Engine/Plugins/Runtime/WarpUtils/Shaders/Private/PFMExporterShaders.usf
+++ b/Engine/Plugins/Runtime/WarpUtils/Shaders/Private/PFMExporterShaders.usf
@@ -1,6 +1,6 @@
 // Copyright Epic Games, Inc. All Rights Reserved.
 
-#include "/engine/public/platform.ush"
+#include "/Engine/Public/Platform.ush"
 
 
 float4x4	MeshToPFMMatrix;
diff --git a/Engine/Plugins/Runtime/nDisplay/Shaders/Private/MPCDIUtils.ush b/Engine/Plugins/Runtime/nDisplay/Shaders/Private/MPCDIUtils.ush
--- a/Engine/Plugins/Runtime/nDisplay/Shaders/Private/MPCDIUtils.ush
+++ b/Engine/Plugins/Runtime/nDisplay/Shaders/Private/MPCDIUtils.ush
@@ -1,6 +1,6 @@
 // Copyright Epic Games, Inc. All Rights Reserved.
 
-#include "/engine/public/platform.ush"
+#include "/Engine/Public/Platform.ush"
 #include "/Engine/Private/GammaCorrectionCommon.ush"
 
 
diff --git a/Engine/Plugins/Runtime/nDisplay/Shaders/Private/OutputRemapShaders.usf b/Engine/Plugins/Runtime/nDisplay/Shaders/Private/OutputRemapShaders.usf
--- a/Engine/Plugins/Runtime/nDisplay/Shaders/Private/OutputRemapShaders.usf
+++ b/Engine/Plugins/Runtime/nDisplay/Shaders/Private/OutputRemapShaders.usf
@@ -1,6 +1,6 @@
 // Copyright Epic Games, Inc. All Rights Reserved.
 
-#include "/engine/public/platform.ush"
+#include "/Engine/Public/Platform.ush"
 
 Texture2D PostprocessInput0;
 SamplerState PostprocessInput0Sampler;
