[nDisplay] Fix a couple crashes that could happen when using the Epic Stage App

diff --git a/Engine/Plugins/Runtime/nDisplay/Source/DisplayClusterLightCardEditorShaders/Private/DisplayClusterMeshProjectionRenderer.cpp b/Engine/Plugins/Runtime/nDisplay/Source/DisplayClusterLightCardEditorShaders/Private/DisplayClusterMeshProjectionRenderer.cpp
--- a/Engine/Plugins/Runtime/nDisplay/Source/DisplayClusterLightCardEditorShaders/Private/DisplayClusterMeshProjectionRenderer.cpp
+++ b/Engine/Plugins/Runtime/nDisplay/Source/DisplayClusterLightCardEditorShaders/Private/DisplayClusterMeshProjectionRenderer.cpp
@@ -717,6 +717,22 @@ FVector FDisplayClusterMeshProjectionTransform::UnprojectPosition(const FVector&
 	return ProjectedPosition;
 }
 
+FDisplayClusterMeshProjectionRenderer::~FDisplayClusterMeshProjectionRenderer()
+{
+#if WITH_EDITOR
+	for (TWeakObjectPtr<UPrimitiveComponent> PrimitiveComponent : PrimitiveComponents)
+	{
+		if (!PrimitiveComponent.IsValid())
+		{
+			continue;
+		}
+
+		// Unbind this since it will still have a raw pointer to this renderer
+		PrimitiveComponent->SelectionOverrideDelegate.Unbind();
+	}
+#endif
+}
+
 FVector FDisplayClusterMeshProjectionRenderer::ProjectViewPosition(const FVector& ViewPosition, EDisplayClusterMeshProjectionType  ProjectionType)
 {
 	FVector ProjectedViewPosition(ViewPosition);
diff --git a/Engine/Plugins/Runtime/nDisplay/Source/DisplayClusterLightCardEditorShaders/Public/DisplayClusterMeshProjectionRenderer.h b/Engine/Plugins/Runtime/nDisplay/Source/DisplayClusterLightCardEditorShaders/Public/DisplayClusterMeshProjectionRenderer.h
--- a/Engine/Plugins/Runtime/nDisplay/Source/DisplayClusterLightCardEditorShaders/Public/DisplayClusterMeshProjectionRenderer.h
+++ b/Engine/Plugins/Runtime/nDisplay/Source/DisplayClusterLightCardEditorShaders/Public/DisplayClusterMeshProjectionRenderer.h
@@ -63,6 +63,9 @@ private:
 class DISPLAYCLUSTERLIGHTCARDEDITORSHADERS_API FDisplayClusterMeshProjectionRenderer
 {
 public:
+	/** Clean up any references to the renderer. */
+	~FDisplayClusterMeshProjectionRenderer();
+
 	/** Projects a position in view coordinates into the projected view space of the specified projection type */
 	static FVector ProjectViewPosition(const FVector& ViewPosition, EDisplayClusterMeshProjectionType  ProjectionType);
 
diff --git a/Engine/Plugins/Runtime/nDisplay/Source/DisplayClusterScenePreview/DisplayClusterScenePreview.Build.cs b/Engine/Plugins/Runtime/nDisplay/Source/DisplayClusterScenePreview/DisplayClusterScenePreview.Build.cs
--- a/Engine/Plugins/Runtime/nDisplay/Source/DisplayClusterScenePreview/DisplayClusterScenePreview.Build.cs
+++ b/Engine/Plugins/Runtime/nDisplay/Source/DisplayClusterScenePreview/DisplayClusterScenePreview.Build.cs
@@ -20,6 +20,7 @@ public class DisplayClusterScenePreview : ModuleRules
 				"Core",
 				"CoreUObject",
 				"Engine",
+				"RenderCore",
 				"UnrealEd"
 			});
 	}
diff --git a/Engine/Plugins/Runtime/nDisplay/Source/DisplayClusterScenePreview/Private/DisplayClusterScenePreviewModule.cpp b/Engine/Plugins/Runtime/nDisplay/Source/DisplayClusterScenePreview/Private/DisplayClusterScenePreviewModule.cpp
--- a/Engine/Plugins/Runtime/nDisplay/Source/DisplayClusterScenePreview/Private/DisplayClusterScenePreviewModule.cpp
+++ b/Engine/Plugins/Runtime/nDisplay/Source/DisplayClusterScenePreview/Private/DisplayClusterScenePreviewModule.cpp
@@ -358,6 +358,9 @@ bool FDisplayClusterScenePreviewModule::OnTick(float DeltaTime)
 			{
 				// Resize to match the new size
 				RenderTarget->ResizeTarget(Job.Size.X, Job.Size.Y);
+				
+				// Flush commands so target is immediately ready to render at the new size
+				FlushRenderingCommands();
 			}
 
 			FTextureRenderTargetResource* RenderTargetResource = RenderTarget->GameThread_GetRenderTargetResource();
