nDisplay: added support for CineCamera's custom nearplane support for Camera projection policy

diff --git a/Engine/Plugins/Runtime/nDisplay/Source/DisplayClusterProjection/DisplayClusterProjection.Build.cs b/Engine/Plugins/Runtime/nDisplay/Source/DisplayClusterProjection/DisplayClusterProjection.Build.cs
--- a/Engine/Plugins/Runtime/nDisplay/Source/DisplayClusterProjection/DisplayClusterProjection.Build.cs
+++ b/Engine/Plugins/Runtime/nDisplay/Source/DisplayClusterProjection/DisplayClusterProjection.Build.cs
@@ -10,11 +10,12 @@ public class DisplayClusterProjection : ModuleRules
 		PublicIncludePathModuleNames.AddRange(
 			new string[] {
 				"DisplayClusterConfiguration",
-				"DisplayClusterShaders",
+				"DisplayClusterShaders"
 			});
 
 		PublicDependencyModuleNames.AddRange(
 			new string[] {
+				"CinematicCamera",
 				"Core",
 				"CoreUObject",
 				"DisplayCluster",
@@ -29,7 +30,7 @@ public class DisplayClusterProjection : ModuleRules
 				"ProceduralMeshComponent",
 				"Projects",
 				"RenderCore",
-				"RHI",
+				"RHI"
 			});
 
 		if (Target.bBuildEditor == true)
diff --git a/Engine/Plugins/Runtime/nDisplay/Source/DisplayClusterProjection/Private/Policy/Camera/DisplayClusterProjectionCameraPolicy.cpp b/Engine/Plugins/Runtime/nDisplay/Source/DisplayClusterProjection/Private/Policy/Camera/DisplayClusterProjectionCameraPolicy.cpp
--- a/Engine/Plugins/Runtime/nDisplay/Source/DisplayClusterProjection/Private/Policy/Camera/DisplayClusterProjectionCameraPolicy.cpp
+++ b/Engine/Plugins/Runtime/nDisplay/Source/DisplayClusterProjection/Private/Policy/Camera/DisplayClusterProjectionCameraPolicy.cpp
@@ -10,6 +10,7 @@
 #include "Engine/World.h"
 #include "GameFramework/PlayerController.h"
 #include "Camera/CameraComponent.h"
+#include "CineCameraComponent.h"
 #include "ComposurePostMoves.h"
 
 #include "Render/Viewport/IDisplayClusterViewport.h"
@@ -83,14 +84,23 @@ bool FDisplayClusterProjectionCameraPolicy::CalculateView(IDisplayClusterViewpor
 
 	InOutViewLocation = FVector::ZeroVector;
 	InOutViewRotation = FRotator::ZeroRotator;
-
+	
 	// Save Z values
 	ZNear = NCP;
 	ZFar  = FCP;
-	
-	// Use transform of an assigned camera
+
+	// Use assigned camera
 	if (UCameraComponent* CameraComponent = GetCameraComponent())
 	{
+		if (UCineCameraComponent* CineCameraComponent = Cast<UCineCameraComponent>(CameraComponent))
+		{
+			if (CineCameraComponent->bOverride_CustomNearClippingPlane)
+			{
+				ZNear = CineCameraComponent->CustomNearClippingPlane;
+				ZFar = ZNear;
+			}
+		}
+
 		InOutViewLocation = CameraComponent->GetComponentLocation();
 		InOutViewRotation = CameraComponent->GetComponentRotation();
 	}
