nDisplay: Added orthographic projection mode to the light card editor, as well as options to orient the view to face a particular direction of the stage (top, bottom, left, right, front, back). Also fixed several small bugs, including one where the stage actor proxy would disappear from the viewport when a light card was deleted.

diff --git a/Engine/Plugins/Runtime/nDisplay/Source/DisplayClusterLightCardEditor/DisplayClusterLightCardEditor.Build.cs b/Engine/Plugins/Runtime/nDisplay/Source/DisplayClusterLightCardEditor/DisplayClusterLightCardEditor.Build.cs
--- a/Engine/Plugins/Runtime/nDisplay/Source/DisplayClusterLightCardEditor/DisplayClusterLightCardEditor.Build.cs
+++ b/Engine/Plugins/Runtime/nDisplay/Source/DisplayClusterLightCardEditor/DisplayClusterLightCardEditor.Build.cs
@@ -25,6 +25,7 @@ public class DisplayClusterLightCardEditor : ModuleRules
 				"InputCore",
 				"OpenCV",
 				"OpenCVHelper",
+				"ProceduralMeshComponent",
 				"PropertyEditor",
 				"RenderCore",
 				"Renderer",
diff --git a/Engine/Plugins/Runtime/nDisplay/Source/DisplayClusterLightCardEditor/Private/DisplayClusterLightCardEditorCommands.cpp b/Engine/Plugins/Runtime/nDisplay/Source/DisplayClusterLightCardEditor/Private/DisplayClusterLightCardEditorCommands.cpp
--- a/Engine/Plugins/Runtime/nDisplay/Source/DisplayClusterLightCardEditor/Private/DisplayClusterLightCardEditorCommands.cpp
+++ b/Engine/Plugins/Runtime/nDisplay/Source/DisplayClusterLightCardEditor/Private/DisplayClusterLightCardEditorCommands.cpp
@@ -10,7 +10,15 @@ void FDisplayClusterLightCardEditorCommands::RegisterCommands()
 	UI_COMMAND(ResetCamera, "Reset Camera", "Resets the camera to focus on the mesh", EUserInterfaceActionType::Button, FInputChord());
 
 	UI_COMMAND(PerspectiveProjection, "Perspective", "A perspective projection from the stage's view origin", EUserInterfaceActionType::Button, FInputChord());
-	UI_COMMAND(AzimuthalProjection, "Azimuthal", "An azimuthal hemisphere projection from the stage's center pointed upwards", EUserInterfaceActionType::Button, FInputChord());
+	UI_COMMAND(OrthographicProjection, "Orthographic", "An orthographic projection from the stage's view origin", EUserInterfaceActionType::Button, FInputChord());
+	UI_COMMAND(AzimuthalProjection, "Azimuthal", "An azimuthal hemisphere projection from the stage's view origin", EUserInterfaceActionType::Button, FInputChord());
+
+	UI_COMMAND(ViewOrientationTop, "Top", "Orient the view to look at the top of the stage", EUserInterfaceActionType::Button, FInputChord());
+	UI_COMMAND(ViewOrientationBottom, "Bottom", "Orient the view look at the bottom of the stage", EUserInterfaceActionType::Button, FInputChord());
+	UI_COMMAND(ViewOrientationLeft, "Left", "Orient the view to look at the left of the stage", EUserInterfaceActionType::Button, FInputChord());
+	UI_COMMAND(ViewOrientationRight, "Right", "Orient the view to look at the right of the stage", EUserInterfaceActionType::Button, FInputChord());
+	UI_COMMAND(ViewOrientationFront, "Front", "Orient the view to look at the front of the stage", EUserInterfaceActionType::Button, FInputChord());
+	UI_COMMAND(ViewOrientationBack, "Back", "Orient the view to look at the back of the stage", EUserInterfaceActionType::Button, FInputChord());
 
 	UI_COMMAND(AddNewLightCard, "Add New Light Card", "Add and assign a new Light Card to the actor", EUserInterfaceActionType::Button, FInputChord());
 	UI_COMMAND(AddExistingLightCard, "Add Existing Light Card", "Add an existing Light Card to the actor", EUserInterfaceActionType::Button, FInputChord());
diff --git a/Engine/Plugins/Runtime/nDisplay/Source/DisplayClusterLightCardEditor/Private/DisplayClusterLightCardEditorCommands.h b/Engine/Plugins/Runtime/nDisplay/Source/DisplayClusterLightCardEditor/Private/DisplayClusterLightCardEditorCommands.h
--- a/Engine/Plugins/Runtime/nDisplay/Source/DisplayClusterLightCardEditor/Private/DisplayClusterLightCardEditorCommands.h
+++ b/Engine/Plugins/Runtime/nDisplay/Source/DisplayClusterLightCardEditor/Private/DisplayClusterLightCardEditorCommands.h
@@ -26,8 +26,16 @@ public:
 	TSharedPtr<FUICommandInfo> ResetCamera;
 
 	TSharedPtr<FUICommandInfo> PerspectiveProjection;
+	TSharedPtr<FUICommandInfo> OrthographicProjection;
 	TSharedPtr<FUICommandInfo> AzimuthalProjection;
 
+	TSharedPtr<FUICommandInfo> ViewOrientationTop;
+	TSharedPtr<FUICommandInfo> ViewOrientationBottom;
+	TSharedPtr<FUICommandInfo> ViewOrientationLeft;
+	TSharedPtr<FUICommandInfo> ViewOrientationRight;
+	TSharedPtr<FUICommandInfo> ViewOrientationFront;
+	TSharedPtr<FUICommandInfo> ViewOrientationBack;
+
 	TSharedPtr<FUICommandInfo> AddNewLightCard;
 	TSharedPtr<FUICommandInfo> AddExistingLightCard;
 	TSharedPtr<FUICommandInfo> RemoveLightCard;
diff --git a/Engine/Plugins/Runtime/nDisplay/Source/DisplayClusterLightCardEditor/Private/DisplayClusterLightCardEditorStyle.h b/Engine/Plugins/Runtime/nDisplay/Source/DisplayClusterLightCardEditor/Private/DisplayClusterLightCardEditorStyle.h
--- a/Engine/Plugins/Runtime/nDisplay/Source/DisplayClusterLightCardEditor/Private/DisplayClusterLightCardEditorStyle.h
+++ b/Engine/Plugins/Runtime/nDisplay/Source/DisplayClusterLightCardEditor/Private/DisplayClusterLightCardEditorStyle.h
@@ -29,6 +29,13 @@ public:
 
 		Set("DisplayClusterLightCardEditor.PasteHere", new CORE_IMAGE_BRUSH_SVG("Starship/Actors/paste-here", Icon16x16));
 
+		Set("DisplayClusterLightCardEditor.ViewOrientationTop", new CORE_IMAGE_BRUSH_SVG("Starship/Common/ViewTop", Icon16x16));
+		Set("DisplayClusterLightCardEditor.ViewOrientationBottom", new CORE_IMAGE_BRUSH_SVG("Starship/Common/ViewBottom", Icon16x16));
+		Set("DisplayClusterLightCardEditor.ViewOrientationLeft", new CORE_IMAGE_BRUSH_SVG("Starship/Common/ViewLeft", Icon16x16));
+		Set("DisplayClusterLightCardEditor.ViewOrientationRight", new CORE_IMAGE_BRUSH_SVG("Starship/Common/ViewRight", Icon16x16));
+		Set("DisplayClusterLightCardEditor.ViewOrientationFront", new CORE_IMAGE_BRUSH_SVG("Starship/Common/ViewFront", Icon16x16));
+		Set("DisplayClusterLightCardEditor.ViewOrientationBack", new CORE_IMAGE_BRUSH_SVG("Starship/Common/ViewBack", Icon16x16));
+
 		FSlateStyleRegistry::RegisterSlateStyle(*this);
 	}
 
diff --git a/Engine/Plugins/Runtime/nDisplay/Source/DisplayClusterLightCardEditor/Private/Viewport/DisplayClusterLightCardEditorViewportClient.cpp b/Engine/Plugins/Runtime/nDisplay/Source/DisplayClusterLightCardEditor/Private/Viewport/DisplayClusterLightCardEditorViewportClient.cpp
--- a/Engine/Plugins/Runtime/nDisplay/Source/DisplayClusterLightCardEditor/Private/Viewport/DisplayClusterLightCardEditorViewportClient.cpp
+++ b/Engine/Plugins/Runtime/nDisplay/Source/DisplayClusterLightCardEditor/Private/Viewport/DisplayClusterLightCardEditorViewportClient.cpp
@@ -35,9 +35,11 @@
 #include "ImageUtils.h"
 #include "Kismet/KismetMathLibrary.h"
 #include "Kismet/GameplayStatics.h"
+#include "KismetProceduralMeshLibrary.h"
 #include "LegacyScreenPercentageDriver.h"
 #include "Math/UnrealMathUtility.h"
 #include "PreviewScene.h"
+#include "ProceduralMeshComponent.h"
 #include "RayTracingDebugVisualizationMenuCommands.h"
 #include "Renderer/Private/SceneRendering.h"
 #include "ScopedTransaction.h"
@@ -232,8 +234,8 @@ bool FDisplayClusterLightCardEditorViewportClient::FNormalMap::GetNormalAndDista
 {
 	auto GetPixel = [this](uint32 InX, uint32 InY)
 	{
-		uint32 ClampedX = FMath::Clamp(InX, (uint32)0, SizeX);
-		uint32 ClampedY = FMath::Clamp(InY, (uint32)0, SizeY);
+		uint32 ClampedX = FMath::Clamp(InX, (uint32)0, SizeX - 1);
+		uint32 ClampedY = FMath::Clamp(InY, (uint32)0, SizeY - 1);
 
 		return CachedNormalData[ClampedY * SizeX + ClampedX].GetFloats();
 	};
@@ -283,6 +285,43 @@ bool FDisplayClusterLightCardEditorViewportClient::FNormalMap::GetNormalAndDista
 	}
 }
 
+void FDisplayClusterLightCardEditorViewportClient::FNormalMap::MorphProceduralMesh(UProceduralMeshComponent* InProceduralMeshComponent) const
+{
+	const FVector ViewOrigin = ViewMatrices.GetViewOrigin();
+	const FVector ViewDirection = ViewMatrices.GetViewMatrix().TransformVector(FVector::ZAxisVector);
+	const float MaxAngle = 1.0f / ViewMatrices.GetProjectionMatrix().M[0][0];
+
+	FProcMeshSection& Section = *InProceduralMeshComponent->GetProcMeshSection(0);
+
+
+	for (int32 Index = 0; Index < Section.ProcVertexBuffer.Num(); ++Index)
+	{
+		FProcMeshVertex& Vertex = Section.ProcVertexBuffer[Index];
+
+		const FVector VertexPosition = Vertex.Position;
+		const FVector VertexWorldPosition = VertexPosition + ViewOrigin;
+		const FVector VertexDirection = Vertex.Position.GetSafeNormal();
+		const float VertexAngle = FMath::Acos(VertexDirection | ViewDirection);
+
+		if (VertexAngle < MaxAngle)
+		{
+			FVector Normal;
+			float Depth;
+			GetNormalAndDistanceAtPosition(VertexWorldPosition, Normal, Depth);
+
+			const FVector NewPosition = VertexDirection * Depth;
+			Vertex.Position = NewPosition;
+
+			const FMatrix RadialBasis = FRotationMatrix::MakeFromX(VertexDirection);
+
+			const FVector WorldNormal = RadialBasis.TransformVector(Normal);
+			Vertex.Normal = WorldNormal;
+		}
+	}
+
+	InProceduralMeshComponent->SetProcMeshSection(0, Section);
+}
+
 UTexture2D* FDisplayClusterLightCardEditorViewportClient::FNormalMap::GenerateNormalMapTexture(const FString& TextureName)
 {
 	if (NormalMapTexture.IsValid())
@@ -314,7 +353,7 @@ UTexture2D* FDisplayClusterLightCardEditorViewportClient::FNormalMap::GenerateNo
 //////////////////////////////////////////////////////////////////////////
 // FDisplayClusterLightCardEditorViewportClient
 
-FDisplayClusterLightCardEditorViewportClient::FDisplayClusterLightCardEditorViewportClient(FAdvancedPreviewScene& InPreviewScene,
+FDisplayClusterLightCardEditorViewportClient::FDisplayClusterLightCardEditorViewportClient(FPreviewScene& InPreviewScene,
 	const TWeakPtr<SDisplayClusterLightCardEditorViewport>& InEditorViewportWidget)
 	: FEditorViewportClient(nullptr, &InPreviewScene, InEditorViewportWidget)
 	, LightCardEditorViewportPtr(InEditorViewportWidget)
@@ -365,7 +404,7 @@ FDisplayClusterLightCardEditorViewportClient::FDisplayClusterLightCardEditorView
 	
 	UpdatePreviewActor(LightCardEditorPtr.Pin()->GetActiveRootActor().Get());
 
-	SetProjectionMode(EDisplayClusterMeshProjectionType::Azimuthal);
+	SetProjectionMode(EDisplayClusterMeshProjectionType::Azimuthal, ELevelViewportType::LVT_Perspective);
 }
 
 FDisplayClusterLightCardEditorViewportClient::~FDisplayClusterLightCardEditorViewportClient()
@@ -429,12 +468,12 @@ void FDisplayClusterLightCardEditorViewportClient::Tick(float DeltaSeconds)
 	}
 
 	// EditorViewportClient sets the cursor settings based on the state of the built in FWidget, which isn't being used here, so
-	// force a software cursor if we are dragging an actor so that the correct mouse cursor shows up
+	// force a hardware cursor if we are dragging an actor so that the correct mouse cursor shows up
 	switch (InputMode)
 	{
 	case EInputMode::DraggingActor:
-		SetRequiredCursor(false, true);
-		SetRequiredCursorOverride(true, EMouseCursor::CardinalCross);
+		SetRequiredCursor(true, false);
+		SetRequiredCursorOverride(true, EMouseCursor::GrabHandClosed);
 		break;
 
 	case EInputMode::DrawingLightCard:
@@ -469,6 +508,13 @@ void FDisplayClusterLightCardEditorViewportClient::Draw(FViewport* InViewport, F
 	{
 		RenderNormalMap(NorthNormalMap, FVector::UpVector);
 		RenderNormalMap(SouthNormalMap, -FVector::UpVector);
+
+		if (NormalMapMeshComponent.IsValid())
+		{
+			NorthNormalMap.MorphProceduralMesh(NormalMapMeshComponent.Get());
+			SouthNormalMap.MorphProceduralMesh(NormalMapMeshComponent.Get());
+		}
+
 		bNormalMapInvalid = false;
 	}
 
@@ -568,13 +614,20 @@ void FDisplayClusterLightCardEditorViewportClient::Draw(FViewport* InViewport, F
 
 	Canvas->Clear(FLinearColor::Black);
 
-	FDisplayClusterScenePreviewRenderSettings RenderSettings;
-	RenderSettings.RenderType = FDisplayClusterScenePreviewRenderSettings::ERenderType::Color;
-	RenderSettings.EngineShowFlags = EngineShowFlags;
-	RenderSettings.ProjectionType = ProjectionMode;
-	GetSceneViewInitOptions(RenderSettings.ViewInitOptions);
+	if (!bDisableCustomRenderer)
+	{
+		FDisplayClusterScenePreviewRenderSettings RenderSettings;
+		RenderSettings.RenderType = FDisplayClusterScenePreviewRenderSettings::ERenderType::Color;
+		RenderSettings.EngineShowFlags = EngineShowFlags;
+		RenderSettings.ProjectionType = ProjectionMode;
+		GetSceneViewInitOptions(RenderSettings.ViewInitOptions);
 
-	IDisplayClusterScenePreview::Get().Render(PreviewRendererId, RenderSettings, *Canvas);
+		IDisplayClusterScenePreview::Get().Render(PreviewRendererId, RenderSettings, *Canvas);
+	}
+	else
+	{
+		GetRendererModule().BeginRenderingViewFamily(Canvas, &ViewFamily);
+	}
 
 	if (View)
 	{
@@ -828,6 +881,16 @@ void FDisplayClusterLightCardEditorViewportClient::TrackingStarted(
 		FVector Direction;
 		PixelToWorld(*View, MousePos, Origin, Direction);
 
+		if (RenderViewportType != LVT_Perspective)
+		{
+			// For orthogonal projections, PixelToWorld does not return the view origin or a direction from the view origin. Use TraceScreenRay
+			// to find a useful direction away from the view origin to use
+			const FVector ViewOrigin = View->ViewLocation;
+
+			Direction = TraceScreenRay(Origin, Direction, ViewOrigin);
+			Origin = ViewOrigin;
+		}
+
 		DragWidgetOffset = Direction - (CachedEditorWidgetWorldTransform.GetTranslation() - Origin).GetSafeNormal();
 		LastWidgetMousePos = MousePos;
 	}
@@ -900,14 +963,27 @@ void FDisplayClusterLightCardEditorViewportClient::CreateDrawnLightCard(const TA
 	).SetRealtimeUpdate(IsRealtime()));
 
 	FSceneView* View = CalcSceneView(&ViewFamily);
+	const FVector ViewOrigin = View->ViewLocation;
 
-	FVector ViewOrigin;
 	TArray<FVector> MouseWorldDirections; // directions from view origin
+
+	FVector ScreenOrigin;
+	FVector ScreenDirection;
 	MouseWorldDirections.AddUninitialized(MousePositions.Num());
 
 	for (int32 PointIdx = 0; PointIdx < MousePositions.Num(); ++PointIdx)
 	{
-		PixelToWorld(*View, MousePositions[PointIdx], ViewOrigin, MouseWorldDirections[PointIdx]);
+		PixelToWorld(*View, MousePositions[PointIdx], ScreenOrigin, ScreenDirection);
+
+		if (RenderViewportType != LVT_Perspective)
+		{
+			// For orthogonal projections, PixelToWorld does not return the view origin or a direction from the view origin. Use TraceScreenRay
+			// to find a useful direction away from the view origin to use
+
+			ScreenDirection = TraceScreenRay(ScreenOrigin, ScreenDirection, ViewOrigin);
+		}
+
+		MouseWorldDirections[PointIdx] = ScreenDirection;
 	}
 
 	//
@@ -1142,9 +1218,6 @@ void FDisplayClusterLightCardEditorViewportClient::ProcessClick(FSceneView& View
 		return;
 	}
 
-	UWorld* PreviewWorld = PreviewScene->GetWorld();
-	check(PreviewWorld);
-
 	const bool bIsCtrlKeyDown = Viewport->KeyState(EKeys::LeftControl) || Viewport->KeyState(EKeys::RightControl);
 	const bool bMultiSelect = Key == EKeys::LeftMouseButton && bIsCtrlKeyDown;
 	const bool bIsRightClickSelection = Key == EKeys::RightMouseButton && !bIsCtrlKeyDown && !Viewport->KeyState(EKeys::LeftMouseButton);
@@ -1315,7 +1388,42 @@ void FDisplayClusterLightCardEditorViewportClient::UpdatePreviewActor(ADisplayCl
 					ProxyConfig->StageSettings.Lightcard.bEnable = false;
 				}
 
-				RootActorBoundingRadius = 0.5f * RootActorProxy->GetComponentsBoundingBox().GetSize().Length();
+				FBox BoundingBox = RootActorProxy->GetComponentsBoundingBox();
+				RootActorBoundingRadius = FMath::Max(BoundingBox.Min.Length(), BoundingBox.Max.Length());
+
+				// Create the normal map mesh component, which is used to run traces on the generated normal map of the stage
+				{
+					NormalMapMeshComponent = NewObject<UProceduralMeshComponent>(GetTransientPackage(), TEXT("NormalMapMesh"));
+					NormalMapMeshComponent->SetCollisionEnabled(ECollisionEnabled::QueryOnly);
+					PreviewScene->AddComponent(NormalMapMeshComponent.Get(), FTransform(ProjectionOriginComponent.Get()->GetComponentLocation()));
+
+					UStaticMesh* IcoSphereMesh = LoadObject<UStaticMesh>(nullptr, TEXT("/nDisplay/Meshes/SM_IcoSphere.SM_IcoSphere"), nullptr, LOAD_None, nullptr);
+
+					if (ensure(IcoSphereMesh))
+					{
+						int32 NumSections = IcoSphereMesh->GetNumSections(0);
+						for (int32 SectionIndex = 0; SectionIndex < NumSections; SectionIndex++)
+						{
+							TArray<FVector> Vertices;
+							TArray<int32> Triangles;
+							TArray<FVector> Normals;
+							TArray<FVector2D> UVs;
+							TArray<FProcMeshTangent> Tangents;
+
+							UKismetProceduralMeshLibrary::GetSectionFromStaticMesh(IcoSphereMesh, 0, SectionIndex, Vertices, Triangles, Normals, UVs, Tangents);
+
+							TArray<FVector2D> EmptyUVs;
+							TArray<FLinearColor> EmptyColors;
+							NormalMapMeshComponent.Get()->CreateMeshSection_LinearColor(SectionIndex, Vertices, Triangles, Normals, UVs, EmptyUVs, EmptyUVs, EmptyUVs, EmptyColors, Tangents, true);
+
+							for (int32 Index = 0; Index < IcoSphereMesh->GetStaticMaterials().Num(); ++Index)
+							{
+								UMaterialInterface* MaterialInterface = IcoSphereMesh->GetStaticMaterials()[Index].MaterialInterface;
+								NormalMapMeshComponent.Get()->SetMaterial(Index, MaterialInterface);
+							}
+						}
+					}
+				}
 			}
 
 			// Filter out any primitives hidden in game except screen components
@@ -1407,9 +1515,33 @@ void FDisplayClusterLightCardEditorViewportClient::UpdateProxyTransforms()
 }
 
 void FDisplayClusterLightCardEditorViewportClient::DestroyProxies(
-EDisplayClusterLightCardEditorProxyType ProxyType)
+	EDisplayClusterLightCardEditorProxyType ProxyType)
 {
-	IDisplayClusterScenePreview::Get().ClearRendererScene(PreviewRendererId);
+	// Clear the primitives from the scene renderer based on the type of proxy that is being destroyed
+	switch (ProxyType)
+	{
+	case EDisplayClusterLightCardEditorProxyType::RootActor:
+		if (RootActorProxy.IsValid())
+		{
+			IDisplayClusterScenePreview::Get().RemoveActorFromRenderer(PreviewRendererId, RootActorProxy.Get());
+		}
+		break;
+
+	case EDisplayClusterLightCardEditorProxyType::LightCards:
+		for (const FLightCardProxy& LightCardProxy : LightCardProxies)
+		{
+			if (LightCardProxy.Proxy.IsValid())
+			{
+				IDisplayClusterScenePreview::Get().RemoveActorFromRenderer(PreviewRendererId, LightCardProxy.Proxy.Get());
+			}
+		}
+		break;
+
+	case EDisplayClusterLightCardEditorProxyType::All:
+	default:
+		IDisplayClusterScenePreview::Get().ClearRendererScene(PreviewRendererId);
+		break;
+	}
 
 	UWorld* PreviewWorld = PreviewScene->GetWorld();
 	check(PreviewWorld);
@@ -1428,6 +1560,12 @@ EDisplayClusterLightCardEditorProxyType ProxyType)
 			RootActorLevelInstance->UnsubscribeFromPostProcessRenderTarget(reinterpret_cast<uint8*>(this));
 			RootActorLevelInstance.Reset();
 		}
+
+		if (NormalMapMeshComponent.IsValid())
+		{
+			PreviewScene->RemoveComponent(NormalMapMeshComponent.Get());
+			NormalMapMeshComponent.Reset();
+		}
 	}
 	
 	if (ProxyType == EDisplayClusterLightCardEditorProxyType::All ||
@@ -1460,11 +1598,12 @@ void FDisplayClusterLightCardEditorViewportClient::SelectLightCards(const TArray
 	}
 }
 
-void FDisplayClusterLightCardEditorViewportClient::SetProjectionMode(EDisplayClusterMeshProjectionType InProjectionMode)
+void FDisplayClusterLightCardEditorViewportClient::SetProjectionMode(EDisplayClusterMeshProjectionType InProjectionMode, ELevelViewportType InViewportType)
 {
 	ProjectionMode = InProjectionMode;
+	RenderViewportType = InViewportType;
 
-	if (ProjectionMode == EDisplayClusterMeshProjectionType::Perspective)
+	if (ProjectionMode == EDisplayClusterMeshProjectionType::Linear)
 	{
 		// TODO: Do we want to cache the perspective rotation and restore it when the user switches back?
 		SetViewRotation(FVector::ForwardVector.Rotation());
@@ -1530,7 +1669,7 @@ void FDisplayClusterLightCardEditorViewportClient::ResetCamera(bool bLocationOnl
 		return;
 	}
 	
-	SetProjectionMode(GetProjectionMode());
+	SetProjectionMode(GetProjectionMode(), GetRenderViewportType());
 
 	ResetFOVs();
 }
@@ -1609,7 +1748,18 @@ void FDisplayClusterLightCardEditorViewportClient::MoveSelectedLightCardsToPixel
 	FVector Origin;
 	FVector Direction;
 	PixelToWorld(*View, PixelPos, Origin, Direction);
- 
+
+	if (RenderViewportType != LVT_Perspective)
+	{
+		// For orthogonal projections, PixelToWorld does not return the view origin or a direction from the view origin. Use TraceScreenRay
+		// to find a useful direction away from the view origin to use
+
+		const FVector ViewOrigin = View->ViewLocation;
+
+		Direction = TraceScreenRay(Origin, Direction, ViewOrigin);
+		Origin = ViewOrigin;
+	}
+
 	// Compute desired coordinates (radius doesn't matter here since we will use the flush constraint on the light cards after moving them)
 	const FSphericalCoordinates DesiredCoords(Direction * 100.0f);
 	const FSphericalCoordinates DeltaCoords = DesiredCoords - AverageCoords;
@@ -1683,10 +1833,10 @@ void FDisplayClusterLightCardEditorViewportClient::GetSceneViewInitOptions(FScen
 
 	const float MinZ = GetNearClipPlane();
 	const float MaxZ = MinZ;
-	const float FieldOfView = GetProjectionModeFOV(ProjectionMode);
+	const float FieldOfView = FMath::DegreesToRadians(GetProjectionModeFOV(ProjectionMode));
 
 	// Avoid zero ViewFOV's which cause divide by zero's in projection matrix
-	const float MatrixFOV = FMath::Max(0.001f, FieldOfView) * (float)PI / 360.0f;
+	const float HalfFOV = FMath::Max(0.5f * FieldOfView, 0.001f);
 
 	float XAxisMultiplier;
 	float YAxisMultiplier;
@@ -1706,27 +1856,59 @@ void FDisplayClusterLightCardEditorViewportClient::GetSceneViewInitOptions(FScen
 		YAxisMultiplier = 1.0f;
 	}
 
-	if ((bool)ERHIZBuffer::IsInverted)
+	if (RenderViewportType == LVT_Perspective)
 	{
-		ViewInitOptions.ProjectionMatrix = FReversedZPerspectiveMatrix(
-			MatrixFOV,
-			MatrixFOV,
-			XAxisMultiplier,
-			YAxisMultiplier,
-			MinZ,
-			MaxZ
+		if ((bool)ERHIZBuffer::IsInverted)
+		{
+			ViewInitOptions.ProjectionMatrix = FReversedZPerspectiveMatrix(
+				HalfFOV,
+				HalfFOV,
+				XAxisMultiplier,
+				YAxisMultiplier,
+				MinZ,
+				MaxZ
 			);
+		}
+		else
+		{
+			ViewInitOptions.ProjectionMatrix = FPerspectiveMatrix(
+				HalfFOV,
+				HalfFOV,
+				XAxisMultiplier,
+				YAxisMultiplier,
+				MinZ,
+				MaxZ
+			);
+		}
 	}
 	else
 	{
-		ViewInitOptions.ProjectionMatrix = FPerspectiveMatrix(
-			MatrixFOV,
-			MatrixFOV,
-			XAxisMultiplier,
-			YAxisMultiplier,
-			MinZ,
-			MaxZ
+		const float ZScale = 0.5f / HALF_WORLD_MAX;
+		const float ZOffset = HALF_WORLD_MAX;
+
+		const float FOVScale = FMath::Tan(HalfFOV) / GetDPIScale();
+
+		const float OrthoWidth = 0.5f * FOVScale * ViewportSize.X;
+		const float OrthoHeight = 0.5f * FOVScale * ViewportSize.Y;
+
+		if ((bool)ERHIZBuffer::IsInverted)
+		{
+			ViewInitOptions.ProjectionMatrix = FReversedZOrthoMatrix(
+				OrthoWidth,
+				OrthoHeight,
+				ZScale,
+				ZOffset
 			);
+		}
+		else
+		{
+			ViewInitOptions.ProjectionMatrix = FOrthoMatrix(
+				OrthoWidth,
+				OrthoHeight,
+				ZScale,
+				ZOffset
+			);
+		}
 	}
 
 	if (!ViewInitOptions.IsValidViewRectangle())
@@ -1742,10 +1924,6 @@ void FDisplayClusterLightCardEditorViewportClient::GetSceneViewInitOptions(FScen
 
 	ViewInitOptions.EditorViewBitflag = (uint64)1 << ViewIndex, // send the bit for this view - each actor will check it's visibility bits against this
 
-	// for ortho views to steal perspective view origin
-	ViewInitOptions.OverrideLODViewOrigin = FVector::ZeroVector;
-	ViewInitOptions.bUseFauxOrthoViewPos = true;
-
 	ViewInitOptions.FOV = FieldOfView;
 	ViewInitOptions.OverrideFarClippingPlaneDistance = GetFarClipPlaneOverride();
 	ViewInitOptions.CursorPos = CurrentMousePos;
@@ -2026,9 +2204,6 @@ void FDisplayClusterLightCardEditorViewportClient::VerifyAndFixLightCardOrigin(A
 
 void FDisplayClusterLightCardEditorViewportClient::MoveSelectedLightCards(FViewport* InViewport, EAxisList::Type CurrentAxis)
 {
-	FIntPoint MousePos;
-	InViewport->GetMousePos(MousePos);
-
 	FSceneViewFamilyContext ViewFamily(FSceneViewFamily::ConstructionValues(
 		InViewport,
 		GetScene(),
@@ -2038,10 +2213,6 @@ void FDisplayClusterLightCardEditorViewportClient::MoveSelectedLightCards(FViewp
 
 	FSceneView* View = CalcSceneView(&ViewFamily);
 
-	FVector Origin;
-	FVector Direction;
-	PixelToWorld(*View, MousePos, Origin, Direction);
-
 	const TWeakObjectPtr<ADisplayClusterLightCardActor>& LastSelectedLightCard = SelectedLightCards.Last();
 
 	if (LastSelectedLightCard.IsValid())
@@ -2204,6 +2375,17 @@ FDisplayClusterLightCardEditorViewportClient::FSphericalCoordinates FDisplayClus
 	FVector Direction;
 	PixelToWorld(*View, MousePos, Origin, Direction);
 
+	if (RenderViewportType != LVT_Perspective)
+	{
+		// For orthogonal projections, PixelToWorld does not return the view origin or a direction from the view origin. Use TraceScreenRay
+		// to find a useful direction away from the view origin to use
+
+		const FVector ViewOrigin = View->ViewLocation;
+
+		Direction = TraceScreenRay(Origin, Direction, ViewOrigin);
+		Origin = ViewOrigin;
+	}
+
 	Direction = (Direction - DragWidgetOffset).GetSafeNormal();
 
 	const FVector LocalDirection = LightCard->GetActorRotation().RotateVector(Direction);
@@ -2318,7 +2500,9 @@ FVector2D FDisplayClusterLightCardEditorViewportClient::GetLightCardScaleDelta(F
 	const FVector LightCardSize3D = LightCard->GetLightCardBounds(bLocalSpace).GetSize();
 	const FVector2D SizeToScale = FVector2D((CurrentAxis & EAxisList::X) * LightCardSize3D.Y, (CurrentAxis & EAxisList::Y) * LightCardSize3D.Z);
 	const double DistanceFromCamera = FMath::Max(FVector::Dist(WorldWidgetOrigin, View->ViewMatrices.GetViewOrigin()), 1.0f);
-	const double ScreenSize = View->ViewMatrices.GetScreenScale() * SizeToScale.Length() / DistanceFromCamera;
+	const double ScreenSize = RenderViewportType == LVT_Perspective 
+		? View->ViewMatrices.GetScreenScale() * SizeToScale.Length() / DistanceFromCamera 
+		: SizeToScale.Length() * View->ViewMatrices.GetProjectionMatrix().M[0][0] * View->UnscaledViewRect.Width();
 
 	// Compute the scale delta as s' - s = 2d / h_0
 	const double ScaleMagnitude = 2.0f * (ScaleDir | DragDir) / ScreenSize;
@@ -2406,6 +2590,89 @@ FDisplayClusterLightCardEditorViewportClient::FSphericalCoordinates FDisplayClus
 	return LightCardCoords;
 }
 
+bool FDisplayClusterLightCardEditorViewportClient::TraceStage(const FVector& RayStart, const FVector& RayEnd, FVector& OutHitLocation) const
+{
+	UWorld* PreviewWorld = PreviewScene->GetWorld();
+	check(PreviewWorld);
+
+	FCollisionQueryParams TraceParams(SCENE_QUERY_STAT(DisplayClusterStageTrace), true);
+
+	for (const FLightCardProxy& ProxyRef : LightCardProxies)
+	{
+		if (ProxyRef.Proxy.IsValid())
+		{
+			TraceParams.AddIgnoredActor(ProxyRef.Proxy.Get());
+		}
+	}
+
+	TraceParams.AddIgnoredComponent(NormalMapMeshComponent.Get());
+
+	FHitResult HitResult;
+	if (PreviewWorld->LineTraceSingleByObjectType(HitResult, RayStart, RayEnd, FCollisionObjectQueryParams(FCollisionObjectQueryParams::InitType::AllObjects), TraceParams))
+	{
+		// If the actor we hit was the stage root actor, return the hit location.
+		if (RootActorProxy.Get() == HitResult.GetActor())
+		{
+			OutHitLocation = HitResult.Location;
+			return true;
+		}
+	}
+
+	OutHitLocation = FVector::ZeroVector;
+	return false;
+}
+
+FVector FDisplayClusterLightCardEditorViewportClient::TraceScreenRay(const FVector& OrthogonalOrigin, const FVector& OrthogonalDirection, const FVector& ViewOrigin)
+{
+	const FVector RayStart = OrthogonalOrigin;
+	const FVector RayEnd = OrthogonalOrigin + OrthogonalDirection * WORLD_MAX;
+
+	FVector Direction = FVector::ZeroVector;
+
+	// First, trace against the stage actor to see if the screen ray hits it; if so, simply return the direction from the view origin to this hit point
+	FVector HitLocation = FVector::ZeroVector;
+	if (TraceStage(RayStart, RayEnd, HitLocation))
+	{
+		Direction = (HitLocation - ViewOrigin).GetSafeNormal();
+	}
+	else
+	{
+		// If we didn't hit any stage geometry, try to trace against the normal map mesh. Procedural meshes does not appear to handle inward pointing normals correctly,
+		// so we need to reverse the trace start and end locations to get a useful hit result
+
+		UWorld* PreviewWorld = PreviewScene->GetWorld();
+		check(PreviewWorld);
+
+		FCollisionQueryParams TraceParams(SCENE_QUERY_STAT(DisplayClusterStageTrace), true);
+
+		for (const FLightCardProxy& ProxyRef : LightCardProxies)
+		{
+			if (ProxyRef.Proxy.IsValid())
+			{
+				TraceParams.AddIgnoredActor(ProxyRef.Proxy.Get());
+			}
+		}
+
+		TraceParams.AddIgnoredActor(RootActorProxy.Get());
+
+		FHitResult HitResult;
+		if (PreviewWorld->LineTraceSingleByObjectType(HitResult, RayEnd, RayStart, FCollisionObjectQueryParams(FCollisionObjectQueryParams::InitType::AllObjects), TraceParams))
+		{
+			HitLocation = HitResult.Location;
+			Direction = (HitLocation - ViewOrigin).GetSafeNormal();
+		}
+		else
+		{
+			// If the screen ray does not hit the stage or the normal map mesh, then simply use the closest point on the ray to the view origin
+			const FVector ClosestPoint = OrthogonalOrigin + ((ViewOrigin - OrthogonalOrigin) | OrthogonalDirection) * OrthogonalDirection;
+
+			Direction = (ClosestPoint - ViewOrigin).GetSafeNormal();
+		}
+	}
+
+	return Direction;
+}
+
 ADisplayClusterLightCardActor* FDisplayClusterLightCardEditorViewportClient::TraceScreenForLightCard(const FSceneView& View, int32 HitX, int32 HitY)
 {
 	UWorld* PreviewWorld = PreviewScene->GetWorld();
@@ -2416,13 +2683,15 @@ ADisplayClusterLightCardActor* FDisplayClusterLightCardEditorViewportClient::Tra
 	PixelToWorld(View, FIntPoint(HitX, HitY), Origin, Direction);
 
 	const FVector CursorRayStart = Origin;
-	const FVector CursorRayEnd = CursorRayStart + Direction * HALF_WORLD_MAX;
+	const FVector CursorRayEnd = CursorRayStart + Direction * (RenderViewportType == LVT_Perspective ? HALF_WORLD_MAX : WORLD_MAX);
 
-	FCollisionQueryParams Param(SCENE_QUERY_STAT(DragDropTrace), true);
+	FCollisionQueryParams TraceParams(SCENE_QUERY_STAT(LightCardTrace), true);
+
+	TraceParams.AddIgnoredComponent(NormalMapMeshComponent.Get());
 
 	bool bHitLightCard = false;
 	FHitResult ScreenHitResult;
-	if (PreviewWorld->LineTraceSingleByObjectType(ScreenHitResult, CursorRayStart, CursorRayEnd, FCollisionObjectQueryParams(FCollisionObjectQueryParams::InitType::AllObjects), Param))
+	if (PreviewWorld->LineTraceSingleByObjectType(ScreenHitResult, CursorRayStart, CursorRayEnd, FCollisionObjectQueryParams(FCollisionObjectQueryParams::InitType::AllObjects), TraceParams))
 	{
 		if (AActor* HitActor = ScreenHitResult.GetActor())
 		{
@@ -2449,7 +2718,7 @@ ADisplayClusterLightCardActor* FDisplayClusterLightCardEditorViewportClient::Tra
 						const FVector ViewOriginRayEnd = ViewOriginRayStart + (ScreenHitResult.Location - ViewOriginRayStart) * HALF_WORLD_MAX;
 
 						TArray<FHitResult> HitResults;
-						if (PreviewWorld->LineTraceMultiByObjectType(HitResults, ViewOriginRayStart, ViewOriginRayEnd, FCollisionObjectQueryParams(FCollisionObjectQueryParams::InitType::AllObjects), Param))
+						if (PreviewWorld->LineTraceMultiByObjectType(HitResults, ViewOriginRayStart, ViewOriginRayEnd, FCollisionObjectQueryParams(FCollisionObjectQueryParams::InitType::AllObjects), TraceParams))
 						{
 							for (FHitResult& HitResult : HitResults)
 							{
@@ -2487,11 +2756,23 @@ void FDisplayClusterLightCardEditorViewportClient::PixelToWorld(const FSceneView
 	const FMatrix& InvViewMatrix = View.ViewMatrices.GetInvViewMatrix();
 
 	const FVector4 ScreenPos = View.PixelToScreen(PixelPos.X, PixelPos.Y, 0);
-	const FVector ViewPos = FVector(InvProjMatrix.TransformFVector4(FVector4(ScreenPos.X * GNearClippingPlane, ScreenPos.Y * GNearClippingPlane, 0.0f, GNearClippingPlane)));
+	const FVector4 HomogeneousPos = RenderViewportType == LVT_Perspective
+		? FVector4(ScreenPos.X * GNearClippingPlane, ScreenPos.Y * GNearClippingPlane, 0.0f, GNearClippingPlane)
+		: FVector4(ScreenPos.X, ScreenPos.Y, 1.0f, 1.0f);
+
+	const FVector ViewPos = FVector(InvProjMatrix.TransformFVector4(HomogeneousPos));
 	const FVector UnprojectedViewPos = FDisplayClusterMeshProjectionRenderer::UnprojectViewPosition(ViewPos, ProjectionMode);
 
-	OutOrigin = View.ViewMatrices.GetViewOrigin();
-	OutDirection = InvViewMatrix.TransformVector(UnprojectedViewPos).GetSafeNormal();
+	if (RenderViewportType == LVT_Perspective)
+	{
+		OutOrigin = View.ViewMatrices.GetViewOrigin();
+		OutDirection = InvViewMatrix.TransformVector(UnprojectedViewPos).GetSafeNormal();
+	}
+	else
+	{
+		OutOrigin = InvViewMatrix.TransformFVector4(UnprojectedViewPos);
+		OutDirection = InvViewMatrix.TransformVector(FVector(0, 0, 1)).GetSafeNormal();
+	}
 }
 
 bool FDisplayClusterLightCardEditorViewportClient::WorldToPixel(const FSceneView& View, const FVector& WorldPos, FVector2D& OutPixelPos) const
@@ -2687,7 +2968,7 @@ void FDisplayClusterLightCardEditorViewportClient::ResetFOVs()
 	{
 		ProjectionFOVs.AddDefaulted(MaxFOVs - ProjectionFOVs.Num());
 	}
-	ProjectionFOVs[static_cast<int32>(EDisplayClusterMeshProjectionType::Perspective)] = 90.0f;
+	ProjectionFOVs[static_cast<int32>(EDisplayClusterMeshProjectionType::Linear)] = 90.0f;
 	ProjectionFOVs[static_cast<int32>(EDisplayClusterMeshProjectionType::Azimuthal)] = 130.0f;
 }
 
diff --git a/Engine/Plugins/Runtime/nDisplay/Source/DisplayClusterLightCardEditor/Private/Viewport/DisplayClusterLightCardEditorViewportClient.h b/Engine/Plugins/Runtime/nDisplay/Source/DisplayClusterLightCardEditor/Private/Viewport/DisplayClusterLightCardEditorViewportClient.h
--- a/Engine/Plugins/Runtime/nDisplay/Source/DisplayClusterLightCardEditor/Private/Viewport/DisplayClusterLightCardEditorViewportClient.h
+++ b/Engine/Plugins/Runtime/nDisplay/Source/DisplayClusterLightCardEditor/Private/Viewport/DisplayClusterLightCardEditorViewportClient.h
@@ -17,6 +17,7 @@ class SDisplayClusterLightCardEditor;
 class SDisplayClusterLightCardEditorViewport;
 class FScopedTransaction;
 class UDisplayClusterConfigurationViewport;
+class UProceduralMeshComponent;
 
 /** Viewport Client for the preview viewport */
 class FDisplayClusterLightCardEditorViewportClient : public FEditorViewportClient, public TSharedFromThis<FDisplayClusterLightCardEditorViewportClient>
@@ -97,6 +98,9 @@ private:
 		/** Gets the normal vector and distance at the specified world location. The normal and distance are bilinearly interpolated from the nearest pixels in the normal map */
 		bool GetNormalAndDistanceAtPosition(FVector Position, FVector& OutNormal, float& OutDistance) const;
 
+		/** Morphs the vertices of the specified prodedural mesh to match the normal map */
+		void MorphProceduralMesh(UProceduralMeshComponent* InProceduralMeshComponent) const;
+
 		/** Generates a texture object that can be used to visualize the normal map */
 		UTexture2D* GenerateNormalMapTexture(const FString& TextureName);
 
@@ -121,7 +125,7 @@ private:
 	};
 
 public:
-	FDisplayClusterLightCardEditorViewportClient(FAdvancedPreviewScene& InPreviewScene,
+	FDisplayClusterLightCardEditorViewportClient(FPreviewScene& InPreviewScene,
 		const TWeakPtr<SDisplayClusterLightCardEditorViewport>& InEditorViewportWidget);
 	virtual ~FDisplayClusterLightCardEditorViewportClient() override;
 	
@@ -169,8 +173,14 @@ public:
 	FDisplayClusterLightCardEditorWidget::EWidgetMode GetEditorWidgetMode() const { return EditorWidget->GetWidgetMode(); }
 	void SetEditorWidgetMode(FDisplayClusterLightCardEditorWidget::EWidgetMode InWidgetMode) { EditorWidget->SetWidgetMode(InWidgetMode); }
 
+	/** Gets the current projection mode of the editor viewport */
 	EDisplayClusterMeshProjectionType GetProjectionMode() const { return ProjectionMode; }
-	void SetProjectionMode(EDisplayClusterMeshProjectionType InProjectionMode);
+
+	/** Gets the current viewport type the viewport is being rendered with */
+	ELevelViewportType GetRenderViewportType() const { return RenderViewportType; }
+
+	/** Sets the projection mode and the render viewport type of the viewport */
+	void SetProjectionMode(EDisplayClusterMeshProjectionType InProjectionMode, ELevelViewportType InViewportType);
 
 	/** Gets the field of view of the specified projection mode */
 	float GetProjectionModeFOV(EDisplayClusterMeshProjectionType InProjectionMode) const;
@@ -268,6 +278,12 @@ private:
 	/** Sets the light card position to the given spherical coordinates */
 	void SetLightCardCoordinates(ADisplayClusterLightCardActor* LightCard, const FSphericalCoordinates& SphericalCoords) const;
 
+	/** Performs a ray trace against the stage's geometry, and returns the hit point */
+	bool TraceStage(const FVector& RayStart, const FVector& RayEnd, FVector& OutHitLocation) const;
+
+	/** Traces the world geometry to find the best direction vector from the view origin to a valid point in space using a screen ray */
+	FVector TraceScreenRay(const FVector& RayOrigin, const FVector& RayDirection, const FVector& ViewOrigin);
+
 	/** Traces to find the light card corresponding to a click on a stage screen */
 	ADisplayClusterLightCardActor* TraceScreenForLightCard(const FSceneView& View, int32 HitX, int32 HitY);
 
@@ -363,6 +379,9 @@ private:
 	/** The current projection mode the 3D viewport is being displayed with */
 	EDisplayClusterMeshProjectionType ProjectionMode = EDisplayClusterMeshProjectionType::Azimuthal;
 
+	/** The viewport type (perspective or orthogonal) to use when rendering the viewport. Separate from ViewportType since ViewportType also determines input functionality */
+	ELevelViewportType RenderViewportType = LVT_Perspective;
+
 	/** The component of the root actor that is acting as the projection origin. Can be either the root component (stage origin) or a view origin component */
 	TWeakObjectPtr<USceneComponent> ProjectionOriginComponent;
 
@@ -378,6 +397,9 @@ private:
 	/** The render target used to render a map of the screens' normals for the southern hemisphere of the view */
 	FNormalMap SouthNormalMap;
 
+	/** A morphed ico-sphere mesh component that approximates the normal and depth map */
+	TWeakObjectPtr<UProceduralMeshComponent> NormalMapMeshComponent;
+
 	/** Indicates if the cached normal map is invalid and needs to be redrawn */
 	bool bNormalMapInvalid = false;
 
@@ -392,4 +414,7 @@ private:
 
 	/** Multicast delegate that stores callbacks to be invoked the next time the scene is refreshed */
 	FOnNextSceneRefresh OnNextSceneRefreshDelegate;
+
+	/** A flag that disables drawing with the custom projection renderer and instead renders the viewport client with the editor's default renderer */
+	bool bDisableCustomRenderer = false;
 };
\ No newline at end of file
diff --git a/Engine/Plugins/Runtime/nDisplay/Source/DisplayClusterLightCardEditor/Private/Viewport/DisplayClusterLightCardEditorWidget.cpp b/Engine/Plugins/Runtime/nDisplay/Source/DisplayClusterLightCardEditor/Private/Viewport/DisplayClusterLightCardEditorWidget.cpp
--- a/Engine/Plugins/Runtime/nDisplay/Source/DisplayClusterLightCardEditor/Private/Viewport/DisplayClusterLightCardEditorWidget.cpp
+++ b/Engine/Plugins/Runtime/nDisplay/Source/DisplayClusterLightCardEditor/Private/Viewport/DisplayClusterLightCardEditorWidget.cpp
@@ -2,30 +2,34 @@
 
 #include "DisplayClusterLightCardEditorWidget.h"
 
+#include "DisplayClusterLightCardEditorViewportClient.h"
+
 #include "SceneManagement.h"
 #include "UnrealWidget.h"
-#include "EditorViewportClient.h"
 
-void FDisplayClusterLightCardEditorWidget::Draw(const FSceneView* View, const FEditorViewportClient* ViewportClient, FPrimitiveDrawInterface* PDI)
+void FDisplayClusterLightCardEditorWidget::Draw(const FSceneView* View, const FDisplayClusterLightCardEditorViewportClient* ViewportClient, FPrimitiveDrawInterface* PDI)
 {
+	const bool bIsOrthographic = ViewportClient->GetRenderViewportType() != LVT_Perspective;
+
 	const float SizeScalar = GetSizeScreenScalar(View, ViewportClient);
 	const float LengthScalar = GetLengthScreenScalar(View, ViewportClient, Transform.GetTranslation());
+	const float OrthoScalar = bIsOrthographic ? 1.0f / View->ViewMatrices.GetProjectionMatrix().M[0][0] : 1.0f;
 
 	if (WidgetMode == EWidgetMode::WM_RotateZ)
 	{
-		DrawAxis(PDI, EAxisList::Type::Y, SizeScalar, LengthScalar);
+		DrawAxis(PDI, EAxisList::Type::Y, SizeScalar, LengthScalar, OrthoScalar);
 		DrawCircle(PDI, EAxisList::Type::Z, SizeScalar, LengthScalar);
-		DrawOrigin(PDI, SizeScalar);
+		DrawOrigin(PDI, SizeScalar, OrthoScalar);
 	}
 	else
 	{
-		DrawAxis(PDI, EAxisList::Type::X, SizeScalar, LengthScalar);
-		DrawAxis(PDI, EAxisList::Type::Y, SizeScalar, LengthScalar);
-		DrawOrigin(PDI, SizeScalar);
+		DrawAxis(PDI, EAxisList::Type::X, SizeScalar, LengthScalar, OrthoScalar);
+		DrawAxis(PDI, EAxisList::Type::Y, SizeScalar, LengthScalar, OrthoScalar);
+		DrawOrigin(PDI, SizeScalar, OrthoScalar);
 	}
 }
 
-void FDisplayClusterLightCardEditorWidget::DrawAxis(FPrimitiveDrawInterface* PDI, EAxisList::Type Axis, float SizeScalar, float LengthScalar)
+void FDisplayClusterLightCardEditorWidget::DrawAxis(FPrimitiveDrawInterface* PDI, EAxisList::Type Axis, float SizeScalar, float LengthScalar, float OrthoScalar)
 {
 	const FVector Origin = Transform.GetTranslation();
 
@@ -42,26 +46,26 @@ void FDisplayClusterLightCardEditorWidget::DrawAxis(FPrimitiveDrawInterface* PDI
 
 		if (WidgetMode == EWidgetMode::WM_Scale)
 		{
-			PDI->DrawPoint(AxisEnd, AxisColor, AxisCapSize * SizeScalar, ESceneDepthPriorityGroup::SDPG_Foreground);
+			PDI->DrawPoint(AxisEnd, AxisColor, AxisCapSize * SizeScalar * OrthoScalar, ESceneDepthPriorityGroup::SDPG_Foreground);
 		}
 	}
 	PDI->SetHitProxy(nullptr);
 }
 
-void FDisplayClusterLightCardEditorWidget::DrawOrigin(FPrimitiveDrawInterface* PDI, float SizeScalar)
+void FDisplayClusterLightCardEditorWidget::DrawOrigin(FPrimitiveDrawInterface* PDI, float SizeScalar, float OrthoScalar)
 {
 	const FVector Origin = ProjectionTransform.ProjectPosition(Transform.GetTranslation());
 	const FLinearColor Color = HighlightedAxis == EAxisList::Type::XYZ ? HighlightColor : FLinearColor::Black;
 
 	PDI->SetHitProxy(new HWidgetAxis(EAxisList::Type::XYZ));
-	PDI->DrawPoint(Origin, Color, OriginSize * SizeScalar, ESceneDepthPriorityGroup::SDPG_Foreground);
+	PDI->DrawPoint(Origin, Color, OriginSize * SizeScalar * OrthoScalar, ESceneDepthPriorityGroup::SDPG_Foreground);
 	PDI->SetHitProxy(nullptr);
 }
 
 void FDisplayClusterLightCardEditorWidget::DrawCircle(FPrimitiveDrawInterface* PDI, EAxisList::Type Axis, float SizeScalar, float LengthScalar)
 {
 	const float ScaledRadius = CirlceRadius * LengthScalar;
-	const int32 NumSides = FMath::Max(ScaledRadius / 2, 24.0f);
+	const int32 NumSides = FMath::Clamp(ScaledRadius / 2, 24.0f, 100.0f);
 	const float AngleDelta = 2.0f * UE_PI / NumSides;
 
 	const FVector Origin = Transform.GetTranslation();
@@ -123,7 +127,7 @@ FLinearColor FDisplayClusterLightCardEditorWidget::GetAxisColor(EAxisList::Type
 	return FLinearColor::Transparent;
 }
 
-float FDisplayClusterLightCardEditorWidget::GetLengthScreenScalar(const FSceneView* View, const FEditorViewportClient* ViewportClient, const FVector& Origin) const
+float FDisplayClusterLightCardEditorWidget::GetLengthScreenScalar(const FSceneView* View, const FDisplayClusterLightCardEditorViewportClient* ViewportClient, const FVector& Origin) const
 {
 	// The ideal behavior for the length of the widget axes is to remain a fixed length regardless of the field of view or distance from the camera,
 	// but change proportionally with the size of the viewport, so that the widget takes up the same percentage of viewport space. 
@@ -135,17 +139,19 @@ float FDisplayClusterLightCardEditorWidget::GetLengthScreenScalar(const FSceneVi
 	// Note we use the x axis sizes to compute the scalars as the viewport will only scale smaller when the width is resized; when the 
 	// height is resized, the view is clamped
 
+	const bool bIsOrthographic = ViewportClient->GetRenderViewportType() != LVT_Perspective;
+
 	const float DPIScale = ViewportClient->GetDPIScale();
 	const float ResolutionScale = View->UnconstrainedViewRect.Size().X / (DPIScale * ReferenceResolution);
 
 	const float DistanceFromView = FMath::Max(FVector::Dist(Origin, View->ViewMatrices.GetViewOrigin()), 1.f);
-	const float ProjectionScale = DistanceFromView / View->ViewMatrices.GetScreenScale();
+	const float ProjectionScale = bIsOrthographic ? 1.0 / View->ViewMatrices.GetScreenScale() : DistanceFromView / View->ViewMatrices.GetScreenScale();
 
 	const float FinalScalar = DPIScale * ResolutionScale * ProjectionScale;
 	return FinalScalar;
 }
 
-float FDisplayClusterLightCardEditorWidget::GetSizeScreenScalar(const FSceneView* View, const FEditorViewportClient* ViewportClient) const
+float FDisplayClusterLightCardEditorWidget::GetSizeScreenScalar(const FSceneView* View, const FDisplayClusterLightCardEditorViewportClient* ViewportClient) const
 {
 	// The ideal behavior for the size of the widget is to remain a fixed size regardless of the field of view or distance from the camera,
 	// but change proportionally with the size of the viewport, so that the widget takes up the same percentage of viewport space. 
@@ -156,11 +162,13 @@ float FDisplayClusterLightCardEditorWidget::GetSizeScreenScalar(const FSceneView
 	// Note we use the x axis sizes to compute the scalars as the viewport will only scale smaller when the width is resized; when the 
 	// height is resized, the view is clamped
 
+	const bool bIsOrthographic = ViewportClient->GetRenderViewportType() != LVT_Perspective;
+
 	const float DPIScale = ViewportClient->GetDPIScale();
 	const float ResolutionScale = View->UnconstrainedViewRect.Size().X / (DPIScale * ReferenceResolution);
 
 	const FMatrix& ProjectionMatrix = View->ViewMatrices.GetProjectionMatrix();
-	const float ProjectionScale = FMath::Abs(ProjectionMatrix.M[0][0]);
+	const float ProjectionScale = !bIsOrthographic ? FMath::Abs(ProjectionMatrix.M[0][0]) : 1.0f;
 
 	const float FinalScalar = DPIScale * ResolutionScale / ProjectionScale;
 	return FinalScalar;
diff --git a/Engine/Plugins/Runtime/nDisplay/Source/DisplayClusterLightCardEditor/Private/Viewport/DisplayClusterLightCardEditorWidget.h b/Engine/Plugins/Runtime/nDisplay/Source/DisplayClusterLightCardEditor/Private/Viewport/DisplayClusterLightCardEditorWidget.h
--- a/Engine/Plugins/Runtime/nDisplay/Source/DisplayClusterLightCardEditor/Private/Viewport/DisplayClusterLightCardEditorWidget.h
+++ b/Engine/Plugins/Runtime/nDisplay/Source/DisplayClusterLightCardEditor/Private/Viewport/DisplayClusterLightCardEditorWidget.h
@@ -8,8 +8,8 @@
 #include "DisplayClusterMeshProjectionRenderer.h"
 
 class FSceneView;
-class FEditorViewportClient;
 class FPrimitiveDrawInterface;
+class FDisplayClusterLightCardEditorViewportClient;
 
 /** Widget used to manipulate entities in the light card editor's viewport */
 class FDisplayClusterLightCardEditorWidget
@@ -50,7 +50,7 @@ public:
 		WM_Max
 	};
 
-	void Draw(const FSceneView* View, const FEditorViewportClient* ViewportClient, FPrimitiveDrawInterface* PDI);
+	void Draw(const FSceneView* View, const FDisplayClusterLightCardEditorViewportClient* ViewportClient, FPrimitiveDrawInterface* PDI);
 
 	EWidgetMode GetWidgetMode() const { return WidgetMode; }
 	void SetWidgetMode(EWidgetMode NewWidgetMode) { WidgetMode = NewWidgetMode; }
@@ -62,10 +62,10 @@ public:
 
 private:
 	/** Draws the specified axis to the PDI, sizing and coloring it appropriately */
-	void DrawAxis(FPrimitiveDrawInterface* PDI, EAxisList::Type Axis, float SizeScalar, float LengthScalar);
+	void DrawAxis(FPrimitiveDrawInterface* PDI, EAxisList::Type Axis, float SizeScalar, float LengthScalar, float OrthoScalar);
 
 	/** Draws the widget origin to the PDI, sizing and coloring it appropriately */
-	void DrawOrigin(FPrimitiveDrawInterface* PDI, float SizeScalar);
+	void DrawOrigin(FPrimitiveDrawInterface* PDI, float SizeScalar, float OrthoScalar);
 
 	void DrawCircle(FPrimitiveDrawInterface* PDI, EAxisList::Type Axis, float SizeScalar, float LengthScalar);
 
@@ -76,10 +76,10 @@ private:
 	FLinearColor GetAxisColor(EAxisList::Type Axis) const;
 
 	/** Calculates a scalar to use to keep lengths and lines a fixed size on the screen regardless of DPI, viewport size, FOV, or distance from the view origin */
-	float GetLengthScreenScalar(const FSceneView* View, const FEditorViewportClient* ViewportClient, const FVector& Origin) const;
+	float GetLengthScreenScalar(const FSceneView* View, const FDisplayClusterLightCardEditorViewportClient* ViewportClient, const FVector& Origin) const;
 
 	/** Calculates a scalar to use to keep sizes a fixed size on the screen regardless of DPI, viewport size, FOV, or distance from the view origin */
-	float GetSizeScreenScalar(const FSceneView* View, const FEditorViewportClient* ViewportClient) const;
+	float GetSizeScreenScalar(const FSceneView* View, const FDisplayClusterLightCardEditorViewportClient* ViewportClient) const;
 
 private:
 	/** The current mode the widget is in */
diff --git a/Engine/Plugins/Runtime/nDisplay/Source/DisplayClusterLightCardEditor/Private/Viewport/DisplayClusterLightcardEditorViewport.cpp b/Engine/Plugins/Runtime/nDisplay/Source/DisplayClusterLightCardEditor/Private/Viewport/DisplayClusterLightcardEditorViewport.cpp
--- a/Engine/Plugins/Runtime/nDisplay/Source/DisplayClusterLightCardEditor/Private/Viewport/DisplayClusterLightcardEditorViewport.cpp
+++ b/Engine/Plugins/Runtime/nDisplay/Source/DisplayClusterLightCardEditor/Private/Viewport/DisplayClusterLightcardEditorViewport.cpp
@@ -60,6 +60,16 @@ public:
 					.OnGetMenuContent(this, &SDisplayClusterLightCardEditorViewportToolBar::GenerateProjectionMenu)
 				]
 				+ SHorizontalBox::Slot()
+				.AutoWidth()
+				.Padding(2.0f, 2.0f)
+				[
+					SNew(SEditorViewportToolbarMenu)
+					.ParentToolBar(SharedThis(this))
+					.Cursor(EMouseCursor::Default)
+					.Label(LOCTEXT("LightCardEditorViewMenuLabel", "View"))
+					.OnGetMenuContent(this, &SDisplayClusterLightCardEditorViewportToolBar::GenerateViewMenu)
+				]
+				+ SHorizontalBox::Slot()
 				.Padding(3.0f, 1.0f)
 				.HAlign(HAlign_Right)
 				[
@@ -86,11 +96,7 @@ public:
 
 		FMenuBuilder PreviewOptionsMenuBuilder(bInShouldCloseWindowAfterMenuSelection, CommandList);
 		{
-			PreviewOptionsMenuBuilder.BeginSection("LightCardEditorViewportOptions", LOCTEXT("ViewportOptionsMenuHeader", "Viewport Options"));
-			{
-				PreviewOptionsMenuBuilder.AddMenuEntry(FDisplayClusterLightCardEditorCommands::Get().ResetCamera);
-			}
-			PreviewOptionsMenuBuilder.EndSection();
+			// TODO: Any advanced viewport options can be added here
 		}
 
 		return PreviewOptionsMenuBuilder.MakeWidget();
@@ -102,10 +108,18 @@ public:
 
 		if (EditorViewport.IsValid())
 		{
-			switch (EditorViewport.Pin()->GetLightCardEditorViewportClient()->GetProjectionMode())
+			TSharedRef<FDisplayClusterLightCardEditorViewportClient> ViewportClient = EditorViewport.Pin()->GetLightCardEditorViewportClient();
+			switch (ViewportClient->GetProjectionMode())
 			{
-			case EDisplayClusterMeshProjectionType::Perspective:
-				Label = LOCTEXT("ProjectionwMenuTitle_Perspective", "Perspective");
+			case EDisplayClusterMeshProjectionType::Linear:
+				if (ViewportClient->GetRenderViewportType() == ELevelViewportType::LVT_Perspective)
+				{
+					Label = LOCTEXT("ProjectionwMenuTitle_Perspective", "Perspective");
+				}
+				else
+				{
+					Label = LOCTEXT("ProjectionwMenuTitle_Orthographic", "Orthographic");
+				}
 				break;
 
 			case EDisplayClusterMeshProjectionType::Azimuthal:
@@ -125,34 +139,39 @@ public:
 		FMenuBuilder ViewMenuBuilder(bInShouldCloseWindowAfterMenuSelection, CommandList);
 
 		ViewMenuBuilder.AddMenuEntry(FDisplayClusterLightCardEditorCommands::Get().PerspectiveProjection);
+		ViewMenuBuilder.AddMenuEntry(FDisplayClusterLightCardEditorCommands::Get().OrthographicProjection);
 		ViewMenuBuilder.AddMenuEntry(FDisplayClusterLightCardEditorCommands::Get().AzimuthalProjection);
 
 		return ViewMenuBuilder.MakeWidget();
 	}
 
-	FText GetViewMenuLabel() const
+	TSharedRef<SWidget> GenerateViewMenu() const
 	{
-		FText Label = NSLOCTEXT("BlueprintEditor", "ViewMenuTitle_Default", "View");
+		TSharedPtr<const FUICommandList> CommandList = EditorViewport.IsValid() ? EditorViewport.Pin()->GetCommandList() : nullptr;
 
-		if (EditorViewport.IsValid())
+		const bool bInShouldCloseWindowAfterMenuSelection = true;
+
+		FMenuBuilder MenuBuilder(bInShouldCloseWindowAfterMenuSelection, CommandList);
 		{
-			switch (EditorViewport.Pin()->GetViewportClient()->GetViewMode())
+			MenuBuilder.BeginSection("LightCardEditorViewOrientation", LOCTEXT("ViewOrientationMenuHeader", "View Orientation"));
 			{
-			case VMI_Lit:
-				Label = NSLOCTEXT("BlueprintEditor", "ViewMenuTitle_Lit", "Lit");
-				break;
-
-			case VMI_Unlit:
-				Label = NSLOCTEXT("BlueprintEditor", "ViewMenuTitle_Unlit", "Unlit");
-				break;
+				MenuBuilder.AddMenuEntry(FDisplayClusterLightCardEditorCommands::Get().ViewOrientationTop);
+				MenuBuilder.AddMenuEntry(FDisplayClusterLightCardEditorCommands::Get().ViewOrientationBottom);
+				MenuBuilder.AddMenuEntry(FDisplayClusterLightCardEditorCommands::Get().ViewOrientationLeft);
+				MenuBuilder.AddMenuEntry(FDisplayClusterLightCardEditorCommands::Get().ViewOrientationRight);
+				MenuBuilder.AddMenuEntry(FDisplayClusterLightCardEditorCommands::Get().ViewOrientationFront);
+				MenuBuilder.AddMenuEntry(FDisplayClusterLightCardEditorCommands::Get().ViewOrientationBack);
+			}
+			MenuBuilder.EndSection();
 
-			case VMI_BrushWireframe:
-				Label = NSLOCTEXT("BlueprintEditor", "ViewMenuTitle_Wireframe", "Wireframe");
-				break;
+			MenuBuilder.BeginSection("LightCardEditorViewOptions", LOCTEXT("ViewOptionsMenuHeader", "View Options"));
+			{
+				MenuBuilder.AddMenuEntry(FDisplayClusterLightCardEditorCommands::Get().ResetCamera);
 			}
+			MenuBuilder.EndSection();
 		}
 
-		return Label;
+		return MenuBuilder.MakeWidget();
 	}
 
 	TSharedRef<SWidget> MakeTransformToolBar()
@@ -228,12 +247,18 @@ private:
 	TWeakPtr<SDisplayClusterLightCardEditorViewport> EditorViewport;
 };
 
+const FVector SDisplayClusterLightCardEditorViewport::ViewDirectionTop = FVector(0.0f, 0.0f, 1.0f);
+const FVector SDisplayClusterLightCardEditorViewport::ViewDirectionBottom = FVector(0.0f, 0.0f, -1.0f);
+const FVector SDisplayClusterLightCardEditorViewport::ViewDirectionLeft = FVector(0.0f, -1.0f, 0.0f);
+const FVector SDisplayClusterLightCardEditorViewport::ViewDirectionRight = FVector(0.0f, 1.0f, 0.0f);
+const FVector SDisplayClusterLightCardEditorViewport::ViewDirectionFront = FVector(1.0f, 0.0f, 0.0f);
+const FVector SDisplayClusterLightCardEditorViewport::ViewDirectionBack = FVector(-1.0f, 0.0f, 0.0f);
+
 void SDisplayClusterLightCardEditorViewport::Construct(const FArguments& InArgs, TSharedPtr<SDisplayClusterLightCardEditor> InLightCardEditor, TSharedPtr<class FUICommandList> InCommandList)
 {
 	LightCardEditorPtr = InLightCardEditor;
 	
-	AdvancedPreviewScene = MakeShared<FAdvancedPreviewScene>(FPreviewScene::ConstructionValues());
-	AdvancedPreviewScene->SetFloorVisibility(true);
+	PreviewScene = MakeShared<FPreviewScene>(FPreviewScene::ConstructionValues());
 		
 	SEditorViewport::Construct(SEditorViewport::FArguments());
 
@@ -302,9 +327,9 @@ void SDisplayClusterLightCardEditorViewport::SummonContextMenu()
 
 TSharedRef<FEditorViewportClient> SDisplayClusterLightCardEditorViewport::MakeEditorViewportClient()
 {
-	check(AdvancedPreviewScene.IsValid());
+	check(PreviewScene.IsValid());
 	
-	ViewportClient = MakeShareable(new FDisplayClusterLightCardEditorViewportClient(*AdvancedPreviewScene.Get(),
+	ViewportClient = MakeShareable(new FDisplayClusterLightCardEditorViewportClient(*PreviewScene.Get(),
 		SharedThis(this)));
 	return ViewportClient.ToSharedRef();
 }
@@ -336,7 +361,6 @@ void SDisplayClusterLightCardEditorViewport::BindCommands()
 			FCanExecuteAction(),
 			FIsActionChecked::CreateSP(this, &SDisplayClusterLightCardEditorViewport::IsEditorWidgetModeSelected, FDisplayClusterLightCardEditorWidget::EWidgetMode::WM_Translate)
 		);
-
 		
 		CommandList->MapAction(
 			Commands.RotateMode,
@@ -364,15 +388,45 @@ void SDisplayClusterLightCardEditorViewport::BindCommands()
 
 		CommandList->MapAction(
 			FDisplayClusterLightCardEditorCommands::Get().PerspectiveProjection,
-			FExecuteAction::CreateSP(this, &SDisplayClusterLightCardEditorViewport::SetProjectionMode, EDisplayClusterMeshProjectionType::Perspective),
+			FExecuteAction::CreateSP(this, &SDisplayClusterLightCardEditorViewport::SetProjectionMode, EDisplayClusterMeshProjectionType::Linear, ELevelViewportType::LVT_Perspective),
+			FCanExecuteAction(),
+			FIsActionChecked::CreateSP(this, &SDisplayClusterLightCardEditorViewport::IsProjectionModeSelected, EDisplayClusterMeshProjectionType::Linear, ELevelViewportType::LVT_Perspective));
+
+		CommandList->MapAction(
+			FDisplayClusterLightCardEditorCommands::Get().OrthographicProjection,
+			FExecuteAction::CreateSP(this, &SDisplayClusterLightCardEditorViewport::SetProjectionMode, EDisplayClusterMeshProjectionType::Linear, ELevelViewportType::LVT_OrthoFreelook),
 			FCanExecuteAction(),
-			FIsActionChecked::CreateSP(this, &SDisplayClusterLightCardEditorViewport::IsProjectionModeSelected, EDisplayClusterMeshProjectionType::Perspective));
+			FIsActionChecked::CreateSP(this, &SDisplayClusterLightCardEditorViewport::IsProjectionModeSelected, EDisplayClusterMeshProjectionType::Linear, ELevelViewportType::LVT_OrthoFreelook));
 
 		CommandList->MapAction(
 			FDisplayClusterLightCardEditorCommands::Get().AzimuthalProjection,
-			FExecuteAction::CreateSP(this, &SDisplayClusterLightCardEditorViewport::SetProjectionMode, EDisplayClusterMeshProjectionType::Azimuthal),
+			FExecuteAction::CreateSP(this, &SDisplayClusterLightCardEditorViewport::SetProjectionMode, EDisplayClusterMeshProjectionType::Azimuthal, ELevelViewportType::LVT_Perspective),
 			FCanExecuteAction(),
-			FIsActionChecked::CreateSP(this, &SDisplayClusterLightCardEditorViewport::IsProjectionModeSelected, EDisplayClusterMeshProjectionType::Azimuthal));
+			FIsActionChecked::CreateSP(this, &SDisplayClusterLightCardEditorViewport::IsProjectionModeSelected, EDisplayClusterMeshProjectionType::Azimuthal, ELevelViewportType::LVT_Perspective));
+
+		CommandList->MapAction(
+			Commands.ViewOrientationTop,
+			FExecuteAction::CreateSP(this, &SDisplayClusterLightCardEditorViewport::SetViewDirection, ViewDirectionTop));
+
+		CommandList->MapAction(
+			Commands.ViewOrientationBottom,
+			FExecuteAction::CreateSP(this, &SDisplayClusterLightCardEditorViewport::SetViewDirection, ViewDirectionBottom));
+
+		CommandList->MapAction(
+			Commands.ViewOrientationLeft,
+			FExecuteAction::CreateSP(this, &SDisplayClusterLightCardEditorViewport::SetViewDirection, ViewDirectionLeft));
+
+		CommandList->MapAction(
+			Commands.ViewOrientationRight,
+			FExecuteAction::CreateSP(this, &SDisplayClusterLightCardEditorViewport::SetViewDirection, ViewDirectionRight));
+
+		CommandList->MapAction(
+			Commands.ViewOrientationFront,
+			FExecuteAction::CreateSP(this, &SDisplayClusterLightCardEditorViewport::SetViewDirection, ViewDirectionFront));
+
+		CommandList->MapAction(
+			Commands.ViewOrientationBack,
+			FExecuteAction::CreateSP(this, &SDisplayClusterLightCardEditorViewport::SetViewDirection, ViewDirectionBack));
 
 		CommandList->MapAction(
 			Commands.ResetCamera,
@@ -457,24 +511,32 @@ void SDisplayClusterLightCardEditorViewport::CycleEditorWidgetMode()
 	SetEditorWidgetMode((FDisplayClusterLightCardEditorWidget::EWidgetMode)WidgetModeAsInt);
 }
 
-void SDisplayClusterLightCardEditorViewport::SetProjectionMode(EDisplayClusterMeshProjectionType InProjectionMode)
+void SDisplayClusterLightCardEditorViewport::SetProjectionMode(EDisplayClusterMeshProjectionType InProjectionMode, ELevelViewportType InViewportType)
 {
 	if (ViewportClient.IsValid())
 	{
-		ViewportClient->SetProjectionMode(InProjectionMode);
+		ViewportClient->SetProjectionMode(InProjectionMode, InViewportType);
 	}
 }
 
-bool SDisplayClusterLightCardEditorViewport::IsProjectionModeSelected(EDisplayClusterMeshProjectionType InProjectionMode) const
+bool SDisplayClusterLightCardEditorViewport::IsProjectionModeSelected(EDisplayClusterMeshProjectionType InProjectionMode, ELevelViewportType ViewportType) const
 {
 	if (ViewportClient.IsValid())
 	{
-		return ViewportClient->GetProjectionMode() == InProjectionMode;
+		return ViewportClient->GetProjectionMode() == InProjectionMode && ViewportClient->GetRenderViewportType() == ViewportType;
 	}
 
 	return false;
 }
 
+void SDisplayClusterLightCardEditorViewport::SetViewDirection(FVector InViewDirection)
+{
+	if (ViewportClient.IsValid())
+	{
+		ViewportClient->SetViewRotation(InViewDirection.Rotation());
+	}
+}
+
 bool SDisplayClusterLightCardEditorViewport::IsDrawingLightCard() const
 {
 	if (ViewportClient.IsValid())
diff --git a/Engine/Plugins/Runtime/nDisplay/Source/DisplayClusterLightCardEditor/Private/Viewport/DisplayClusterLightcardEditorViewport.h b/Engine/Plugins/Runtime/nDisplay/Source/DisplayClusterLightCardEditor/Private/Viewport/DisplayClusterLightcardEditorViewport.h
--- a/Engine/Plugins/Runtime/nDisplay/Source/DisplayClusterLightCardEditor/Private/Viewport/DisplayClusterLightcardEditorViewport.h
+++ b/Engine/Plugins/Runtime/nDisplay/Source/DisplayClusterLightCardEditor/Private/Viewport/DisplayClusterLightcardEditorViewport.h
@@ -19,6 +19,14 @@ class FDisplayClusterLightCardEditorViewportClient;
  */
 class SDisplayClusterLightCardEditorViewport : public SEditorViewport, public ICommonEditorViewportToolbarInfoProvider
 {
+public:
+	static const FVector ViewDirectionTop;
+	static const FVector ViewDirectionBottom;
+	static const FVector ViewDirectionLeft;
+	static const FVector ViewDirectionRight;
+	static const FVector ViewDirectionFront;
+	static const FVector ViewDirectionBack;
+
 public:
 	SLATE_BEGIN_ARGS(SDisplayClusterLightCardEditorViewport) {}
 	SLATE_END_ARGS()
@@ -56,12 +64,23 @@ private:
 	/** Create the popup context menu */
 	TSharedRef<SWidget> MakeContextMenu();
 
+	/** Sets the widget mode for the editor viewport to the specified mode */
 	void SetEditorWidgetMode(FDisplayClusterLightCardEditorWidget::EWidgetMode InWidgetMode);
+
+	/** Returns true if the widget mode of the editor viewport matches the specified mode */
 	bool IsEditorWidgetModeSelected(FDisplayClusterLightCardEditorWidget::EWidgetMode InWidgetMode) const;
+
+	/** Cycles the widget mode of the editor viewport */
 	void CycleEditorWidgetMode();
 
-	void SetProjectionMode(EDisplayClusterMeshProjectionType InProjectionMode);
-	bool IsProjectionModeSelected(EDisplayClusterMeshProjectionType InProjectionMode) const;
+	/** Sets the projection mode and viewport type of the editor viewport */
+	void SetProjectionMode(EDisplayClusterMeshProjectionType InProjectionMode, ELevelViewportType InViewportType);
+
+	/** Returns true if the editor viewport's projection mode and viewport type match the specified parameters */
+	bool IsProjectionModeSelected(EDisplayClusterMeshProjectionType InProjectionMode, ELevelViewportType InViewportType) const;
+
+	/** Sets the direction of the editor viewport */
+	void SetViewDirection(FVector InViewDirection);
 
 	/** Handles the Draw LC button being pressed. */
 	void DrawLightCard();
@@ -76,8 +95,8 @@ private:
 	bool CanPasteLightCardsHere() const;
 
 private:
-	/** Preview Scene - uses advanced preview settings */
-	TSharedPtr<FAdvancedPreviewScene> AdvancedPreviewScene;
+	/** The preview scene to use by the viewport client to manage the preview 3D world */
+	TSharedPtr<FPreviewScene> PreviewScene;
 	
 	/** Level viewport client */
 	TSharedPtr<FDisplayClusterLightCardEditorViewportClient> ViewportClient;
diff --git a/Engine/Plugins/Runtime/nDisplay/Source/DisplayClusterLightCardEditorShaders/Private/DisplayClusterMeshProjectionRenderer.cpp b/Engine/Plugins/Runtime/nDisplay/Source/DisplayClusterLightCardEditorShaders/Private/DisplayClusterMeshProjectionRenderer.cpp
--- a/Engine/Plugins/Runtime/nDisplay/Source/DisplayClusterLightCardEditorShaders/Private/DisplayClusterMeshProjectionRenderer.cpp
+++ b/Engine/Plugins/Runtime/nDisplay/Source/DisplayClusterLightCardEditorShaders/Private/DisplayClusterMeshProjectionRenderer.cpp
@@ -42,7 +42,7 @@ public:
 	}
 };
 
-using FMeshPerspectiveProjectionVS = FMeshProjectionVS<EDisplayClusterMeshProjectionType::Perspective>;
+using FMeshPerspectiveProjectionVS = FMeshProjectionVS<EDisplayClusterMeshProjectionType::Linear>;
 using FMeshAzimuthalProjectionVS = FMeshProjectionVS<EDisplayClusterMeshProjectionType::Azimuthal>;
 
 IMPLEMENT_MATERIAL_SHADER_TYPE(template<>, FMeshPerspectiveProjectionVS, TEXT("/Plugin/nDisplay/Private/MeshProjectionShaders.usf"), TEXT("PerspectiveVS"), SF_Vertex);
@@ -140,7 +140,7 @@ public:
 
 		const FMeshDrawingPolicyOverrideSettings OverrideSettings = ComputeMeshOverrideSettings(MeshBatch);
 		ERasterizerFillMode MeshFillMode = ComputeMeshFillMode(MeshBatch, Material, OverrideSettings);
-		ERasterizerCullMode MeshCullMode = CM_None;
+		ERasterizerCullMode MeshCullMode = ComputeMeshCullMode(MeshBatch, Material, OverrideSettings);
 
 		FMeshMaterialShaderElementData ShaderElementData;
 		ShaderElementData.InitializeMeshMaterialData(ViewIfDynamicMeshCommand, PrimitiveSceneProxy, MeshBatch, StaticMeshId, false);
@@ -343,7 +343,7 @@ public:
 
 				const FMeshDrawingPolicyOverrideSettings OverrideSettings = ComputeMeshOverrideSettings(MeshBatch);
 				ERasterizerFillMode MeshFillMode = ComputeMeshFillMode(MeshBatch, *Material, OverrideSettings);
-				ERasterizerCullMode MeshCullMode = CM_None;
+				ERasterizerCullMode MeshCullMode = ComputeMeshCullMode(MeshBatch, *Material, OverrideSettings);
 
 				FMeshProjectionHitProxyShaderElementData ShaderElementData(MeshBatch.BatchHitProxyId);
 				ShaderElementData.InitializeMeshMaterialData(ViewIfDynamicMeshCommand, PrimitiveSceneProxy, MeshBatch, StaticMeshId, false);
@@ -693,7 +693,7 @@ FVector FDisplayClusterMeshProjectionTransform::ProjectPosition(const FVector& W
 {
 	FVector ProjectedPosition(WorldPosition);
 
-	if (Projection != EDisplayClusterMeshProjectionType::Perspective)
+	if (Projection != EDisplayClusterMeshProjectionType::Linear)
 	{
 		const FVector ViewPos = ViewMatrix.TransformPosition(WorldPosition);
 		const FVector ProjectedViewPos = FDisplayClusterMeshProjectionRenderer::ProjectViewPosition(ViewPos, Projection);
@@ -707,7 +707,7 @@ FVector FDisplayClusterMeshProjectionTransform::UnprojectPosition(const FVector&
 {
 	FVector UnprojectedPosition(ProjectedPosition);
 
-	if (Projection != EDisplayClusterMeshProjectionType::Perspective)
+	if (Projection != EDisplayClusterMeshProjectionType::Linear)
 	{
 		const FVector ViewPos = ViewMatrix.TransformPosition(ProjectedPosition);
 		const FVector ProjectedViewPos = FDisplayClusterMeshProjectionRenderer::UnprojectViewPosition(ViewPos, Projection);
@@ -1053,7 +1053,7 @@ void FDisplayClusterMeshProjectionRenderer::RenderNormals(FCanvas* Canvas,
 #define SWITCH_ON_PROJECTION_TYPE(FuncName, ProjectionType, ...) switch (ProjectionType) \
 	{ \
 	PROJECTION_TYPE_CASE(FuncName, EDisplayClusterMeshProjectionType::Azimuthal, __VA_ARGS__) \
-	PROJECTION_TYPE_CASE(FuncName, EDisplayClusterMeshProjectionType::Perspective, __VA_ARGS__) \
+	PROJECTION_TYPE_CASE(FuncName, EDisplayClusterMeshProjectionType::Linear, __VA_ARGS__) \
 	}
 
 void FDisplayClusterMeshProjectionRenderer::AddBaseRenderPass(FRDGBuilder& GraphBuilder,
@@ -1577,7 +1577,7 @@ bool FDisplayClusterMeshProjectionRenderer::IsPrimitiveComponentSelected(const U
 // Explicit template specializations for each projection type
 #define PROJECTION_TYPE_TEMPLATE_SPECIALIZATION(FuncName, ProjectionType, ...) template void FDisplayClusterMeshProjectionRenderer::FuncName<ProjectionType>(__VA_ARGS__);
 #define CREATE_PROJECTION_TYPE_TEMPLATE_SPECIALIZATIONS(FuncName, ...) \
-	PROJECTION_TYPE_TEMPLATE_SPECIALIZATION(FuncName, EDisplayClusterMeshProjectionType::Perspective, __VA_ARGS__) \
+	PROJECTION_TYPE_TEMPLATE_SPECIALIZATION(FuncName, EDisplayClusterMeshProjectionType::Linear, __VA_ARGS__) \
 	PROJECTION_TYPE_TEMPLATE_SPECIALIZATION(FuncName, EDisplayClusterMeshProjectionType::Azimuthal, __VA_ARGS__)
 
 CREATE_PROJECTION_TYPE_TEMPLATE_SPECIALIZATIONS(RenderPrimitives_RenderThread, const FSceneView* View, FRHICommandList& RHICmdList, bool bTranslucencyPass)
diff --git a/Engine/Plugins/Runtime/nDisplay/Source/DisplayClusterLightCardEditorShaders/Public/DisplayClusterMeshProjectionRenderer.h b/Engine/Plugins/Runtime/nDisplay/Source/DisplayClusterLightCardEditorShaders/Public/DisplayClusterMeshProjectionRenderer.h
--- a/Engine/Plugins/Runtime/nDisplay/Source/DisplayClusterLightCardEditorShaders/Public/DisplayClusterMeshProjectionRenderer.h
+++ b/Engine/Plugins/Runtime/nDisplay/Source/DisplayClusterLightCardEditorShaders/Public/DisplayClusterMeshProjectionRenderer.h
@@ -18,7 +18,7 @@ struct FEngineShowFlags;
 enum DISPLAYCLUSTERLIGHTCARDEDITORSHADERS_API EDisplayClusterMeshProjectionType
 {
 	/** Default linear projection */
-	Perspective,
+	Linear,
 
 	/** Non-linear spherical projection based on the azimuthal equidistant map projection */
 	Azimuthal
@@ -41,7 +41,7 @@ class DISPLAYCLUSTERLIGHTCARDEDITORSHADERS_API FDisplayClusterMeshProjectionTran
 {
 public:
 	FDisplayClusterMeshProjectionTransform()
-		: FDisplayClusterMeshProjectionTransform(EDisplayClusterMeshProjectionType::Perspective, FMatrix::Identity)
+		: FDisplayClusterMeshProjectionTransform(EDisplayClusterMeshProjectionType::Linear, FMatrix::Identity)
 	{ }
 
 	FDisplayClusterMeshProjectionTransform(EDisplayClusterMeshProjectionType InProjection, const FMatrix& InViewMatrix)
diff --git a/Engine/Plugins/Runtime/nDisplay/Source/DisplayClusterScenePreview/Private/DisplayClusterScenePreviewModule.cpp b/Engine/Plugins/Runtime/nDisplay/Source/DisplayClusterScenePreview/Private/DisplayClusterScenePreviewModule.cpp
--- a/Engine/Plugins/Runtime/nDisplay/Source/DisplayClusterScenePreview/Private/DisplayClusterScenePreviewModule.cpp
+++ b/Engine/Plugins/Runtime/nDisplay/Source/DisplayClusterScenePreview/Private/DisplayClusterScenePreviewModule.cpp
@@ -99,6 +99,17 @@ bool FDisplayClusterScenePreviewModule::AddActorToRenderer(int32 RendererId, AAc
 	return false;
 }
 
+bool FDisplayClusterScenePreviewModule::RemoveActorFromRenderer(int32 RendererId, AActor* Actor)
+{
+	if (FRendererConfig* Config = RendererConfigs.Find(RendererId))
+	{
+		Config->Renderer->RemoveActor(Actor);
+		return true;
+	}
+
+	return false;
+}
+
 bool FDisplayClusterScenePreviewModule::ClearRendererScene(int32 RendererId)
 {
 	if (FRendererConfig* Config = RendererConfigs.Find(RendererId))
diff --git a/Engine/Plugins/Runtime/nDisplay/Source/DisplayClusterScenePreview/Private/DisplayClusterScenePreviewModule.h b/Engine/Plugins/Runtime/nDisplay/Source/DisplayClusterScenePreview/Private/DisplayClusterScenePreviewModule.h
--- a/Engine/Plugins/Runtime/nDisplay/Source/DisplayClusterScenePreview/Private/DisplayClusterScenePreviewModule.h
+++ b/Engine/Plugins/Runtime/nDisplay/Source/DisplayClusterScenePreview/Private/DisplayClusterScenePreviewModule.h
@@ -30,6 +30,7 @@ public:
 	virtual bool SetRendererRootActor(int32 RendererId, ADisplayClusterRootActor* Actor, bool bAutoUpdateLightcards = false) override;
 	virtual ADisplayClusterRootActor* GetRendererRootActor(int32 RendererId) override;
 	virtual bool AddActorToRenderer(int32 RendererId, AActor* Actor) override;
+	virtual bool RemoveActorFromRenderer(int32 RendererId, AActor* Actor) override;
 	virtual bool ClearRendererScene(int32 RendererId) override;
 	virtual bool SetRendererActorSelectedDelegate(int32 RendererId, FDisplayClusterMeshProjectionRenderer::FSelection ActorSelectedDelegate) override;
 	virtual bool SetRendererRenderSimpleElementsDelegate(int32 RendererId, FDisplayClusterMeshProjectionRenderer::FSimpleElementPass RenderSimpleElementsDelegate) override;
diff --git a/Engine/Plugins/Runtime/nDisplay/Source/DisplayClusterScenePreview/Public/IDisplayClusterScenePreview.h b/Engine/Plugins/Runtime/nDisplay/Source/DisplayClusterScenePreview/Public/IDisplayClusterScenePreview.h
--- a/Engine/Plugins/Runtime/nDisplay/Source/DisplayClusterScenePreview/Public/IDisplayClusterScenePreview.h
+++ b/Engine/Plugins/Runtime/nDisplay/Source/DisplayClusterScenePreview/Public/IDisplayClusterScenePreview.h
@@ -134,6 +134,14 @@ public:
 	 */
 	virtual bool AddActorToRenderer(int32 RendererId, AActor* Actor) = 0;
 
+	/**
+	 * Removes an actor from the preview scene for a renderer.
+	 *
+	 * @param RendererId The ID of the renderer as returned from CreateRenderer.
+	 * @param Actor The actor to remove to the renderer's scene.
+	 */
+	virtual bool RemoveActorFromRenderer(int32 RendererId, AActor* Actor) = 0;
+
 	/**
 	 * Clear the scene of a renderer.
 	 * 
