nDisplay: Add light card template feature.
Light cards can be saved to a template. Templates are displayed in a new dockable tab and each template can be dragged into the light card viewport.
Templates can be favorited and filtered.

- Refactor light card editor widget to use internal dockable tabs.
- Change operator details panel to use kismet inspector. This allows components to be hidden.

diff --git a/Engine/Plugins/Runtime/nDisplay/Source/DisplayClusterLightCardEditor/Private/DisplayClusterLightCardEditorCommands.cpp b/Engine/Plugins/Runtime/nDisplay/Source/DisplayClusterLightCardEditor/Private/DisplayClusterLightCardEditorCommands.cpp
--- a/Engine/Plugins/Runtime/nDisplay/Source/DisplayClusterLightCardEditor/Private/DisplayClusterLightCardEditorCommands.cpp
+++ b/Engine/Plugins/Runtime/nDisplay/Source/DisplayClusterLightCardEditor/Private/DisplayClusterLightCardEditorCommands.cpp
@@ -25,6 +25,8 @@ void FDisplayClusterLightCardEditorCommands::RegisterCommands()
 	UI_COMMAND(RemoveLightCard, "Remove from Actor", "Remove the Light Card from the actor but do not delete it", EUserInterfaceActionType::Button, FInputChord());
 	UI_COMMAND(PasteHere, "Paste Here", "Paste clipboard contents at the click location", EUserInterfaceActionType::Button, FInputChord());
 
+	UI_COMMAND(SaveLightCardTemplate, "Save As Template", "Save a template of the light card's appearance", EUserInterfaceActionType::Button, FInputChord(EModifierKey::Control | EModifierKey::Alt, EKeys::S));
+
 	UI_COMMAND(DrawLightCard, "Draw Light Card", "Draw polygon light card on viewport", EUserInterfaceActionType::ToggleButton, FInputChord());
 }
 
diff --git a/Engine/Plugins/Runtime/nDisplay/Source/DisplayClusterLightCardEditor/Private/DisplayClusterLightCardEditorCommands.h b/Engine/Plugins/Runtime/nDisplay/Source/DisplayClusterLightCardEditor/Private/DisplayClusterLightCardEditorCommands.h
--- a/Engine/Plugins/Runtime/nDisplay/Source/DisplayClusterLightCardEditor/Private/DisplayClusterLightCardEditorCommands.h
+++ b/Engine/Plugins/Runtime/nDisplay/Source/DisplayClusterLightCardEditor/Private/DisplayClusterLightCardEditorCommands.h
@@ -41,5 +41,7 @@ public:
 	TSharedPtr<FUICommandInfo> RemoveLightCard;
 	TSharedPtr<FUICommandInfo> PasteHere;
 
+	TSharedPtr<FUICommandInfo> SaveLightCardTemplate;
+	
 	TSharedPtr<FUICommandInfo> DrawLightCard;
 };
diff --git a/Engine/Plugins/Runtime/nDisplay/Source/DisplayClusterLightCardEditor/Private/DisplayClusterLightCardEditorModule.cpp b/Engine/Plugins/Runtime/nDisplay/Source/DisplayClusterLightCardEditor/Private/DisplayClusterLightCardEditorModule.cpp
--- a/Engine/Plugins/Runtime/nDisplay/Source/DisplayClusterLightCardEditor/Private/DisplayClusterLightCardEditorModule.cpp
+++ b/Engine/Plugins/Runtime/nDisplay/Source/DisplayClusterLightCardEditor/Private/DisplayClusterLightCardEditorModule.cpp
@@ -4,12 +4,16 @@
 
 #include "DisplayClusterLightCardEditorCommands.h"
 #include "SDisplayClusterLightCardEditor.h"
+#include "Settings/DisplayClusterLightCardEditorSettings.h"
+
+#include "ISettingsModule.h"
 
 #define LOCTEXT_NAMESPACE "DisplayClusterLightCardEditor"
 
 void FDisplayClusterLightCardEditorModule::StartupModule()
 {
 	RegisterTabSpawners();
+	RegisterSettings();
 
 	FDisplayClusterLightCardEditorCommands::Register();
 }
@@ -17,6 +21,7 @@ void FDisplayClusterLightCardEditorModule::StartupModule()
 void FDisplayClusterLightCardEditorModule::ShutdownModule()
 {
 	UnregisterTabSpawners();
+	UnregisterSettings();
 
 	FDisplayClusterLightCardEditorCommands::Unregister();
 }
@@ -31,6 +36,25 @@ void FDisplayClusterLightCardEditorModule::UnregisterTabSpawners()
 	SDisplayClusterLightCardEditor::UnregisterTabSpawner();
 }
 
+void FDisplayClusterLightCardEditorModule::RegisterSettings()
+{
+	if (ISettingsModule* SettingsModule = FModuleManager::GetModulePtr<ISettingsModule>("Settings"))
+	{
+		SettingsModule->RegisterSettings("Project", "Plugins", "nDisplayLightCardEditor",
+			LOCTEXT("nDisplayLightCardEditorName", "nDisplay Light Card Editor"),
+			LOCTEXT("nDisplayLightCardEditorDescription", "Configure settings for the nDisplay Light Card Editor."),
+			GetMutableDefault<UDisplayClusterLightCardEditorProjectSettings>());
+	}
+}
+
+void FDisplayClusterLightCardEditorModule::UnregisterSettings()
+{
+	if (ISettingsModule* SettingsModule = FModuleManager::GetModulePtr<ISettingsModule>("Settings"))
+	{
+		SettingsModule->UnregisterSettings("Project", "Plugins", "nDisplayLightCardEditor");
+	}
+}
+
 IMPLEMENT_MODULE(FDisplayClusterLightCardEditorModule, DisplayClusterLightCardEditor);
 
 #undef LOCTEXT_NAMESPACE
diff --git a/Engine/Plugins/Runtime/nDisplay/Source/DisplayClusterLightCardEditor/Private/DisplayClusterLightCardEditorModule.h b/Engine/Plugins/Runtime/nDisplay/Source/DisplayClusterLightCardEditor/Private/DisplayClusterLightCardEditorModule.h
--- a/Engine/Plugins/Runtime/nDisplay/Source/DisplayClusterLightCardEditor/Private/DisplayClusterLightCardEditorModule.h
+++ b/Engine/Plugins/Runtime/nDisplay/Source/DisplayClusterLightCardEditor/Private/DisplayClusterLightCardEditorModule.h
@@ -19,4 +19,7 @@ public:
 private:
 	void RegisterTabSpawners();
 	void UnregisterTabSpawners();
+
+	void RegisterSettings();
+	void UnregisterSettings();
 };
diff --git a/Engine/Plugins/Runtime/nDisplay/Source/DisplayClusterLightCardEditor/Private/LightCardTemplates/DisplayClusterLightCardTemplate.h b/Engine/Plugins/Runtime/nDisplay/Source/DisplayClusterLightCardEditor/Private/LightCardTemplates/DisplayClusterLightCardTemplate.h
new file mode 100644
--- /dev/null
+++ b/Engine/Plugins/Runtime/nDisplay/Source/DisplayClusterLightCardEditor/Private/LightCardTemplates/DisplayClusterLightCardTemplate.h
@@ -0,0 +1,25 @@
+﻿// Copyright Epic Games, Inc. All Rights Reserved.
+
+#pragma once
+
+#include "DisplayClusterLightCardActor.h"
+
+#include "DisplayClusterLightCardTemplate.generated.h"
+
+/**
+ * A template asset to store appearance settings from Light Card actors.
+ */
+UCLASS(NotBlueprintType, NotBlueprintable, NotPlaceable, config=EditorPerProjectUserSettings)
+class DISPLAYCLUSTERLIGHTCARDEDITOR_API UDisplayClusterLightCardTemplate : public UObject
+{
+	GENERATED_BODY()
+
+public:
+	/** The instanced template object containing user settings for the light card. */
+	UPROPERTY(VisibleAnywhere, Instanced, Category = Template, NoClear, meta = (ShowOnlyInnerProperties))
+	TObjectPtr<ADisplayClusterLightCardActor> LightCardActor;
+
+	/** If the user has marked this a favorite template. */
+	UPROPERTY(Config)
+	bool bIsFavorite;
+};
diff --git a/Engine/Plugins/Runtime/nDisplay/Source/DisplayClusterLightCardEditor/Private/LightCardTemplates/DisplayClusterLightCardTemplateDragDropOp.cpp b/Engine/Plugins/Runtime/nDisplay/Source/DisplayClusterLightCardEditor/Private/LightCardTemplates/DisplayClusterLightCardTemplateDragDropOp.cpp
new file mode 100644
--- /dev/null
+++ b/Engine/Plugins/Runtime/nDisplay/Source/DisplayClusterLightCardEditor/Private/LightCardTemplates/DisplayClusterLightCardTemplateDragDropOp.cpp
@@ -0,0 +1,41 @@
+﻿// Copyright Epic Games, Inc. All Rights Reserved.
+
+#include "DisplayClusterLightCardTemplateDragDropOp.h"
+
+#define LOCTEXT_NAMESPACE "FDisplayClusterLightCardTemplateDragDropOp"
+
+TSharedRef<FDisplayClusterLightCardTemplateDragDropOp> FDisplayClusterLightCardTemplateDragDropOp::New(
+	TWeakObjectPtr<UDisplayClusterLightCardTemplate> InTemplate)
+{
+	TSharedRef<FDisplayClusterLightCardTemplateDragDropOp> DragDropOp = MakeShared<FDisplayClusterLightCardTemplateDragDropOp>();
+	DragDropOp->LightCardTemplate = InTemplate;
+	DragDropOp->MouseCursor = EMouseCursor::GrabHandClosed;
+
+	return DragDropOp;
+}
+
+void FDisplayClusterLightCardTemplateDragDropOp::Construct()
+{
+	DefaultHoverText = GetHoverText();
+	ResetToDefaultToolTip();
+	
+	FDecoratedDragDropOp::Construct();
+}
+
+void FDisplayClusterLightCardTemplateDragDropOp::SetDropAsInvalid(FText Message)
+{
+	bCanBeDropped = false;
+	CurrentHoverText = !Message.IsEmpty() ? FText::Format(LOCTEXT("HoverText_Formatted", "{0}: {1}"), GetHoverText(), Message) : DefaultHoverText;
+
+	CurrentIconBrush = FAppStyle::GetBrush(TEXT("Graph.ConnectorFeedback.Error"));
+}
+
+void FDisplayClusterLightCardTemplateDragDropOp::SetDropAsValid(FText Message)
+{
+	bCanBeDropped = true;
+	CurrentHoverText = !Message.IsEmpty() ? FText::Format(LOCTEXT("HoverText_Formatted", "{0}: {1}"), GetHoverText(), Message) : DefaultHoverText;
+
+	CurrentIconBrush = FAppStyle::GetBrush(TEXT("Graph.ConnectorFeedback.OK"));
+}
+
+#undef LOCTEXT_NAMESPACE
\ No newline at end of file
diff --git a/Engine/Plugins/Runtime/nDisplay/Source/DisplayClusterLightCardEditor/Private/LightCardTemplates/DisplayClusterLightCardTemplateDragDropOp.h b/Engine/Plugins/Runtime/nDisplay/Source/DisplayClusterLightCardEditor/Private/LightCardTemplates/DisplayClusterLightCardTemplateDragDropOp.h
new file mode 100644
--- /dev/null
+++ b/Engine/Plugins/Runtime/nDisplay/Source/DisplayClusterLightCardEditor/Private/LightCardTemplates/DisplayClusterLightCardTemplateDragDropOp.h
@@ -0,0 +1,41 @@
+﻿// Copyright Epic Games, Inc. All Rights Reserved.
+
+#pragma once
+
+#include "Input/DragAndDrop.h"
+#include "DragAndDrop/DecoratedDragDropOp.h"
+
+class UDisplayClusterLightCardTemplate;
+
+class FDisplayClusterLightCardTemplateDragDropOp : public FDecoratedDragDropOp
+{
+public:
+	DRAG_DROP_OPERATOR_TYPE(FDisplayClusterLightCardTemplateDragDropOp, FDecoratedDragDropOp)
+
+	FDisplayClusterLightCardTemplateDragDropOp() :
+		bCanBeDropped(true)
+	{ }
+
+	static TSharedRef<FDisplayClusterLightCardTemplateDragDropOp> New(TWeakObjectPtr<UDisplayClusterLightCardTemplate> InTemplate);
+
+	TWeakObjectPtr<UDisplayClusterLightCardTemplate> GetTemplate() const { return LightCardTemplate; }
+	
+	// Begin FDragDropOperation
+protected:
+	virtual void Construct() override;
+	// End FDragDropOperation
+
+public:
+	void SetDropAsInvalid(FText Message = FText::GetEmpty());
+	void SetDropAsValid(FText Message = FText::GetEmpty());
+	bool CanBeDropped() const { return bCanBeDropped; }
+
+protected:
+	virtual FText GetHoverText() const { return FText::GetEmpty(); }
+
+protected:
+	bool bCanBeDropped;
+
+private:
+	TWeakObjectPtr<UDisplayClusterLightCardTemplate> LightCardTemplate;
+};
\ No newline at end of file
diff --git a/Engine/Plugins/Runtime/nDisplay/Source/DisplayClusterLightCardEditor/Private/LightCardTemplates/SDisplayClusterLightCardTemplateList.cpp b/Engine/Plugins/Runtime/nDisplay/Source/DisplayClusterLightCardEditor/Private/LightCardTemplates/SDisplayClusterLightCardTemplateList.cpp
new file mode 100644
--- /dev/null
+++ b/Engine/Plugins/Runtime/nDisplay/Source/DisplayClusterLightCardEditor/Private/LightCardTemplates/SDisplayClusterLightCardTemplateList.cpp
@@ -0,0 +1,349 @@
+// Copyright Epic Games, Inc. All Rights Reserved.
+
+#include "SDisplayClusterLightCardTemplateList.h"
+
+#include "DisplayClusterLightCardTemplate.h"
+#include "DisplayClusterLightCardTemplateDragDropOp.h"
+#include "SDisplayClusterLightCardEditor.h"
+
+#include "AssetRegistry/AssetRegistryModule.h"
+#include "Engine/Texture.h"
+#include "Styling/SlateIconFinder.h"
+#include "Widgets/Input/SSearchBox.h"
+#include "Widgets/Layout/SBorder.h"
+#include "Widgets/Layout/SBox.h"
+#include "Widgets/Views/STableRow.h"
+#include "Widgets/Views/STreeView.h"
+
+#define LOCTEXT_NAMESPACE "SDisplayClusterLightCardTemplateList"
+
+namespace DisplayClusterLightCardTemplateListColumnNames
+{
+	const static FName LightCardTemplateName(TEXT("LightCardTemplateName"));
+};
+
+class SLightCardTemplateTreeItemRow : public SMultiColumnTableRow<TSharedPtr<SDisplayClusterLightCardTemplateList::FLightCardTemplateTreeItem>>
+{
+	SLATE_BEGIN_ARGS(SLightCardTemplateTreeItemRow) {}
+	SLATE_END_ARGS()
+
+public:
+
+	void Construct(const FArguments& InArgs, const TSharedRef<STableViewBase>& InOwnerTable, const TSharedPtr<SDisplayClusterLightCardTemplateList>& InTemplateList, const TSharedPtr<SDisplayClusterLightCardTemplateList::FLightCardTemplateTreeItem> InLightCardTreeItem)
+	{
+		LightCardTreeItem = InLightCardTreeItem;
+		OwningTemplateListPtr = InTemplateList;
+
+		SMultiColumnTableRow<TSharedPtr<SDisplayClusterLightCardTemplateList::FLightCardTemplateTreeItem>>::Construct(
+			SMultiColumnTableRow<TSharedPtr<SDisplayClusterLightCardTemplateList::FLightCardTemplateTreeItem>>::FArguments()
+			.Style(&FAppStyle::Get().GetWidgetStyle<FTableRowStyle>("SceneOutliner.TableViewRow"))
+			.OnDragDetected(this, &SLightCardTemplateTreeItemRow::HandleDragDetected)
+		, InOwnerTable);
+	}
+
+	virtual TSharedRef<SWidget> GenerateWidgetForColumn(const FName& ColumnName) override
+	{
+		if (ColumnName == DisplayClusterLightCardTemplateListColumnNames::LightCardTemplateName)
+		{
+			return SNew(SHorizontalBox)
+
+			+ SHorizontalBox::Slot()
+			.AutoWidth()
+			.Padding(6.f, 0.f)
+			.VAlign(VAlign_Center)
+			.HAlign(HAlign_Center)
+			[
+				SAssignNew(FavoriteIcon, SBorder)
+				.BorderImage(FAppStyle::GetBrush("NoBorder"))
+				.OnMouseButtonDown(this, &SLightCardTemplateTreeItemRow::OnFavoriteMouseDown)
+				[
+					SNew(SImage)
+					.Image(this, &SLightCardTemplateTreeItemRow::GetFavoriteIcon)
+					.ToolTipText(LOCTEXT("FavoriteToolTip", "Favorite this template"))
+				]
+			]
+
+			+ SHorizontalBox::Slot()
+			.AutoWidth()
+			.Padding(FMargin(0, 1, 6, 1))
+			.VAlign(VAlign_Center)
+			[
+				SNew(SBox)
+				.WidthOverride(16)
+				.HeightOverride(16)
+				[
+					SNew(SImage)
+					.ColorAndOpacity(FSlateColor::UseForeground())
+					.Image(this, &SLightCardTemplateTreeItemRow::GetLightCardIcon)
+				]
+			]
+
+			+ SHorizontalBox::Slot()
+			.AutoWidth()
+			.VAlign(VAlign_Center)
+			.HAlign(HAlign_Left)
+			[
+				SNew(STextBlock)
+				.Text(this, &SLightCardTemplateTreeItemRow::GetLightCardDisplayName)
+			];
+		}
+
+		return SNullWidget::NullWidget;
+	}
+
+private:
+	FText GetLightCardDisplayName() const
+	{
+		FString ItemName = TEXT("");
+		if (LightCardTreeItem.IsValid())
+		{
+			const TSharedPtr<SDisplayClusterLightCardTemplateList::FLightCardTemplateTreeItem> LightCardTreeItemPin = LightCardTreeItem.Pin();
+			ItemName = LightCardTreeItemPin->TemplateName.ToString();
+		}
+
+		return FText::FromString(*ItemName);
+	}
+
+	const FSlateBrush* GetLightCardIcon() const
+	{
+		if (LightCardTreeItem.IsValid())
+		{
+			if (LightCardTreeItem.Pin()->SlateBrush.IsValid())
+			{
+				return LightCardTreeItem.Pin()->SlateBrush.Get();
+			}
+			if (LightCardTreeItem.Pin()->LightCardTemplate.IsValid())
+			{
+				return FSlateIconFinder::FindIconBrushForClass(LightCardTreeItem.Pin()->LightCardTemplate->GetClass());
+			}
+		}
+
+		return nullptr;
+	}
+	
+	FReply HandleDragDetected(const FGeometry& MyGeometry, const FPointerEvent& MouseEvent)
+	{
+		if (LightCardTreeItem.IsValid() && MouseEvent.IsMouseButtonDown(EKeys::LeftMouseButton))
+		{
+			const TSharedPtr<FDragDropOperation> DragDropOp = FDisplayClusterLightCardTemplateDragDropOp::New(LightCardTreeItem.Pin()->LightCardTemplate);
+
+			if (DragDropOp.IsValid())
+			{
+				return FReply::Handled().BeginDragDrop(DragDropOp.ToSharedRef());
+			}
+		}
+
+		return FReply::Unhandled();
+	}
+	
+	const FSlateBrush* GetFavoriteIcon() const
+	{
+		const bool bIsFavorite = (LightCardTreeItem.IsValid() && LightCardTreeItem.Pin()->LightCardTemplate.IsValid()) ? LightCardTreeItem.Pin()->LightCardTemplate->bIsFavorite : false;
+        const FName FavoriteIconName = bIsFavorite ? "DetailsView.PropertyIsFavorite" : "DetailsView.PropertyIsNotFavorite";
+		return FSlateIcon(FAppStyle::Get().GetStyleSetName(), FavoriteIconName).GetIcon();
+	}
+
+	FReply OnFavoriteMouseDown(const FGeometry& MyGeometry, const FPointerEvent& MouseEvent)
+	{
+		if (LightCardTreeItem.IsValid() && LightCardTreeItem.Pin()->LightCardTemplate.IsValid())
+		{
+			LightCardTreeItem.Pin()->LightCardTemplate->bIsFavorite = !LightCardTreeItem.Pin()->LightCardTemplate->bIsFavorite;
+			LightCardTreeItem.Pin()->LightCardTemplate->SaveConfig();
+			LightCardTreeItem.Pin()->LightCardTemplate->PostEditChange();
+			
+			if (OwningTemplateListPtr.IsValid())
+			{
+				OwningTemplateListPtr.Pin()->ApplyFilter();
+			}
+		}
+		return FReply::Handled();
+	}
+
+private:
+	TWeakPtr<SDisplayClusterLightCardTemplateList> OwningTemplateListPtr;
+	TWeakPtr<SDisplayClusterLightCardTemplateList::FLightCardTemplateTreeItem> LightCardTreeItem;
+	TSharedPtr<SBorder> FavoriteIcon;
+};
+
+bool SDisplayClusterLightCardTemplateList::FLightCardTemplateTreeItem::IsFavorite() const
+{
+	return LightCardTemplate.IsValid() && LightCardTemplate->bIsFavorite;
+}
+
+SDisplayClusterLightCardTemplateList::~SDisplayClusterLightCardTemplateList()
+{
+	IAssetRegistry& AssetRegistry = FModuleManager::GetModuleChecked<FAssetRegistryModule>(AssetRegistryConstants::ModuleName).Get();
+	AssetRegistry.OnAssetAdded().RemoveAll(this);
+	AssetRegistry.OnAssetRemoved().RemoveAll(this);
+	AssetRegistry.OnFilesLoaded().RemoveAll(this);
+}
+
+void SDisplayClusterLightCardTemplateList::Construct(const FArguments& InArgs, TSharedPtr<SDisplayClusterLightCardEditor> InLightCardEditor)
+{
+	LightCardEditorPtr = InLightCardEditor;
+
+	IAssetRegistry& AssetRegistry = FModuleManager::LoadModuleChecked<FAssetRegistryModule>(AssetRegistryConstants::ModuleName).Get();
+	AssetRegistry.OnAssetAdded().AddRaw(this, &SDisplayClusterLightCardTemplateList::OnAssetAddedOrRemoved);
+	AssetRegistry.OnAssetRemoved().AddRaw(this, &SDisplayClusterLightCardTemplateList::OnAssetAddedOrRemoved);
+	AssetRegistry.OnFilesLoaded().AddRaw(this, &SDisplayClusterLightCardTemplateList::OnAssetsLoaded);
+	
+	RefreshTemplateList();
+
+	ChildSlot
+	[
+		SNew(SVerticalBox)
+		+SVerticalBox::Slot()
+		.Padding(0.f, 8.f, 0.f, 5.f)
+		.AutoHeight()
+		[
+			SNew(SHorizontalBox)
+			+SHorizontalBox::Slot()
+			.Padding(6.f, 0.f)
+			.AutoWidth()
+			.VAlign(VAlign_Center)
+			.HAlign(HAlign_Center)
+			[
+				SNew(SImage)
+				.Image(this, &SDisplayClusterLightCardTemplateList::GetFavoriteFilterIcon)
+				.ToolTipText(LOCTEXT("FavoriteFilterToolTip", "Filter by favorites"))
+				.OnMouseButtonDown_Lambda([this](const FGeometry& MyGeometry, const FPointerEvent& MouseEvent)
+				{
+					bFilterFavorites = !bFilterFavorites;
+					ApplyFilter();
+					return FReply::Handled();
+				})
+			]
+			+SHorizontalBox::Slot()
+			.FillWidth(1.f)
+			[
+				SNew(SSearchBox)
+				.SelectAllTextWhenFocused(true)
+				.OnTextChanged(this, &SDisplayClusterLightCardTemplateList::OnFilterTextChanged)
+				.HintText(LOCTEXT("SearchBoxHint", "Search Light Card Templates..."))
+			]
+		]
+		+SVerticalBox::Slot()
+		[
+			SAssignNew(LightCardTemplateTreeView, STreeView<TSharedPtr<FLightCardTemplateTreeItem>>)
+			.TreeItemsSource(&FilteredLightCardTemplateTree)
+			.ItemHeight(28)
+			.SelectionMode(ESelectionMode::Single)
+			.OnGenerateRow(this, &SDisplayClusterLightCardTemplateList::GenerateTreeItemRow)
+			.OnGetChildren(this, &SDisplayClusterLightCardTemplateList::GetChildrenForTreeItem)
+			.HeaderRow
+			(
+				SNew(SHeaderRow)
+
+				+ SHeaderRow::Column(DisplayClusterLightCardTemplateListColumnNames::LightCardTemplateName)
+				.DefaultLabel(LOCTEXT("LightCardTemplateName", "Light Card Templates"))
+				.FillWidth(0.8f)
+			)
+		]
+	];
+}
+
+void SDisplayClusterLightCardTemplateList::RefreshTemplateList()
+{
+	LightCardTemplateTree.Reset();
+	
+	const IAssetRegistry& AssetRegistry = FModuleManager::GetModuleChecked<FAssetRegistryModule>(AssetRegistryConstants::ModuleName).Get();
+	
+	TArray<FAssetData> TemplateAssets;
+	const FTopLevelAssetPath ClassPath = UDisplayClusterLightCardTemplate::StaticClass()->GetClassPathName();
+	AssetRegistry.GetAssetsByClass(ClassPath, TemplateAssets, true);
+
+	for (const FAssetData& AssetData : TemplateAssets)
+	{
+		TSharedPtr<FLightCardTemplateTreeItem> TemplateTreeItem = MakeShared<FLightCardTemplateTreeItem>();
+		UDisplayClusterLightCardTemplate* AssetObject = CastChecked<UDisplayClusterLightCardTemplate>(AssetData.GetAsset());
+		
+		TemplateTreeItem->TemplateName = AssetData.AssetName;
+		TemplateTreeItem->LightCardTemplate = AssetObject;
+
+		if (AssetObject->LightCardActor && AssetObject->LightCardActor->Texture.Get())
+		{
+			TemplateTreeItem->SlateBrush = MakeShared<FSlateBrush>();
+		
+			TemplateTreeItem->SlateBrush->SetResourceObject(AssetObject->LightCardActor->Texture.Get());
+			TemplateTreeItem->SlateBrush->ImageSize = FVector2D(64.f, 64.f);
+		}
+		
+		LightCardTemplateTree.Add(TemplateTreeItem);
+	}
+
+	ApplyFilter();
+}
+
+void SDisplayClusterLightCardTemplateList::ApplyFilter(const FString& InFilterText, bool bFavorite)
+{
+	FilterText = InFilterText;
+	
+	FilteredLightCardTemplateTree.Reset();
+
+	for (TSharedPtr<FLightCardTemplateTreeItem>& LightCardTemplate : LightCardTemplateTree)
+	{
+		if (LightCardTemplate.IsValid() &&
+			((InFilterText.IsEmpty() || LightCardTemplate->TemplateName.ToString().Contains(InFilterText)) &&
+			(!bFavorite || LightCardTemplate->IsFavorite())))
+		{
+			FilteredLightCardTemplateTree.Add(LightCardTemplate);
+		}
+	}
+
+	if (LightCardTemplateTreeView.IsValid())
+	{
+		LightCardTemplateTreeView->RequestTreeRefresh();
+	}
+}
+
+void SDisplayClusterLightCardTemplateList::ApplyFilter()
+{
+	ApplyFilter(FilterText, bFilterFavorites);
+}
+
+TSharedRef<ITableRow> SDisplayClusterLightCardTemplateList::GenerateTreeItemRow(
+	TSharedPtr<FLightCardTemplateTreeItem> Item, const TSharedRef<STableViewBase>& OwnerTable)
+{
+	return SNew(SLightCardTemplateTreeItemRow, OwnerTable, SharedThis(this), Item);
+}
+
+void SDisplayClusterLightCardTemplateList::GetChildrenForTreeItem(TSharedPtr<FLightCardTemplateTreeItem> InItem,
+	TArray<TSharedPtr<FLightCardTemplateTreeItem>>& OutChildren)
+{
+}
+
+const FSlateBrush* SDisplayClusterLightCardTemplateList::GetFavoriteFilterIcon() const
+{
+	const FName FavoriteIconName = bFilterFavorites ? "DetailsView.PropertyIsFavorite" : "DetailsView.PropertyIsNotFavorite";
+	return FSlateIcon(FAppStyle::Get().GetStyleSetName(), FavoriteIconName).GetIcon();
+}
+
+void SDisplayClusterLightCardTemplateList::OnFilterTextChanged(const FText& SearchText)
+{
+	ApplyFilter(SearchText.ToString(), bFilterFavorites);
+}
+
+void SDisplayClusterLightCardTemplateList::OnAssetAddedOrRemoved(const FAssetData& InAssetData)
+{
+	const IAssetRegistry& AssetRegistry = FModuleManager::GetModuleChecked<FAssetRegistryModule>(AssetRegistryConstants::ModuleName).Get();
+	if (AssetRegistry.IsLoadingAssets())
+	{
+		// This will be triggered once assets have finished loaded.
+		return;
+	}
+	
+	if (InAssetData.IsValid() && !InAssetData.IsRedirector())
+	{
+		if (InAssetData.AssetClassPath == UDisplayClusterLightCardTemplate::StaticClass()->GetClassPathName())
+		{
+			RefreshTemplateList();
+		}
+	}
+}
+
+void SDisplayClusterLightCardTemplateList::OnAssetsLoaded()
+{
+	RefreshTemplateList();
+}
+
+#undef LOCTEXT_NAMESPACE
diff --git a/Engine/Plugins/Runtime/nDisplay/Source/DisplayClusterLightCardEditor/Private/LightCardTemplates/SDisplayClusterLightCardTemplateList.h b/Engine/Plugins/Runtime/nDisplay/Source/DisplayClusterLightCardEditor/Private/LightCardTemplates/SDisplayClusterLightCardTemplateList.h
new file mode 100644
--- /dev/null
+++ b/Engine/Plugins/Runtime/nDisplay/Source/DisplayClusterLightCardEditor/Private/LightCardTemplates/SDisplayClusterLightCardTemplateList.h
@@ -0,0 +1,78 @@
+// Copyright Epic Games, Inc. All Rights Reserved.
+
+#pragma once
+
+#include "CoreMinimal.h"
+#include "Widgets/DeclarativeSyntaxSupport.h"
+#include "Widgets/SCompoundWidget.h"
+
+class SDisplayClusterLightCardEditor;
+class UDisplayClusterLightCardTemplate;
+class ITableRow;
+class STableViewBase;
+template<class T>
+class STreeView;
+
+/** Displays available light card templates for users to drop into the light card editor scene */
+class SDisplayClusterLightCardTemplateList : public SCompoundWidget
+{
+public:
+	struct FLightCardTemplateTreeItem
+	{
+		TWeakObjectPtr<UDisplayClusterLightCardTemplate> LightCardTemplate;
+		FName TemplateName;
+		TSharedPtr<FSlateBrush> SlateBrush;
+
+		FLightCardTemplateTreeItem()
+		{ }
+
+		/** If the template is a user favorite */
+		bool IsFavorite() const;
+	};
+	
+	SLATE_BEGIN_ARGS(SDisplayClusterLightCardTemplateList)
+	{}
+	SLATE_END_ARGS()
+
+	virtual ~SDisplayClusterLightCardTemplateList() override;
+
+	void Construct(const FArguments& InArgs, TSharedPtr<SDisplayClusterLightCardEditor> InLightCardEditor);
+
+	/** Find available templates and refresh the list */
+	void RefreshTemplateList();
+
+	/** Apply a filter to the template list */
+	void ApplyFilter(const FString& InFilterText, bool bFavorite);
+	
+	/** Reapply the current filter */
+	void ApplyFilter();
+	
+private:
+	TSharedRef<ITableRow> GenerateTreeItemRow(TSharedPtr<FLightCardTemplateTreeItem> Item, const TSharedRef<STableViewBase>& OwnerTable);
+	void GetChildrenForTreeItem(TSharedPtr<FLightCardTemplateTreeItem> InItem, TArray<TSharedPtr<FLightCardTemplateTreeItem>>& OutChildren);
+	const FSlateBrush* GetFavoriteFilterIcon() const;
+	
+	void OnFilterTextChanged(const FText& SearchText);
+	
+	void OnAssetAddedOrRemoved(const FAssetData& InAssetData);
+	void OnAssetsLoaded();
+	
+private:
+	/** Pointer to the light card editor that owns this widget */
+	TWeakPtr<SDisplayClusterLightCardEditor> LightCardEditorPtr;
+
+	/** A list of all light card tree items to be displayed in the tree view */
+	TArray<TSharedPtr<FLightCardTemplateTreeItem>> LightCardTemplateTree;
+
+	/** The visible list of light card tree items filtered by the user */
+	TArray<TSharedPtr<FLightCardTemplateTreeItem>> FilteredLightCardTemplateTree;
+	
+	/** The tree view widget displaying the light card templates */
+	TSharedPtr<STreeView<TSharedPtr<FLightCardTemplateTreeItem>>> LightCardTemplateTreeView;
+
+	/** The search string text */
+	FString FilterText;
+
+	/** If the favorite filter should be applied */
+	bool bFilterFavorites = false;
+};
\ No newline at end of file
diff --git a/Engine/Plugins/Runtime/nDisplay/Source/DisplayClusterLightCardEditor/Private/SDisplayClusterLightCardEditor.cpp b/Engine/Plugins/Runtime/nDisplay/Source/DisplayClusterLightCardEditor/Private/SDisplayClusterLightCardEditor.cpp
--- a/Engine/Plugins/Runtime/nDisplay/Source/DisplayClusterLightCardEditor/Private/SDisplayClusterLightCardEditor.cpp
+++ b/Engine/Plugins/Runtime/nDisplay/Source/DisplayClusterLightCardEditor/Private/SDisplayClusterLightCardEditor.cpp
@@ -4,26 +4,34 @@
 
 #include "DisplayClusterLightCardEditorCommands.h"
 #include "SDisplayClusterLightCardList.h"
+#include "LightCardTemplates/SDisplayClusterLightCardTemplateList.h"
+#include "LightCardTemplates/DisplayClusterLightCardTemplate.h"
+#include "Settings/DisplayClusterLightCardEditorSettings.h"
 
-#include "IDisplayClusterOperator.h"
 #include "Viewport/DisplayClusterLightcardEditorViewport.h"
 #include "Viewport/DisplayClusterLightCardEditorViewportClient.h"
 
+#include "IDisplayClusterOperator.h"
+
 #include "DisplayClusterRootActor.h"
 #include "DisplayClusterConfigurationTypes.h"
 #include "Components/DisplayClusterCameraComponent.h"
 
+#include "AssetRegistry/AssetRegistryModule.h"
+#include "ContentBrowserModule.h"
+#include "FileHelpers.h"
 #include "Framework/Commands/GenericCommands.h"
 #include "Framework/Docking/TabManager.h"
 #include "Framework/Docking/LayoutExtender.h"
 #include "Framework/MultiBox/MultiBoxBuilder.h"
 #include "Framework/MultiBox/MultiBoxExtender.h"
+#include "IContentBrowserSingleton.h"
+#include "Misc/FileHelper.h"
 #include "PropertyCustomizationHelpers.h"
 #include "ScopedTransaction.h"
 #include "Selection.h"
 #include "Toolkits/AssetEditorToolkit.h"
 #include "Widgets/Docking/SDockTab.h"
-#include "Widgets/Layout/SSplitter.h"
 #include "Widgets/Workflow/SWizard.h"
 
 #define LOCTEXT_NAMESPACE "SDisplayClusterLightCardEditor"
@@ -101,32 +109,87 @@ void SDisplayClusterLightCardEditor::Construct(const FArguments& InArgs, const T
 	OnObjectTransactedHandle = FCoreUObjectDelegates::OnObjectTransacted.AddSP(this, &SDisplayClusterLightCardEditor::OnObjectTransacted);
 
 	BindCommands();
+	
+	TabManager = FGlobalTabmanager::Get()->NewTabManager(MajorTabOwner);
 
-	ChildSlot                     
-	[
-		SNew(SSplitter)
-		.Orientation(Orient_Horizontal)
+	const FName LightCardTab = TEXT("LightCards");
+	const FName LightCardTemplateTab = TEXT("LightCardTemplates");
+	const FName ViewportTab = TEXT("Viewport"); 
+	
+	const TSharedRef<FTabManager::FLayout> Layout = FTabManager::NewLayout("LightCardEditor")
+	->AddArea
+	(
+		FTabManager::NewPrimaryArea()
 		
-		+SSplitter::Slot()
-		.Value(0.25f)
+		->Split
+		(
+			FTabManager::NewSplitter()
+			->SetSizeCoefficient(0.25f)
+			->SetOrientation(Orient_Vertical)
+			->Split
+			(
+			FTabManager::NewStack()
+				->AddTab(LightCardTab, ETabState::OpenedTab)
+				->SetHideTabWell(true)
+				->SetForegroundTab(LightCardTab)
+			)
+			->Split
+			(
+			FTabManager::NewStack()
+				->AddTab(LightCardTemplateTab, ETabState::OpenedTab)
+				->SetHideTabWell(true)
+				->SetForegroundTab(LightCardTab)
+			)
+		)
+		->Split
+		(
+			FTabManager::NewStack()
+			->AddTab(ViewportTab, ETabState::OpenedTab)
+			->SetSizeCoefficient(0.75f)
+			->SetHideTabWell(true)
+			->SetForegroundTab(ViewportTab)
+		)
+	);
+
+	TabManager->RegisterTabSpawner(LightCardTab, FOnSpawnTab::CreateLambda([this](const FSpawnTabArgs&) -> TSharedRef<SDockTab>
+	{
+		return SNew(SDockTab)
+		.TabRole(ETabRole::PanelTab)
+		.CanEverClose(false)
+		.OnCanCloseTab(SDockTab::FCanCloseTab::CreateLambda([](){ return false; }))
+		.ContentPadding(0)
 		[
-			// Vertical box for the left hand panel of the editor. Add new slots here as needed for any editor UI controls
-			SNew(SVerticalBox)
-			+ SVerticalBox::Slot()
-			.FillHeight(1.0f)
-			[
-				CreateLightCardListWidget()
-			]
-		]
-		+SSplitter::Slot()
+			CreateLightCardListWidget()
+		];
+	}));
+
+	TabManager->RegisterTabSpawner(LightCardTemplateTab, FOnSpawnTab::CreateLambda([this](const FSpawnTabArgs&) -> TSharedRef<SDockTab>
+	{
+		return SNew(SDockTab)
+		.TabRole(ETabRole::PanelTab)
+		.CanEverClose(false)
+		.OnCanCloseTab(SDockTab::FCanCloseTab::CreateLambda([](){ return false; }))
+		.ContentPadding(0)
 		[
-			SNew(SVerticalBox)
-			+SVerticalBox::Slot()
-			.FillHeight(0.75f)
-			[
-				CreateViewportWidget()
-			]
-		]
+			CreateLightCardTemplateWidget()
+		];
+	}));
+	
+	TabManager->RegisterTabSpawner(ViewportTab, FOnSpawnTab::CreateLambda([this](const FSpawnTabArgs&) -> TSharedRef<SDockTab>
+	{
+		return SNew(SDockTab)
+		.TabRole(ETabRole::PanelTab)
+		.CanEverClose(false)
+		.OnCanCloseTab(SDockTab::FCanCloseTab::CreateLambda([](){ return false; }))
+		.ContentPadding(0)
+		[
+			CreateViewportWidget()
+		];
+	}));
+	
+	ChildSlot                     
+	[
+		TabManager->RestoreFrom(Layout, WindowOwner).ToSharedRef()
 	];
 
 	BindCompileDelegates();
@@ -186,6 +249,41 @@ ADisplayClusterLightCardActor* SDisplayClusterLightCardEditor::SpawnLightCard()
 	return NewLightCard;
 }
 
+ADisplayClusterLightCardActor* SDisplayClusterLightCardEditor::SpawnLightCardFromTemplate(
+	const UDisplayClusterLightCardTemplate* InTemplate, ULevel* InLevel, bool bIsPreview)
+{
+	if (!ActiveRootActor.IsValid())
+	{
+		return nullptr;
+	}
+
+	check(InTemplate);
+
+	ULevel* Level = InLevel ? InLevel : ActiveRootActor->GetWorld()->GetCurrentLevel();
+	
+	FName UniqueName = *InTemplate->GetName().Replace(TEXT("Template"), TEXT(""));
+	if (StaticFindObjectFast(InTemplate->LightCardActor->GetClass(), Level, UniqueName))
+	{
+		UniqueName = MakeUniqueObjectName(Level, InTemplate->LightCardActor->GetClass(), UniqueName);
+	}
+	
+	// Duplicate, don't copy properties or spawn from a template. Doing so will copy component data incorrectly,
+	// specifically the static mesh override textures. They will be parented to the template, not the level instance
+	// and prevent the map from saving.
+	ADisplayClusterLightCardActor* NewLightCard = CastChecked<ADisplayClusterLightCardActor>(StaticDuplicateObject(InTemplate->LightCardActor.Get(), Level, UniqueName));
+	Level->AddLoadedActor(NewLightCard);
+	
+	if (!bIsPreview)
+	{
+		NewLightCard->SetActorLabel(NewLightCard->GetName());
+	
+		const TArray<ADisplayClusterLightCardActor*> LightCards { NewLightCard } ;
+		AddLightCardsToActor(LightCards);
+	}
+	
+	return NewLightCard;
+}
+
 void SDisplayClusterLightCardEditor::AddNewLightCard()
 {
 	check(ActiveRootActor.IsValid());
@@ -504,6 +602,77 @@ bool SDisplayClusterLightCardEditor::CanRemoveLightCards() const
 	return SelectedLightCards.Num() > 0;
 }
 
+void SDisplayClusterLightCardEditor::CreateLightCardTemplate()
+{
+	TArray<ADisplayClusterLightCardActor*> SelectedLightCards;
+	LightCardList->GetSelectedLightCards(SelectedLightCards);
+
+	check(SelectedLightCards.Num() == 1);
+	const ADisplayClusterLightCardActor* LightCardActor = SelectedLightCards[0];
+
+	const UDisplayClusterLightCardEditorProjectSettings* Settings = GetDefault<UDisplayClusterLightCardEditorProjectSettings>();
+	FString DefaultPath = Settings->LightCardTemplateDefaultPath.Path;
+	if (DefaultPath.IsEmpty())
+	{
+		DefaultPath = TEXT("/Game");
+	}
+	
+	FString DefaultAssetName = LightCardActor->GetActorLabel();
+	if (!DefaultAssetName.EndsWith(TEXT("Template")))
+	{
+		DefaultAssetName += TEXT("Template");
+	}
+
+	FString PackageName;
+	bool FilenameValid = false;
+	while (!FilenameValid)
+	{
+		FSaveAssetDialogConfig SaveAssetDialogConfig;
+		{
+			SaveAssetDialogConfig.DefaultPath = MoveTemp(DefaultPath);
+			SaveAssetDialogConfig.DefaultAssetName = MoveTemp(DefaultAssetName);
+			SaveAssetDialogConfig.AssetClassNames.Add(UDisplayClusterLightCardTemplate::StaticClass()->GetClassPathName());
+			SaveAssetDialogConfig.ExistingAssetPolicy = ESaveAssetDialogExistingAssetPolicy::AllowButWarn;
+			SaveAssetDialogConfig.DialogTitleOverride = LOCTEXT("SaveAssetDialogTitle", "Save Asset As");
+		}
+
+		const FContentBrowserModule& ContentBrowserModule = FModuleManager::LoadModuleChecked<FContentBrowserModule>("ContentBrowser");
+		const FString SaveObjectPath = ContentBrowserModule.Get().CreateModalSaveAssetDialog(SaveAssetDialogConfig);
+
+		if (SaveObjectPath.IsEmpty())
+		{
+			return;
+		}
+
+		PackageName = FPackageName::ObjectPathToPackageName(SaveObjectPath);
+		
+		FText OutError;
+		FilenameValid = FFileHelper::IsFilenameValidForSaving(PackageName, OutError);
+	}
+
+	const FString NewAssetName = FPackageName::GetLongPackageAssetName(PackageName);
+	UPackage* AssetPackage = CreatePackage(*PackageName);
+	check(AssetPackage);
+
+	const EObjectFlags Flags = RF_Public | RF_Standalone;
+	UDisplayClusterLightCardTemplate* Template = NewObject<UDisplayClusterLightCardTemplate>(AssetPackage, FName(*NewAssetName), Flags);
+	Template->LightCardActor = CastChecked<ADisplayClusterLightCardActor>(StaticDuplicateObject(LightCardActor, Template));
+
+	FAssetRegistryModule::AssetCreated(Template);
+
+	AssetPackage->MarkPackageDirty();
+
+	FEditorFileUtils::PromptForCheckoutAndSave({AssetPackage}, true, false);
+}
+
+bool SDisplayClusterLightCardEditor::CanCreateLightCardTemplate() const
+{
+	TArray<ADisplayClusterLightCardActor*> SelectedLightCards;
+	LightCardList->GetSelectedLightCards(SelectedLightCards);
+
+	return SelectedLightCards.Num() == 1;
+}
+
 void SDisplayClusterLightCardEditor::OnActiveRootActorChanged(ADisplayClusterRootActor* NewRootActor)
 {
 	RemoveCompileDelegates();
@@ -529,6 +698,11 @@ TSharedRef<SWidget> SDisplayClusterLightCardEditor::CreateViewportWidget()
 	return SAssignNew(ViewportView, SDisplayClusterLightCardEditorViewport, SharedThis(this), CommandList);
 }
 
+TSharedRef<SWidget> SDisplayClusterLightCardEditor::CreateLightCardTemplateWidget()
+{
+	return SAssignNew(LightCardTemplateList, SDisplayClusterLightCardTemplateList, SharedThis(this));
+}
+
 void SDisplayClusterLightCardEditor::BindCommands()
 {
 	CommandList = MakeShareable(new FUICommandList);
@@ -572,6 +746,11 @@ void SDisplayClusterLightCardEditor::BindCommands()
 		FGenericCommands::Get().Delete,
 		FExecuteAction::CreateSP(this, &SDisplayClusterLightCardEditor::RemoveLightCards, true),
 		FCanExecuteAction::CreateSP(this, &SDisplayClusterLightCardEditor::CanRemoveLightCards));
+
+	CommandList->MapAction(
+		FDisplayClusterLightCardEditorCommands::Get().SaveLightCardTemplate,
+		FExecuteAction::CreateSP(this, &SDisplayClusterLightCardEditor::CreateLightCardTemplate),
+		FCanExecuteAction::CreateSP(this, &SDisplayClusterLightCardEditor::CanCreateLightCardTemplate));
 }
 
 void SDisplayClusterLightCardEditor::RefreshPreviewActors(EDisplayClusterLightCardEditorProxyType ProxyType)
diff --git a/Engine/Plugins/Runtime/nDisplay/Source/DisplayClusterLightCardEditor/Private/SDisplayClusterLightCardEditor.h b/Engine/Plugins/Runtime/nDisplay/Source/DisplayClusterLightCardEditor/Private/SDisplayClusterLightCardEditor.h
--- a/Engine/Plugins/Runtime/nDisplay/Source/DisplayClusterLightCardEditor/Private/SDisplayClusterLightCardEditor.h
+++ b/Engine/Plugins/Runtime/nDisplay/Source/DisplayClusterLightCardEditor/Private/SDisplayClusterLightCardEditor.h
@@ -15,9 +15,11 @@ class FToolBarBuilder;
 class FUICommandList;
 class SDockTab;
 class SDisplayClusterLightCardList;
+class SDisplayClusterLightCardTemplateList;
 class SDisplayClusterLightCardEditorViewport;
 class ADisplayClusterRootActor;
 class ADisplayClusterLightCardActor;
+class UDisplayClusterLightCardTemplate;
 
 /** A panel that can be spawned in a tab that contains all the UI elements that make up the 2D light cards editor */
 class SDisplayClusterLightCardEditor : public SCompoundWidget
@@ -67,6 +69,9 @@ public:
 	/** Spawns a new light card and adds it to the root actor */
 	ADisplayClusterLightCardActor* SpawnLightCard();
 
+	/** Spawns a new light card from a light card template */
+	ADisplayClusterLightCardActor* SpawnLightCardFromTemplate(const UDisplayClusterLightCardTemplate* InTemplate, ULevel* InLevel, bool bIsPreview);
+
 	/** Adds a new light card to the root actor and centers it in the viewport */
 	void AddNewLightCard();
 
@@ -112,6 +117,12 @@ public:
 	/** If the selected Light Card can be removed */
 	bool CanRemoveLightCards() const;
 
+	/** Creates a template of the selected light card */
+	void CreateLightCardTemplate();
+
+	/** Determines if there are selected light cards that can have templates created */
+	bool CanCreateLightCardTemplate() const;
+
 private:
 	/** Raised when the active Display cluster root actor has been changed in the operator panel */
 	void OnActiveRootActorChanged(ADisplayClusterRootActor* NewRootActor);
@@ -119,6 +130,9 @@ private:
 	/** Creates the widget used to show the list of light cards associated with the active root actor */
 	TSharedRef<SWidget> CreateLightCardListWidget();
 
+	/** Create the widget for selecting light card templates */
+	TSharedRef<SWidget> CreateLightCardTemplateWidget();
+	
 	/** Create the 3d viewport widget */
 	TSharedRef<SWidget> CreateViewportWidget();
 
@@ -158,9 +172,14 @@ private:
 	void OnLightCardListChanged();
 	
 private:
+	TSharedPtr<FTabManager> TabManager;
+	
 	/** The light card list widget */
 	TSharedPtr<SDisplayClusterLightCardList> LightCardList;
 
+	/** The light card template list */
+	TSharedPtr<SDisplayClusterLightCardTemplateList> LightCardTemplateList;
+	
 	/** The 3d viewport */
 	TSharedPtr<SDisplayClusterLightCardEditorViewport> ViewportView;
 
diff --git a/Engine/Plugins/Runtime/nDisplay/Source/DisplayClusterLightCardEditor/Private/SDisplayClusterLightCardList.cpp b/Engine/Plugins/Runtime/nDisplay/Source/DisplayClusterLightCardEditor/Private/SDisplayClusterLightCardList.cpp
--- a/Engine/Plugins/Runtime/nDisplay/Source/DisplayClusterLightCardEditor/Private/SDisplayClusterLightCardList.cpp
+++ b/Engine/Plugins/Runtime/nDisplay/Source/DisplayClusterLightCardEditor/Private/SDisplayClusterLightCardList.cpp
@@ -146,6 +146,7 @@ void SDisplayClusterLightCardList::Construct(const FArguments& InArgs, TSharedPt
 		SNew(SVerticalBox)
 		+SVerticalBox::Slot()
 		.AutoHeight()
+		.Padding(0.f, 5.f, 0.f, 0.f)
 		[
 			SNew(SBorder)
 			.Padding(0)
@@ -409,6 +410,7 @@ TSharedPtr<SWidget> SDisplayClusterLightCardList::CreateContextMenu()
 	const bool bCloseAfterSelection = true;
 	FMenuBuilder MenuBuilder(bCloseAfterSelection, CommandList, Extenders);
 	MenuBuilder.AddMenuEntry(FDisplayClusterLightCardEditorCommands::Get().RemoveLightCard);
+	MenuBuilder.AddMenuEntry(FDisplayClusterLightCardEditorCommands::Get().SaveLightCardTemplate);
 
 	MenuBuilder.BeginSection("Edit", LOCTEXT("EditSection", "Edit"));
 	{
diff --git a/Engine/Plugins/Runtime/nDisplay/Source/DisplayClusterLightCardEditor/Private/SDisplayClusterLightCardList.h b/Engine/Plugins/Runtime/nDisplay/Source/DisplayClusterLightCardEditor/Private/SDisplayClusterLightCardList.h
--- a/Engine/Plugins/Runtime/nDisplay/Source/DisplayClusterLightCardEditor/Private/SDisplayClusterLightCardList.h
+++ b/Engine/Plugins/Runtime/nDisplay/Source/DisplayClusterLightCardEditor/Private/SDisplayClusterLightCardList.h
@@ -56,18 +56,12 @@ public:
 	/** Refreshes the list */
 	void Refresh();
 
-	/** Spawn and add a new Light Card */
-	ADisplayClusterLightCardActor* AddNewLightCard();
-
 private:
 	/**
 	 * Fill the LightCard list with available LightCards
 	 * @return True if the list has been modified
 	 */
 	bool FillLightCardList();
-	
-	/** Handles the Add Light Card button. Spawn a new Light Card placed in the middle of the viewport. */
-	void AddNewLightCardHandler();
 
 	TSharedRef<ITableRow> GenerateTreeItemRow(TSharedPtr<FLightCardTreeItem> Item, const TSharedRef<STableViewBase>& OwnerTable);
 	void GetChildrenForTreeItem(TSharedPtr<FLightCardTreeItem> InItem, TArray<TSharedPtr<FLightCardTreeItem>>& OutChildren);
diff --git a/Engine/Plugins/Runtime/nDisplay/Source/DisplayClusterLightCardEditor/Private/Settings/DisplayClusterLightCardEditorSettings.cpp b/Engine/Plugins/Runtime/nDisplay/Source/DisplayClusterLightCardEditor/Private/Settings/DisplayClusterLightCardEditorSettings.cpp
new file mode 100644
--- /dev/null
+++ b/Engine/Plugins/Runtime/nDisplay/Source/DisplayClusterLightCardEditor/Private/Settings/DisplayClusterLightCardEditorSettings.cpp
@@ -0,0 +1,8 @@
+// Copyright Epic Games, Inc. All Rights Reserved.
+
+#include "DisplayClusterLightCardEditorSettings.h"
+
+UDisplayClusterLightCardEditorProjectSettings::UDisplayClusterLightCardEditorProjectSettings()
+{
+	LightCardTemplateDefaultPath.Path = TEXT("/Game/VP/LightCards");
+}
diff --git a/Engine/Plugins/Runtime/nDisplay/Source/DisplayClusterLightCardEditor/Private/Settings/DisplayClusterLightCardEditorSettings.h b/Engine/Plugins/Runtime/nDisplay/Source/DisplayClusterLightCardEditor/Private/Settings/DisplayClusterLightCardEditorSettings.h
new file mode 100644
--- /dev/null
+++ b/Engine/Plugins/Runtime/nDisplay/Source/DisplayClusterLightCardEditor/Private/Settings/DisplayClusterLightCardEditorSettings.h
@@ -0,0 +1,20 @@
+// Copyright Epic Games, Inc. All Rights Reserved.
+
+#pragma once
+
+#include "CoreMinimal.h"
+#include "DisplayClusterLightCardEditorSettings.generated.h"
+
+
+UCLASS(config = Editor, defaultconfig)
+class UDisplayClusterLightCardEditorProjectSettings : public UObject
+{
+	GENERATED_BODY()
+
+public:
+	UDisplayClusterLightCardEditorProjectSettings();
+
+	/** The default path to save new light card templates */
+	UPROPERTY(config, EditAnywhere, Category = LightCardTemplates)
+	FDirectoryPath LightCardTemplateDefaultPath;
+};
diff --git a/Engine/Plugins/Runtime/nDisplay/Source/DisplayClusterLightCardEditor/Private/Viewport/DisplayClusterLightCardEditorViewportClient.cpp b/Engine/Plugins/Runtime/nDisplay/Source/DisplayClusterLightCardEditor/Private/Viewport/DisplayClusterLightCardEditorViewportClient.cpp
--- a/Engine/Plugins/Runtime/nDisplay/Source/DisplayClusterLightCardEditor/Private/Viewport/DisplayClusterLightCardEditorViewportClient.cpp
+++ b/Engine/Plugins/Runtime/nDisplay/Source/DisplayClusterLightCardEditor/Private/Viewport/DisplayClusterLightCardEditorViewportClient.cpp
@@ -4,6 +4,7 @@
 
 #include "DisplayClusterLightcardEditorViewport.h"
 #include "DisplayClusterLightCardEditorWidget.h"
+#include "LightCardTemplates/DisplayClusterLightCardTemplate.h"
 
 #include "Blueprints/DisplayClusterBlueprintLib.h"
 #include "Components/DisplayClusterPreviewComponent.h"
@@ -1324,6 +1325,119 @@
 	return MouseCursor;
 }
 
+void FDisplayClusterLightCardEditorViewportClient::DestroyDropPreviewActors()
+{
+	if (HasDropPreviewActors())
+	{
+		for (auto ActorIt = DropPreviewLightCards.CreateConstIterator(); ActorIt; ++ActorIt)
+		{
+			ADisplayClusterLightCardActor* PreviewActor = (*ActorIt).Get();
+			if (PreviewActor)
+			{
+				IDisplayClusterScenePreview::Get().RemoveActorFromRenderer(PreviewRendererId, PreviewActor);
+				GetWorld()->DestroyActor(PreviewActor);
+			}
+		}
+		DropPreviewLightCards.Empty();
+	}
+}
+
+bool FDisplayClusterLightCardEditorViewportClient::UpdateDropPreviewActors(int32 MouseX, int32 MouseY,
+                                                                           const TArray<UObject*>& DroppedObjects, bool& bOutDroppedObjectsVisible, UActorFactory* FactoryToUse)
+{
+	bOutDroppedObjectsVisible = false;
+	if(!HasDropPreviewActors())
+	{
+		return false;
+	}
+
+	bNeedsRedraw = true;
+
+	const FIntPoint NewMousePos { MouseX, MouseY };
+
+	for (TWeakObjectPtr<ADisplayClusterLightCardActor>& LightCard : DropPreviewLightCards)
+	{
+		VerifyAndFixLightCardOrigin(LightCard.Get());
+	}
+	
+	MoveLightCardsToPixel(NewMousePos, DropPreviewLightCards);
+	
+	return true;
+}
+
+bool FDisplayClusterLightCardEditorViewportClient::DropObjectsAtCoordinates(int32 MouseX, int32 MouseY,
+                                                                            const TArray<UObject*>& DroppedObjects, TArray<AActor*>& OutNewActors, bool bOnlyDropOnTarget,
+                                                                            bool bCreateDropPreview, bool bSelectActors, UActorFactory* FactoryToUse)
+{
+	if (!LightCardEditorPtr.IsValid() || LightCardEditorPtr.Pin()->GetActiveRootActor() == nullptr)
+	{
+		return false;
+	}
+	
+	UWorld* PreviewWorld = PreviewScene->GetWorld();
+	check(PreviewWorld);
+	
+	DestroyDropPreviewActors();
+	
+	Viewport->InvalidateHitProxy();
+
+	bool bSuccess = false;
+
+	if (bSelectActors)
+	{
+		SelectedLightCards.Empty();
+	}
+	
+	TArray<TWeakObjectPtr<ADisplayClusterLightCardActor>> CreatedLightCards;
+	CreatedLightCards.Reserve(DroppedObjects.Num());
+
+	SelectLightCard(nullptr);
+	PropagateLightCardSelection();
+	
+	for (UObject* DroppedObject : DroppedObjects)
+	{
+		if (UDisplayClusterLightCardTemplate* Template = Cast<UDisplayClusterLightCardTemplate>(DroppedObject))
+		{
+			ADisplayClusterLightCardActor* LightCardActor = LightCardEditorPtr.Pin()->SpawnLightCardFromTemplate(Template,
+				bCreateDropPreview ? PreviewWorld->GetCurrentLevel() : nullptr, bCreateDropPreview);
+			check(LightCardActor);
+
+			VerifyAndFixLightCardOrigin(LightCardActor);
+			
+			CreatedLightCards.Add(LightCardActor);
+			
+			if (bCreateDropPreview)
+			{
+				DropPreviewLightCards.Add(LightCardActor);
+				
+				LightCardActor->bIsProxy = true;
+				IDisplayClusterScenePreview::Get().AddActorToRenderer(PreviewRendererId, LightCardActor);
+			}
+			else
+			{
+				LightCardActor->UpdatePolygonTexture();
+				LightCardActor->UpdateLightCardMaterialInstance();
+				if (bSelectActors)
+				{
+					GetOnNextSceneRefresh().AddLambda([this, LightCardActor]()
+					{
+						// Select on next refresh so the persistent level and corresponding proxy has spawned.
+						SelectLightCard(LightCardActor, true);
+						PropagateLightCardSelection();
+					});
+				}
+			}
+			
+			bSuccess = true;
+		}
+	}
+
+	const FIntPoint NewMousePos { MouseX, MouseY };
+	MoveLightCardsToPixel(NewMousePos, CreatedLightCards);
+
+	return bSuccess;
+}
+
 void FDisplayClusterLightCardEditorViewportClient::UpdatePreviewActor(ADisplayClusterRootActor* RootActor, bool bForce,
                                                                       EDisplayClusterLightCardEditorProxyType ProxyType)
 {
@@ -1434,7 +1548,7 @@
 			{
 				TSet<ADisplayClusterLightCardActor*> LightCards;
 				UDisplayClusterBlueprintLib::FindLightCardsForRootActor(RootActor, LightCards);
-
+				
 				SelectLightCard(nullptr);
 				
 				for (ADisplayClusterLightCardActor* LightCard : LightCards)
@@ -1715,67 +1829,7 @@
 
 void FDisplayClusterLightCardEditorViewportClient::MoveSelectedLightCardsToPixel(const FIntPoint& PixelPos)
 {
-	// First, find the average position of all selected light cards
-	FSphericalCoordinates AverageCoords = FSphericalCoordinates();
-	int32 NumLightCards = 0;
-
-	for (const TWeakObjectPtr<ADisplayClusterLightCardActor>& LightCard : SelectedLightCards)
-	{
-		if (LightCard.IsValid())
-		{
-			const FSphericalCoordinates LightCardCoords = GetLightCardCoordinates(LightCard.Get());
-
-			AverageCoords = AverageCoords + LightCardCoords;
-			++NumLightCards;
-		}
-	}
-
-	AverageCoords.Radius /= NumLightCards;
-	AverageCoords.Azimuth /= NumLightCards;
-	AverageCoords.Inclination /= NumLightCards;
-	AverageCoords.Conform();
-
-	// Now, find the desired coordinates based on the mouse position
-	FSceneViewFamilyContext ViewFamily(FSceneViewFamily::ConstructionValues(
-		Viewport,
-		GetScene(),
-		EngineShowFlags)
-		.SetRealtimeUpdate(IsRealtime())
-	);
-
-	FSceneView* View = CalcSceneView(&ViewFamily);
-
-	FVector Origin;
-	FVector Direction;
-	PixelToWorld(*View, PixelPos, Origin, Direction);
-
-	if (RenderViewportType != LVT_Perspective)
-	{
-		// For orthogonal projections, PixelToWorld does not return the view origin or a direction from the view origin. Use TraceScreenRay
-		// to find a useful direction away from the view origin to use
-
-		const FVector ViewOrigin = View->ViewLocation;
-
-		Direction = TraceScreenRay(Origin, Direction, ViewOrigin);
-		Origin = ViewOrigin;
-	}
-
-	// Compute desired coordinates (radius doesn't matter here since we will use the flush constraint on the light cards after moving them)
-	const FSphericalCoordinates DesiredCoords(Direction * 100.0f);
-	const FSphericalCoordinates DeltaCoords = DesiredCoords - AverageCoords;
-
-	// Update each light card with the delta coordinates; the flush constraint is applied by MoveLightCardTo, ensuring the light card is always flush to screens
-	for (const TWeakObjectPtr<ADisplayClusterLightCardActor>& LightCard : SelectedLightCards)
-	{
-		if (LightCard.IsValid())
-		{
-			const FSphericalCoordinates LightCardCoords = GetLightCardCoordinates(LightCard.Get());
-			const FSphericalCoordinates NewCoords = LightCardCoords + DeltaCoords;
-
-			MoveLightCardTo(*LightCard.Get(), NewCoords);
-			PropagateLightCardTransform(LightCard.Get());
-		}
-	}
+	MoveLightCardsToPixel(PixelPos, SelectedLightCards);
 }
 
 void FDisplayClusterLightCardEditorViewportClient::BeginTransaction(const FText& Description)
@@ -2328,6 +2382,71 @@
 	}
 }
 
+void FDisplayClusterLightCardEditorViewportClient::MoveLightCardsToPixel(const FIntPoint& PixelPos, const TArray<TWeakObjectPtr<ADisplayClusterLightCardActor>>& InLightCards)
+{
+	// First, find the average position of all selected light cards
+	FSphericalCoordinates AverageCoords = FSphericalCoordinates();
+	int32 NumLightCards = 0;
+
+	for (const TWeakObjectPtr<ADisplayClusterLightCardActor>& LightCard : InLightCards)
+	{
+		if (LightCard.IsValid())
+		{
+			const FSphericalCoordinates LightCardCoords = GetLightCardCoordinates(LightCard.Get());
+
+			AverageCoords = AverageCoords + LightCardCoords;
+			++NumLightCards;
+		}
+	}
+
+	AverageCoords.Radius /= NumLightCards;
+	AverageCoords.Azimuth /= NumLightCards;
+	AverageCoords.Inclination /= NumLightCards;
+	AverageCoords.Conform();
+
+	// Now, find the desired coordinates based on the mouse position
+	FSceneViewFamilyContext ViewFamily(FSceneViewFamily::ConstructionValues(
+		Viewport,
+		GetScene(),
+		EngineShowFlags)
+		.SetRealtimeUpdate(IsRealtime())
+	);
+
+	FSceneView* View = CalcSceneView(&ViewFamily);
+
+	FVector Origin;
+	FVector Direction;
+	PixelToWorld(*View, PixelPos, Origin, Direction);
+
+	if (RenderViewportType != LVT_Perspective)
+	{
+		// For orthogonal projections, PixelToWorld does not return the view origin or a direction from the view origin. Use TraceScreenRay
+		// to find a useful direction away from the view origin to use
+
+		const FVector ViewOrigin = View->ViewLocation;
+
+		Direction = TraceScreenRay(Origin, Direction, ViewOrigin);
+		Origin = ViewOrigin;
+	}
+
+	// Compute desired coordinates (radius doesn't matter here since we will use the flush constraint on the light cards after moving them)
+	const FSphericalCoordinates DesiredCoords(Direction * 100.0f);
+	const FSphericalCoordinates DeltaCoords = DesiredCoords - AverageCoords;
+
+	// Update each light card with the delta coordinates; the flush constraint is applied by MoveLightCardTo, ensuring the light card is always flush to screens
+	for (const TWeakObjectPtr<ADisplayClusterLightCardActor>& LightCard : InLightCards)
+	{
+		if (LightCard.IsValid())
+		{
+			const FSphericalCoordinates LightCardCoords = GetLightCardCoordinates(LightCard.Get());
+			const FSphericalCoordinates NewCoords = LightCardCoords + DeltaCoords;
+
+			MoveLightCardTo(*LightCard.Get(), NewCoords);
+			PropagateLightCardTransform(LightCard.Get());
+		}
+	}
+}
+
 FRotator FDisplayClusterLightCardEditorViewportClient::GetLightCardRotationDelta(
 	FViewport* InViewport,
 	ADisplayClusterLightCardActor* LightCard,
diff --git a/Engine/Plugins/Runtime/nDisplay/Source/DisplayClusterLightCardEditor/Private/Viewport/DisplayClusterLightCardEditorViewportClient.h b/Engine/Plugins/Runtime/nDisplay/Source/DisplayClusterLightCardEditor/Private/Viewport/DisplayClusterLightCardEditorViewportClient.h
--- a/Engine/Plugins/Runtime/nDisplay/Source/DisplayClusterLightCardEditor/Private/Viewport/DisplayClusterLightCardEditorViewportClient.h
+++ b/Engine/Plugins/Runtime/nDisplay/Source/DisplayClusterLightCardEditor/Private/Viewport/DisplayClusterLightCardEditorViewportClient.h
@@ -146,6 +146,10 @@ public:
 	virtual void ProcessClick(FSceneView& View, HHitProxy* HitProxy, FKey Key, EInputEvent Event, uint32 HitX, uint32 HitY) override;
 	virtual EMouseCursor::Type GetCursor(FViewport* Viewport,int32 X,int32 Y) override;
 	virtual ELevelViewportType GetViewportType() const override { return LVT_Perspective; }
+	virtual bool HasDropPreviewActors() const override { return DropPreviewLightCards.Num() > 0; }
+	virtual void DestroyDropPreviewActors() override;
+	virtual bool UpdateDropPreviewActors(int32 MouseX, int32 MouseY, const TArray<UObject*>& DroppedObjects, bool& bOutDroppedObjectsVisible, UActorFactory* FactoryToUse) override;
+	virtual bool DropObjectsAtCoordinates(int32 MouseX, int32 MouseY, const TArray<UObject*>& DroppedObjects, TArray<AActor*>& OutNewActors, bool bOnlyDropOnTarget, bool bCreateDropPreview, bool bSelectActors, UActorFactory* FactoryToUse) override;
 	// ~FEditorViewportClient
 
 	/** Returns a delegate that is invoked on the next scene refresh. The delegate is cleared afterwards */
@@ -251,6 +255,9 @@ private:
 	/** Moves the currently selected light cards */
 	void MoveSelectedLightCards(FViewport* InViewport, EAxisList::Type CurrentAxis);
 
+	/** Moves all given light cards to the specified pixel position */
+	void MoveLightCardsToPixel(const FIntPoint& PixelPos, const TArray<TWeakObjectPtr<ADisplayClusterLightCardActor>>& InLightCards);
+	
 	/** Ensures that the light card root component is at the same location as the projection/view origin */
 	void VerifyAndFixLightCardOrigin(ADisplayClusterLightCardActor* LightCard) const;
 
@@ -348,6 +355,9 @@ private:
 
 	TArray<FLightCardProxy> LightCardProxies;
 	TArray<TWeakObjectPtr<ADisplayClusterLightCardActor>> SelectedLightCards;
+
+	/** Light card preview actors being dropped on the scene */
+	TArray<TWeakObjectPtr<ADisplayClusterLightCardActor>> DropPreviewLightCards;
 	
 	/** The index of the scene preview renderer returned from IDisplayClusterScenePreview */
 	int32 PreviewRendererId = -1;
diff --git a/Engine/Plugins/Runtime/nDisplay/Source/DisplayClusterLightCardEditor/Private/Viewport/DisplayClusterLightcardEditorViewport.cpp b/Engine/Plugins/Runtime/nDisplay/Source/DisplayClusterLightCardEditor/Private/Viewport/DisplayClusterLightcardEditorViewport.cpp
--- a/Engine/Plugins/Runtime/nDisplay/Source/DisplayClusterLightCardEditor/Private/Viewport/DisplayClusterLightcardEditorViewport.cpp
+++ b/Engine/Plugins/Runtime/nDisplay/Source/DisplayClusterLightCardEditor/Private/Viewport/DisplayClusterLightcardEditorViewport.cpp
@@ -6,6 +6,8 @@
 #include "DisplayClusterLightCardEditorViewportClient.h"
 #include "DisplayClusterLightCardEditorCommands.h"
 #include "SDisplayClusterLightCardEditor.h"
+#include "LightCardTemplates/DisplayClusterLightCardTemplateDragDropOp.h"
+#include "LightCardTemplates/DisplayClusterLightCardTemplate.h"
 
 #include "EditorViewportCommands.h"
 #include "Framework/Commands/GenericCommands.h"
@@ -13,7 +15,7 @@
 #include "SEditorViewportToolBarMenu.h"
 #include "STransformViewportToolbar.h"
 #include "Slate/SceneViewport.h"
-
+#include "ScopedTransaction.h"
 #include "Kismet2/DebuggerCommands.h"
 #include "Styling/AppStyle.h"
 
@@ -445,13 +447,111 @@ void SDisplayClusterLightCardEditorViewport::BindCommands()
 	}
 }
 
+FReply SDisplayClusterLightCardEditorViewport::OnDragOver(const FGeometry& MyGeometry,
+	const FDragDropEvent& DragDropEvent)
+{
+	const TSharedPtr<FDisplayClusterLightCardTemplateDragDropOp> TemplateDragDropOp = DragDropEvent.GetOperationAs<FDisplayClusterLightCardTemplateDragDropOp>();
+	if (TemplateDragDropOp.IsValid() && LightCardEditorPtr.IsValid() &&
+		LightCardEditorPtr.Pin()->GetActiveRootActor().IsValid() && TemplateDragDropOp->GetTemplate().IsValid())
+	{
+		TemplateDragDropOp->SetDropAsValid(FText::Format(LOCTEXT("TemplateDragDropOp_LightCardTemplate", "Spawn light card from template {0}"),
+			FText::FromString(TemplateDragDropOp->GetTemplate()->GetName())));
+
+		FIntPoint ViewportOrigin, ViewportSize;
+		ViewportClient->GetViewportDimensions(ViewportOrigin, ViewportSize);
+		const FVector2D MousePos = MyGeometry.AbsoluteToLocal(DragDropEvent.GetScreenSpacePosition()) * MyGeometry.Scale - ViewportOrigin;
+
+		const TArray<UObject*> DroppedObjects;
+		bool bDroppedObjectsVisible = true; 
+		ViewportClient->UpdateDropPreviewActors(MousePos.X, MousePos.Y, DroppedObjects, bDroppedObjectsVisible, nullptr);
+			
+		return FReply::Handled();
+	}
+
+	return SEditorViewport::OnDragOver(MyGeometry, DragDropEvent);
+}
+
+void SDisplayClusterLightCardEditorViewport::OnDragEnter(const FGeometry& MyGeometry,
+	const FDragDropEvent& DragDropEvent)
+{
+	const TSharedPtr<FDisplayClusterLightCardTemplateDragDropOp> TemplateDragDropOp = DragDropEvent.GetOperationAs<FDisplayClusterLightCardTemplateDragDropOp>();
+	if (TemplateDragDropOp.IsValid() && LightCardEditorPtr.IsValid())
+	{
+		FIntPoint ViewportOrigin, ViewportSize;
+		ViewportClient->GetViewportDimensions(ViewportOrigin, ViewportSize);
+
+		const FVector2D MousePos = MyGeometry.AbsoluteToLocal(DragDropEvent.GetScreenSpacePosition()) * MyGeometry.Scale - ViewportOrigin;
+
+		const TWeakObjectPtr<UDisplayClusterLightCardTemplate> LightCardTemplate = TemplateDragDropOp->GetTemplate();
+
+		if (LightCardTemplate.IsValid())
+		{
+			const TArray<UObject*> DroppedObjects{ LightCardTemplate.Get() };
+			
+			TArray<AActor*> TemporaryActors;
+			
+			const bool bIsPreview = true;
+			ViewportClient->DropObjectsAtCoordinates(MousePos.X, MousePos.Y, DroppedObjects, TemporaryActors, false, bIsPreview, false, nullptr);
+
+			return;
+		}
+	}
+	
+	SEditorViewport::OnDragEnter(MyGeometry, DragDropEvent);
+}
+
+void SDisplayClusterLightCardEditorViewport::OnDragLeave(const FDragDropEvent& DragDropEvent)
+{
+	ViewportClient->DestroyDropPreviewActors();
+	
+	const TSharedPtr<FDisplayClusterLightCardTemplateDragDropOp> TemplateDragDropOp = DragDropEvent.GetOperationAs<FDisplayClusterLightCardTemplateDragDropOp>();
+	if (TemplateDragDropOp.IsValid())
+	{
+		TemplateDragDropOp->SetDropAsInvalid();
+		return;
+	}
+	
+	SEditorViewport::OnDragLeave(DragDropEvent);
+}
+
+FReply SDisplayClusterLightCardEditorViewport::OnDrop(const FGeometry& MyGeometry, const FDragDropEvent& DragDropEvent)
+{
+	const TSharedPtr<FDisplayClusterLightCardTemplateDragDropOp> TemplateDragDropOp = DragDropEvent.GetOperationAs<FDisplayClusterLightCardTemplateDragDropOp>();
+	if (TemplateDragDropOp.IsValid() && TemplateDragDropOp->CanBeDropped() && LightCardEditorPtr.IsValid())
+	{
+		FIntPoint ViewportOrigin, ViewportSize;
+		ViewportClient->GetViewportDimensions(ViewportOrigin, ViewportSize);
+
+		const FVector2D MousePos = MyGeometry.AbsoluteToLocal(DragDropEvent.GetScreenSpacePosition()) * MyGeometry.Scale - ViewportOrigin;
+
+		const TWeakObjectPtr<UDisplayClusterLightCardTemplate> LightCardTemplate = TemplateDragDropOp->GetTemplate();
+
+		if (LightCardTemplate.IsValid())
+		{
+			const TArray<UObject*> DroppedObjects{ LightCardTemplate.Get() };
+			
+			TArray<AActor*> TemporaryActors;
+			// Only select actor on drop
+			const bool bSelectActor = true;
+
+			const FScopedTransaction Transaction(LOCTEXT("CreateLightCardFromTemplate", "Create Light Card from Template"));
+			ViewportClient->DropObjectsAtCoordinates(MousePos.X, MousePos.Y, DroppedObjects, TemporaryActors, false, false, bSelectActor, nullptr);
+
+			return FReply::Handled();
+		}
+	}
+	
+	return SEditorViewport::OnDrop(MyGeometry, DragDropEvent);
+}
+
 TSharedRef<SWidget> SDisplayClusterLightCardEditorViewport::MakeContextMenu()
 {
 	const bool bInShouldCloseWindowAfterMenuSelection = true;
 	FMenuBuilder MenuBuilder(bInShouldCloseWindowAfterMenuSelection, CommandList);
 
 	MenuBuilder.AddMenuEntry(FDisplayClusterLightCardEditorCommands::Get().RemoveLightCard);
-
+	MenuBuilder.AddMenuEntry(FDisplayClusterLightCardEditorCommands::Get().SaveLightCardTemplate);
+	
 	MenuBuilder.BeginSection("Edit", LOCTEXT("EditSection", "Edit"));
 	{
 		MenuBuilder.AddMenuEntry(FGenericCommands::Get().Cut);
diff --git a/Engine/Plugins/Runtime/nDisplay/Source/DisplayClusterLightCardEditor/Private/Viewport/DisplayClusterLightcardEditorViewport.h b/Engine/Plugins/Runtime/nDisplay/Source/DisplayClusterLightCardEditor/Private/Viewport/DisplayClusterLightcardEditorViewport.h
--- a/Engine/Plugins/Runtime/nDisplay/Source/DisplayClusterLightCardEditor/Private/Viewport/DisplayClusterLightcardEditorViewport.h
+++ b/Engine/Plugins/Runtime/nDisplay/Source/DisplayClusterLightCardEditor/Private/Viewport/DisplayClusterLightcardEditorViewport.h
@@ -59,6 +59,10 @@ private:
 	virtual TSharedPtr<SWidget> MakeViewportToolbar() override;
 	virtual void PopulateViewportOverlays(TSharedRef<SOverlay> Overlay) override;
 	virtual void BindCommands() override;
+	virtual FReply OnDragOver(const FGeometry& MyGeometry, const FDragDropEvent& DragDropEvent) override;
+	virtual void OnDragEnter(const FGeometry& MyGeometry, const FDragDropEvent& DragDropEvent) override;
+	virtual void OnDragLeave(const FDragDropEvent& DragDropEvent) override;
+	virtual FReply OnDrop(const FGeometry& MyGeometry, const FDragDropEvent& DragDropEvent) override;
 	// ~SEditorViewport
 
 	/** Create the popup context menu */
diff --git a/Engine/Plugins/Runtime/nDisplay/Source/DisplayClusterOperator/DisplayClusterOperator.Build.cs b/Engine/Plugins/Runtime/nDisplay/Source/DisplayClusterOperator/DisplayClusterOperator.Build.cs
--- a/Engine/Plugins/Runtime/nDisplay/Source/DisplayClusterOperator/DisplayClusterOperator.Build.cs
+++ b/Engine/Plugins/Runtime/nDisplay/Source/DisplayClusterOperator/DisplayClusterOperator.Build.cs
@@ -20,6 +20,7 @@ public class DisplayClusterOperator : ModuleRules
 				"UnrealEd",
 				"EditorStyle",
 				"WorkspaceMenuStructure",
+				"Kismet"
 			});
 	}
 }
diff --git a/Engine/Plugins/Runtime/nDisplay/Source/DisplayClusterOperator/Private/SDisplayClusterOperatorPanel.cpp b/Engine/Plugins/Runtime/nDisplay/Source/DisplayClusterOperator/Private/SDisplayClusterOperatorPanel.cpp
--- a/Engine/Plugins/Runtime/nDisplay/Source/DisplayClusterOperator/Private/SDisplayClusterOperatorPanel.cpp
+++ b/Engine/Plugins/Runtime/nDisplay/Source/DisplayClusterOperator/Private/SDisplayClusterOperatorPanel.cpp
@@ -9,11 +9,13 @@
 #include "Framework/Docking/TabManager.h"
 #include "Framework/Docking/LayoutExtender.h"
 #include "Framework/MultiBox/MultiBoxBuilder.h"
-#include "PropertyEditorModule.h"
+#include "SKismetInspector.h"
 #include "Widgets/Docking/SDockTab.h"
 #include "WorkspaceMenuStructure.h"
 #include "WorkspaceMenuStructureModule.h"
 
+#include "Components/ActorComponent.h"
+
 #define LOCTEXT_NAMESPACE "SDisplayClusterOperatorPanel"
 
 const FName SDisplayClusterOperatorPanel::TabName = TEXT("DisplayClusterOperatorTab");
@@ -131,14 +133,8 @@ TSharedRef<SDockTab> SDisplayClusterOperatorPanel::SpawnToolbarTab(const FSpawnT
 
 TSharedRef<SDockTab> SDisplayClusterOperatorPanel::SpawnDetailsTab(const FSpawnTabArgs& Args)
 {
-	FPropertyEditorModule& EditModule = FModuleManager::Get().GetModuleChecked<FPropertyEditorModule>("PropertyEditor");
-
-	FDetailsViewArgs DetailsViewArgs;
-	DetailsViewArgs.bAllowSearch = false;
-	DetailsViewArgs.NameAreaSettings = FDetailsViewArgs::HideNameArea;
-	DetailsViewArgs.bHideSelectionTip = true;
-
-	DetailsView = EditModule.CreateDetailView(DetailsViewArgs);
+	SAssignNew(DetailsView, SKismetInspector)
+	.HideNameArea(true);
 
 	return SNew(SDockTab)
 		.TabRole(ETabRole::PanelTab)
@@ -151,7 +147,9 @@ void SDisplayClusterOperatorPanel::DisplayObjectsInDetailsPanel(const TArray<UOb
 {
 	if (DetailsView.IsValid())
 	{
-		DetailsView->SetObjects(Objects);
+		SKismetInspector::FShowDetailsOptions Options;
+		Options.bShowComponents = false;
+		DetailsView->ShowDetailsForObjects(Objects, MoveTemp(Options));
 	}
 }
 
diff --git a/Engine/Plugins/Runtime/nDisplay/Source/DisplayClusterOperator/Private/SDisplayClusterOperatorPanel.h b/Engine/Plugins/Runtime/nDisplay/Source/DisplayClusterOperator/Private/SDisplayClusterOperatorPanel.h
--- a/Engine/Plugins/Runtime/nDisplay/Source/DisplayClusterOperator/Private/SDisplayClusterOperatorPanel.h
+++ b/Engine/Plugins/Runtime/nDisplay/Source/DisplayClusterOperator/Private/SDisplayClusterOperatorPanel.h
@@ -68,7 +68,7 @@ private:
 	TSharedPtr<SDisplayClusterOperatorToolbar> Toolbar;
 
 	/** A reference to the operator panel's details view */
-	TSharedPtr<class IDetailsView> DetailsView;
+	TSharedPtr<class SKismetInspector> DetailsView;
 
 	/** The delegate handle for the operator module's OnDetailObjectsChanged event */
 	FDelegateHandle DetailObjectsChangedHandle;
