nDisplay: Fixed ChromakeyTexture and OverrideTexture types to support RTT and multimedia as input.

diff --git a/Engine/Plugins/Runtime/nDisplay/Source/DisplayClusterConfiguration/Public/DisplayClusterConfigurationTypes_ICVFX.h b/Engine/Plugins/Runtime/nDisplay/Source/DisplayClusterConfiguration/Public/DisplayClusterConfigurationTypes_ICVFX.h
--- a/Engine/Plugins/Runtime/nDisplay/Source/DisplayClusterConfiguration/Public/DisplayClusterConfigurationTypes_ICVFX.h
+++ b/Engine/Plugins/Runtime/nDisplay/Source/DisplayClusterConfiguration/Public/DisplayClusterConfigurationTypes_ICVFX.h
@@ -160,7 +160,7 @@ public:
 
 	/** Texture to use as the chromakey marker tile. */
 	UPROPERTY(EditAnywhere, BlueprintReadWrite, Category = NDisplay, meta = (EditCondition = "bEnable"))
-	UTexture2D* MarkerTileRGBA = nullptr;
+	UTexture* MarkerTileRGBA = nullptr;
 
 	/** Scale value for the size of each chromakey marker tile. */
 	UPROPERTY(EditAnywhere, BlueprintReadWrite, Category = NDisplay, meta = (EditCondition = "bEnable", ClampMin = "0", UIMin = "0", DisplayName = "Marker Scale"))
diff --git a/Engine/Plugins/Runtime/nDisplay/Source/DisplayClusterConfiguration/Public/DisplayClusterConfigurationTypes_PostRender.h b/Engine/Plugins/Runtime/nDisplay/Source/DisplayClusterConfiguration/Public/DisplayClusterConfigurationTypes_PostRender.h
--- a/Engine/Plugins/Runtime/nDisplay/Source/DisplayClusterConfiguration/Public/DisplayClusterConfigurationTypes_PostRender.h
+++ b/Engine/Plugins/Runtime/nDisplay/Source/DisplayClusterConfiguration/Public/DisplayClusterConfigurationTypes_PostRender.h
@@ -27,7 +27,7 @@ public:
 
 	/** Texture to use in place of the inner frustum. */
 	UPROPERTY(EditAnywhere, BlueprintReadWrite, Category = "NDisplay Render", meta = (EditCondition = "bAllowReplace"))
-	UTexture2D* SourceTexture = nullptr;
+	UTexture* SourceTexture = nullptr;
 
 	/** Set to True to crop the texture for the inner frustum as specified below. */
 	UPROPERTY(EditAnywhere, BlueprintReadWrite, Category = "NDisplay Render", meta = (DisplayName = "Use Texture Crop", EditCondition = "bAllowReplace"))
