nDisplay: media integration

diff --git a/Engine/Plugins/Runtime/nDisplay/Source/DisplayCluster/DisplayCluster.Build.cs b/Engine/Plugins/Runtime/nDisplay/Source/DisplayCluster/DisplayCluster.Build.cs
--- a/Engine/Plugins/Runtime/nDisplay/Source/DisplayCluster/DisplayCluster.Build.cs
+++ b/Engine/Plugins/Runtime/nDisplay/Source/DisplayCluster/DisplayCluster.Build.cs
@@ -9,10 +9,10 @@
 	{
 		PublicDefinitions.Add("WITH_OCIO=0");
 
-		PrivateIncludePaths.AddRange(
-			new string[] {
-				"../../../../Source/Runtime/Renderer/Private",
-			});
+		// [temporary] We need this to be able to use some private data types. This should
+		// be removed once we move the nD rendering pipeline to RDG.
+		string EngineDir = Path.GetFullPath(Target.RelativeEnginePath);
+		PrivateIncludePaths.Add(Path.Combine(EngineDir, "Source", "Runtime", "Renderer", "Private"));
 
 		PublicDependencyModuleNames.AddRange(
 			new string[] {
@@ -24,23 +24,23 @@
 			});
 
 		PrivateDependencyModuleNames.AddRange(
-		new string[] {
-			"HeadMountedDisplay",
-			"InputCore",
-			"Json",
-			"JsonUtilities",
-			"Networking",
-			"OpenColorIO",
-			"OpenCV",
-			"OpenCVHelper",
-			"ProceduralMeshComponent",
-			"Renderer",
-			"RenderCore",
-			"RHI",
-			"Slate",
-			"SlateCore",
-			"Sockets",
-		});
+			new string[] {
+				"HeadMountedDisplay",
+				"InputCore",
+				"Json",
+				"JsonUtilities",
+				"Networking",
+				"OpenColorIO",
+				"OpenCV",
+				"OpenCVHelper",
+				"ProceduralMeshComponent",
+				"Renderer",
+				"RenderCore",
+				"RHI",
+				"Slate",
+				"SlateCore",
+				"Sockets",
+			});
 
 		if (Target.bBuildEditor == true)
 		{
diff --git a/Engine/Plugins/Runtime/nDisplay/Source/DisplayCluster/Private/DisplayClusterCallbacks.h b/Engine/Plugins/Runtime/nDisplay/Source/DisplayCluster/Private/DisplayClusterCallbacks.h
--- a/Engine/Plugins/Runtime/nDisplay/Source/DisplayCluster/Private/DisplayClusterCallbacks.h
+++ b/Engine/Plugins/Runtime/nDisplay/Source/DisplayCluster/Private/DisplayClusterCallbacks.h
@@ -77,6 +77,41 @@ public:
 		return DisplayClusterFailoverNodeDown;
 	}
 
+	virtual FDisplayClusterPostRenderViewFamily_RenderThread& OnDisplayClusterPostRenderViewFamily_RenderThread() override
+	{
+		return DisplayClusterPostRenderViewFamily;
+	}
+
+	virtual FDisplayClusterPreWarp_RenderThread& OnDisplayClusterPreWarp_RenderThread() override
+	{
+		return DisplayClusterPreWarpEvent;
+	}
+
+	virtual FDisplayClusterPreWarpViewport_RenderThread& OnDisplayClusterPreWarpViewport_RenderThread() override
+	{
+		return DisplayClusterPreWarpViewportEvent;
+	}
+
+	virtual FDisplayClusterPostWarp_RenderThread& OnDisplayClusterPostWarp_RenderThread() override
+	{
+		return DisplayClusterPostWarpEvent;
+	}
+
+	virtual FDisplayClusterPostWarpViewport_RenderThread& OnDisplayClusterPostWarpViewport_RenderThread() override
+	{
+		return DisplayClusterPostWarpViewportEvent;
+	}
+
+	virtual FDisplayClusterPreFrameRender_RenderThread& OnDisplayClusterPreFrameRender_RenderThread() override
+	{
+		return DisplayClusterPreFrameRenderEvent;
+	}
+
+	virtual FDisplayClusterPostFrameRender_RenderThread& OnDisplayClusterPostFrameRender_RenderThread() override
+	{
+		return DisplayClusterPostFrameRenderEvent;
+	}
+
 private:
 	FDisplayClusterStartSessionEvent         DisplayClusterStartSessionEvent;
 	FDisplayClusterEndSessionEvent           DisplayClusterEndSessionEvent;
@@ -88,7 +123,16 @@ private:
 	FDisplayClusterStartSceneEvent           DisplayClusterStartSceneEvent;
 	FDisplayClusterEndSceneEvent             DisplayClusterEndSceneEvent;
 	FDisplayClusterCustomPresentSetEvent     DisplayClusterCustomPresentSetEvent;
-	FDisplayClusterPresentationPreSynchronization_RHIThread DisplayClusterPresentationPreSynchronizationEvent;
-	FDisplayClusterPresentationPostSynchronization_RHIThread DisplayClusterPresentationPostSynchronizationEvent;
 	FDisplayClusterFailoverNodeDown          DisplayClusterFailoverNodeDown;
+
+	FDisplayClusterPresentationPreSynchronization_RHIThread  DisplayClusterPresentationPreSynchronizationEvent;
+	FDisplayClusterPresentationPostSynchronization_RHIThread DisplayClusterPresentationPostSynchronizationEvent;
+
+	FDisplayClusterPostRenderViewFamily_RenderThread DisplayClusterPostRenderViewFamily;
+	FDisplayClusterPreWarp_RenderThread          DisplayClusterPreWarpEvent;
+	FDisplayClusterPreWarpViewport_RenderThread  DisplayClusterPreWarpViewportEvent;
+	FDisplayClusterPostWarp_RenderThread         DisplayClusterPostWarpEvent;
+	FDisplayClusterPostWarpViewport_RenderThread DisplayClusterPostWarpViewportEvent;
+	FDisplayClusterPreFrameRender_RenderThread   DisplayClusterPreFrameRenderEvent;
+	FDisplayClusterPostFrameRender_RenderThread  DisplayClusterPostFrameRenderEvent;
 };
diff --git a/Engine/Plugins/Runtime/nDisplay/Source/DisplayCluster/Private/Game/EngineClasses/Basics/DisplayClusterViewportClient.cpp b/Engine/Plugins/Runtime/nDisplay/Source/DisplayCluster/Private/Game/EngineClasses/Basics/DisplayClusterViewportClient.cpp
--- a/Engine/Plugins/Runtime/nDisplay/Source/DisplayCluster/Private/Game/EngineClasses/Basics/DisplayClusterViewportClient.cpp
+++ b/Engine/Plugins/Runtime/nDisplay/Source/DisplayCluster/Private/Game/EngineClasses/Basics/DisplayClusterViewportClient.cpp
@@ -530,7 +530,7 @@
 				FRotator	ViewRotation;
 				FSceneView* View = LocalPlayer->CalcSceneView(&ViewFamily, ViewLocation, ViewRotation, InViewport, nullptr, ViewportContext.StereoViewIndex);
 
-				if (View && DCView.IsShouldRenderView() == false)
+				if (View && !DCView.ShouldRenderSceneView())
 				{
 					ViewFamily.Views.Remove(View);
 
diff --git a/Engine/Plugins/Runtime/nDisplay/Source/DisplayCluster/Private/Render/Viewport/Configuration/DisplayClusterViewportConfigurationHelpers.cpp b/Engine/Plugins/Runtime/nDisplay/Source/DisplayCluster/Private/Render/Viewport/Configuration/DisplayClusterViewportConfigurationHelpers.cpp
--- a/Engine/Plugins/Runtime/nDisplay/Source/DisplayCluster/Private/Render/Viewport/Configuration/DisplayClusterViewportConfigurationHelpers.cpp
+++ b/Engine/Plugins/Runtime/nDisplay/Source/DisplayCluster/Private/Render/Viewport/Configuration/DisplayClusterViewportConfigurationHelpers.cpp
@@ -14,6 +14,9 @@
 #include "DisplayClusterConfigurationTypes_ICVFX.h"
 #include "DisplayClusterConfigurationTypes_PostRender.h"
 
+#include "IDisplayCluster.h"
+#include "Cluster/IDisplayClusterClusterManager.h"
+
 #include "ShaderParameters/DisplayClusterShaderParameters_PostprocessBlur.h"
 #include "ShaderParameters/DisplayClusterShaderParameters_GenerateMips.h"
 #include "ShaderParameters/DisplayClusterShaderParameters_Override.h"
@@ -202,6 +205,28 @@ void FDisplayClusterViewportConfigurationHelpers::UpdateBaseViewportSetting(FDis
 		DstViewport.RenderSettings.RenderFamilyGroup = InRenderSettings.RenderFamilyGroup;
 	}
 
+	// Set media related configuration (runtime only for now)
+	if (IDisplayCluster::Get().GetOperationMode() == EDisplayClusterOperationMode::Cluster)
+	{
+		const FDisplayClusterConfigurationMedia& MediaSettings = InConfigurationViewport.RenderSettings.Media;
+
+		if (MediaSettings.bEnabled)
+		{
+			const FString ThisClusterNodeId = IDisplayCluster::Get().GetClusterMgr()->GetNodeId();
+			const bool bThisNodeSharesMedia = MediaSettings.MediaOutputNode.Equals(ThisClusterNodeId, ESearchCase::IgnoreCase);
+
+			// In most cases there is no need to render viewport if media input set up. The only exception
+			// is when this viewport is used as media source for other cluster nodes. In this case media input
+			// settings are ignored, and media capture is used.
+			DstViewport.RenderSettings.bSkipSceneRenderingButLeaveResourcesAvailable = MediaSettings.MediaInput.bEnabled ?
+					!(MediaSettings.MediaOutput.bEnabled && bThisNodeSharesMedia) :
+					false;
+
+			// Set media capture flag if media capture is used
+			DstViewport.RenderSettings.bIsBeingCaptured = InConfigurationViewport.RenderSettings.Media.MediaOutput.bEnabled;
+		}
+	}
+
 	// FDisplayClusterConfigurationViewport_ICVFX property:
 	{
 		EDisplayClusterViewportICVFXFlags& TargetFlags = DstViewport.RenderSettingsICVFX.Flags;
diff --git a/Engine/Plugins/Runtime/nDisplay/Source/DisplayCluster/Private/Render/Viewport/Configuration/DisplayClusterViewportConfigurationHelpers_ICVFX.cpp b/Engine/Plugins/Runtime/nDisplay/Source/DisplayCluster/Private/Render/Viewport/Configuration/DisplayClusterViewportConfigurationHelpers_ICVFX.cpp
--- a/Engine/Plugins/Runtime/nDisplay/Source/DisplayCluster/Private/Render/Viewport/Configuration/DisplayClusterViewportConfigurationHelpers_ICVFX.cpp
+++ b/Engine/Plugins/Runtime/nDisplay/Source/DisplayCluster/Private/Render/Viewport/Configuration/DisplayClusterViewportConfigurationHelpers_ICVFX.cpp
@@ -18,6 +18,9 @@
 #include "DisplayClusterConfigurationTypes_ICVFX.h"
 #include "DisplayClusterConfigurationTypes_PostRender.h"
 
+#include "IDisplayCluster.h"
+#include "Cluster/IDisplayClusterClusterManager.h"
+
 #include "ShaderParameters/DisplayClusterShaderParameters_PostprocessBlur.h"
 #include "ShaderParameters/DisplayClusterShaderParameters_GenerateMips.h"
 #include "ShaderParameters/DisplayClusterShaderParameters_Override.h"
@@ -560,6 +563,27 @@ void FDisplayClusterViewportConfigurationHelpers_ICVFX::UpdateCameraViewportSett
 	// Support inner camera custom frustum
 	FDisplayClusterViewportConfigurationHelpers_ICVFX::UpdateCameraCustomFrustum(DstViewport, CameraSettings.CustomFrustum);
 	UpdateCameraViewportBufferRatio(DstViewport, CameraSettings);
+
+	// Set media related configuration (runtime only for now)
+	if (IDisplayCluster::Get().GetOperationMode() == EDisplayClusterOperationMode::Cluster)
+	{
+		const FDisplayClusterConfigurationMedia& MediaSettings = InCameraComponent.CameraSettings.RenderSettings.Media;
+		if (MediaSettings.bEnabled)
+		{
+			const FString ThisClusterNodeId = IDisplayCluster::Get().GetClusterMgr()->GetNodeId();
+			const bool bThisNodeSharesMedia = MediaSettings.MediaOutputNode.Equals(ThisClusterNodeId, ESearchCase::IgnoreCase);
+
+			// In most cases there is no need to render ICVFX view if media input set up. The only exception
+			// is when this camera is used as media source for other cluster nodes. In this case media input
+			// settings are ignored, and media capture is used.
+			DstViewport.RenderSettings.bSkipSceneRenderingButLeaveResourcesAvailable = MediaSettings.MediaInput.bEnabled ?
+					!(MediaSettings.MediaOutput.bEnabled && bThisNodeSharesMedia) :
+					false;
+
+			// Set media capture flag if media capture is used
+			DstViewport.RenderSettings.bIsBeingCaptured = InCameraComponent.CameraSettings.RenderSettings.Media.MediaOutput.bEnabled;
+		}
+	}
 }
 
 void FDisplayClusterViewportConfigurationHelpers_ICVFX::UpdateChromakeyViewportSettings(FDisplayClusterViewport& DstViewport, FDisplayClusterViewport& InCameraViewport, ADisplayClusterRootActor& RootActor, UDisplayClusterICVFXCameraComponent& InCameraComponent)
diff --git a/Engine/Plugins/Runtime/nDisplay/Source/DisplayCluster/Private/Render/Viewport/DisplayClusterViewportManagerProxy.cpp b/Engine/Plugins/Runtime/nDisplay/Source/DisplayCluster/Private/Render/Viewport/DisplayClusterViewportManagerProxy.cpp
--- a/Engine/Plugins/Runtime/nDisplay/Source/DisplayCluster/Private/Render/Viewport/DisplayClusterViewportManagerProxy.cpp
+++ b/Engine/Plugins/Runtime/nDisplay/Source/DisplayCluster/Private/Render/Viewport/DisplayClusterViewportManagerProxy.cpp
@@ -6,6 +6,7 @@
 #include "Misc/DisplayClusterLog.h"
 
 #include "IDisplayCluster.h"
+#include "IDisplayClusterCallbacks.h"
 #include "Render/IDisplayClusterRenderManager.h"
 #include "Render/Viewport/DisplayClusterViewportManager.h"
 #include "Render/Projection/IDisplayClusterProjectionPolicyFactory.h"
@@ -223,6 +224,8 @@ void FDisplayClusterViewportManagerProxy::ImplRenderFrame(FViewport* InViewport)
 			// Now all resources on GPU#0
 		}
 
+		IDisplayCluster::Get().GetCallbacks().OnDisplayClusterPreFrameRender_RenderThread().Broadcast(RHICmdList, ViewportManagerProxy, InViewport);
+
 		// Update viewports resources: overlay, vp-overla, blur, nummips, etc
 		ViewportManagerProxy->UpdateDeferredResources_RenderThread(RHICmdList);
 
@@ -259,6 +262,8 @@ void FDisplayClusterViewportManagerProxy::ImplRenderFrame(FViewport* InViewport)
 				ViewportManagerProxy->ResolveFrameTargetToBackBuffer_RenderThread(RHICmdList, 0, 0, FrameOutputRTT, FrameOutputRTT->GetSizeXY());
 			}
 		}
+
+		IDisplayCluster::Get().GetCallbacks().OnDisplayClusterPostFrameRender_RenderThread().Broadcast(RHICmdList, ViewportManagerProxy, InViewport);
 	});
 }
 
@@ -351,6 +356,8 @@ void FDisplayClusterViewportManagerProxy::UpdateFrameResources_RenderThread(FRHI
 		ImplClearFrameTargets_RenderThread(RHICmdList);
 	}
 
+	IDisplayCluster::Get().GetCallbacks().OnDisplayClusterPreWarp_RenderThread().Broadcast(RHICmdList, this);
+
 	// Handle warped viewport projection policy logic:
 	for (uint8 WarpPass = 0; WarpPass < (uint8)EWarpPass::COUNT; WarpPass++)
 	{
@@ -376,7 +383,7 @@ void FDisplayClusterViewportManagerProxy::UpdateFrameResources_RenderThread(FRHI
 						else
 						{
 							
-							// Projection policy must support warp blend op
+							// Ask current projection policy if it's warp&blend compatible
 							bShouldApplyWarpBlend = PrjPolicy.IsValid() && PrjPolicy->IsWarpBlendSupported();
 						}
 					}
@@ -386,6 +393,7 @@ void FDisplayClusterViewportManagerProxy::UpdateFrameResources_RenderThread(FRHI
 						switch ((EWarpPass)WarpPass)
 						{
 						case EWarpPass::Begin:
+							IDisplayCluster::Get().GetCallbacks().OnDisplayClusterPreWarpViewport_RenderThread().Broadcast(RHICmdList, ViewportProxy);
 							PrjPolicy->BeginWarpBlend_RenderThread(RHICmdList, ViewportProxy);
 							break;
 
@@ -397,6 +405,7 @@ void FDisplayClusterViewportManagerProxy::UpdateFrameResources_RenderThread(FRHI
 
 						case EWarpPass::End:
 							PrjPolicy->EndWarpBlend_RenderThread(RHICmdList, ViewportProxy);
+							IDisplayCluster::Get().GetCallbacks().OnDisplayClusterPostWarpViewport_RenderThread().Broadcast(RHICmdList, ViewportProxy);
 							break;
 
 						default:
@@ -423,6 +432,8 @@ void FDisplayClusterViewportManagerProxy::UpdateFrameResources_RenderThread(FRHI
 		}
 	}
 
+	IDisplayCluster::Get().GetCallbacks().OnDisplayClusterPostWarp_RenderThread().Broadcast(RHICmdList, this);
+
 	if (PostProcessManager.IsValid())
 	{
 		PostProcessManager->PerformPostProcessAfterWarpBlend_RenderThread(RHICmdList, this);
diff --git a/Engine/Plugins/Runtime/nDisplay/Source/DisplayCluster/Private/Render/Viewport/DisplayClusterViewportManagerProxy.h b/Engine/Plugins/Runtime/nDisplay/Source/DisplayCluster/Private/Render/Viewport/DisplayClusterViewportManagerProxy.h
--- a/Engine/Plugins/Runtime/nDisplay/Source/DisplayCluster/Private/Render/Viewport/DisplayClusterViewportManagerProxy.h
+++ b/Engine/Plugins/Runtime/nDisplay/Source/DisplayCluster/Private/Render/Viewport/DisplayClusterViewportManagerProxy.h
@@ -13,6 +13,7 @@
 class FDisplayClusterRenderTargetManager;
 class FDisplayClusterViewportPostProcessManager;
 class FDisplayClusterViewportManager;
+class FDisplayClusterViewportResource;
 class IDisplayClusterProjectionPolicy;
 
 class FViewport;
diff --git a/Engine/Plugins/Runtime/nDisplay/Source/DisplayCluster/Private/Render/Viewport/DisplayClusterViewportManagerViewExtension.cpp b/Engine/Plugins/Runtime/nDisplay/Source/DisplayCluster/Private/Render/Viewport/DisplayClusterViewportManagerViewExtension.cpp
--- a/Engine/Plugins/Runtime/nDisplay/Source/DisplayCluster/Private/Render/Viewport/DisplayClusterViewportManagerViewExtension.cpp
+++ b/Engine/Plugins/Runtime/nDisplay/Source/DisplayCluster/Private/Render/Viewport/DisplayClusterViewportManagerViewExtension.cpp
@@ -2,8 +2,14 @@
 
 #include "DisplayClusterViewportManagerViewExtension.h"
 #include "DisplayClusterSceneViewExtensions.h"
+#include "Render/Viewport/IDisplayClusterViewportManager.h"
+#include "Render/Viewport/IDisplayClusterViewportManagerProxy.h"
 #include "SceneRendering.h"
 
+#include "IDisplayCluster.h"
+#include "IDisplayClusterCallbacks.h"
+
+
 FDisplayClusterViewportManagerViewExtension::FDisplayClusterViewportManagerViewExtension(const FAutoRegister& AutoRegister, const IDisplayClusterViewportManager* InViewportManager)
 	: FSceneViewExtensionBase(AutoRegister)
 	, ViewportManager(InViewportManager)
@@ -35,3 +41,39 @@ bool FDisplayClusterViewportManagerViewExtension::IsActiveThisFrame_Internal(con
 
 	return false;
 }
+
+void FDisplayClusterViewportManagerViewExtension::PostRenderViewFamily_RenderThread(FRHICommandListImmediate& RHICmdList, FSceneViewFamily& InViewFamily)
+{
+	// Since FSceneViewFamily has no direct references to the nD viewport it relates to,
+	// we find a proper viewport ID by comparing RenderTarget addresses.
+	if (ViewportManager)
+	{
+		const IDisplayClusterViewportManagerProxy* const ViewportMgrPoxy = ViewportManager->GetProxy();
+		if (ViewportMgrPoxy)
+		{
+			// Get all available viewport proxies
+			const TArrayView<IDisplayClusterViewportProxy*> ViewportProxies = ViewportMgrPoxy->GetViewports_RenderThread();
+
+			// Filter those viewports that refer to the RenderTarget used in the ViewFamily
+			const TArray<IDisplayClusterViewportProxy*> FoundViewportProxies = ViewportProxies.FilterByPredicate([&InViewFamily](const IDisplayClusterViewportProxy* ViewportProxy)
+			{
+				TArray<FRHITexture*> RenderTargets;
+				if (ViewportProxy->GetResources_RenderThread(EDisplayClusterViewportResourceType::InternalRenderTargetResource, RenderTargets))
+				{
+					if (RenderTargets.Contains(InViewFamily.RenderTarget->GetRenderTargetTexture()))
+					{
+						return true;
+					}
+				}
+
+				return false;
+			});
+
+			// Now we can perform per-viewport notification
+			for (const IDisplayClusterViewportProxy* ViewportProxy : FoundViewportProxies)
+			{
+				IDisplayCluster::Get().GetCallbacks().OnDisplayClusterPostRenderViewFamily_RenderThread().Broadcast(RHICmdList, InViewFamily, ViewportProxy);
+			}
+		}
+	}
+}
diff --git a/Engine/Plugins/Runtime/nDisplay/Source/DisplayCluster/Private/Render/Viewport/DisplayClusterViewportManagerViewExtension.h b/Engine/Plugins/Runtime/nDisplay/Source/DisplayCluster/Private/Render/Viewport/DisplayClusterViewportManagerViewExtension.h
--- a/Engine/Plugins/Runtime/nDisplay/Source/DisplayCluster/Private/Render/Viewport/DisplayClusterViewportManagerViewExtension.h
+++ b/Engine/Plugins/Runtime/nDisplay/Source/DisplayCluster/Private/Render/Viewport/DisplayClusterViewportManagerViewExtension.h
@@ -22,6 +22,8 @@ public:
 
 	virtual void PreRenderView_RenderThread(FRHICommandListImmediate& RHICmdList, FSceneView& InView) override;
 	virtual bool IsActiveThisFrame_Internal(const FSceneViewExtensionContext& Context) const override;
+
+	virtual void PostRenderViewFamily_RenderThread(FRHICommandListImmediate& RHICmdList, FSceneViewFamily& InViewFamily) override;
 	//~End ISceneVIewExtension interface
 
 private:
diff --git a/Engine/Plugins/Runtime/nDisplay/Source/DisplayCluster/Private/Render/Viewport/DisplayClusterViewportManager_EditorPreview.cpp b/Engine/Plugins/Runtime/nDisplay/Source/DisplayCluster/Private/Render/Viewport/DisplayClusterViewportManager_EditorPreview.cpp
--- a/Engine/Plugins/Runtime/nDisplay/Source/DisplayCluster/Private/Render/Viewport/DisplayClusterViewportManager_EditorPreview.cpp
+++ b/Engine/Plugins/Runtime/nDisplay/Source/DisplayCluster/Private/Render/Viewport/DisplayClusterViewportManager_EditorPreview.cpp
@@ -185,7 +185,7 @@ bool FDisplayClusterViewportManager::RenderInEditor(class FDisplayClusterRenderF
 						FRotator ViewRotation;
 						FSceneView* View = ViewportPtr->ImplCalcScenePreview(ViewFamily, ViewIt.ContextNum);
 
-						if (View != nullptr && ViewIt.IsShouldRenderView() == false)
+						if (View != nullptr && ViewIt.ShouldRenderSceneView() == false)
 						{
 							ViewFamily.Views.Remove(View);
 
diff --git a/Engine/Plugins/Runtime/nDisplay/Source/DisplayCluster/Private/Render/Viewport/RenderFrame/DisplayClusterRenderFrameManager.cpp b/Engine/Plugins/Runtime/nDisplay/Source/DisplayCluster/Private/Render/Viewport/RenderFrame/DisplayClusterRenderFrameManager.cpp
--- a/Engine/Plugins/Runtime/nDisplay/Source/DisplayCluster/Private/Render/Viewport/RenderFrame/DisplayClusterRenderFrameManager.cpp
+++ b/Engine/Plugins/Runtime/nDisplay/Source/DisplayCluster/Private/Render/Viewport/RenderFrame/DisplayClusterRenderFrameManager.cpp
@@ -49,8 +49,31 @@ bool FDisplayClusterRenderFrameManager::BuildRenderFrame(FViewport* InViewport,
 	}
 
 	SortedViewports.Append(ChildsViewports);
-	ChildsViewports.Empty();
 	
+	// At this point, we have the list of the viewports sorted in a way that
+	// child (dependent) viewports follow after the parent ones. Here we do
+	// some re-ordering so all the viewports that have a media capture
+	// device assigned, will be first in the list. Child viewports can't
+	// have media capture assigned, so there won't be any problems.
+	{
+		// Find all viewports being captured
+		TArray<FDisplayClusterViewport*> ViewportsBeingCaptured = SortedViewports.FilterByPredicate([](const FDisplayClusterViewport* Viewport)
+		{
+			return Viewport->GetRenderSettings().bIsBeingCaptured;
+		});
+
+		// Put them all in the beginning of the list
+		const int32 CapturingViewportsAmount = ViewportsBeingCaptured.Num();
+		if (CapturingViewportsAmount > 0 && CapturingViewportsAmount != SortedViewports.Num())
+		{
+			for (FDisplayClusterViewport* ViewportBeingCaptured : ViewportsBeingCaptured)
+			{
+				SortedViewports.Remove(ViewportBeingCaptured);
+				SortedViewports.Insert(ViewportBeingCaptured, 0);
+			}
+		}
+	}
+
 	bool bResult = false;
 
 	if (InRenderFrameSettings.bAllowRenderTargetAtlasing)
@@ -99,7 +122,7 @@ bool FDisplayClusterRenderFrameManager::BuildSimpleFrame(FViewport* InViewport,
 				{
 					FrameView.ContextNum = ContextIt.ContextNum;
 					FrameView.Viewport = ViewportIt;
-					FrameView.bDisableRender = ContextIt.bDisableRender;
+					FrameView.bDisableRender = ContextIt.bDisableRender || ViewportIt->RenderSettings.bSkipSceneRenderingButLeaveResourcesAvailable;
 					FrameView.bFreezeRendering = ViewportIt->RenderSettings.bFreezeRendering;
 				}
 
diff --git a/Engine/Plugins/Runtime/nDisplay/Source/DisplayCluster/Public/Components/DisplayClusterICVFXCameraComponent.h b/Engine/Plugins/Runtime/nDisplay/Source/DisplayCluster/Public/Components/DisplayClusterICVFXCameraComponent.h
--- a/Engine/Plugins/Runtime/nDisplay/Source/DisplayCluster/Public/Components/DisplayClusterICVFXCameraComponent.h
+++ b/Engine/Plugins/Runtime/nDisplay/Source/DisplayCluster/Public/Components/DisplayClusterICVFXCameraComponent.h
@@ -163,6 +163,9 @@ private:
 	UPROPERTY(EditAnywhere, Transient, Category = "Texture Replacement", meta = (PropertyPath = "CameraSettings.RenderSettings.Replace.TextureRegion", EditConditionPath = "CameraSettings.bEnable && CameraSettings.RenderSettings.Replace.bAllowReplace && CameraSettings.RenderSettings.Replace.bShouldUseTextureRegion"))
 	FDisplayClusterEditorPropertyReference TextureRegionRef;
 
+	UPROPERTY(EditAnywhere, Transient, Category = "Media", meta = (PropertyPath = "CameraSettings.RenderSettings.Media", EditConditionPath = "CameraSettings.bEnable"))
+	FDisplayClusterEditorPropertyReference MediaRef;
+
 	UPROPERTY(EditAnywhere, Transient, Category = Configuration, meta = (PropertyPath = "CameraSettings.RenderSettings.RenderOrder", EditConditionPath = "CameraSettings.bEnable"))
 	FDisplayClusterEditorPropertyReference RenderOrderRef;
 
diff --git a/Engine/Plugins/Runtime/nDisplay/Source/DisplayCluster/Public/IDisplayClusterCallbacks.h b/Engine/Plugins/Runtime/nDisplay/Source/DisplayCluster/Public/IDisplayClusterCallbacks.h
--- a/Engine/Plugins/Runtime/nDisplay/Source/DisplayCluster/Public/IDisplayClusterCallbacks.h
+++ b/Engine/Plugins/Runtime/nDisplay/Source/DisplayCluster/Public/IDisplayClusterCallbacks.h
@@ -4,6 +4,12 @@
 
 #include "DisplayClusterEnums.h"
 
+class IDisplayClusterViewportProxy;
+class IDisplayClusterViewportManagerProxy;
+class FRHICommandListImmediate;
+class FSceneViewFamily;
+class FViewport;
+
 
 /**
  * DisplayCluster callbacks API
@@ -65,4 +71,32 @@ public:
 	/** Failover notification **/
 	DECLARE_EVENT_OneParam(IDisplayClusterCallbacks, FDisplayClusterFailoverNodeDown, const FString&);
 	virtual FDisplayClusterFailoverNodeDown& OnDisplayClusterFailoverNodeDown() = 0;
+
+	/** Called once the ViewFamily of this viewport is rendered **/
+	DECLARE_EVENT_ThreeParams(IDisplayClusterCallbacks, FDisplayClusterPostRenderViewFamily_RenderThread, FRHICommandListImmediate&, const FSceneViewFamily&, const IDisplayClusterViewportProxy*);
+	virtual FDisplayClusterPostRenderViewFamily_RenderThread& OnDisplayClusterPostRenderViewFamily_RenderThread() = 0;
+
+	/** Called once before warping all available viewports **/
+	DECLARE_EVENT_TwoParams(IDisplayClusterCallbacks, FDisplayClusterPreWarp_RenderThread, FRHICommandListImmediate&, const IDisplayClusterViewportManagerProxy*);
+	virtual FDisplayClusterPreWarp_RenderThread& OnDisplayClusterPreWarp_RenderThread() = 0;
+
+	/** Called before warping a specific viewport **/
+	DECLARE_EVENT_TwoParams(IDisplayClusterCallbacks, FDisplayClusterPreWarpViewport_RenderThread, FRHICommandListImmediate&, const IDisplayClusterViewportProxy*);
+	virtual FDisplayClusterPreWarpViewport_RenderThread& OnDisplayClusterPreWarpViewport_RenderThread() = 0;
+
+	/** Called once after warping all the viewports **/
+	DECLARE_EVENT_TwoParams(IDisplayClusterCallbacks, FDisplayClusterPostWarp_RenderThread, FRHICommandListImmediate&, const IDisplayClusterViewportManagerProxy*);
+	virtual FDisplayClusterPostWarp_RenderThread& OnDisplayClusterPostWarp_RenderThread() = 0;
+
+	/** Called after warping a specific viewport **/
+	DECLARE_EVENT_TwoParams(IDisplayClusterCallbacks, FDisplayClusterPostWarpViewport_RenderThread, FRHICommandListImmediate&, const IDisplayClusterViewportProxy*);
+	virtual FDisplayClusterPostWarpViewport_RenderThread& OnDisplayClusterPostWarpViewport_RenderThread() = 0;
+
+	/** Called before staring nDisplay rendering chain **/
+	DECLARE_EVENT_ThreeParams(IDisplayClusterCallbacks, FDisplayClusterPreFrameRender_RenderThread, FRHICommandListImmediate&, const IDisplayClusterViewportManagerProxy*, FViewport*);
+	virtual FDisplayClusterPreFrameRender_RenderThread& OnDisplayClusterPreFrameRender_RenderThread() = 0;
+
+	/** Last call from the nDisplay rendering chain **/
+	DECLARE_EVENT_ThreeParams(IDisplayClusterCallbacks, FDisplayClusterPostFrameRender_RenderThread, FRHICommandListImmediate&, const IDisplayClusterViewportManagerProxy*, FViewport*);
+	virtual FDisplayClusterPostFrameRender_RenderThread& OnDisplayClusterPostFrameRender_RenderThread() = 0;
 };
diff --git a/Engine/Plugins/Runtime/nDisplay/Source/DisplayCluster/Public/Render/Viewport/Containers/DisplayClusterViewport_RenderSettings.h b/Engine/Plugins/Runtime/nDisplay/Source/DisplayCluster/Public/Render/Viewport/Containers/DisplayClusterViewport_RenderSettings.h
--- a/Engine/Plugins/Runtime/nDisplay/Source/DisplayCluster/Public/Render/Viewport/Containers/DisplayClusterViewport_RenderSettings.h
+++ b/Engine/Plugins/Runtime/nDisplay/Source/DisplayCluster/Public/Render/Viewport/Containers/DisplayClusterViewport_RenderSettings.h
@@ -39,6 +39,11 @@ public:
 	// Freeze viewport resources, skip rendering internal viewport resources. But still use it for final compositing
 	bool bFreezeRendering = false;
 
+	// This flag means no scene rendering required, but all internal resources should still be valid for
+	// the media subsystem. It's a temporary solution. The flags 'bSkipRendering' and 'bFreezeRendering'
+	// above, plus this one, need to be refactored at some point.
+	bool bSkipSceneRenderingButLeaveResourcesAvailable = false;
+
 	// Render alpha channel from input texture to warp output
 	bool bWarpBlendRenderAlphaChannel = false;
 
@@ -48,6 +53,9 @@ public:
 	// Useful to render some viewports in mono, then copied to stereo backbuffers identical image
 	bool bForceMono = false;
 
+	// Is this viewport being captured by a media capture device?
+	bool bIsBeingCaptured = false;
+
 	// Performance, Multi-GPU: Asign GPU for viewport rendering. The Value '-1' used to default gpu mapping
 	int32 GPUIndex = -1;
 
diff --git a/Engine/Plugins/Runtime/nDisplay/Source/DisplayCluster/Public/Render/Viewport/RenderFrame/DisplayClusterRenderFrame.h b/Engine/Plugins/Runtime/nDisplay/Source/DisplayCluster/Public/Render/Viewport/RenderFrame/DisplayClusterRenderFrame.h
--- a/Engine/Plugins/Runtime/nDisplay/Source/DisplayCluster/Public/Render/Viewport/RenderFrame/DisplayClusterRenderFrame.h
+++ b/Engine/Plugins/Runtime/nDisplay/Source/DisplayCluster/Public/Render/Viewport/RenderFrame/DisplayClusterRenderFrame.h
@@ -20,7 +20,7 @@ public:
 	class FFrameView
 	{
 	public:
-		inline bool IsShouldRenderView() const
+		inline bool ShouldRenderSceneView() const
 		{
 			return !bDisableRender && !bFreezeRendering;
 		}
diff --git a/Engine/Plugins/Runtime/nDisplay/Source/DisplayClusterConfiguration/DisplayClusterConfiguration.Build.cs b/Engine/Plugins/Runtime/nDisplay/Source/DisplayClusterConfiguration/DisplayClusterConfiguration.Build.cs
--- a/Engine/Plugins/Runtime/nDisplay/Source/DisplayClusterConfiguration/DisplayClusterConfiguration.Build.cs
+++ b/Engine/Plugins/Runtime/nDisplay/Source/DisplayClusterConfiguration/DisplayClusterConfiguration.Build.cs
@@ -17,9 +17,11 @@ public class DisplayClusterConfiguration : ModuleRules
 
 		PublicDependencyModuleNames.AddRange(
 			new string[] {
-				"CinematicCamera",
 				"ActorLayerUtilities",
-				"OpenColorIO"
+				"CinematicCamera",
+				"MediaAssets",
+				"MediaIOCore",
+				"OpenColorIO",
 			});
 
 		PrivateDependencyModuleNames.AddRange(
diff --git a/Engine/Plugins/Runtime/nDisplay/Source/DisplayClusterConfiguration/Public/DisplayClusterConfigurationTypes.h b/Engine/Plugins/Runtime/nDisplay/Source/DisplayClusterConfiguration/Public/DisplayClusterConfigurationTypes.h
--- a/Engine/Plugins/Runtime/nDisplay/Source/DisplayClusterConfiguration/Public/DisplayClusterConfigurationTypes.h
+++ b/Engine/Plugins/Runtime/nDisplay/Source/DisplayClusterConfiguration/Public/DisplayClusterConfigurationTypes.h
@@ -350,6 +350,10 @@ public:
 	UPROPERTY(EditDefaultsOnly, BlueprintReadOnly, Category = "Configuration", meta = (DisplayName = "Custom Output Settings", DisplayThumbnail = false, ShowInnerProperties))
 	TMap<FString, FDisplayClusterConfigurationPostprocess> Postprocess;
 
+	// Media settings
+	UPROPERTY(EditAnywhere, BlueprintReadWrite, Category = "Media", meta = (ShowOnlyInnerProperties))
+	FDisplayClusterConfigurationMedia Media;
+
 #if WITH_EDITORONLY_DATA
 	UPROPERTY(EditDefaultsOnly, Category = "Configuration", meta = (HideProperty))
 	bool bIsVisible;
diff --git a/Engine/Plugins/Runtime/nDisplay/Source/DisplayClusterConfiguration/Public/DisplayClusterConfigurationTypes_ICVFX.h b/Engine/Plugins/Runtime/nDisplay/Source/DisplayClusterConfiguration/Public/DisplayClusterConfigurationTypes_ICVFX.h
--- a/Engine/Plugins/Runtime/nDisplay/Source/DisplayClusterConfiguration/Public/DisplayClusterConfigurationTypes_ICVFX.h
+++ b/Engine/Plugins/Runtime/nDisplay/Source/DisplayClusterConfiguration/Public/DisplayClusterConfigurationTypes_ICVFX.h
@@ -12,6 +12,7 @@
 
 #include "CineCameraActor.h"
 
+#include "DisplayClusterConfigurationTypes_Media.h"
 #include "DisplayClusterConfigurationTypes_PostRender.h"
 #include "DisplayClusterConfigurationTypes_Postprocess.h"
 #include "DisplayClusterConfigurationTypes_OCIO.h"
@@ -317,9 +318,13 @@ public:
 	bool bUseCameraComponentPostprocess = true;
 
 	// Replace viewport render from source texture
-	UPROPERTY(EditAnywhere, BlueprintReadWrite, Category = NDisplay, meta = (DisplayName = "Mipmapping"))
+	UPROPERTY(EditAnywhere, BlueprintReadWrite, Category = NDisplay, meta = (DisplayName = "Replace Output With a Texture"))
 	FDisplayClusterConfigurationPostRender_Override Replace;
 
+	// Media settings
+	UPROPERTY(EditAnywhere, BlueprintReadWrite, Category = NDisplay, meta = (DisplayName = "Media"))
+	FDisplayClusterConfigurationMedia Media;
+
 	UPROPERTY()
 	FDisplayClusterConfigurationPostRender_BlurPostprocess PostprocessBlur;
 
diff --git a/Engine/Plugins/Runtime/nDisplay/Source/DisplayClusterConfiguration/Public/DisplayClusterConfigurationTypes_Media.h b/Engine/Plugins/Runtime/nDisplay/Source/DisplayClusterConfiguration/Public/DisplayClusterConfigurationTypes_Media.h
new file mode 100644
--- /dev/null
+++ b/Engine/Plugins/Runtime/nDisplay/Source/DisplayClusterConfiguration/Public/DisplayClusterConfigurationTypes_Media.h
@@ -0,0 +1,98 @@
+// Copyright Epic Games, Inc. All Rights Reserved.
+
+#pragma once
+
+#include "CoreMinimal.h"
+
+#include "MediaPlayer.h"
+#include "MediaSource.h"
+#include "MediaTexture.h"
+#include "MediaOutput.h"
+
+#include "Engine/TextureRenderTarget2D.h"
+
+#include "DisplayClusterConfigurationTypes_Media.generated.h"
+
+
+/*
+ * Media input settings
+ */
+USTRUCT(Blueprintable)
+struct DISPLAYCLUSTERCONFIGURATION_API FDisplayClusterConfigurationMediaInput
+{
+	GENERATED_BODY()
+
+public:
+	/** Enable/disable media input */
+	UPROPERTY(EditAnywhere, BlueprintReadWrite, Category = "Media Input", meta = (DisplayName = "Enable"))
+	bool bEnabled = false;
+
+	/** Media source to use */
+	UPROPERTY(EditAnywhere, BlueprintReadWrite, Category = "Media Input", meta = (EditCondition = "bEnabled"))
+	TObjectPtr<UMediaSource> MediaSource = nullptr;
+
+	/** Media player to use */
+	UPROPERTY(EditAnywhere, BlueprintReadWrite, Category = "Media Input", meta = (EditCondition = "bEnabled"))
+	TObjectPtr<UMediaPlayer> MediaPlayer = nullptr;
+
+	/** Media texture to use */
+	UPROPERTY(EditAnywhere, BlueprintReadWrite, Category = "Media Input", meta = (EditCondition = "bEnabled"))
+	TObjectPtr<UMediaTexture> MediaTexture = nullptr;
+};
+
+
+/*
+ * Media output settings
+ */
+USTRUCT(Blueprintable)
+struct DISPLAYCLUSTERCONFIGURATION_API FDisplayClusterConfigurationMediaOutput
+{
+	GENERATED_BODY()
+
+public:
+	/** Enable/disable media capture */
+	UPROPERTY(EditAnywhere, BlueprintReadWrite, Category = "Media Output", meta = (DisplayName = "Enable"))
+	bool bEnabled = false;
+
+	/** Media output to use */
+	UPROPERTY(EditAnywhere, BlueprintReadWrite, Category = "Media Output", meta = (EditCondition = "bEnabled"))
+	TObjectPtr<UMediaOutput> MediaOutput = nullptr;
+
+	/** Render target to use */
+	UPROPERTY(EditAnywhere, BlueprintReadWrite, Category = "Media Output", meta = (EditCondition = "bEnabled"))
+	TObjectPtr<UTextureRenderTarget2D> RenderTarget = nullptr;
+};
+
+
+/*
+ * ICVFX media override settings
+ */
+USTRUCT(Blueprintable)
+struct DISPLAYCLUSTERCONFIGURATION_API FDisplayClusterConfigurationMedia
+{
+	GENERATED_BODY()
+
+public:
+	/** Enable/disable media features */
+	UPROPERTY(EditAnywhere, BlueprintReadWrite, Category = "Media", meta = (DisplayName = "Enable"))
+	bool bEnabled = false;
+
+	/** When in-cluster media sharing us used, the cluster node specified here will be used as a source (Tx node) */
+	UPROPERTY(EditAnywhere, BlueprintReadWrite, Category = "Media", meta = (EditCondition = "bEnabled"))
+	FString MediaOutputNode;
+
+	/** Media input settings */
+	UPROPERTY(EditAnywhere, BlueprintReadWrite, Category = "Media", meta = (EditCondition = "bEnabled"))
+	FDisplayClusterConfigurationMediaInput MediaInput;
+
+	/** Media capture settings */
+	UPROPERTY(EditAnywhere, BlueprintReadWrite, Category = "Media", meta = (EditCondition = "bEnabled"))
+	FDisplayClusterConfigurationMediaOutput MediaOutput;
+
+public:
+	/** Returns true if media sharing is used */
+	bool IsMediaSharingRequired() const
+	{
+		return bEnabled && MediaInput.bEnabled && MediaOutput.bEnabled;
+	}
+};
diff --git a/Engine/Plugins/Runtime/nDisplay/Source/DisplayClusterConfiguration/Public/DisplayClusterConfigurationTypes_Viewport.h b/Engine/Plugins/Runtime/nDisplay/Source/DisplayClusterConfiguration/Public/DisplayClusterConfigurationTypes_Viewport.h
--- a/Engine/Plugins/Runtime/nDisplay/Source/DisplayClusterConfiguration/Public/DisplayClusterConfigurationTypes_Viewport.h
+++ b/Engine/Plugins/Runtime/nDisplay/Source/DisplayClusterConfiguration/Public/DisplayClusterConfigurationTypes_Viewport.h
@@ -8,6 +8,8 @@
 #include "UObject/Object.h"
 
 #include "DisplayClusterConfigurationStrings.h"
+#include "DisplayClusterConfigurationTypes_Base.h"
+#include "DisplayClusterConfigurationTypes_Media.h"
 #include "DisplayClusterConfigurationTypes_PostRender.h"
 #include "DisplayClusterConfigurationTypes_Postprocess.h"
 #include "DisplayClusterConfigurationTypes_ViewportRemap.h"
@@ -116,6 +118,10 @@ public:
 	UPROPERTY(EditAnywhere, BlueprintReadWrite, Category = "Configuration")
 	FDisplayClusterConfigurationViewport_Overscan Overscan;
 
+	// Media settings
+	UPROPERTY(EditAnywhere, BlueprintReadWrite, Category = "Configuration", meta = (DisplayName = "Media"))
+	FDisplayClusterConfigurationMedia Media;
+
 	// Experimental: Support special frame builder mode - merge viewports to single viewfamily by group num
 	// [not implemented yet]
 	UPROPERTY()
diff --git a/Engine/Plugins/Runtime/nDisplay/Source/DisplayClusterConfigurator/DisplayClusterConfigurator.Build.cs b/Engine/Plugins/Runtime/nDisplay/Source/DisplayClusterConfigurator/DisplayClusterConfigurator.Build.cs
--- a/Engine/Plugins/Runtime/nDisplay/Source/DisplayClusterConfigurator/DisplayClusterConfigurator.Build.cs
+++ b/Engine/Plugins/Runtime/nDisplay/Source/DisplayClusterConfigurator/DisplayClusterConfigurator.Build.cs
@@ -34,6 +34,7 @@ public class DisplayClusterConfigurator : ModuleRules
 				"Kismet",
 				"KismetCompiler",
 				"MainFrame",
+				"MediaAssets",
 				"MessageLog",
 				"PinnedCommandList",
 				"PlacementMode",
diff --git a/Engine/Plugins/Runtime/nDisplay/Source/DisplayClusterConfigurator/Private/DisplayClusterConfiguratorModule.cpp b/Engine/Plugins/Runtime/nDisplay/Source/DisplayClusterConfigurator/Private/DisplayClusterConfiguratorModule.cpp
--- a/Engine/Plugins/Runtime/nDisplay/Source/DisplayClusterConfigurator/Private/DisplayClusterConfiguratorModule.cpp
+++ b/Engine/Plugins/Runtime/nDisplay/Source/DisplayClusterConfigurator/Private/DisplayClusterConfiguratorModule.cpp
@@ -17,6 +17,7 @@
 #include "Views/Details/Cluster/DisplayClusterConfiguratorClusterDetailsCustomization.h"
 #include "Views/Details/Cluster/DisplayClusterConfiguratorExternalImageTypeCustomization.h"
 #include "Views/Details/Cluster/DisplayClusterConfiguratorGenerateMipsCustomization.h"
+#include "Views/Details/Cluster/DisplayClusterConfiguratorMediaTypeCustomization.h"
 #include "Views/Details/Cluster/DisplayClusterConfiguratorNodeSelectionCustomization.h"
 #include "Views/Details/Cluster/DisplayClusterConfiguratorClusterReferenceListCustomization.h"
 #include "Views/Details/Cluster/DisplayClusterConfiguratorViewportDetailsCustomization.h"
@@ -205,6 +206,7 @@ void FDisplayClusterConfiguratorModule::RegisterCustomLayouts()
 	REGISTER_PROPERTY_LAYOUT(FDisplayClusterConfigurationClusterItemReferenceList, FDisplayClusterConfiguratorClusterReferenceListCustomization);
 	REGISTER_PROPERTY_LAYOUT(FDisplayClusterConfigurationViewport_RemapData, FDisplayClusterConfiguratorViewportRemapCustomization);
 	REGISTER_PROPERTY_LAYOUT(FDisplayClusterConfigurationRectangle, FDisplayClusterConfiguratorRectangleCustomization);
+	REGISTER_PROPERTY_LAYOUT(FDisplayClusterConfigurationMedia, FDisplayClusterConfiguratorMediaTypeCustomization);
 }
 
 void FDisplayClusterConfiguratorModule::UnregisterCustomLayouts()
diff --git a/Engine/Plugins/Runtime/nDisplay/Source/DisplayClusterConfigurator/Private/Views/Details/Cluster/DisplayClusterConfiguratorMediaTypeCustomization.cpp b/Engine/Plugins/Runtime/nDisplay/Source/DisplayClusterConfigurator/Private/Views/Details/Cluster/DisplayClusterConfiguratorMediaTypeCustomization.cpp
new file mode 100644
--- /dev/null
+++ b/Engine/Plugins/Runtime/nDisplay/Source/DisplayClusterConfigurator/Private/Views/Details/Cluster/DisplayClusterConfiguratorMediaTypeCustomization.cpp
@@ -0,0 +1,119 @@
+// Copyright Epic Games, Inc. All Rights Reserved.
+
+#include "Views/Details/Cluster/DisplayClusterConfiguratorMediaTypeCustomization.h"
+
+#include "DisplayClusterConfigurationTypes.h"
+#include "DisplayClusterConfiguratorBlueprintEditor.h"
+#include "DisplayClusterRootActor.h"
+
+#include "DetailWidgetRow.h"
+#include "IDetailChildrenBuilder.h"
+#include "IPropertyTypeCustomization.h"
+#include "PropertyHandle.h"
+
+
+void FDisplayClusterConfiguratorMediaTypeCustomization::CustomizeChildren(TSharedRef<IPropertyHandle> InPropertyHandle, IDetailChildrenBuilder& InChildBuilder, IPropertyTypeCustomizationUtils& CustomizationUtils)
+{
+	OutputNodeProperty = InPropertyHandle->GetChildHandle(GET_MEMBER_NAME_CHECKED(FDisplayClusterConfigurationMedia, MediaOutputNode));
+	check(OutputNodeProperty && OutputNodeProperty->IsValidHandle());
+	OutputNodeProperty->MarkHiddenByCustomization();
+
+	BuildClusterNodeOptionsList();
+
+	// Create GUI representation
+	{
+		FDetailWidgetRow& Row = InChildBuilder.AddCustomRow(OutputNodeProperty->GetPropertyDisplayName())
+			.NameContent()
+			[
+				OutputNodeProperty->CreatePropertyNameWidget()
+			]
+			.ValueContent()
+			[
+				SAssignNew(OutputNodeComboBox, SComboBox<TSharedPtr<FString>>)
+				.OptionsSource(&OutputNodeOptions)
+				.OnGenerateWidget(this, &FDisplayClusterConfiguratorMediaTypeCustomization::CreateComboWidget)
+				.OnSelectionChanged(this, &FDisplayClusterConfiguratorMediaTypeCustomization::OnNodeIdSelected)
+				.ContentPadding(2)
+				.Content()
+				[
+					SNew(STextBlock)
+					.Text(this, &FDisplayClusterConfiguratorMediaTypeCustomization::GetSelectedNodeIdText)
+				]
+			];
+	}
+
+	// Set selected item
+	FString CurrentNodeId;
+	if (OutputNodeProperty->GetValue(CurrentNodeId))
+	{
+		TSharedPtr<FString>* FoundItem = OutputNodeOptions.FindByPredicate([CurrentNodeId](const TSharedPtr<FString>& Item)
+		{
+			return Item->Equals(CurrentNodeId, ESearchCase::IgnoreCase);
+		});
+
+		if (FoundItem)
+		{
+			OutputNodeComboBox->SetSelectedItem(*FoundItem);
+		}
+	}
+
+	FDisplayClusterConfiguratorBaseTypeCustomization::CustomizeChildren(InPropertyHandle, InChildBuilder, CustomizationUtils);
+}
+
+void FDisplayClusterConfiguratorMediaTypeCustomization::BuildClusterNodeOptionsList()
+{
+	ADisplayClusterRootActor* DCRA = nullptr;
+
+	if (FDisplayClusterConfiguratorBlueprintEditor* Editor = FindBlueprintEditor())
+	{
+		DCRA = Editor->GetDefaultRootActor();
+	}
+
+	if (!DCRA)
+	{
+		DCRA = FindRootActor();
+	}
+	
+	if (!DCRA)
+	{
+		return;
+	}
+
+	// Get cluster node IDs
+	TArray<FString> NodeIds;
+	if (const UDisplayClusterConfigurationData* const ConfigurationData = DCRA->GetConfigData())
+	{
+		ConfigurationData->Cluster->GetNodeIds(NodeIds);
+	}
+
+	// Sort the node IDs in alphabet order
+	NodeIds.Sort();
+
+	// Prepare combobox options
+	OutputNodeOptions.Reset();
+	for (const FString& NodeId : NodeIds)
+	{
+		if (!NodeId.IsEmpty())
+		{
+			OutputNodeOptions.Emplace(MakeShared<FString>(NodeId));
+		}
+	}
+}
+
+
+void FDisplayClusterConfiguratorMediaTypeCustomization::OnNodeIdSelected(TSharedPtr<FString> PreviewNodeId, ESelectInfo::Type SelectInfo)
+{
+	const FString NewValue = PreviewNodeId.IsValid() ? *PreviewNodeId : FString();
+	OutputNodeProperty->SetValue(NewValue);
+}
+
+FText FDisplayClusterConfiguratorMediaTypeCustomization::GetSelectedNodeIdText() const
+{
+	TSharedPtr<FString> CurSelection = OutputNodeComboBox->GetSelectedItem();
+	return FText::FromString(CurSelection.IsValid() ? *CurSelection : *FString());
+}
+
+TSharedRef<SWidget> FDisplayClusterConfiguratorMediaTypeCustomization::CreateComboWidget(TSharedPtr<FString> InItem)
+{
+	return SNew(STextBlock).Text(FText::FromString(*InItem));
+}
diff --git a/Engine/Plugins/Runtime/nDisplay/Source/DisplayClusterConfigurator/Private/Views/Details/Cluster/DisplayClusterConfiguratorMediaTypeCustomization.h b/Engine/Plugins/Runtime/nDisplay/Source/DisplayClusterConfigurator/Private/Views/Details/Cluster/DisplayClusterConfiguratorMediaTypeCustomization.h
new file mode 100644
--- /dev/null
+++ b/Engine/Plugins/Runtime/nDisplay/Source/DisplayClusterConfigurator/Private/Views/Details/Cluster/DisplayClusterConfiguratorMediaTypeCustomization.h
@@ -0,0 +1,47 @@
+// Copyright Epic Games, Inc. All Rights Reserved.
+
+#pragma once
+
+#include "Views/Details/DisplayClusterConfiguratorBaseTypeCustomization.h"
+
+template<class T>
+class SComboBox;
+
+
+class FDisplayClusterConfiguratorMediaTypeCustomization final
+	: public FDisplayClusterConfiguratorBaseTypeCustomization
+{
+public:
+	static TSharedRef<IPropertyTypeCustomization> MakeInstance()
+	{
+		return MakeShared<FDisplayClusterConfiguratorMediaTypeCustomization>();
+	}
+
+protected:
+	//~ IPropertyTypeCustomization interface begin
+	virtual void CustomizeChildren(TSharedRef<IPropertyHandle> InPropertyHandle, IDetailChildrenBuilder& InChildBuilder, IPropertyTypeCustomizationUtils& CustomizationUtils) override;
+	//~ IPropertyTypeCustomization interface end
+
+private:
+	/**  Generates internal array of cluster node IDs */
+	void BuildClusterNodeOptionsList();
+
+	/**  Processes combobox selection changes */
+	void OnNodeIdSelected(TSharedPtr<FString> PreviewNodeId, ESelectInfo::Type SelectInfo);
+
+	/**  Returns text of currently selected item */
+	FText GetSelectedNodeIdText() const;
+
+	/** Creates combobox widget */
+	TSharedRef<SWidget> CreateComboWidget(TSharedPtr<FString> InItem);
+
+private:
+	/** OutputNode original property handler */
+	TSharedPtr<IPropertyHandle> OutputNodeProperty;
+
+	/** List of cluster node IDs for the combobox */
+	TArray<TSharedPtr<FString>> OutputNodeOptions;
+
+	/** Node ID selection combobox */
+	TSharedPtr<SComboBox<TSharedPtr<FString>>> OutputNodeComboBox;
+};
diff --git a/Engine/Plugins/Runtime/nDisplay/Source/DisplayClusterConfigurator/Private/Views/Details/DisplayClusterConfiguratorBaseDetailCustomization.cpp b/Engine/Plugins/Runtime/nDisplay/Source/DisplayClusterConfigurator/Private/Views/Details/DisplayClusterConfiguratorBaseDetailCustomization.cpp
--- a/Engine/Plugins/Runtime/nDisplay/Source/DisplayClusterConfigurator/Private/Views/Details/DisplayClusterConfiguratorBaseDetailCustomization.cpp
+++ b/Engine/Plugins/Runtime/nDisplay/Source/DisplayClusterConfigurator/Private/Views/Details/DisplayClusterConfiguratorBaseDetailCustomization.cpp
@@ -105,6 +105,17 @@ IDetailLayoutBuilder* FDisplayClusterConfiguratorBaseDetailCustomization::GetLay
 	return LayoutBuilderPtr;
 }
 
+FDisplayClusterConfiguratorBlueprintEditor* FDisplayClusterConfiguratorBaseDetailCustomization::GetBlueprintEditor() const
+{
+	FDisplayClusterConfiguratorBlueprintEditor* BlueprintEditorPtr = nullptr;
+	if (ToolkitPtr.IsValid())
+	{
+		BlueprintEditorPtr = ToolkitPtr.Pin().Get();
+	}
+
+	return BlueprintEditorPtr;
+}
+
 ADisplayClusterRootActor* FDisplayClusterConfiguratorBaseDetailCustomization::GetRootActor() const
 {
 	ADisplayClusterRootActor* RootActor = nullptr;
diff --git a/Engine/Plugins/Runtime/nDisplay/Source/DisplayClusterConfigurator/Private/Views/Details/DisplayClusterConfiguratorBaseDetailCustomization.h b/Engine/Plugins/Runtime/nDisplay/Source/DisplayClusterConfigurator/Private/Views/Details/DisplayClusterConfiguratorBaseDetailCustomization.h
--- a/Engine/Plugins/Runtime/nDisplay/Source/DisplayClusterConfigurator/Private/Views/Details/DisplayClusterConfiguratorBaseDetailCustomization.h
+++ b/Engine/Plugins/Runtime/nDisplay/Source/DisplayClusterConfigurator/Private/Views/Details/DisplayClusterConfiguratorBaseDetailCustomization.h
@@ -52,6 +52,9 @@ protected:
 	/** Gets a pointer to the details layout builder used by this customization */
 	IDetailLayoutBuilder* GetLayoutBuilder() const;
 
+	/** Gets a pointer to the blueprint editor */
+	FDisplayClusterConfiguratorBlueprintEditor* GetBlueprintEditor() const;
+
 	/** Gets a pointer to the display cluster root actor being edited by the details panel */
 	ADisplayClusterRootActor* GetRootActor() const;
 
diff --git a/Engine/Plugins/Runtime/nDisplay/Source/DisplayClusterConfigurator/Private/Views/Details/DisplayClusterEditorPropertyReferenceTypeCustomization.cpp b/Engine/Plugins/Runtime/nDisplay/Source/DisplayClusterConfigurator/Private/Views/Details/DisplayClusterEditorPropertyReferenceTypeCustomization.cpp
--- a/Engine/Plugins/Runtime/nDisplay/Source/DisplayClusterConfigurator/Private/Views/Details/DisplayClusterEditorPropertyReferenceTypeCustomization.cpp
+++ b/Engine/Plugins/Runtime/nDisplay/Source/DisplayClusterConfigurator/Private/Views/Details/DisplayClusterEditorPropertyReferenceTypeCustomization.cpp
@@ -7,6 +7,8 @@
 #include "DetailLayoutBuilder.h"
 #include "IDetailGroup.h"
 #include "IDetailChildrenBuilder.h"
+#include "IPropertyTypeCustomization.h"
+#include "IPropertyUtilities.h"
 #include "Widgets/Text/STextBlock.h"
 
 const FName FDisplayClusterEditorPropertyReferenceTypeCustomization::PropertyPathMetadataKey = TEXT("PropertyPath");
@@ -19,6 +21,8 @@ TSharedRef<IPropertyTypeCustomization> FDisplayClusterEditorPropertyReferenceTyp
 
 void FDisplayClusterEditorPropertyReferenceTypeCustomization::CustomizeHeader(TSharedRef<IPropertyHandle> InPropertyHandle, FDetailWidgetRow& InHeaderRow, IPropertyTypeCustomizationUtils& CustomizationUtils)
 {
+	CustomizationUtilsPtr = &CustomizationUtils;
+
 	check(InPropertyHandle->IsValidHandle());
 	if (InPropertyHandle->HasMetaData(EditConditionPathMetadataKey))
 	{
@@ -137,6 +141,8 @@ void FDisplayClusterEditorPropertyReferenceTypeCustomization::CustomizeChildren(
 					ReferencedPropertyHandle->SetToolTipText(InPropertyHandle->GetToolTipText());
 				}
 
+				ReferencedPropertyHandle->SetOnPropertyValueChanged(FSimpleDelegate::CreateRaw(this, &FDisplayClusterEditorPropertyReferenceTypeCustomization::OnReferencedPropertyValueChanged));
+
 				IDetailPropertyRow& PropertyRow = InChildBuilder.AddProperty(ReferencedPropertyHandle.ToSharedRef());
 
 				// Mark the property with the "IsCustomized" flag so that any subsequent layout builders can account for the property
@@ -152,6 +158,15 @@ void FDisplayClusterEditorPropertyReferenceTypeCustomization::CustomizeChildren(
 	}
 }
 
+void FDisplayClusterEditorPropertyReferenceTypeCustomization::OnReferencedPropertyValueChanged()
+{
+	// When a referenced property is changed, we have to trigger layout refresh
+	if (CustomizationUtilsPtr)
+	{
+		CustomizationUtilsPtr->GetPropertyUtilities()->ForceRefresh();
+	}
+}
+
 bool FDisplayClusterEditorPropertyReferenceTypeCustomization::FindPropertyHandles(const FString& PropertyPath, TSharedRef<IPropertyHandle>& InPropertyHandle)
 {
 	bIsIteratedList = false;
diff --git a/Engine/Plugins/Runtime/nDisplay/Source/DisplayClusterConfigurator/Private/Views/Details/DisplayClusterEditorPropertyReferenceTypeCustomization.h b/Engine/Plugins/Runtime/nDisplay/Source/DisplayClusterConfigurator/Private/Views/Details/DisplayClusterEditorPropertyReferenceTypeCustomization.h
--- a/Engine/Plugins/Runtime/nDisplay/Source/DisplayClusterConfigurator/Private/Views/Details/DisplayClusterEditorPropertyReferenceTypeCustomization.h
+++ b/Engine/Plugins/Runtime/nDisplay/Source/DisplayClusterConfigurator/Private/Views/Details/DisplayClusterEditorPropertyReferenceTypeCustomization.h
@@ -5,6 +5,8 @@
 #include "CoreMinimal.h"
 #include "IPropertyTypeCustomization.h"
 
+class IPropertyTypeCustomizationUtils;
+
 
 /** A type customizer that expands a property reference to display the referenced property directly within a details panel. */
 class FDisplayClusterEditorPropertyReferenceTypeCustomization : public IPropertyTypeCustomization
@@ -33,11 +35,16 @@ private:
 	TSharedPtr<IPropertyHandle> FindRootPropertyHandle(const FString& PropertyPath, TSharedRef<IPropertyHandle>& InPropertyHandle);
 	TSharedPtr<IPropertyHandle> GetChildPropertyHandle(const FString& PropertyName, TSharedPtr<IPropertyHandle>& PropertyHandle);
 
+	void OnReferencedPropertyValueChanged();
+
 	bool IsListType(const TSharedPtr<IPropertyHandle>& PropertyHandle);
 
 	TAttribute<bool> CreateEditConditional() const;
 	bool IsEditable() const;
 
+	// Customization utils instance to trigger GUI update on referenced value changes
+	IPropertyTypeCustomizationUtils* CustomizationUtilsPtr = nullptr;
+
 	TArray<FPropertyHandlePair> ReferencedPropertyHandles;
 	TArray<TSharedRef<IPropertyHandle>> EditConditionPropertyHandles;
 	bool bIsIteratedList;
diff --git a/Engine/Plugins/Runtime/nDisplay/Source/DisplayClusterMedia/DisplayClusterMedia.Build.cs b/Engine/Plugins/Runtime/nDisplay/Source/DisplayClusterMedia/DisplayClusterMedia.Build.cs
new file mode 100644
--- /dev/null
+++ b/Engine/Plugins/Runtime/nDisplay/Source/DisplayClusterMedia/DisplayClusterMedia.Build.cs
@@ -0,0 +1,34 @@
+// Copyright Epic Games, Inc. All Rights Reserved.
+
+using UnrealBuildTool;
+using System.IO;
+
+public class DisplayClusterMedia : ModuleRules
+{
+	public DisplayClusterMedia(ReadOnlyTargetRules ROTargetRules) : base(ROTargetRules)
+	{
+		// [temporary] We need this to be able to use some private data types. This should
+		// be removed once we move the nD rendering pipeline to RDG.
+		string EngineDir = Path.GetFullPath(Target.RelativeEnginePath);
+		PrivateIncludePaths.Add(Path.Combine(EngineDir, "Source", "Runtime", "Renderer", "Private"));
+
+		PrivateDependencyModuleNames.AddRange(
+			new string[] {
+				"Core",
+				"CoreUObject",
+				"DisplayCluster",
+				"DisplayClusterConfiguration",
+				"Engine",
+				"MediaAssets",
+				"MediaIOCore",
+				"Renderer",
+				"RenderCore",
+				"RHI",
+			});
+
+		if (Target.bBuildEditor == true)
+		{
+			PrivateDependencyModuleNames.Add("UnrealEd");
+		}
+	}
+}
diff --git a/Engine/Plugins/Runtime/nDisplay/Source/DisplayClusterMedia/Private/Capture/DisplayClusterMediaCaptureBase.cpp b/Engine/Plugins/Runtime/nDisplay/Source/DisplayClusterMedia/Private/Capture/DisplayClusterMediaCaptureBase.cpp
new file mode 100644
--- /dev/null
+++ b/Engine/Plugins/Runtime/nDisplay/Source/DisplayClusterMedia/Private/Capture/DisplayClusterMediaCaptureBase.cpp
@@ -0,0 +1,99 @@
+// Copyright Epic Games, Inc. All Rights Reserved.
+
+#include "Capture/DisplayClusterMediaCaptureBase.h"
+#include "DisplayClusterMediaHelpers.h"
+#include "DisplayClusterMediaLog.h"
+
+#include "MediaCapture.h"
+#include "MediaOutput.h"
+#include "Engine/TextureRenderTarget2D.h"
+
+#include "RHIResources.h"
+
+
+FDisplayClusterMediaCaptureBase::FDisplayClusterMediaCaptureBase(const FString& InMediaId, const FString& InClusterNodeId, UMediaOutput* InMediaOutput, UTextureRenderTarget2D* InRenderTarget)
+	: FDisplayClusterMediaBase(InMediaId, InClusterNodeId)
+	, MediaOutput(InMediaOutput)
+	, RenderTarget(InRenderTarget)
+{
+	// We expect these to always be valid
+	check(InMediaOutput && InRenderTarget);
+}
+
+
+void FDisplayClusterMediaCaptureBase::AddReferencedObjects(FReferenceCollector& Collector)
+{
+	Collector.AddReferencedObject(MediaOutput);
+	Collector.AddReferencedObject(RenderTarget);
+
+	if (MediaCapture)
+	{
+		Collector.AddReferencedObject(MediaCapture);
+	}
+}
+
+bool FDisplayClusterMediaCaptureBase::StartCapture()
+{
+	if (MediaOutput && RenderTarget && !MediaCapture)
+	{
+		MediaCapture = MediaOutput->CreateMediaCapture();
+		if (MediaCapture)
+		{
+			MediaCapture->SetMediaOutput(MediaOutput);
+
+			FMediaCaptureOptions MediaCaptureOptions;
+			MediaCaptureOptions.NumberOfFramesToCapture = -1;
+			MediaCaptureOptions.bResizeSourceBuffer = true;
+			MediaCaptureOptions.bSkipFrameWhenRunningExpensiveTasks = false;
+
+			const bool bCaptureStarted = MediaCapture->CaptureTextureRenderTarget2D(RenderTarget, MediaCaptureOptions);
+			if (bCaptureStarted)
+			{
+				UE_LOG(LogDisplayClusterMedia, Log, TEXT("Started media capture: '%s'"), *GetMediaId());
+			}
+			else
+			{
+				UE_LOG(LogDisplayClusterMedia, Warning, TEXT("Couldn't start media capture '%s'"), *GetMediaId());
+			}
+
+			return bCaptureStarted;
+		}
+	}
+
+	return false;
+}
+
+void FDisplayClusterMediaCaptureBase::StopCapture()
+{
+	if (MediaCapture)
+	{
+		MediaCapture->StopCapture(false);
+		MediaCapture = nullptr;
+	}
+}
+
+void FDisplayClusterMediaCaptureBase::ExportMediaData(FRHICommandListImmediate& RHICmdList, const FMediaTextureInfo& TextureInfo)
+{
+	FRHITexture* const SrcTexture = TextureInfo.Texture;
+	FRHITexture* const DstTexture = GetRenderTarget()->GetRenderTargetResource()->GetTextureRHI();
+
+	if (SrcTexture && DstTexture)
+	{
+		const FIntPoint SrcRegionSize = TextureInfo.Region.Size();
+
+		if (SrcTexture->GetDesc().Format == DstTexture->GetDesc().Format &&
+			SrcRegionSize == DstTexture->GetDesc().Extent)
+		{
+			FRHICopyTextureInfo CopyInfo;
+
+			CopyInfo.SourcePosition = FIntVector(TextureInfo.Region.Min.X, TextureInfo.Region.Min.Y, 0);
+			CopyInfo.Size = FIntVector(SrcRegionSize.X, SrcRegionSize.Y, 0);
+
+			TransitionAndCopyTexture(RHICmdList, SrcTexture, DstTexture, CopyInfo);
+		}
+		else
+		{
+			DisplayClusterMediaHelpers::ResampleTexture_RenderThread(RHICmdList, SrcTexture, DstTexture, TextureInfo.Region, FIntRect(FIntPoint::ZeroValue, DstTexture->GetDesc().Extent));
+		}
+	}
+}
diff --git a/Engine/Plugins/Runtime/nDisplay/Source/DisplayClusterMedia/Private/Capture/DisplayClusterMediaCaptureBase.h b/Engine/Plugins/Runtime/nDisplay/Source/DisplayClusterMedia/Private/Capture/DisplayClusterMediaCaptureBase.h
new file mode 100644
--- /dev/null
+++ b/Engine/Plugins/Runtime/nDisplay/Source/DisplayClusterMedia/Private/Capture/DisplayClusterMediaCaptureBase.h
@@ -0,0 +1,52 @@
+// Copyright Epic Games, Inc. All Rights Reserved.
+
+#pragma once
+
+#include "CoreMinimal.h"
+#include "DisplayClusterMediaBase.h"
+#include "UObject/GCObject.h"
+
+class FRHICommandListImmediate;
+class UMediaCapture;
+class UMediaOutput;
+class UTextureRenderTarget2D;
+
+
+/**
+ * Base media capture class
+ */
+class FDisplayClusterMediaCaptureBase
+	: public FDisplayClusterMediaBase
+	, public FGCObject
+{
+public:
+	FDisplayClusterMediaCaptureBase(const FString& MediaId, const FString& ClusterNodeId, UMediaOutput* MediaOutput, UTextureRenderTarget2D* RenderTarget);
+
+public:
+	//~ Begin FGCObject interface
+	virtual void AddReferencedObjects(FReferenceCollector& Collector) override;
+	virtual FString GetReferencerName() const override
+	{
+		return TEXT("FDisplayClusterMediaCaptureBase");
+	}
+	//~ End FGCObject interface
+
+public:
+	virtual bool StartCapture();
+	virtual void StopCapture();
+
+	UTextureRenderTarget2D* GetRenderTarget()
+	{
+		return RenderTarget;
+	}
+
+protected:
+	void ExportMediaData(FRHICommandListImmediate& RHICmdList, const FMediaTextureInfo& TextureInfo);
+
+private:
+	//~ Begin GC by AddReferencedObjects
+	UMediaOutput*           MediaOutput  = nullptr;
+	UMediaCapture*          MediaCapture = nullptr;
+	UTextureRenderTarget2D* RenderTarget = nullptr;
+	//~ End GC by AddReferencedObjects
+};
diff --git a/Engine/Plugins/Runtime/nDisplay/Source/DisplayClusterMedia/Private/Capture/DisplayClusterMediaCaptureNode.cpp b/Engine/Plugins/Runtime/nDisplay/Source/DisplayClusterMedia/Private/Capture/DisplayClusterMediaCaptureNode.cpp
new file mode 100644
--- /dev/null
+++ b/Engine/Plugins/Runtime/nDisplay/Source/DisplayClusterMedia/Private/Capture/DisplayClusterMediaCaptureNode.cpp
@@ -0,0 +1,56 @@
+// Copyright Epic Games, Inc. All Rights Reserved.
+
+#include "Capture/DisplayClusterMediaCaptureNode.h"
+
+#include "IDisplayCluster.h"
+#include "IDisplayClusterCallbacks.h"
+
+#include "Render/Viewport/IDisplayClusterViewportManagerProxy.h"
+#include "RHICommandList.h"
+#include "RHIResources.h"
+
+#include "UnrealClient.h"
+
+
+FDisplayClusterMediaCaptureNode::FDisplayClusterMediaCaptureNode(const FString& InMediaId, const FString& InClusterNodeId, UMediaOutput* InMediaOutput, UTextureRenderTarget2D* InRenderTarget)
+	: FDisplayClusterMediaCaptureBase(InMediaId, InClusterNodeId, InMediaOutput, InRenderTarget)
+{
+}
+
+
+bool FDisplayClusterMediaCaptureNode::StartCapture()
+{
+	// If capturing initialized and started successfully, subscribe for rendering callbacks
+	if (FDisplayClusterMediaCaptureBase::StartCapture())
+	{
+		IDisplayCluster::Get().GetCallbacks().OnDisplayClusterPostFrameRender_RenderThread().AddRaw(this, &FDisplayClusterMediaCaptureNode::OnPostFrameRender_RenderThread);
+		return true;
+	}
+
+	return false;
+}
+
+void FDisplayClusterMediaCaptureNode::StopCapture()
+{
+	// Stop rendering notifications
+	IDisplayCluster::Get().GetCallbacks().OnDisplayClusterPostFrameRender_RenderThread().RemoveAll(this);
+	// Stop capturing
+	FDisplayClusterMediaCaptureBase::StopCapture();
+}
+
+void FDisplayClusterMediaCaptureNode::OnPostFrameRender_RenderThread(FRHICommandListImmediate& RHICmdList, const IDisplayClusterViewportManagerProxy* ViewportManagerProxy, FViewport* Viewport)
+{
+	ensure(ViewportManagerProxy);
+
+	TArray<FRHITexture*> Textures;
+	TArray<FIntPoint>    Regions;
+
+	if (ViewportManagerProxy && ViewportManagerProxy->GetFrameTargets_RenderThread(Textures, Regions))
+	{
+		if (Textures.Num() > 0 && Regions.Num() > 0 && Textures[0])
+		{
+			FMediaTextureInfo TextureInfo{ Textures[0], FIntRect(Regions[0], Regions[0] + Textures[0]->GetDesc().Extent) };
+			ExportMediaData(RHICmdList, TextureInfo);
+		}
+	}
+}
diff --git a/Engine/Plugins/Runtime/nDisplay/Source/DisplayClusterMedia/Private/Capture/DisplayClusterMediaCaptureNode.h b/Engine/Plugins/Runtime/nDisplay/Source/DisplayClusterMedia/Private/Capture/DisplayClusterMediaCaptureNode.h
new file mode 100644
--- /dev/null
+++ b/Engine/Plugins/Runtime/nDisplay/Source/DisplayClusterMedia/Private/Capture/DisplayClusterMediaCaptureNode.h
@@ -0,0 +1,27 @@
+// Copyright Epic Games, Inc. All Rights Reserved.
+
+#pragma once
+
+#include "Capture/DisplayClusterMediaCaptureBase.h"
+
+class FRHICommandListImmediate;
+class FViewport;
+class IDisplayClusterViewportManagerProxy;
+
+
+/**
+ * Node backbuffer media capture
+ */
+class FDisplayClusterMediaCaptureNode
+	: public FDisplayClusterMediaCaptureBase
+{
+public:
+	FDisplayClusterMediaCaptureNode(const FString& MediaId, const FString& ClusterNodeId, UMediaOutput* MediaOutput, UTextureRenderTarget2D* RenderTarget);
+
+public:
+	virtual bool StartCapture() override;
+	virtual void StopCapture() override;
+
+private:
+	void OnPostFrameRender_RenderThread(FRHICommandListImmediate& RHICmdList, const IDisplayClusterViewportManagerProxy* ViewportManagerProxy, FViewport* Viewport);
+};
diff --git a/Engine/Plugins/Runtime/nDisplay/Source/DisplayClusterMedia/Private/Capture/DisplayClusterMediaCaptureViewport.cpp b/Engine/Plugins/Runtime/nDisplay/Source/DisplayClusterMedia/Private/Capture/DisplayClusterMediaCaptureViewport.cpp
new file mode 100644
--- /dev/null
+++ b/Engine/Plugins/Runtime/nDisplay/Source/DisplayClusterMedia/Private/Capture/DisplayClusterMediaCaptureViewport.cpp
@@ -0,0 +1,59 @@
+// Copyright Epic Games, Inc. All Rights Reserved.
+
+#include "Capture/DisplayClusterMediaCaptureViewport.h"
+
+#include "IDisplayCluster.h"
+#include "IDisplayClusterCallbacks.h"
+
+#include "Render/Viewport/IDisplayClusterViewportManagerProxy.h"
+#include "RHICommandList.h"
+#include "RHIResources.h"
+
+
+FDisplayClusterMediaCaptureViewport::FDisplayClusterMediaCaptureViewport(const FString& InMediaId, const FString& InClusterNodeId, const FString& InViewportId, UMediaOutput* InMediaOutput, UTextureRenderTarget2D* InRenderTarget)
+	: FDisplayClusterMediaCaptureBase(InMediaId, InClusterNodeId, InMediaOutput, InRenderTarget)
+	, ViewportId(InViewportId)
+{
+}
+
+
+bool FDisplayClusterMediaCaptureViewport::StartCapture()
+{
+	// If capturing has started successfully, subscribe for rendering callbacks
+	if (FDisplayClusterMediaCaptureBase::StartCapture())
+	{
+		IDisplayCluster::Get().GetCallbacks().OnDisplayClusterPostRenderViewFamily_RenderThread().AddRaw(this, &FDisplayClusterMediaCaptureViewport::OnPostRenderViewFamily_RenderThread);
+		return true;
+	}
+
+	return false;
+}
+
+void FDisplayClusterMediaCaptureViewport::StopCapture()
+{
+	// Stop rendering notifications
+	IDisplayCluster::Get().GetCallbacks().OnDisplayClusterPostRenderViewFamily_RenderThread().RemoveAll(this);
+	// Stop capturing
+	FDisplayClusterMediaCaptureBase::StopCapture();
+}
+
+void FDisplayClusterMediaCaptureViewport::OnPostRenderViewFamily_RenderThread(FRHICommandListImmediate& RHICmdList, const FSceneViewFamily& ViewFamily, const IDisplayClusterViewportProxy* ViewportProxy)
+{
+	ensure(ViewportProxy);
+
+	if (ViewportProxy && ViewportProxy->GetId().Equals(GetViewportId(), ESearchCase::IgnoreCase))
+	{
+		TArray<FRHITexture*> Textures;
+		TArray<FIntRect>     Regions;
+
+		// Get RHI texture
+		if (ViewportProxy->GetResourcesWithRects_RenderThread(EDisplayClusterViewportResourceType::InternalRenderTargetResource, Textures, Regions))
+		{
+			if (Textures.Num() > 0 && Regions.Num() > 0)
+			{
+				FMediaTextureInfo TextureInfo{ Textures[0], Regions[0] };
+				ExportMediaData(RHICmdList, TextureInfo);
+			}
+		}
+	}
+}
diff --git a/Engine/Plugins/Runtime/nDisplay/Source/DisplayClusterMedia/Private/Capture/DisplayClusterMediaCaptureViewport.h b/Engine/Plugins/Runtime/nDisplay/Source/DisplayClusterMedia/Private/Capture/DisplayClusterMediaCaptureViewport.h
new file mode 100644
--- /dev/null
+++ b/Engine/Plugins/Runtime/nDisplay/Source/DisplayClusterMedia/Private/Capture/DisplayClusterMediaCaptureViewport.h
@@ -0,0 +1,36 @@
+// Copyright Epic Games, Inc. All Rights Reserved.
+
+#pragma once
+
+#include "CoreMinimal.h"
+#include "Capture/DisplayClusterMediaCaptureBase.h"
+
+class FRHICommandListImmediate;
+class FSceneViewFamily;
+class IDisplayClusterViewportProxy;
+
+
+/**
+ * Viewport media capture
+ */
+class FDisplayClusterMediaCaptureViewport
+	: public FDisplayClusterMediaCaptureBase
+{
+public:
+	FDisplayClusterMediaCaptureViewport(const FString& MediaId, const FString& ClusterNodeId, const FString& ViewportId, UMediaOutput* MediaOutput, UTextureRenderTarget2D* RenderTarget);
+
+public:
+	virtual bool StartCapture() override;
+	virtual void StopCapture() override;
+
+	const FString& GetViewportId() const
+	{
+		return ViewportId;
+	}
+
+private:
+	void OnPostRenderViewFamily_RenderThread(FRHICommandListImmediate& RHICmdList, const FSceneViewFamily& ViewFamily, const IDisplayClusterViewportProxy* ViewportProxy);
+
+private:
+	const FString ViewportId;
+};
diff --git a/Engine/Plugins/Runtime/nDisplay/Source/DisplayClusterMedia/Private/DisplayClusterMediaBase.h b/Engine/Plugins/Runtime/nDisplay/Source/DisplayClusterMedia/Private/DisplayClusterMediaBase.h
new file mode 100644
--- /dev/null
+++ b/Engine/Plugins/Runtime/nDisplay/Source/DisplayClusterMedia/Private/DisplayClusterMediaBase.h
@@ -0,0 +1,44 @@
+// Copyright Epic Games, Inc. All Rights Reserved.
+
+#pragma once
+
+#include "CoreMinimal.h"
+
+class FRHITexture;
+
+
+/**
+ * Base media class
+ */
+class FDisplayClusterMediaBase
+{
+public:
+	FDisplayClusterMediaBase(const FString& InMediaId, const FString& InClusterNodeId)
+		: MediaId(InMediaId)
+		, ClusterNodeId(InClusterNodeId)
+	{ }
+
+	virtual ~FDisplayClusterMediaBase() = default;
+
+public:
+	const FString& GetMediaId() const
+	{
+		return MediaId;
+	}
+
+	const FString& GetClusterNodeId() const
+	{
+		return ClusterNodeId;
+	}
+
+protected:
+	struct FMediaTextureInfo
+	{
+		FRHITexture* Texture = nullptr;
+		FIntRect     Region;
+	};
+
+private:
+	const FString MediaId;
+	const FString ClusterNodeId;
+};
diff --git a/Engine/Plugins/Runtime/nDisplay/Source/DisplayClusterMedia/Private/DisplayClusterMediaHelpers.cpp b/Engine/Plugins/Runtime/nDisplay/Source/DisplayClusterMedia/Private/DisplayClusterMediaHelpers.cpp
new file mode 100644
--- /dev/null
+++ b/Engine/Plugins/Runtime/nDisplay/Source/DisplayClusterMedia/Private/DisplayClusterMediaHelpers.cpp
@@ -0,0 +1,83 @@
+// Copyright Epic Games, Inc. All Rights Reserved.
+
+#include "DisplayClusterMediaHelpers.h"
+
+#include "CommonRenderResources.h"
+#include "PixelShaderUtils.h"
+#include "ScreenRendering.h"
+#include "PostProcess/SceneFilterRendering.h"
+
+
+namespace DisplayClusterMediaHelpers
+{
+	//@todo This needs to be exposed from the DisplayCluster core module after its refactoring
+	FString GenerateICVFXViewportName(const FString& ClusterNodeId, const FString& ICVFXCameraName)
+	{
+		return FString::Printf(TEXT("%s_icvfx_%s_incamera"), *ClusterNodeId, *ICVFXCameraName);
+	}
+
+	template<class TScreenPixelShader>
+	void ResampleTextureImpl_RenderThread(FRHICommandListImmediate& RHICmdList, FRHITexture* SrcTexture, FRHITexture* DstTexture, const FIntRect& SrcRect, const FIntRect& DstRect)
+	{
+		FRHIRenderPassInfo RPInfo(DstTexture, ERenderTargetActions::Load_Store);
+		RHICmdList.Transition(FRHITransitionInfo(DstTexture, ERHIAccess::Unknown, ERHIAccess::RTV));
+
+		RHICmdList.BeginRenderPass(RPInfo, TEXT("DisaplyClusterMedia_ResampleTexture"));
+		{
+			FIntVector SrcSizeXYZ = SrcTexture->GetSizeXYZ();
+			FIntVector DstSizeXYZ = DstTexture->GetSizeXYZ();
+
+			FIntPoint SrcSize(SrcSizeXYZ.X, SrcSizeXYZ.Y);
+			FIntPoint DstSize(DstSizeXYZ.X, DstSizeXYZ.Y);
+
+			RHICmdList.SetViewport(0.f, 0.f, 0.0f, DstSize.X, DstSize.Y, 1.0f);
+
+			FGraphicsPipelineStateInitializer GraphicsPSOInit;
+			RHICmdList.ApplyCachedRenderTargets(GraphicsPSOInit);
+
+			GraphicsPSOInit.BlendState = TStaticBlendState<>::GetRHI();
+			GraphicsPSOInit.RasterizerState = TStaticRasterizerState<>::GetRHI();
+			GraphicsPSOInit.DepthStencilState = TStaticDepthStencilState<false, CF_Always>::GetRHI();
+
+
+			FGlobalShaderMap* ShaderMap = GetGlobalShaderMap(GMaxRHIFeatureLevel);
+			TShaderMapRef<FScreenVS> VertexShader(ShaderMap);
+			TShaderMapRef<TScreenPixelShader> PixelShader(ShaderMap);
+
+			GraphicsPSOInit.BoundShaderState.VertexDeclarationRHI = GFilterVertexDeclaration.VertexDeclarationRHI;
+			GraphicsPSOInit.BoundShaderState.VertexShaderRHI = VertexShader.GetVertexShader();
+			GraphicsPSOInit.BoundShaderState.PixelShaderRHI = PixelShader.GetPixelShader();
+			GraphicsPSOInit.PrimitiveType = PT_TriangleList;
+
+			SetGraphicsPipelineState(RHICmdList, GraphicsPSOInit, 0);
+
+			if (SrcRect.Size() != DstRect.Size())
+			{
+				PixelShader->SetParameters(RHICmdList, TStaticSamplerState<SF_Bilinear>::GetRHI(), SrcTexture);
+			}
+			else
+			{
+				PixelShader->SetParameters(RHICmdList, TStaticSamplerState<SF_Point>::GetRHI(), SrcTexture);
+			}
+
+			// Set up vertex uniform parameters for scaling and biasing the rectangle.
+			// Note: Use DrawRectangle in the vertex shader to calculate the correct vertex position and uv.
+			FDrawRectangleParameters Parameters;
+			Parameters.PosScaleBias = FVector4f(DstRect.Size().X, DstRect.Size().Y, DstRect.Min.X, DstRect.Min.Y);
+			Parameters.UVScaleBias = FVector4f(SrcRect.Size().X, SrcRect.Size().Y, SrcRect.Min.X, SrcRect.Min.Y);
+			Parameters.InvTargetSizeAndTextureSize = FVector4f(1.0f / DstSize.X, 1.0f / DstSize.Y, 1.0f / SrcSize.X, 1.0f / SrcSize.Y);
+
+			SetUniformBufferParameterImmediate(RHICmdList, VertexShader.GetVertexShader(), VertexShader->GetUniformBufferParameter<FDrawRectangleParameters>(), Parameters);
+
+			FPixelShaderUtils::DrawFullscreenQuad(RHICmdList, 1);
+		}
+
+		RHICmdList.EndRenderPass();
+		RHICmdList.Transition(FRHITransitionInfo(DstTexture, ERHIAccess::Unknown, ERHIAccess::SRVMask));
+	}
+
+	void ResampleTexture_RenderThread(FRHICommandListImmediate& RHICmdList, FRHITexture* SrcTexture, FRHITexture* DstTexture, const FIntRect& SrcRect, const FIntRect& DstRect)
+	{
+		ResampleTextureImpl_RenderThread<FScreenPS>(RHICmdList, SrcTexture, DstTexture, SrcRect, DstRect);
+	}
+}
diff --git a/Engine/Plugins/Runtime/nDisplay/Source/DisplayClusterMedia/Private/DisplayClusterMediaHelpers.h b/Engine/Plugins/Runtime/nDisplay/Source/DisplayClusterMedia/Private/DisplayClusterMediaHelpers.h
new file mode 100644
--- /dev/null
+++ b/Engine/Plugins/Runtime/nDisplay/Source/DisplayClusterMedia/Private/DisplayClusterMediaHelpers.h
@@ -0,0 +1,18 @@
+// Copyright Epic Games, Inc. All Rights Reserved.
+
+#pragma once
+
+#include "CoreMinimal.h"
+
+class FRHICommandListImmediate;
+class FRHITexture;
+
+
+namespace DisplayClusterMediaHelpers
+{
+	// Generates internal ICVFX viewport IDs
+	FString GenerateICVFXViewportName(const FString& ClusterNodeId, const FString& ICVFXCameraName);
+
+	// Copy and resize an RHI texture
+	void ResampleTexture_RenderThread(FRHICommandListImmediate& RHICmdList, FRHITexture* SrcTexture, FRHITexture* DstTexture, const FIntRect& SrcRect, const FIntRect& DstRect);
+}
diff --git a/Engine/Plugins/Runtime/nDisplay/Source/DisplayClusterMedia/Private/DisplayClusterMediaLog.cpp b/Engine/Plugins/Runtime/nDisplay/Source/DisplayClusterMedia/Private/DisplayClusterMediaLog.cpp
new file mode 100644
--- /dev/null
+++ b/Engine/Plugins/Runtime/nDisplay/Source/DisplayClusterMedia/Private/DisplayClusterMediaLog.cpp
@@ -0,0 +1,7 @@
+// Copyright Epic Games, Inc. All Rights Reserved.
+
+#include "DisplayClusterMediaLog.h"
+
+
+// Module-wide log categories
+DEFINE_LOG_CATEGORY(LogDisplayClusterMedia);
diff --git a/Engine/Plugins/Runtime/nDisplay/Source/DisplayClusterMedia/Private/DisplayClusterMediaLog.h b/Engine/Plugins/Runtime/nDisplay/Source/DisplayClusterMedia/Private/DisplayClusterMediaLog.h
new file mode 100644
--- /dev/null
+++ b/Engine/Plugins/Runtime/nDisplay/Source/DisplayClusterMedia/Private/DisplayClusterMediaLog.h
@@ -0,0 +1,9 @@
+// Copyright Epic Games, Inc. All Rights Reserved.
+
+#pragma once
+
+#include "CoreMinimal.h"
+
+
+// Module-wide log categories
+DECLARE_LOG_CATEGORY_EXTERN(LogDisplayClusterMedia, Log, All);
diff --git a/Engine/Plugins/Runtime/nDisplay/Source/DisplayClusterMedia/Private/DisplayClusterMediaModule.cpp b/Engine/Plugins/Runtime/nDisplay/Source/DisplayClusterMedia/Private/DisplayClusterMediaModule.cpp
new file mode 100644
--- /dev/null
+++ b/Engine/Plugins/Runtime/nDisplay/Source/DisplayClusterMedia/Private/DisplayClusterMediaModule.cpp
@@ -0,0 +1,346 @@
+// Copyright Epic Games, Inc. All Rights Reserved.
+
+#include "DisplayClusterMediaModule.h"
+#include "DisplayClusterMediaLog.h"
+#include "DisplayClusterMediaHelpers.h"
+
+#include "IDisplayCluster.h"
+#include "IDisplayClusterCallbacks.h"
+#include "DisplayClusterEnums.h"
+#include "DisplayClusterRootActor.h"
+#include "DisplayClusterConfigurationTypes.h"
+
+#include "Components/DisplayClusterICVFXCameraComponent.h"
+#include "Config/IDisplayClusterConfigManager.h"
+#include "Cluster/IDisplayClusterClusterManager.h"
+#include "Game/IDisplayClusterGameManager.h"
+
+#include "Capture/DisplayClusterMediaCaptureNode.h"
+#include "Capture/DisplayClusterMediaCaptureViewport.h"
+#include "Input/DisplayClusterMediaInputNode.h"
+#include "Input/DisplayClusterMediaInputViewport.h"
+
+#include "Misc/CoreDelegates.h"
+
+
+void FDisplayClusterMediaModule::StartupModule()
+{
+	UE_LOG(LogDisplayClusterMedia, Log, TEXT("Starting module 'DisplayClusterMedia'..."));
+
+	IDisplayCluster::Get().GetCallbacks().OnDisplayClusterCustomPresentSet().AddRaw(this, &FDisplayClusterMediaModule::OnCustomPresentSet);
+	FCoreDelegates::OnEnginePreExit.AddRaw(this, &FDisplayClusterMediaModule::OnEnginePreExit);
+}
+
+void FDisplayClusterMediaModule::ShutdownModule()
+{
+	UE_LOG(LogDisplayClusterMedia, Log, TEXT("Shutting down module 'DisplayClusterMedia'..."));
+
+	IDisplayCluster::Get().GetCallbacks().OnDisplayClusterCustomPresentSet().RemoveAll(this);
+	FCoreDelegates::OnEnginePreExit.RemoveAll(this);
+}
+
+void FDisplayClusterMediaModule::OnCustomPresentSet()
+{
+	// We initialize and start media when backbuffer is available. This CustomPresentSet event is a
+	// sign that the backbuffer is already available. This allows us to prevent any potential problems
+	// if any of the following is used under the hood, or will be used in the future:
+	// - UMediaCapture::CaptureActiveSceneViewport
+	// - UMediaCapture::CaptureSceneViewport
+	InitializeMedia();
+	StartCapture();
+	PlayMedia();
+}
+
+void FDisplayClusterMediaModule::OnEnginePreExit()
+{
+	ReleaseMedia();
+}
+
+void FDisplayClusterMediaModule::InitializeMedia()
+{
+	// Runtime only for now
+	if (IDisplayCluster::Get().GetOperationMode() != EDisplayClusterOperationMode::Cluster)
+	{
+		UE_LOG(LogDisplayClusterMedia, Warning, TEXT("DisplayClusterMedia is available in 'cluster' operation mode only"));
+		return;
+	}
+
+	// Parse DCRA configuration and initialize media
+	if (const ADisplayClusterRootActor* const RootActor = IDisplayCluster::Get().GetGameMgr()->GetRootActor())
+	{
+		const FString ClusterNodeId = IDisplayCluster::Get().GetClusterMgr()->GetNodeId();
+		if (const UDisplayClusterConfigurationClusterNode* const ClusterNode = RootActor->GetConfigData()->Cluster->GetNode(ClusterNodeId))
+		{
+			///////////////////////////////
+			// Node backbuffer media setup
+			if (ClusterNode->Media.bEnabled)
+			{
+				const FDisplayClusterConfigurationMedia& MediaSettings = ClusterNode->Media;
+
+				const FString MediaCaptureId = RootActor->GetName() + FString("_") + ClusterNodeId + FString("_backbuffer_capture");
+				const FString MediaInputId   = RootActor->GetName() + FString("_") + ClusterNodeId + FString("_backbuffer_input");
+				const bool bInputDataValid   = MediaSettings.MediaInput.MediaPlayer && MediaSettings.MediaInput.MediaSource && MediaSettings.MediaInput.MediaTexture;
+				const bool bOutputDataValid  = MediaSettings.MediaOutput.MediaOutput && MediaSettings.MediaOutput.RenderTarget;
+
+				// Media sharing
+				if (MediaSettings.IsMediaSharingRequired() && bInputDataValid && bOutputDataValid)
+				{
+					const bool bThisNodeSharesMedia = ClusterNodeId.Equals(MediaSettings.MediaOutputNode, ESearchCase::IgnoreCase);
+
+					// Share media
+					if (bThisNodeSharesMedia)
+					{
+						UE_LOG(LogDisplayClusterMedia, Log, TEXT("Initializing backbuffer media capture for node '%s' (in-cluster shared)..."), *ClusterNodeId);
+
+						CaptureNode = MakeUnique<FDisplayClusterMediaCaptureNode>(
+							MediaCaptureId, ClusterNodeId,
+							MediaSettings.MediaOutput.MediaOutput, MediaSettings.MediaOutput.RenderTarget);
+
+						AllCaptures.Add(CaptureNode.Get());
+					}
+					// Consume media
+					else
+					{
+						UE_LOG(LogDisplayClusterMedia, Log, TEXT("Initializing backbuffer media input for node '%s' (in-cluster shared)..."), *ClusterNodeId);
+
+						InputNode = MakeUnique<FDisplayClusterMediaInputNode>(
+							MediaInputId, ClusterNodeId,
+							MediaSettings.MediaInput.MediaSource, MediaSettings.MediaInput.MediaPlayer, MediaSettings.MediaInput.MediaTexture);
+
+						AllInputs.Add(InputNode.Get());
+					}
+				}
+				// Media input only
+				else if (MediaSettings.MediaInput.bEnabled && bInputDataValid)
+				{
+					UE_LOG(LogDisplayClusterMedia, Log, TEXT("Initializing backbuffer media input for node '%s'..."), *ClusterNodeId);
+
+					InputNode = MakeUnique<FDisplayClusterMediaInputNode>(
+						MediaInputId, ClusterNodeId,
+						MediaSettings.MediaInput.MediaSource, MediaSettings.MediaInput.MediaPlayer, MediaSettings.MediaInput.MediaTexture);
+
+					AllInputs.Add(InputNode.Get());
+				}
+				// Media capture only
+				else if (MediaSettings.MediaOutput.bEnabled && bOutputDataValid)
+				{
+					UE_LOG(LogDisplayClusterMedia, Log, TEXT("Initializing backbuffer media capture for node '%s'..."), *ClusterNodeId);
+
+					CaptureNode = MakeUnique<FDisplayClusterMediaCaptureNode>(
+						MediaCaptureId, ClusterNodeId,
+						MediaSettings.MediaOutput.MediaOutput, MediaSettings.MediaOutput.RenderTarget);
+
+					AllCaptures.Add(CaptureNode.Get());
+				}
+			}
+
+			///////////////////////////////
+			// Viewports media setup
+			for (const TPair<FString, UDisplayClusterConfigurationViewport*>& ViewportIt : ClusterNode->Viewports)
+			{
+				if (const UDisplayClusterConfigurationViewport* const Viewport = ViewportIt.Value)
+				{
+					const FDisplayClusterConfigurationMedia& MediaSettings = Viewport->RenderSettings.Media;
+					if (MediaSettings.bEnabled)
+					{
+						const FString MediaCaptureId = RootActor->GetName() + FString("_") + ClusterNodeId + FString("_") + ViewportIt.Key + FString("_viewport_capture");
+						const FString MediaInputId   = RootActor->GetName() + FString("_") + ClusterNodeId + FString("_") + ViewportIt.Key + FString("_viewport_input");
+						const bool bInputDataValid   = MediaSettings.MediaInput.MediaPlayer && MediaSettings.MediaInput.MediaSource && MediaSettings.MediaInput.MediaTexture;
+						const bool bOutputDataValid  = MediaSettings.MediaOutput.MediaOutput && MediaSettings.MediaOutput.RenderTarget;
+
+						// Media sharing
+						if (MediaSettings.IsMediaSharingRequired() && bInputDataValid && bOutputDataValid)
+						{
+							const bool bThisNodeSharesMedia = ClusterNodeId.Equals(MediaSettings.MediaOutputNode, ESearchCase::IgnoreCase);
+
+							// Share media
+							if (bThisNodeSharesMedia)
+							{
+								UE_LOG(LogDisplayClusterMedia, Log, TEXT("Initializing viewport media capture for viewport '%s' (in-cluster shared)..."), *ClusterNodeId);
+
+								TUniquePtr<FDisplayClusterMediaCaptureViewport> NewViewportCapture = MakeUnique<FDisplayClusterMediaCaptureViewport>(
+									MediaCaptureId, ClusterNodeId,
+									ViewportIt.Key,
+									MediaSettings.MediaOutput.MediaOutput, MediaSettings.MediaOutput.RenderTarget);
+
+								AllCaptures.Add(NewViewportCapture.Get());
+								CaptureViewports.Emplace(MediaCaptureId, MoveTemp(NewViewportCapture));
+							}
+							// Consume media
+							else
+							{
+								UE_LOG(LogDisplayClusterMedia, Log, TEXT("Initializing viewport media input for viewport '%s' (in-cluster shared)..."), *ClusterNodeId);
+
+								TUniquePtr<FDisplayClusterMediaInputViewport> NewViewportInput = MakeUnique<FDisplayClusterMediaInputViewport>(
+									MediaInputId, ClusterNodeId,
+									ViewportIt.Key,
+									MediaSettings.MediaInput.MediaSource, MediaSettings.MediaInput.MediaPlayer, MediaSettings.MediaInput.MediaTexture);
+
+								AllInputs.Add(NewViewportInput.Get());
+								InputViewports.Emplace(MediaCaptureId, MoveTemp(NewViewportInput));
+							}
+						}
+						// Media input only
+						else if (MediaSettings.MediaInput.bEnabled && bInputDataValid)
+						{
+							UE_LOG(LogDisplayClusterMedia, Log, TEXT("Initializing viewport media input for viewport '%s'..."), *ClusterNodeId);
+
+							TUniquePtr<FDisplayClusterMediaInputViewport> NewViewportInput = MakeUnique<FDisplayClusterMediaInputViewport>(
+								MediaInputId, ClusterNodeId,
+								ViewportIt.Key,
+								MediaSettings.MediaInput.MediaSource, MediaSettings.MediaInput.MediaPlayer, MediaSettings.MediaInput.MediaTexture);
+
+							AllInputs.Add(NewViewportInput.Get());
+							InputViewports.Emplace(MediaCaptureId, MoveTemp(NewViewportInput));
+						}
+						// Media capture only
+						else if (MediaSettings.MediaOutput.bEnabled && bOutputDataValid)
+						{
+							UE_LOG(LogDisplayClusterMedia, Log, TEXT("Initializing viewport capture for viewport '%s'..."), *ViewportIt.Key);
+
+							TUniquePtr<FDisplayClusterMediaCaptureViewport> NewViewportCapture = MakeUnique<FDisplayClusterMediaCaptureViewport>(
+								MediaCaptureId, ClusterNodeId,
+								ViewportIt.Key,
+								MediaSettings.MediaOutput.MediaOutput, MediaSettings.MediaOutput.RenderTarget);
+
+							AllCaptures.Add(NewViewportCapture.Get());
+							CaptureViewports.Emplace(MediaCaptureId, MoveTemp(NewViewportCapture));
+						}
+					}
+				}
+			}
+
+			///////////////////////////////
+			// ICVFX media setup
+			{
+				TArray<UActorComponent*> ICVFXCameraComponents;
+				RootActor->GetComponents(UDisplayClusterICVFXCameraComponent::StaticClass(), ICVFXCameraComponents);
+
+				for (const UActorComponent* const Component : ICVFXCameraComponents)
+				{
+					if (const UDisplayClusterICVFXCameraComponent* const ICVFXCamera = Cast<UDisplayClusterICVFXCameraComponent>(Component))
+					{
+						const FDisplayClusterConfigurationMedia& MediaSettings = ICVFXCamera->CameraSettings.RenderSettings.Media;
+
+						if (MediaSettings.bEnabled)
+						{
+							const FString ICVFXCameraName  = ICVFXCamera->GetName();
+							const FString ICVFXViewportId  = DisplayClusterMediaHelpers::GenerateICVFXViewportName(ClusterNodeId, ICVFXCameraName);
+							const FString MediaCaptureId   = RootActor->GetName() + FString("_") + ICVFXCameraName + FString("_icvfx_capture");
+							const FString MediaInputId     = RootActor->GetName() + FString("_") + ICVFXCameraName + FString("_icvfx_input");
+							const bool bInputDataValid     = MediaSettings.MediaInput.MediaPlayer && MediaSettings.MediaInput.MediaSource && MediaSettings.MediaInput.MediaTexture;
+							const bool bOutputDataValid    = MediaSettings.MediaOutput.MediaOutput && MediaSettings.MediaOutput.RenderTarget;
+							const bool bClusterNodeMatches = MediaSettings.MediaOutputNode.Equals(ClusterNodeId, ESearchCase::IgnoreCase);
+
+							// Media sharing
+							if (MediaSettings.IsMediaSharingRequired() && bInputDataValid && bOutputDataValid)
+							{
+								// Share media
+								if (bClusterNodeMatches)
+								{
+									UE_LOG(LogDisplayClusterMedia, Log, TEXT("Initializing ICVFX in-cluster capture for camera '%s' (in-cluster shared)..."), *ICVFXCameraName);
+
+									TUniquePtr<FDisplayClusterMediaCaptureViewport> NewICVFXCapture = MakeUnique<FDisplayClusterMediaCaptureViewport>(
+										MediaCaptureId, ClusterNodeId,
+										ICVFXViewportId,
+										MediaSettings.MediaOutput.MediaOutput, MediaSettings.MediaOutput.RenderTarget);
+
+									AllCaptures.Add(NewICVFXCapture.Get());
+									CaptureViewports.Emplace(MediaCaptureId, MoveTemp(NewICVFXCapture));
+								}
+								// Consume media
+								else
+								{
+									UE_LOG(LogDisplayClusterMedia, Log, TEXT("Initializing ICVFX in-cluster media input for camera '%s' (in-cluster shared)..."), *ICVFXCameraName);
+
+									TUniquePtr<FDisplayClusterMediaInputViewport> NewICVFXInput = MakeUnique<FDisplayClusterMediaInputViewport>(
+										MediaInputId, ClusterNodeId,
+										ICVFXViewportId,
+										MediaSettings.MediaInput.MediaSource, MediaSettings.MediaInput.MediaPlayer, MediaSettings.MediaInput.MediaTexture);
+
+									AllInputs.Add(NewICVFXInput.Get());
+									InputViewports.Emplace(MediaCaptureId, MoveTemp(NewICVFXInput));
+								}
+							}
+							// Media input only
+							else if (MediaSettings.MediaInput.bEnabled && bInputDataValid)
+							{
+								UE_LOG(LogDisplayClusterMedia, Log, TEXT("Initializing ICVFX media input for camera '%s'..."), *ICVFXCameraName);
+
+								TUniquePtr<FDisplayClusterMediaInputViewport> NewICVFXInput = MakeUnique<FDisplayClusterMediaInputViewport>(
+									MediaInputId, ClusterNodeId,
+									ICVFXViewportId,
+									MediaSettings.MediaInput.MediaSource, MediaSettings.MediaInput.MediaPlayer, MediaSettings.MediaInput.MediaTexture);
+
+								AllInputs.Add(NewICVFXInput.Get());
+								InputViewports.Emplace(MediaCaptureId, MoveTemp(NewICVFXInput));
+							}
+							// Media capture only
+							else if (MediaSettings.MediaOutput.bEnabled && bOutputDataValid)
+							{
+								UE_LOG(LogDisplayClusterMedia, Log, TEXT("Initializing ICVFX capture for camera '%s'..."), *ICVFXCameraName);
+
+								TUniquePtr<FDisplayClusterMediaCaptureViewport> NewICVFXCapture = MakeUnique<FDisplayClusterMediaCaptureViewport>(
+									MediaCaptureId, ClusterNodeId,
+									ICVFXViewportId,
+									MediaSettings.MediaOutput.MediaOutput, MediaSettings.MediaOutput.RenderTarget);
+
+								AllCaptures.Add(NewICVFXCapture.Get());
+								CaptureViewports.Emplace(MediaCaptureId, MoveTemp(NewICVFXCapture));
+							}
+						}
+					}
+				}
+			}
+		}
+	}
+}
+
+void FDisplayClusterMediaModule::ReleaseMedia()
+{
+	StopCapture();
+	StopMedia();
+
+	CaptureViewports.Reset();
+	CaptureNode.Reset();
+	AllCaptures.Reset();
+
+	InputViewports.Reset();
+	InputNode.Reset();
+	AllInputs.Reset();
+}
+
+void FDisplayClusterMediaModule::StartCapture()
+{
+	for (FDisplayClusterMediaCaptureBase* Capture : AllCaptures)
+	{
+		Capture->StartCapture();
+	}
+}
+
+void FDisplayClusterMediaModule::StopCapture()
+{
+	for (FDisplayClusterMediaCaptureBase* Capture : AllCaptures)
+	{
+		Capture->StopCapture();
+	}
+}
+
+void FDisplayClusterMediaModule::PlayMedia()
+{
+	for (FDisplayClusterMediaInputBase* MediaInput : AllInputs)
+	{
+		MediaInput->Play();
+	}
+}
+
+void FDisplayClusterMediaModule::StopMedia()
+{
+	for (FDisplayClusterMediaInputBase* MediaInput : AllInputs)
+	{
+		MediaInput->Stop();
+	}
+}
+
+
+IMPLEMENT_MODULE(FDisplayClusterMediaModule, DisplayClusterMedia);
diff --git a/Engine/Plugins/Runtime/nDisplay/Source/DisplayClusterMedia/Private/DisplayClusterMediaModule.h b/Engine/Plugins/Runtime/nDisplay/Source/DisplayClusterMedia/Private/DisplayClusterMediaModule.h
new file mode 100644
--- /dev/null
+++ b/Engine/Plugins/Runtime/nDisplay/Source/DisplayClusterMedia/Private/DisplayClusterMediaModule.h
@@ -0,0 +1,61 @@
+// Copyright Epic Games, Inc. All Rights Reserved.
+
+#pragma once
+
+#include "CoreMinimal.h"
+#include "Modules/ModuleInterface.h"
+
+class FDisplayClusterMediaCaptureBase;
+class FDisplayClusterMediaCaptureICVFX;
+class FDisplayClusterMediaCaptureNode;
+class FDisplayClusterMediaCaptureViewport;
+class FDisplayClusterMediaInputBase;
+class FDisplayClusterMediaInputICVFX;
+class FDisplayClusterMediaInputNode;
+class FDisplayClusterMediaInputViewport;
+
+
+/**
+ * Media module
+ */
+class FDisplayClusterMediaModule :
+	public IModuleInterface
+{
+public:
+	//~ Begin IModuleInterface Implementation
+	virtual void StartupModule() override;
+	virtual void ShutdownModule() override;
+	//~ End IModuleInterface Implementation
+
+protected:
+	/** Initialize all media internals */
+	void InitializeMedia();
+
+	/** Release all media internals */
+	void ReleaseMedia();
+
+	/** Start media capture (all sources) */
+	void StartCapture();
+
+	/** Stop media capture (all soureces) */
+	void StopCapture();
+
+	/** Start playing media (all inputs) */
+	void PlayMedia();
+
+	/** Stop playing media (all inputs) */
+	void StopMedia();
+
+private:
+	void OnCustomPresentSet();
+	void OnEnginePreExit();
+
+private:
+	TMap<FString, TUniquePtr<FDisplayClusterMediaCaptureViewport>> CaptureViewports;
+	TUniquePtr<FDisplayClusterMediaCaptureNode>                    CaptureNode;
+	TArray<FDisplayClusterMediaCaptureBase*>                       AllCaptures;
+
+	TMap<FString, TUniquePtr<FDisplayClusterMediaInputViewport>>   InputViewports;
+	TUniquePtr<FDisplayClusterMediaInputNode>                      InputNode;
+	TArray<FDisplayClusterMediaInputBase*>                         AllInputs;
+};
diff --git a/Engine/Plugins/Runtime/nDisplay/Source/DisplayClusterMedia/Private/Input/DisplayClusterMediaInputBase.cpp b/Engine/Plugins/Runtime/nDisplay/Source/DisplayClusterMedia/Private/Input/DisplayClusterMediaInputBase.cpp
new file mode 100644
--- /dev/null
+++ b/Engine/Plugins/Runtime/nDisplay/Source/DisplayClusterMedia/Private/Input/DisplayClusterMediaInputBase.cpp
@@ -0,0 +1,122 @@
+// Copyright Epic Games, Inc. All Rights Reserved.
+
+#include "Input/DisplayClusterMediaInputBase.h"
+
+#include "DisplayClusterMediaHelpers.h"
+#include "DisplayClusterMediaLog.h"
+
+#include "MediaPlayer.h"
+#include "MediaSource.h"
+#include "MediaTexture.h"
+
+
+FDisplayClusterMediaInputBase::FDisplayClusterMediaInputBase(const FString& InMediaId, const FString& InClusterNodeId, UMediaSource* InMediaSource, UMediaPlayer* InMediaPlayer, UMediaTexture* InMediaTexture)
+	: FDisplayClusterMediaBase(InMediaId, InClusterNodeId)
+	, MediaSource(InMediaSource)
+	, MediaPlayer(InMediaPlayer)
+	, MediaTexture(InMediaTexture)
+{
+	// We expect these to always be valid
+	check(MediaSource && MediaPlayer && MediaTexture);
+}
+
+
+void FDisplayClusterMediaInputBase::AddReferencedObjects(FReferenceCollector& Collector)
+{
+	Collector.AddReferencedObject(MediaSource);
+	Collector.AddReferencedObject(MediaPlayer);
+	Collector.AddReferencedObject(MediaTexture);
+}
+
+bool FDisplayClusterMediaInputBase::Play()
+{
+	if (MediaSource && MediaPlayer && MediaTexture)
+	{
+		MediaPlayer->PlayOnOpen = true;
+		MediaPlayer->OnMediaEvent().AddRaw(this, &FDisplayClusterMediaInputBase::OnMediaEvent);
+
+		const bool bIsPlaying = MediaPlayer->OpenSource(MediaSource);
+		if (bIsPlaying)
+		{
+			UE_LOG(LogDisplayClusterMedia, Log, TEXT("Started playing media: %s"), *GetMediaId());
+		}
+		else
+		{
+			UE_LOG(LogDisplayClusterMedia, Warning, TEXT("Couldn't start playing media: %s"), *GetMediaId());
+		}
+
+		return bIsPlaying;
+	}
+
+	return false;
+}
+
+void FDisplayClusterMediaInputBase::Stop()
+{
+	if (MediaPlayer)
+	{
+		MediaPlayer->Close();
+		MediaPlayer->OnMediaEvent().RemoveAll(this);
+	}
+}
+
+void FDisplayClusterMediaInputBase::ImportMediaData(FRHICommandListImmediate& RHICmdList, const FMediaTextureInfo& TextureInfo)
+{
+	FRHITexture* const SrcTexture = MediaTexture->GetResource()->GetTextureRHI();
+	FRHITexture* const DstTexture = TextureInfo.Texture;
+
+	if (SrcTexture && DstTexture)
+	{
+		const FIntPoint DstRegionSize = TextureInfo.Region.Size();
+
+		if (SrcTexture->GetDesc().Format == DstTexture->GetDesc().Format &&
+			SrcTexture->GetDesc().Extent == DstRegionSize)
+		{
+			FRHICopyTextureInfo CopyInfo;
+
+			CopyInfo.DestPosition = FIntVector(TextureInfo.Region.Min.X, TextureInfo.Region.Min.Y, 0);
+			CopyInfo.Size = FIntVector(DstRegionSize.X, DstRegionSize.Y, 0);
+
+			TransitionAndCopyTexture(RHICmdList, SrcTexture, DstTexture, CopyInfo);
+		}
+		else
+		{
+			UE_LOG(LogDisplayClusterMedia, Log, TEXT(""));
+
+			DisplayClusterMediaHelpers::ResampleTexture_RenderThread(
+				RHICmdList, SrcTexture, DstTexture,
+				FIntRect(FIntPoint(0, 0), SrcTexture->GetDesc().Extent),
+				TextureInfo.Region);
+		}
+	}
+}
+
+void FDisplayClusterMediaInputBase::OnMediaEvent(EMediaEvent MediaEvent)
+{
+	switch (MediaEvent)
+	{
+	/** The player started connecting to the media source. */
+	case EMediaEvent::MediaConnecting:
+		UE_LOG(LogDisplayClusterMedia, Log, TEXT("Media event for '%s': Connectiong"), *GetMediaId());
+		break;
+
+	/** A new media source has been opened. */
+	case EMediaEvent::MediaOpened:
+		UE_LOG(LogDisplayClusterMedia, Log, TEXT("Media event for '%s': Opened"), *GetMediaId());
+		break;
+
+	/** The current media source has been closed. */
+	case EMediaEvent::MediaClosed:
+		UE_LOG(LogDisplayClusterMedia, Log, TEXT("Media event for '%s': Closed"), *GetMediaId());
+		break;
+		
+	/** A media source failed to open. */
+	case EMediaEvent::MediaOpenFailed:
+		UE_LOG(LogDisplayClusterMedia, Log, TEXT("Media event for '%s': OpenFailed"), *GetMediaId());
+		break;
+
+	default:
+		UE_LOG(LogDisplayClusterMedia, Log, TEXT("Media event for '%s': %d"), *GetMediaId(), static_cast<int32>(MediaEvent));
+		break;
+	}
+}
diff --git a/Engine/Plugins/Runtime/nDisplay/Source/DisplayClusterMedia/Private/Input/DisplayClusterMediaInputBase.h b/Engine/Plugins/Runtime/nDisplay/Source/DisplayClusterMedia/Private/Input/DisplayClusterMediaInputBase.h
new file mode 100644
--- /dev/null
+++ b/Engine/Plugins/Runtime/nDisplay/Source/DisplayClusterMedia/Private/Input/DisplayClusterMediaInputBase.h
@@ -0,0 +1,68 @@
+// Copyright Epic Games, Inc. All Rights Reserved.
+
+#pragma once
+
+#include "CoreMinimal.h"
+#include "DisplayClusterMediaBase.h"
+#include "IMediaEventSink.h"
+#include "UObject/GCObject.h"
+
+class FRHICommandListImmediate;
+
+class UMediaSource;
+class UMediaPlayer;
+class UMediaTexture;
+
+
+/**
+ * Base media input class
+ */
+class FDisplayClusterMediaInputBase
+	: public FDisplayClusterMediaBase
+	, public FGCObject
+{
+public:
+	FDisplayClusterMediaInputBase(const FString& MediaId, const FString& ClusterNodeId, UMediaSource* MediaSource, UMediaPlayer* MediaPlayer, UMediaTexture* MediaTexture);
+
+public:
+	//~ Begin FGCObject interface
+	virtual void AddReferencedObjects(FReferenceCollector& Collector) override;
+	virtual FString GetReferencerName() const override
+	{
+		return TEXT("FDisplayClusterMediaInputBase");
+	}
+	//~ End FGCObject interface
+
+public:
+	virtual bool Play();
+	virtual void Stop();
+
+public:
+	UMediaSource* GetMediaSource() const
+	{
+		return MediaSource;
+	}
+
+	UMediaPlayer* GetMediaPlayer() const
+	{
+		return MediaPlayer;
+	}
+
+	UMediaTexture* GetMediaTexture() const
+	{
+		return MediaTexture;
+	}
+
+protected:
+	void ImportMediaData(FRHICommandListImmediate& RHICmdList, const FMediaTextureInfo& TextureInfo);
+
+private:
+	void OnMediaEvent(EMediaEvent MediaEvent);
+
+private:
+	//~ Begin GC by AddReferencedObjects
+	UMediaSource*  MediaSource = nullptr;
+	UMediaPlayer*  MediaPlayer = nullptr;
+	UMediaTexture* MediaTexture = nullptr;
+	//~ End GC by AddReferencedObjects
+};
diff --git a/Engine/Plugins/Runtime/nDisplay/Source/DisplayClusterMedia/Private/Input/DisplayClusterMediaInputNode.cpp b/Engine/Plugins/Runtime/nDisplay/Source/DisplayClusterMedia/Private/Input/DisplayClusterMediaInputNode.cpp
new file mode 100644
--- /dev/null
+++ b/Engine/Plugins/Runtime/nDisplay/Source/DisplayClusterMedia/Private/Input/DisplayClusterMediaInputNode.cpp
@@ -0,0 +1,55 @@
+// Copyright Epic Games, Inc. All Rights Reserved.
+
+#include "Input/DisplayClusterMediaInputNode.h"
+
+#include "DisplayClusterMediaHelpers.h"
+#include "Render/Viewport/IDisplayClusterViewportManagerProxy.h"
+
+#include "IDisplayCluster.h"
+#include "IDisplayClusterCallbacks.h"
+
+#include "RHICommandList.h"
+#include "RHIResources.h"
+
+
+FDisplayClusterMediaInputNode::FDisplayClusterMediaInputNode(const FString& InMediaId, const FString& InClusterNodeId, UMediaSource* InMediaSource, UMediaPlayer* InMediaPlayer, UMediaTexture* InMediaTexture)
+	: FDisplayClusterMediaInputBase(InMediaId, InClusterNodeId, InMediaSource, InMediaPlayer, InMediaTexture)
+{
+}
+
+
+bool FDisplayClusterMediaInputNode::Play()
+{
+	if (FDisplayClusterMediaInputBase::Play())
+	{
+		IDisplayCluster::Get().GetCallbacks().OnDisplayClusterPostFrameRender_RenderThread().AddRaw(this, &FDisplayClusterMediaInputNode::OnPostFrameRender_RenderThread);
+		return true;
+	}
+
+	return false;
+}
+
+void FDisplayClusterMediaInputNode::Stop()
+{
+	// Stop receiving notifications
+	IDisplayCluster::Get().GetCallbacks().OnDisplayClusterPostFrameRender_RenderThread().RemoveAll(this);
+	// Stop playing
+	FDisplayClusterMediaInputBase::Stop();
+}
+
+void FDisplayClusterMediaInputNode::OnPostFrameRender_RenderThread(FRHICommandListImmediate& RHICmdList, const IDisplayClusterViewportManagerProxy* ViewportManagerProxy, FViewport* Viewport)
+{
+	ensure(ViewportManagerProxy);
+
+	TArray<FRHITexture*> Textures;
+	TArray<FIntPoint>    Regions;
+
+	if (ViewportManagerProxy && ViewportManagerProxy->GetFrameTargets_RenderThread(Textures, Regions))
+	{
+		if (Textures.Num() > 0 && Regions.Num() > 0 && Textures[0])
+		{
+			FMediaTextureInfo TextureInfo{ Textures[0], FIntRect(Regions[0], Regions[0] + Textures[0]->GetDesc().Extent) };
+			ImportMediaData(RHICmdList, TextureInfo);
+		}
+	}
+}
diff --git a/Engine/Plugins/Runtime/nDisplay/Source/DisplayClusterMedia/Private/Input/DisplayClusterMediaInputNode.h b/Engine/Plugins/Runtime/nDisplay/Source/DisplayClusterMedia/Private/Input/DisplayClusterMediaInputNode.h
new file mode 100644
--- /dev/null
+++ b/Engine/Plugins/Runtime/nDisplay/Source/DisplayClusterMedia/Private/Input/DisplayClusterMediaInputNode.h
@@ -0,0 +1,27 @@
+// Copyright Epic Games, Inc. All Rights Reserved.
+
+#pragma once
+
+#include "Input/DisplayClusterMediaInputBase.h"
+
+class FRHICommandListImmediate;
+class FViewport;
+class IDisplayClusterViewportManagerProxy;
+
+
+/**
+ * Node backbuffer media input
+ */
+class FDisplayClusterMediaInputNode
+	: public FDisplayClusterMediaInputBase
+{
+public:
+	FDisplayClusterMediaInputNode(const FString& MediaId, const FString& ClusterNodeId, UMediaSource* MediaSource, UMediaPlayer* MediaPlayer, UMediaTexture* MediaTexture);
+
+public:
+	virtual bool Play() override;
+	virtual void Stop() override;
+
+protected:
+	void OnPostFrameRender_RenderThread(FRHICommandListImmediate& RHICmdList, const IDisplayClusterViewportManagerProxy* ViewportManagerProxy, FViewport* Viewport);
+};
diff --git a/Engine/Plugins/Runtime/nDisplay/Source/DisplayClusterMedia/Private/Input/DisplayClusterMediaInputViewport.cpp b/Engine/Plugins/Runtime/nDisplay/Source/DisplayClusterMedia/Private/Input/DisplayClusterMediaInputViewport.cpp
new file mode 100644
--- /dev/null
+++ b/Engine/Plugins/Runtime/nDisplay/Source/DisplayClusterMedia/Private/Input/DisplayClusterMediaInputViewport.cpp
@@ -0,0 +1,58 @@
+// Copyright Epic Games, Inc. All Rights Reserved.
+
+#include "Input/DisplayClusterMediaInputViewport.h"
+#include "Render/Viewport/IDisplayClusterViewportManagerProxy.h"
+
+#include "IDisplayCluster.h"
+#include "IDisplayClusterCallbacks.h"
+
+#include "RHICommandList.h"
+#include "RHIResources.h"
+
+
+FDisplayClusterMediaInputViewport::FDisplayClusterMediaInputViewport(const FString& InMediaId, const FString& InClusterNodeId, const FString& InViewportId, UMediaSource* InMediaSource, UMediaPlayer* InMediaPlayer, UMediaTexture* InMediaTexture)
+	: FDisplayClusterMediaInputBase(InMediaId, InClusterNodeId, InMediaSource, InMediaPlayer, InMediaTexture)
+	, ViewportId(InViewportId)
+{
+}
+
+
+bool FDisplayClusterMediaInputViewport::Play()
+{
+	// If playback has started successfully, subscribe for rendering callbacks
+	if (FDisplayClusterMediaInputBase::Play())
+	{
+		IDisplayCluster::Get().GetCallbacks().OnDisplayClusterPreFrameRender_RenderThread().AddRaw(this, &FDisplayClusterMediaInputViewport::PreFrameRender_RenderThread);
+		return true;
+	}
+
+	return false;
+}
+
+void FDisplayClusterMediaInputViewport::Stop()
+{
+	// Stop receiving notifications
+	IDisplayCluster::Get().GetCallbacks().OnDisplayClusterPreFrameRender_RenderThread().RemoveAll(this);
+	// Stop playing
+	FDisplayClusterMediaInputBase::Stop();
+}
+
+void FDisplayClusterMediaInputViewport::PreFrameRender_RenderThread(FRHICommandListImmediate& RHICmdList, const IDisplayClusterViewportManagerProxy* ViewportManagerProxy, FViewport* Viewport)
+{
+	checkSlow(ViewportManagerProxy);
+
+	if (const IDisplayClusterViewportProxy* const CaptureViewport = ViewportManagerProxy->FindViewport_RenderThread(GetViewportId()))
+	{
+		TArray<FRHITexture*> Textures;
+		TArray<FIntRect>     Regions;
+
+		if (CaptureViewport->GetResourcesWithRects_RenderThread(EDisplayClusterViewportResourceType::InternalRenderTargetResource, Textures, Regions))
+		{
+			if (Textures.Num() > 0 && Regions.Num() > 0)
+			{
+				FMediaTextureInfo TextureInfo{ Textures[0], Regions[0] };
+				ImportMediaData(RHICmdList, TextureInfo);
+			}
+		}
+	}
+}
diff --git a/Engine/Plugins/Runtime/nDisplay/Source/DisplayClusterMedia/Private/Input/DisplayClusterMediaInputViewport.h b/Engine/Plugins/Runtime/nDisplay/Source/DisplayClusterMedia/Private/Input/DisplayClusterMediaInputViewport.h
new file mode 100644
--- /dev/null
+++ b/Engine/Plugins/Runtime/nDisplay/Source/DisplayClusterMedia/Private/Input/DisplayClusterMediaInputViewport.h
@@ -0,0 +1,35 @@
+// Copyright Epic Games, Inc. All Rights Reserved.
+
+#pragma once
+
+#include "Input/DisplayClusterMediaInputBase.h"
+
+class FRHICommandListImmediate;
+class FViewport;
+class IDisplayClusterViewportManagerProxy;
+
+
+/**
+ * Viewport media input
+ */
+class FDisplayClusterMediaInputViewport
+	: public FDisplayClusterMediaInputBase
+{
+public:
+	FDisplayClusterMediaInputViewport(const FString& MediaId, const FString& ClusterNodeId, const FString& ViewportId, UMediaSource* MediaSource, UMediaPlayer* MediaPlayer, UMediaTexture* MediaTexture);
+
+public:
+	virtual bool Play() override;
+	virtual void Stop() override;
+
+	const FString& GetViewportId() const
+	{
+		return ViewportId;
+	}
+
+private:
+	void PreFrameRender_RenderThread(FRHICommandListImmediate& RHICmdList, const IDisplayClusterViewportManagerProxy* ViewportManagerProxy, FViewport* Viewport);
+
+private:
+	const FString ViewportId;
+};
diff --git a/Engine/Plugins/Runtime/nDisplay/nDisplay.uplugin b/Engine/Plugins/Runtime/nDisplay/nDisplay.uplugin
--- a/Engine/Plugins/Runtime/nDisplay/nDisplay.uplugin
+++ b/Engine/Plugins/Runtime/nDisplay/nDisplay.uplugin
@@ -134,6 +134,15 @@
 				"Linux"
 			]
 		},
+		{
+			"Name": "DisplayClusterMedia",
+			"Type": "Runtime",
+			"LoadingPhase": "Default",
+			"PlatformAllowList": [
+				"Win64",
+				"Linux"
+			]
+		},
 		{
 			"Name": "DisplayClusterMultiUser",
 			"Type": "UncookedOnly",
@@ -161,6 +170,10 @@
 			"Name": "ConcertSyncClient",
 			"Enabled": true
 		},
+		{
+			"Name": "MediaIOFramework",
+			"Enabled": true
+		},
 		{
 			"Name": "OpenCV",
 			"Enabled": true
