Enable the InputKey and InputAxis deprecations on the EditorViewportClient

#rnx

diff --git a/Engine/Plugins/Enterprise/DataprepEditor/Source/DataprepEditor/Private/Widgets/SDataprepEditorViewport.cpp b/Engine/Plugins/Enterprise/DataprepEditor/Source/DataprepEditor/Private/Widgets/SDataprepEditorViewport.cpp
--- a/Engine/Plugins/Enterprise/DataprepEditor/Source/DataprepEditor/Private/Widgets/SDataprepEditorViewport.cpp
+++ b/Engine/Plugins/Enterprise/DataprepEditor/Source/DataprepEditor/Private/Widgets/SDataprepEditorViewport.cpp
@@ -1352,13 +1352,13 @@ FDataprepEditorViewportClient::FDataprepEditorViewportClient(const TSharedRef<SE
 	AdvancedPreviewScene->SetProfileIndex( SDataprepEditorViewport::AssetViewerProfileIndex );
 }
 
-bool FDataprepEditorViewportClient::InputKey(FViewport * InViewport, int32 ControllerId, FKey Key, EInputEvent Event, float AmountDepressed, bool bGamepad)
+bool FDataprepEditorViewportClient::InputKey(const FInputKeyEventArgs& EventArgs)
 {
 	bool bHandled = false;
 	
 	// #ueent_todo: Put code for specific handling
 
-	return bHandled ? true : FEditorViewportClient::InputKey(InViewport, ControllerId, Key, Event, AmountDepressed, bGamepad );
+	return bHandled ? true : FEditorViewportClient::InputKey(EventArgs);
 }
 
 void FDataprepEditorViewportClient::ProcessClick(FSceneView& View, HHitProxy* HitProxy, FKey Key, EInputEvent Event, uint32 HitX, uint32 HitY)
diff --git a/Engine/Plugins/Enterprise/DataprepEditor/Source/DataprepEditor/Private/Widgets/SDataprepEditorViewport.h b/Engine/Plugins/Enterprise/DataprepEditor/Source/DataprepEditor/Private/Widgets/SDataprepEditorViewport.h
--- a/Engine/Plugins/Enterprise/DataprepEditor/Source/DataprepEditor/Private/Widgets/SDataprepEditorViewport.h
+++ b/Engine/Plugins/Enterprise/DataprepEditor/Source/DataprepEditor/Private/Widgets/SDataprepEditorViewport.h
@@ -132,7 +132,7 @@ public:
 	~FDataprepEditorViewportClient() {}
 
 	// FEditorViewportClient interface
-	virtual bool InputKey(FViewport* Viewport, int32 ControllerId, FKey Key, EInputEvent Event, float AmountDepressed = 1.f, bool bGamepad=false) override;
+	virtual bool InputKey(const FInputKeyEventArgs& EventArgs) override;
 	virtual void ProcessClick(class FSceneView& View, class HHitProxy* HitProxy, FKey Key, EInputEvent Event, uint32 HitX, uint32 HitY) override;
 	virtual void DrawCanvas( FViewport& InViewport, FSceneView& View, FCanvas& Canvas ) override;
 	// End of FEditorViewportClient interface
diff --git a/Engine/Plugins/Runtime/nDisplay/Source/DisplayClusterConfigurator/Private/Views/Viewport/DisplayClusterConfiguratorSCSEditorViewportClient.cpp b/Engine/Plugins/Runtime/nDisplay/Source/DisplayClusterConfigurator/Private/Views/Viewport/DisplayClusterConfiguratorSCSEditorViewportClient.cpp
--- a/Engine/Plugins/Runtime/nDisplay/Source/DisplayClusterConfigurator/Private/Views/Viewport/DisplayClusterConfiguratorSCSEditorViewportClient.cpp
+++ b/Engine/Plugins/Runtime/nDisplay/Source/DisplayClusterConfigurator/Private/Views/Viewport/DisplayClusterConfiguratorSCSEditorViewportClient.cpp
@@ -386,13 +386,13 @@ void FDisplayClusterConfiguratorSCSEditorViewportClient::DrawCanvas(FViewport& I
 	}
 }
 
-bool FDisplayClusterConfiguratorSCSEditorViewportClient::InputKey(FViewport* InViewport, int32 ControllerId, FKey Key, EInputEvent Event, float AmountDepressed, bool bGamepad)
+bool FDisplayClusterConfiguratorSCSEditorViewportClient::InputKey(const FInputKeyEventArgs& EventArgs)
 {
-	bool bHandled = GUnrealEd->ComponentVisManager.HandleInputKey(this, InViewport, Key, Event);;
+	bool bHandled = GUnrealEd->ComponentVisManager.HandleInputKey(this, EventArgs.Viewport, EventArgs.Key, EventArgs.Event);
 
 	if (!bHandled)
 	{
-		bHandled = FEditorViewportClient::InputKey(InViewport, ControllerId, Key, Event, AmountDepressed, bGamepad);
+		bHandled = FEditorViewportClient::InputKey(EventArgs);
 	}
 
 	return bHandled;
diff --git a/Engine/Plugins/Runtime/nDisplay/Source/DisplayClusterConfigurator/Private/Views/Viewport/DisplayClusterConfiguratorSCSEditorViewportClient.h b/Engine/Plugins/Runtime/nDisplay/Source/DisplayClusterConfigurator/Private/Views/Viewport/DisplayClusterConfiguratorSCSEditorViewportClient.h
--- a/Engine/Plugins/Runtime/nDisplay/Source/DisplayClusterConfigurator/Private/Views/Viewport/DisplayClusterConfiguratorSCSEditorViewportClient.h
+++ b/Engine/Plugins/Runtime/nDisplay/Source/DisplayClusterConfigurator/Private/Views/Viewport/DisplayClusterConfiguratorSCSEditorViewportClient.h
@@ -33,7 +33,7 @@ public:
 	virtual void Tick(float DeltaSeconds) override;
 	virtual void Draw(const FSceneView* View, FPrimitiveDrawInterface* PDI) override;
 	virtual void DrawCanvas(FViewport& InViewport, FSceneView& View, FCanvas& Canvas) override;
-	virtual bool InputKey(FViewport* Viewport, int32 ControllerId, FKey Key, EInputEvent Event, float AmountDepressed = 1.f, bool bGamepad = false) override;
+	virtual bool InputKey(const FInputKeyEventArgs& EventArgs) override;
 	virtual void ProcessClick(class FSceneView& View, class HHitProxy* HitProxy, FKey Key, EInputEvent Event, uint32 HitX, uint32 HitY) override;
 	virtual bool InputWidgetDelta(FViewport* Viewport, EAxisList::Type CurrentAxis, FVector& Drag, FRotator& Rot, FVector& Scale) override;
 	virtual void TrackingStarted(const struct FInputEventState& InInputState, bool bIsDragging, bool bNudge) override;
diff --git a/Engine/Plugins/Runtime/nDisplay/Source/DisplayClusterLightCardEditor/Private/Viewport/DisplayClusterLightCardEditorViewportClient.cpp b/Engine/Plugins/Runtime/nDisplay/Source/DisplayClusterLightCardEditor/Private/Viewport/DisplayClusterLightCardEditorViewportClient.cpp
--- a/Engine/Plugins/Runtime/nDisplay/Source/DisplayClusterLightCardEditor/Private/Viewport/DisplayClusterLightCardEditorViewportClient.cpp
+++ b/Engine/Plugins/Runtime/nDisplay/Source/DisplayClusterLightCardEditor/Private/Viewport/DisplayClusterLightCardEditorViewportClient.cpp
@@ -750,26 +750,25 @@ FSceneView* FDisplayClusterLightCardEditorViewportClient::CalcSceneView(FSceneVi
 	return View;
 }
 
-bool FDisplayClusterLightCardEditorViewportClient::InputKey(FViewport* InViewport, int32 ControllerId, FKey Key, EInputEvent Event,
-                                            float AmountDepressed, bool bGamepad)
+bool FDisplayClusterLightCardEditorViewportClient::InputKey(const FInputKeyEventArgs& EventArgs)
 {
 	switch (InputMode)
 	{
 	case EInputMode::DrawingLightCard:
 		
-		if ((Key == EKeys::LeftMouseButton) && (Event == IE_Released))
+		if ((EventArgs.Key == EKeys::LeftMouseButton) && (EventArgs.Event == IE_Released))
 		{
 			// Add a new light card polygon point
 
 			FIntPoint MousePos;
-			InViewport->GetMousePos(MousePos);
+			EventArgs.Viewport->GetMousePos(MousePos);
 
 			if (!DrawnMousePositions.Num() || (DrawnMousePositions.Last() != MousePos))
 			{
 				DrawnMousePositions.Add(MousePos);
 			}
 		}
-		else if ((Key == EKeys::RightMouseButton) && (Event == IE_Released))
+		else if ((EventArgs.Key == EKeys::RightMouseButton) && (EventArgs.Event == IE_Released))
 		{
 			// Create the polygon light card.
 			CreateDrawnLightCard(DrawnMousePositions);
@@ -787,9 +786,9 @@ bool FDisplayClusterLightCardEditorViewportClient::InputKey(FViewport* InViewpor
 
 	default:
 
-		if ((Key == EKeys::MouseScrollUp || Key == EKeys::MouseScrollDown) && (Event == IE_Pressed))
+		if ((EventArgs.Key == EKeys::MouseScrollUp || EventArgs.Key == EKeys::MouseScrollDown) && (EventArgs.Event == IE_Pressed))
 		{
-			const int32 Sign = Key == EKeys::MouseScrollUp ? -1 : 1;
+			const int32 Sign = EventArgs.Key == EKeys::MouseScrollUp ? -1 : 1;
 			const float CurrentFOV = GetProjectionModeFOV(ProjectionMode);
 			const float NewFOV = FMath::Clamp(CurrentFOV + Sign * FOVScrollIncrement, CameraController->GetConfig().MinimumAllowedFOV, CameraController->GetConfig().MaximumAllowedFOV);
 
@@ -800,7 +799,7 @@ bool FDisplayClusterLightCardEditorViewportClient::InputKey(FViewport* InViewpor
 		break;
 	}
 
-	return FEditorViewportClient::InputKey(InViewport, ControllerId, Key, Event, AmountDepressed, bGamepad);
+	return FEditorViewportClient::InputKey(EventArgs);
 }
 
 bool FDisplayClusterLightCardEditorViewportClient::InputWidgetDelta(FViewport* InViewport,
diff --git a/Engine/Plugins/Runtime/nDisplay/Source/DisplayClusterLightCardEditor/Private/Viewport/DisplayClusterLightCardEditorViewportClient.h b/Engine/Plugins/Runtime/nDisplay/Source/DisplayClusterLightCardEditor/Private/Viewport/DisplayClusterLightCardEditorViewportClient.h
--- a/Engine/Plugins/Runtime/nDisplay/Source/DisplayClusterLightCardEditor/Private/Viewport/DisplayClusterLightCardEditorViewportClient.h
+++ b/Engine/Plugins/Runtime/nDisplay/Source/DisplayClusterLightCardEditor/Private/Viewport/DisplayClusterLightCardEditorViewportClient.h
@@ -139,7 +139,7 @@ public:
 	virtual UE::Widget::EWidgetMode GetWidgetMode() const override { return UE::Widget::WM_None; }
 	virtual FSceneView* CalcSceneView(FSceneViewFamily* ViewFamily, const int32 StereoViewIndex = INDEX_NONE) override;
 	virtual bool IsLevelEditorClient() const override { return false; }
-	virtual bool InputKey(FViewport* InViewport, int32 ControllerId, FKey Key, EInputEvent Event, float AmountDepressed, bool bGamepad) override;
+	virtual bool InputKey(const FInputKeyEventArgs& EventArgs) override;
 	virtual bool InputWidgetDelta(FViewport* InViewport, EAxisList::Type CurrentAxis, FVector& Drag, FRotator& Rot, FVector& Scale) override;
 	virtual void TrackingStarted(const FInputEventState& InInputState, bool bIsDraggingWidget, bool bNudge) override;
 	virtual void TrackingStopped() override;
diff --git a/Engine/Source/Editor/UnrealEd/Public/EditorViewportClient.h b/Engine/Source/Editor/UnrealEd/Public/EditorViewportClient.h
--- a/Engine/Source/Editor/UnrealEd/Public/EditorViewportClient.h
+++ b/Engine/Source/Editor/UnrealEd/Public/EditorViewportClient.h
@@ -561,10 +561,10 @@ public:
 	virtual bool ProcessScreenShots(FViewport* Viewport) override;
 	virtual void RedrawRequested(FViewport* Viewport) override;
 	virtual void RequestInvalidateHitProxy(FViewport* Viewport) override;
-	//UE_DEPRECATED(5.1, "This version of InputKey is deprecated. Please use the version that takes FInputKeyEventArgs instead.")
+	UE_DEPRECATED(5.1, "This version of InputKey is deprecated. Please use the version that takes FInputKeyEventArgs instead.")
 	virtual bool InputKey(FViewport* InViewport, int32 ControllerId, FKey Key, EInputEvent Event, float AmountDepressed = 1.f, bool bGamepad=false) override;
 	virtual bool InputKey(const FInputKeyEventArgs& EventArgs) override;
-	//UE_DEPRECATED(5.1, "This version of InputAxis is deprecated. Please use the version that takes a DeviceId instead.")
+	UE_DEPRECATED(5.1, "This version of InputAxis is deprecated. Please use the version that takes a DeviceId instead.")
 	virtual bool InputAxis(FViewport* InViewport, int32 ControllerId, FKey Key, float Delta, float DeltaTime, int32 NumSamples=1, bool bGamepad=false) override;
 	virtual bool InputAxis(FViewport* Viewport, FInputDeviceId DeviceID, FKey Key, float Delta, float DeltaTime, int32 NumSamples=1, bool bGamepad=false) override;
 	virtual bool InputGesture(FViewport* Viewport, EGestureEvent GestureType, const FVector2D& GestureDelta, bool bIsDirectionInvertedFromDevice) override;
