nDisplay: media configuration fix (no in-cluster media sharing for viewports and backbuffers)

diff --git a/Engine/Plugins/Runtime/nDisplay/Source/DisplayCluster/Private/Render/Viewport/Configuration/DisplayClusterViewportConfigurationHelpers.cpp b/Engine/Plugins/Runtime/nDisplay/Source/DisplayCluster/Private/Render/Viewport/Configuration/DisplayClusterViewportConfigurationHelpers.cpp
--- a/Engine/Plugins/Runtime/nDisplay/Source/DisplayCluster/Private/Render/Viewport/Configuration/DisplayClusterViewportConfigurationHelpers.cpp
+++ b/Engine/Plugins/Runtime/nDisplay/Source/DisplayCluster/Private/Render/Viewport/Configuration/DisplayClusterViewportConfigurationHelpers.cpp
@@ -212,16 +212,8 @@ void FDisplayClusterViewportConfigurationHelpers::UpdateBaseViewportSetting(FDis
 
 		if (MediaSettings.bEnabled)
 		{
-			const FString ThisClusterNodeId = IDisplayCluster::Get().GetClusterMgr()->GetNodeId();
-			const bool bThisNodeSharesMedia = MediaSettings.MediaOutputNode.Equals(ThisClusterNodeId, ESearchCase::IgnoreCase);
-
-			// In most cases there is no need to render viewport if media input set up. The only exception
-			// is when this viewport is used as media source for other cluster nodes. In this case media input
-			// settings are ignored, and media capture is used.
-			DstViewport.RenderSettings.bSkipSceneRenderingButLeaveResourcesAvailable = MediaSettings.MediaInput.bEnabled ?
-					!(MediaSettings.MediaOutput.bEnabled && bThisNodeSharesMedia) :
-					false;
-
+			// In case this viewport consumes data from a media input device, there is no need to render scene
+			DstViewport.RenderSettings.bSkipSceneRenderingButLeaveResourcesAvailable = MediaSettings.MediaInput.bEnabled;
 			// Set media capture flag if media capture is used
 			DstViewport.RenderSettings.bIsBeingCaptured = InConfigurationViewport.RenderSettings.Media.MediaOutput.bEnabled;
 		}
diff --git a/Engine/Plugins/Runtime/nDisplay/Source/DisplayCluster/Private/Render/Viewport/Configuration/DisplayClusterViewportConfigurationHelpers_ICVFX.cpp b/Engine/Plugins/Runtime/nDisplay/Source/DisplayCluster/Private/Render/Viewport/Configuration/DisplayClusterViewportConfigurationHelpers_ICVFX.cpp
--- a/Engine/Plugins/Runtime/nDisplay/Source/DisplayCluster/Private/Render/Viewport/Configuration/DisplayClusterViewportConfigurationHelpers_ICVFX.cpp
+++ b/Engine/Plugins/Runtime/nDisplay/Source/DisplayCluster/Private/Render/Viewport/Configuration/DisplayClusterViewportConfigurationHelpers_ICVFX.cpp
@@ -567,7 +567,7 @@ void FDisplayClusterViewportConfigurationHelpers_ICVFX::UpdateCameraViewportSett
 	// Set media related configuration (runtime only for now)
 	if (IDisplayCluster::Get().GetOperationMode() == EDisplayClusterOperationMode::Cluster)
 	{
-		const FDisplayClusterConfigurationMedia& MediaSettings = InCameraComponent.CameraSettings.RenderSettings.Media;
+		const FDisplayClusterConfigurationICVFXMedia& MediaSettings = InCameraComponent.CameraSettings.RenderSettings.Media;
 		if (MediaSettings.bEnabled)
 		{
 			const FString ThisClusterNodeId = IDisplayCluster::Get().GetClusterMgr()->GetNodeId();
diff --git a/Engine/Plugins/Runtime/nDisplay/Source/DisplayClusterConfiguration/Public/DisplayClusterConfigurationTypes_ICVFX.h b/Engine/Plugins/Runtime/nDisplay/Source/DisplayClusterConfiguration/Public/DisplayClusterConfigurationTypes_ICVFX.h
--- a/Engine/Plugins/Runtime/nDisplay/Source/DisplayClusterConfiguration/Public/DisplayClusterConfigurationTypes_ICVFX.h
+++ b/Engine/Plugins/Runtime/nDisplay/Source/DisplayClusterConfiguration/Public/DisplayClusterConfigurationTypes_ICVFX.h
@@ -323,7 +323,7 @@ public:
 
 	// Media settings
 	UPROPERTY(EditAnywhere, BlueprintReadWrite, Category = NDisplay, meta = (DisplayName = "Media"))
-	FDisplayClusterConfigurationMedia Media;
+	FDisplayClusterConfigurationICVFXMedia Media;
 
 	UPROPERTY()
 	FDisplayClusterConfigurationPostRender_BlurPostprocess PostprocessBlur;
diff --git a/Engine/Plugins/Runtime/nDisplay/Source/DisplayClusterConfiguration/Public/DisplayClusterConfigurationTypes_Media.h b/Engine/Plugins/Runtime/nDisplay/Source/DisplayClusterConfiguration/Public/DisplayClusterConfigurationTypes_Media.h
--- a/Engine/Plugins/Runtime/nDisplay/Source/DisplayClusterConfiguration/Public/DisplayClusterConfigurationTypes_Media.h
+++ b/Engine/Plugins/Runtime/nDisplay/Source/DisplayClusterConfiguration/Public/DisplayClusterConfigurationTypes_Media.h
@@ -65,7 +65,7 @@ public:
 
 
 /*
- * ICVFX media override settings
+ * Media settings
  */
 USTRUCT(Blueprintable)
 struct DISPLAYCLUSTERCONFIGURATION_API FDisplayClusterConfigurationMedia
@@ -77,10 +77,6 @@ public:
 	UPROPERTY(EditAnywhere, BlueprintReadWrite, Category = "Media", meta = (DisplayName = "Enable"))
 	bool bEnabled = false;
 
-	/** When in-cluster media sharing us used, the cluster node specified here will be used as a source (Tx node) */
-	UPROPERTY(EditAnywhere, BlueprintReadWrite, Category = "Media", meta = (EditCondition = "bEnabled"))
-	FString MediaOutputNode;
-
 	/** Media input settings */
 	UPROPERTY(EditAnywhere, BlueprintReadWrite, Category = "Media", meta = (EditCondition = "bEnabled"))
 	FDisplayClusterConfigurationMediaInput MediaInput;
@@ -88,6 +84,21 @@ public:
 	/** Media capture settings */
 	UPROPERTY(EditAnywhere, BlueprintReadWrite, Category = "Media", meta = (EditCondition = "bEnabled"))
 	FDisplayClusterConfigurationMediaOutput MediaOutput;
+};
+
+/*
+ * ICVFX media settings
+ */
+USTRUCT(Blueprintable)
+struct DISPLAYCLUSTERCONFIGURATION_API FDisplayClusterConfigurationICVFXMedia
+	: public FDisplayClusterConfigurationMedia
+{
+	GENERATED_BODY()
+
+public:
+	/** When in-cluster media sharing us used, the cluster node specified here will be used as a source (Tx node) */
+	UPROPERTY(EditAnywhere, BlueprintReadWrite, Category = "Media", meta = (EditCondition = "bEnabled"))
+	FString MediaOutputNode;
 
 public:
 	/** Returns true if media sharing is used */
diff --git a/Engine/Plugins/Runtime/nDisplay/Source/DisplayClusterConfigurator/Private/DisplayClusterConfiguratorModule.cpp b/Engine/Plugins/Runtime/nDisplay/Source/DisplayClusterConfigurator/Private/DisplayClusterConfiguratorModule.cpp
--- a/Engine/Plugins/Runtime/nDisplay/Source/DisplayClusterConfigurator/Private/DisplayClusterConfiguratorModule.cpp
+++ b/Engine/Plugins/Runtime/nDisplay/Source/DisplayClusterConfigurator/Private/DisplayClusterConfiguratorModule.cpp
@@ -17,7 +17,7 @@
 #include "Views/Details/Cluster/DisplayClusterConfiguratorClusterDetailsCustomization.h"
 #include "Views/Details/Cluster/DisplayClusterConfiguratorExternalImageTypeCustomization.h"
 #include "Views/Details/Cluster/DisplayClusterConfiguratorGenerateMipsCustomization.h"
-#include "Views/Details/Cluster/DisplayClusterConfiguratorMediaTypeCustomization.h"
+#include "Views/Details/Cluster/DisplayClusterConfiguratorICVFXMediaCustomization.h"
 #include "Views/Details/Cluster/DisplayClusterConfiguratorNodeSelectionCustomization.h"
 #include "Views/Details/Cluster/DisplayClusterConfiguratorClusterReferenceListCustomization.h"
 #include "Views/Details/Cluster/DisplayClusterConfiguratorViewportDetailsCustomization.h"
@@ -206,7 +206,7 @@ void FDisplayClusterConfiguratorModule::RegisterCustomLayouts()
 	REGISTER_PROPERTY_LAYOUT(FDisplayClusterConfigurationClusterItemReferenceList, FDisplayClusterConfiguratorClusterReferenceListCustomization);
 	REGISTER_PROPERTY_LAYOUT(FDisplayClusterConfigurationViewport_RemapData, FDisplayClusterConfiguratorViewportRemapCustomization);
 	REGISTER_PROPERTY_LAYOUT(FDisplayClusterConfigurationRectangle, FDisplayClusterConfiguratorRectangleCustomization);
-	REGISTER_PROPERTY_LAYOUT(FDisplayClusterConfigurationMedia, FDisplayClusterConfiguratorMediaTypeCustomization);
+	REGISTER_PROPERTY_LAYOUT(FDisplayClusterConfigurationICVFXMedia, FDisplayClusterConfiguratorICVFXMediaCustomization);
 }
 
 void FDisplayClusterConfiguratorModule::UnregisterCustomLayouts()
diff --git a/Engine/Plugins/Runtime/nDisplay/Source/DisplayClusterConfigurator/Private/Views/Details/Cluster/DisplayClusterConfiguratorICVFXMediaCustomization.cpp b/Engine/Plugins/Runtime/nDisplay/Source/DisplayClusterConfigurator/Private/Views/Details/Cluster/DisplayClusterConfiguratorICVFXMediaCustomization.cpp
new file mode 100644
--- /dev/null
+++ b/Engine/Plugins/Runtime/nDisplay/Source/DisplayClusterConfigurator/Private/Views/Details/Cluster/DisplayClusterConfiguratorICVFXMediaCustomization.cpp
@@ -0,0 +1,119 @@
+// Copyright Epic Games, Inc. All Rights Reserved.
+
+#include "Views/Details/Cluster/DisplayClusterConfiguratorICVFXMediaCustomization.h"
+
+#include "DisplayClusterConfigurationTypes.h"
+#include "DisplayClusterConfiguratorBlueprintEditor.h"
+#include "DisplayClusterRootActor.h"
+
+#include "DetailWidgetRow.h"
+#include "IDetailChildrenBuilder.h"
+#include "IPropertyTypeCustomization.h"
+#include "PropertyHandle.h"
+
+
+void FDisplayClusterConfiguratorICVFXMediaCustomization::CustomizeChildren(TSharedRef<IPropertyHandle> InPropertyHandle, IDetailChildrenBuilder& InChildBuilder, IPropertyTypeCustomizationUtils& CustomizationUtils)
+{
+	OutputNodeProperty = InPropertyHandle->GetChildHandle(GET_MEMBER_NAME_CHECKED(FDisplayClusterConfigurationICVFXMedia, MediaOutputNode));
+	check(OutputNodeProperty && OutputNodeProperty->IsValidHandle());
+	OutputNodeProperty->MarkHiddenByCustomization();
+
+	BuildClusterNodeOptionsList();
+
+	// Create GUI representation
+	{
+		FDetailWidgetRow& Row = InChildBuilder.AddCustomRow(OutputNodeProperty->GetPropertyDisplayName())
+			.NameContent()
+			[
+				OutputNodeProperty->CreatePropertyNameWidget()
+			]
+			.ValueContent()
+			[
+				SAssignNew(OutputNodeComboBox, SComboBox<TSharedPtr<FString>>)
+				.OptionsSource(&OutputNodeOptions)
+				.OnGenerateWidget(this, &FDisplayClusterConfiguratorICVFXMediaCustomization::CreateComboWidget)
+				.OnSelectionChanged(this, &FDisplayClusterConfiguratorICVFXMediaCustomization::OnNodeIdSelected)
+				.ContentPadding(2)
+				.Content()
+				[
+					SNew(STextBlock)
+					.Text(this, &FDisplayClusterConfiguratorICVFXMediaCustomization::GetSelectedNodeIdText)
+				]
+			];
+	}
+
+	// Set selected item
+	FString CurrentNodeId;
+	if (OutputNodeProperty->GetValue(CurrentNodeId))
+	{
+		TSharedPtr<FString>* FoundItem = OutputNodeOptions.FindByPredicate([CurrentNodeId](const TSharedPtr<FString>& Item)
+		{
+			return Item->Equals(CurrentNodeId, ESearchCase::IgnoreCase);
+		});
+
+		if (FoundItem)
+		{
+			OutputNodeComboBox->SetSelectedItem(*FoundItem);
+		}
+	}
+
+	FDisplayClusterConfiguratorBaseTypeCustomization::CustomizeChildren(InPropertyHandle, InChildBuilder, CustomizationUtils);
+}
+
+void FDisplayClusterConfiguratorICVFXMediaCustomization::BuildClusterNodeOptionsList()
+{
+	ADisplayClusterRootActor* DCRA = nullptr;
+
+	if (FDisplayClusterConfiguratorBlueprintEditor* Editor = FindBlueprintEditor())
+	{
+		DCRA = Editor->GetDefaultRootActor();
+	}
+
+	if (!DCRA)
+	{
+		DCRA = FindRootActor();
+	}
+	
+	if (!DCRA)
+	{
+		return;
+	}
+
+	// Get cluster node IDs
+	TArray<FString> NodeIds;
+	if (const UDisplayClusterConfigurationData* const ConfigurationData = DCRA->GetConfigData())
+	{
+		ConfigurationData->Cluster->GetNodeIds(NodeIds);
+	}
+
+	// Sort the node IDs in alphabet order
+	NodeIds.Sort();
+
+	// Prepare combobox options
+	OutputNodeOptions.Reset();
+	for (const FString& NodeId : NodeIds)
+	{
+		if (!NodeId.IsEmpty())
+		{
+			OutputNodeOptions.Emplace(MakeShared<FString>(NodeId));
+		}
+	}
+}
+
+
+void FDisplayClusterConfiguratorICVFXMediaCustomization::OnNodeIdSelected(TSharedPtr<FString> PreviewNodeId, ESelectInfo::Type SelectInfo)
+{
+	const FString NewValue = PreviewNodeId.IsValid() ? *PreviewNodeId : FString();
+	OutputNodeProperty->SetValue(NewValue);
+}
+
+FText FDisplayClusterConfiguratorICVFXMediaCustomization::GetSelectedNodeIdText() const
+{
+	TSharedPtr<FString> CurSelection = OutputNodeComboBox->GetSelectedItem();
+	return FText::FromString(CurSelection.IsValid() ? *CurSelection : *FString());
+}
+
+TSharedRef<SWidget> FDisplayClusterConfiguratorICVFXMediaCustomization::CreateComboWidget(TSharedPtr<FString> InItem)
+{
+	return SNew(STextBlock).Text(FText::FromString(*InItem));
+}
diff --git a/Engine/Plugins/Runtime/nDisplay/Source/DisplayClusterConfigurator/Private/Views/Details/Cluster/DisplayClusterConfiguratorICVFXMediaCustomization.h b/Engine/Plugins/Runtime/nDisplay/Source/DisplayClusterConfigurator/Private/Views/Details/Cluster/DisplayClusterConfiguratorICVFXMediaCustomization.h
new file mode 100644
--- /dev/null
+++ b/Engine/Plugins/Runtime/nDisplay/Source/DisplayClusterConfigurator/Private/Views/Details/Cluster/DisplayClusterConfiguratorICVFXMediaCustomization.h
@@ -0,0 +1,47 @@
+// Copyright Epic Games, Inc. All Rights Reserved.
+
+#pragma once
+
+#include "Views/Details/DisplayClusterConfiguratorBaseTypeCustomization.h"
+
+template<class T>
+class SComboBox;
+
+
+class FDisplayClusterConfiguratorICVFXMediaCustomization final
+	: public FDisplayClusterConfiguratorBaseTypeCustomization
+{
+public:
+	static TSharedRef<IPropertyTypeCustomization> MakeInstance()
+	{
+		return MakeShared<FDisplayClusterConfiguratorICVFXMediaCustomization>();
+	}
+
+protected:
+	//~ IPropertyTypeCustomization interface begin
+	virtual void CustomizeChildren(TSharedRef<IPropertyHandle> InPropertyHandle, IDetailChildrenBuilder& InChildBuilder, IPropertyTypeCustomizationUtils& CustomizationUtils) override;
+	//~ IPropertyTypeCustomization interface end
+
+private:
+	/**  Generates internal array of cluster node IDs */
+	void BuildClusterNodeOptionsList();
+
+	/**  Processes combobox selection changes */
+	void OnNodeIdSelected(TSharedPtr<FString> PreviewNodeId, ESelectInfo::Type SelectInfo);
+
+	/**  Returns text of currently selected item */
+	FText GetSelectedNodeIdText() const;
+
+	/** Creates combobox widget */
+	TSharedRef<SWidget> CreateComboWidget(TSharedPtr<FString> InItem);
+
+private:
+	/** OutputNode original property handler */
+	TSharedPtr<IPropertyHandle> OutputNodeProperty;
+
+	/** List of cluster node IDs for the combobox */
+	TArray<TSharedPtr<FString>> OutputNodeOptions;
+
+	/** Node ID selection combobox */
+	TSharedPtr<SComboBox<TSharedPtr<FString>>> OutputNodeComboBox;
+};
diff --git a/Engine/Plugins/Runtime/nDisplay/Source/DisplayClusterConfigurator/Private/Views/Details/Cluster/DisplayClusterConfiguratorMediaTypeCustomization.cpp b/Engine/Plugins/Runtime/nDisplay/Source/DisplayClusterConfigurator/Private/Views/Details/Cluster/DisplayClusterConfiguratorMediaTypeCustomization.cpp
deleted file mode 100644
diff --git a/Engine/Plugins/Runtime/nDisplay/Source/DisplayClusterConfigurator/Private/Views/Details/Cluster/DisplayClusterConfiguratorMediaTypeCustomization.h b/Engine/Plugins/Runtime/nDisplay/Source/DisplayClusterConfigurator/Private/Views/Details/Cluster/DisplayClusterConfiguratorMediaTypeCustomization.h
deleted file mode 100644
diff --git a/Engine/Plugins/Runtime/nDisplay/Source/DisplayClusterMedia/Private/DisplayClusterMediaModule.cpp b/Engine/Plugins/Runtime/nDisplay/Source/DisplayClusterMedia/Private/DisplayClusterMediaModule.cpp
--- a/Engine/Plugins/Runtime/nDisplay/Source/DisplayClusterMedia/Private/DisplayClusterMediaModule.cpp
+++ b/Engine/Plugins/Runtime/nDisplay/Source/DisplayClusterMedia/Private/DisplayClusterMediaModule.cpp
@@ -77,42 +77,14 @@ void FDisplayClusterMediaModule::InitializeMedia()
 			{
 				const FDisplayClusterConfigurationMedia& MediaSettings = ClusterNode->Media;
 
-				const FString MediaCaptureId = RootActor->GetName() + FString("_") + ClusterNodeId + FString("_backbuffer_capture");
-				const FString MediaInputId   = RootActor->GetName() + FString("_") + ClusterNodeId + FString("_backbuffer_input");
 				const bool bInputDataValid   = MediaSettings.MediaInput.MediaPlayer && MediaSettings.MediaInput.MediaSource && MediaSettings.MediaInput.MediaTexture;
 				const bool bOutputDataValid  = MediaSettings.MediaOutput.MediaOutput && MediaSettings.MediaOutput.RenderTarget;
 
-				// Media sharing
-				if (MediaSettings.IsMediaSharingRequired() && bInputDataValid && bOutputDataValid)
+				// Media input
+				if (MediaSettings.MediaInput.bEnabled && bInputDataValid)
 				{
-					const bool bThisNodeSharesMedia = ClusterNodeId.Equals(MediaSettings.MediaOutputNode, ESearchCase::IgnoreCase);
+					const FString MediaInputId = RootActor->GetName() + FString("_") + ClusterNodeId + FString("_backbuffer_input");
 
-					// Share media
-					if (bThisNodeSharesMedia)
-					{
-						UE_LOG(LogDisplayClusterMedia, Log, TEXT("Initializing backbuffer media capture for node '%s' (in-cluster shared)..."), *ClusterNodeId);
-
-						CaptureNode = MakeUnique<FDisplayClusterMediaCaptureNode>(
-							MediaCaptureId, ClusterNodeId,
-							MediaSettings.MediaOutput.MediaOutput, MediaSettings.MediaOutput.RenderTarget);
-
-						AllCaptures.Add(CaptureNode.Get());
-					}
-					// Consume media
-					else
-					{
-						UE_LOG(LogDisplayClusterMedia, Log, TEXT("Initializing backbuffer media input for node '%s' (in-cluster shared)..."), *ClusterNodeId);
-
-						InputNode = MakeUnique<FDisplayClusterMediaInputNode>(
-							MediaInputId, ClusterNodeId,
-							MediaSettings.MediaInput.MediaSource, MediaSettings.MediaInput.MediaPlayer, MediaSettings.MediaInput.MediaTexture);
-
-						AllInputs.Add(InputNode.Get());
-					}
-				}
-				// Media input only
-				else if (MediaSettings.MediaInput.bEnabled && bInputDataValid)
-				{
 					UE_LOG(LogDisplayClusterMedia, Log, TEXT("Initializing backbuffer media input for node '%s'..."), *ClusterNodeId);
 
 					InputNode = MakeUnique<FDisplayClusterMediaInputNode>(
@@ -121,9 +93,12 @@ void FDisplayClusterMediaModule::InitializeMedia()
 
 					AllInputs.Add(InputNode.Get());
 				}
-				// Media capture only
-				else if (MediaSettings.MediaOutput.bEnabled && bOutputDataValid)
+
+				// Media capture
+				if (MediaSettings.MediaOutput.bEnabled && bOutputDataValid)
 				{
+					const FString MediaCaptureId = RootActor->GetName() + FString("_") + ClusterNodeId + FString("_backbuffer_capture");
+
 					UE_LOG(LogDisplayClusterMedia, Log, TEXT("Initializing backbuffer media capture for node '%s'..."), *ClusterNodeId);
 
 					CaptureNode = MakeUnique<FDisplayClusterMediaCaptureNode>(
@@ -143,46 +118,14 @@ void FDisplayClusterMediaModule::InitializeMedia()
 					const FDisplayClusterConfigurationMedia& MediaSettings = Viewport->RenderSettings.Media;
 					if (MediaSettings.bEnabled)
 					{
-						const FString MediaCaptureId = RootActor->GetName() + FString("_") + ClusterNodeId + FString("_") + ViewportIt.Key + FString("_viewport_capture");
-						const FString MediaInputId   = RootActor->GetName() + FString("_") + ClusterNodeId + FString("_") + ViewportIt.Key + FString("_viewport_input");
 						const bool bInputDataValid   = MediaSettings.MediaInput.MediaPlayer && MediaSettings.MediaInput.MediaSource && MediaSettings.MediaInput.MediaTexture;
 						const bool bOutputDataValid  = MediaSettings.MediaOutput.MediaOutput && MediaSettings.MediaOutput.RenderTarget;
 
-						// Media sharing
-						if (MediaSettings.IsMediaSharingRequired() && bInputDataValid && bOutputDataValid)
+						// Media input
+						if (MediaSettings.MediaInput.bEnabled && bInputDataValid)
 						{
-							const bool bThisNodeSharesMedia = ClusterNodeId.Equals(MediaSettings.MediaOutputNode, ESearchCase::IgnoreCase);
-
-							// Share media
-							if (bThisNodeSharesMedia)
-							{
-								UE_LOG(LogDisplayClusterMedia, Log, TEXT("Initializing viewport media capture for viewport '%s' (in-cluster shared)..."), *ClusterNodeId);
+							const FString MediaInputId = RootActor->GetName() + FString("_") + ClusterNodeId + FString("_") + ViewportIt.Key + FString("_viewport_input");
 
-								TUniquePtr<FDisplayClusterMediaCaptureViewport> NewViewportCapture = MakeUnique<FDisplayClusterMediaCaptureViewport>(
-									MediaCaptureId, ClusterNodeId,
-									ViewportIt.Key,
-									MediaSettings.MediaOutput.MediaOutput, MediaSettings.MediaOutput.RenderTarget);
-
-								AllCaptures.Add(NewViewportCapture.Get());
-								CaptureViewports.Emplace(MediaCaptureId, MoveTemp(NewViewportCapture));
-							}
-							// Consume media
-							else
-							{
-								UE_LOG(LogDisplayClusterMedia, Log, TEXT("Initializing viewport media input for viewport '%s' (in-cluster shared)..."), *ClusterNodeId);
-
-								TUniquePtr<FDisplayClusterMediaInputViewport> NewViewportInput = MakeUnique<FDisplayClusterMediaInputViewport>(
-									MediaInputId, ClusterNodeId,
-									ViewportIt.Key,
-									MediaSettings.MediaInput.MediaSource, MediaSettings.MediaInput.MediaPlayer, MediaSettings.MediaInput.MediaTexture);
-
-								AllInputs.Add(NewViewportInput.Get());
-								InputViewports.Emplace(MediaCaptureId, MoveTemp(NewViewportInput));
-							}
-						}
-						// Media input only
-						else if (MediaSettings.MediaInput.bEnabled && bInputDataValid)
-						{
 							UE_LOG(LogDisplayClusterMedia, Log, TEXT("Initializing viewport media input for viewport '%s'..."), *ClusterNodeId);
 
 							TUniquePtr<FDisplayClusterMediaInputViewport> NewViewportInput = MakeUnique<FDisplayClusterMediaInputViewport>(
@@ -191,11 +134,14 @@ void FDisplayClusterMediaModule::InitializeMedia()
 								MediaSettings.MediaInput.MediaSource, MediaSettings.MediaInput.MediaPlayer, MediaSettings.MediaInput.MediaTexture);
 
 							AllInputs.Add(NewViewportInput.Get());
-							InputViewports.Emplace(MediaCaptureId, MoveTemp(NewViewportInput));
+							InputViewports.Emplace(MediaInputId, MoveTemp(NewViewportInput));
 						}
-						// Media capture only
-						else if (MediaSettings.MediaOutput.bEnabled && bOutputDataValid)
+
+						// Media capture
+						if (MediaSettings.MediaOutput.bEnabled && bOutputDataValid)
 						{
+							const FString MediaCaptureId = RootActor->GetName() + FString("_") + ClusterNodeId + FString("_") + ViewportIt.Key + FString("_viewport_capture");
+
 							UE_LOG(LogDisplayClusterMedia, Log, TEXT("Initializing viewport capture for viewport '%s'..."), *ViewportIt.Key);
 
 							TUniquePtr<FDisplayClusterMediaCaptureViewport> NewViewportCapture = MakeUnique<FDisplayClusterMediaCaptureViewport>(
@@ -220,7 +166,7 @@ void FDisplayClusterMediaModule::InitializeMedia()
 				{
 					if (const UDisplayClusterICVFXCameraComponent* const ICVFXCamera = Cast<UDisplayClusterICVFXCameraComponent>(Component))
 					{
-						const FDisplayClusterConfigurationMedia& MediaSettings = ICVFXCamera->CameraSettings.RenderSettings.Media;
+						const FDisplayClusterConfigurationICVFXMedia& MediaSettings = ICVFXCamera->CameraSettings.RenderSettings.Media;
 
 						if (MediaSettings.bEnabled)
 						{
@@ -259,7 +205,7 @@ void FDisplayClusterMediaModule::InitializeMedia()
 										MediaSettings.MediaInput.MediaSource, MediaSettings.MediaInput.MediaPlayer, MediaSettings.MediaInput.MediaTexture);
 
 									AllInputs.Add(NewICVFXInput.Get());
-									InputViewports.Emplace(MediaCaptureId, MoveTemp(NewICVFXInput));
+									InputViewports.Emplace(MediaInputId, MoveTemp(NewICVFXInput));
 								}
 							}
 							// Media input only
@@ -273,7 +219,7 @@ void FDisplayClusterMediaModule::InitializeMedia()
 									MediaSettings.MediaInput.MediaSource, MediaSettings.MediaInput.MediaPlayer, MediaSettings.MediaInput.MediaTexture);
 
 								AllInputs.Add(NewICVFXInput.Get());
-								InputViewports.Emplace(MediaCaptureId, MoveTemp(NewICVFXInput));
+								InputViewports.Emplace(MediaInputId, MoveTemp(NewICVFXInput));
 							}
 							// Media capture only
 							else if (MediaSettings.MediaOutput.bEnabled && bOutputDataValid)
