Fixed several issues where the nDisplay actor would dereference a null pointer if the nDisplay config isn't setup.

diff --git a/Engine/Plugins/Runtime/nDisplay/Source/DisplayCluster/Private/Game/EngineClasses/Scene/DisplayClusterRootActorEditor.cpp b/Engine/Plugins/Runtime/nDisplay/Source/DisplayCluster/Private/Game/EngineClasses/Scene/DisplayClusterRootActorEditor.cpp
--- a/Engine/Plugins/Runtime/nDisplay/Source/DisplayCluster/Private/Game/EngineClasses/Scene/DisplayClusterRootActorEditor.cpp
+++ b/Engine/Plugins/Runtime/nDisplay/Source/DisplayCluster/Private/Game/EngineClasses/Scene/DisplayClusterRootActorEditor.cpp
@@ -424,7 +424,7 @@ bool ADisplayClusterRootActor::ImplRenderPassPreviewClusterNode_Editor()
 
 void ADisplayClusterRootActor::ImplRenderPreview_Editor()
 {
-	if (!bPreviewEnable || !ViewportManager.IsValid())
+	if (CurrentConfigData == nullptr || !bPreviewEnable || !ViewportManager.IsValid())
 	{
 		// no preview
 		return;
@@ -499,7 +499,7 @@ void ADisplayClusterRootActor::ImplRenderPreview_Editor()
 
 void ADisplayClusterRootActor::ImplRenderPreviewFrustums_Editor()
 {
-	if (!ViewportManager.IsValid())
+	if (CurrentConfigData == nullptr || !ViewportManager.IsValid())
 	{
 		return;
 	}
