CIS Fix for Compile QAVirtualProduction Win64 - nDisplay

diff --git a/Engine/Plugins/Runtime/nDisplay/Source/DisplayClusterRemoteControlInterceptor/DisplayClusterRemoteControlInterceptor.Build.cs b/Engine/Plugins/Runtime/nDisplay/Source/DisplayClusterRemoteControlInterceptor/DisplayClusterRemoteControlInterceptor.Build.cs
--- a/Engine/Plugins/Runtime/nDisplay/Source/DisplayClusterRemoteControlInterceptor/DisplayClusterRemoteControlInterceptor.Build.cs
+++ b/Engine/Plugins/Runtime/nDisplay/Source/DisplayClusterRemoteControlInterceptor/DisplayClusterRemoteControlInterceptor.Build.cs
@@ -8,6 +8,8 @@ public class DisplayClusterRemoteControlInterceptor : ModuleRules
 {
 	public DisplayClusterRemoteControlInterceptor(ReadOnlyTargetRules ROTargetRules) : base(ROTargetRules)
 	{
+		ShortName = "RCInterceptor";
+
 		PrivateDependencyModuleNames.AddRange(
 			new string[] {
 				"Core",
