nDisplay: disabled screen percentage for LightCards/Chromakey rendertargets

diff --git a/Engine/Plugins/Runtime/nDisplay/Source/DisplayCluster/Private/Game/EngineClasses/Basics/DisplayClusterViewportClient.cpp b/Engine/Plugins/Runtime/nDisplay/Source/DisplayCluster/Private/Game/EngineClasses/Basics/DisplayClusterViewportClient.cpp
--- a/Engine/Plugins/Runtime/nDisplay/Source/DisplayCluster/Private/Game/EngineClasses/Basics/DisplayClusterViewportClient.cpp
+++ b/Engine/Plugins/Runtime/nDisplay/Source/DisplayCluster/Private/Game/EngineClasses/Basics/DisplayClusterViewportClient.cpp
@@ -745,7 +745,7 @@ void UDisplayClusterViewportClient::Draw(FViewport* InViewport, FCanvas* SceneCa
 					if (ViewFamily.EngineShowFlags.ScreenPercentage)
 					{
 						// Get global view fraction set by r.ScreenPercentage.
-						GlobalResolutionFraction = FLegacyScreenPercentageDriver::GetCVarResolutionFraction() * CustomBufferRatio;
+						GlobalResolutionFraction = CustomBufferRatio;
 
 						// We need to split the screen percentage if below 0.5 because TAA upscaling only works well up to 2x.
 						if (GlobalResolutionFraction < 0.5f)
diff --git a/Engine/Plugins/Runtime/nDisplay/Source/DisplayCluster/Private/Render/Viewport/DisplayClusterViewport.cpp b/Engine/Plugins/Runtime/nDisplay/Source/DisplayCluster/Private/Render/Viewport/DisplayClusterViewport.cpp
--- a/Engine/Plugins/Runtime/nDisplay/Source/DisplayCluster/Private/Render/Viewport/DisplayClusterViewport.cpp
+++ b/Engine/Plugins/Runtime/nDisplay/Source/DisplayCluster/Private/Render/Viewport/DisplayClusterViewport.cpp
@@ -559,7 +559,18 @@ bool FDisplayClusterViewport::UpdateFrameContexts(const uint32 InStereoViewIndex
 		Context.RenderTargetRect = RenderTargetRect;
 		Context.ContextSize = ContextSize;
 
-		Context.CustomBufferRatio = CustomBufferRatio;
+		// r.ScreenPercentage
+		switch (RenderSettings.CaptureMode)
+		{
+		case EDisplayClusterViewportCaptureMode::Chromakey:
+		case EDisplayClusterViewportCaptureMode::Lightcard:
+			// we should not change the size of the Chromakey\Lighcards due to the way copy\resolve works for RT's
+			// if the viewfamily resolves to RenderTarget it will remove alpha channel
+			// if the viewfamily copying to RenderTarget, the texture would not match the size of RTT (when ScreenPercentage applied)
+			break;
+		default:
+			Context.CustomBufferRatio = CustomBufferRatio;
+		}
 
 		Context.bDisableRender = bDisableRender;
 
diff --git a/Engine/Plugins/Runtime/nDisplay/Source/DisplayCluster/Private/Render/Viewport/DisplayClusterViewportManager.cpp b/Engine/Plugins/Runtime/nDisplay/Source/DisplayCluster/Private/Render/Viewport/DisplayClusterViewportManager.cpp
--- a/Engine/Plugins/Runtime/nDisplay/Source/DisplayCluster/Private/Render/Viewport/DisplayClusterViewportManager.cpp
+++ b/Engine/Plugins/Runtime/nDisplay/Source/DisplayCluster/Private/Render/Viewport/DisplayClusterViewportManager.cpp
@@ -28,6 +28,7 @@
 #include "DisplayClusterRootActor.h" 
 
 #include "Misc/DisplayClusterLog.h"
+#include "LegacyScreenPercentageDriver.h"
 
 #include "Engine/Console.h"
 
@@ -179,13 +180,14 @@ void FDisplayClusterViewportManager::ResetScene()
 
 void FDisplayClusterViewportManager::SetViewportBufferRatio(FDisplayClusterViewport& DstViewport, float InBufferRatio)
 {
-	if (DstViewport.RenderSettings.BufferRatio > InBufferRatio)
+	const float BufferRatio = FLegacyScreenPercentageDriver::GetCVarResolutionFraction() * InBufferRatio;
+	if (DstViewport.RenderSettings.BufferRatio > BufferRatio)
 	{
 		// Reset scene RTT when buffer ratio changed down
 		ResetSceneRenderTargetSize();
 	}
 
-	DstViewport.RenderSettings.BufferRatio = InBufferRatio;
+	DstViewport.RenderSettings.BufferRatio = BufferRatio;
 }
 
 void FDisplayClusterViewportManager::HandleViewportRTTChanges(const TArray<FDisplayClusterViewport_Context>& PrevContexts, const TArray<FDisplayClusterViewport_Context>& Contexts)
