nDisplay: Sanitize cluster and viewport names to uobject friendly names. This prevents crashes when spaces or special characters are added.

diff --git a/Engine/Plugins/Runtime/nDisplay/Source/DisplayClusterConfigurator/Private/ClusterConfiguration/DisplayClusterConfiguratorClusterUtils.cpp b/Engine/Plugins/Runtime/nDisplay/Source/DisplayClusterConfigurator/Private/ClusterConfiguration/DisplayClusterConfiguratorClusterUtils.cpp
--- a/Engine/Plugins/Runtime/nDisplay/Source/DisplayClusterConfigurator/Private/ClusterConfiguration/DisplayClusterConfiguratorClusterUtils.cpp
+++ b/Engine/Plugins/Runtime/nDisplay/Source/DisplayClusterConfigurator/Private/ClusterConfiguration/DisplayClusterConfiguratorClusterUtils.cpp
@@ -14,6 +14,7 @@
 
 #include "Factories.h"
 #include "ISinglePropertyView.h"
+#include "ObjectTools.h"
 #include "Framework/Application/SlateApplication.h"
 #include "HAL/PlatformApplicationMisc.h"
 #include "UnrealExporter.h"
@@ -609,6 +610,8 @@ bool FDisplayClusterConfiguratorClusterUtils::IsClusterNodePrimary(UDisplayClust
 
 FString FDisplayClusterConfiguratorClusterUtils::GetUniqueNameForClusterNode(FString InitialName, UDisplayClusterConfigurationCluster* ParentCluster, bool bAddZero)
 {
+	InitialName = ObjectTools::SanitizeObjectName(InitialName);
+	
 	TArray<FString> UsedNames;
 	ParentCluster->Nodes.GenerateKeyArray(UsedNames);
 
@@ -773,6 +776,8 @@ FString FDisplayClusterConfiguratorClusterUtils::GetViewportName(UDisplayCluster
 
 FString FDisplayClusterConfiguratorClusterUtils::GetUniqueNameForViewport(FString InitialName, UDisplayClusterConfigurationClusterNode* ParentClusterNode, bool bAddZero)
 {
+	InitialName = ObjectTools::SanitizeObjectName(InitialName);
+	
 	// Viewport names must be unique across the entire cluster, not just within its parent cluster nodes. Gather all of the viewport names
 	// in the cluster to check for uniqueness. Add the parent cluster node's viewports first, in case we can't get to the root cluster through
 	// the cluster node's Outer (i.e. the cluster node has not been added to the cluster yet)
