nDisplay: Fix crash on config editor open. Result of changes in 21085991 allowing customization on CreateSingleProperty.

diff --git a/Engine/Plugins/Runtime/nDisplay/Source/DisplayClusterConfigurator/Private/Views/Details/DisplayClusterConfiguratorBaseTypeCustomization.cpp b/Engine/Plugins/Runtime/nDisplay/Source/DisplayClusterConfigurator/Private/Views/Details/DisplayClusterConfiguratorBaseTypeCustomization.cpp
--- a/Engine/Plugins/Runtime/nDisplay/Source/DisplayClusterConfigurator/Private/Views/Details/DisplayClusterConfiguratorBaseTypeCustomization.cpp
+++ b/Engine/Plugins/Runtime/nDisplay/Source/DisplayClusterConfigurator/Private/Views/Details/DisplayClusterConfiguratorBaseTypeCustomization.cpp
@@ -23,6 +23,11 @@ void FDisplayClusterConfiguratorBaseTypeCustomization::CustomizeHeader(TSharedRe
 	check(InPropertyHandle->IsValidHandle());
 
 	PropertyUtilities = CustomizationUtils.GetPropertyUtilities();
+	if (PropertyUtilities == nullptr)
+	{
+		// Can be null if PropertyEditorModule.CreateSingleProperty is what triggers this customization.
+		return;
+	}
 	EditingObjects = PropertyUtilities.Pin()->GetSelectedObjects();
 
 	if (EditingObjects.Num())
