nDisplay: Added UV projection mode to the light card editor, and added the concept of UV light cards, which are light cards that are rendered to UV space which are subsequently rendered to the screens via the second channel UV coordinates of a screen mesh.

diff --git a/Engine/Plugins/Runtime/nDisplay/Shaders/Private/MPCDIOverlayShaders.usf b/Engine/Plugins/Runtime/nDisplay/Shaders/Private/MPCDIOverlayShaders.usf
--- a/Engine/Plugins/Runtime/nDisplay/Shaders/Private/MPCDIOverlayShaders.usf
+++ b/Engine/Plugins/Runtime/nDisplay/Shaders/Private/MPCDIOverlayShaders.usf
@@ -174,6 +174,12 @@ float4 GetInnerCameraColor(float4 WarpedUV, OutputVS IN)
 	return CameraColor;
 }
 
+float3 GetUVLightCard(float3 SourceColor, float2 UV)
+{
+	float4 LightCardColor = UVLightCardMapTexture.Sample(UVLightCardMapSampler, UV);
+	return lerp(SourceColor.rgb, LightCardColor.rgb, LightCardColor.a);
+}
+
 OutputPS IcvfxWarpPS(OutputVS IN)
 {
 	OutputPS OUT;
@@ -205,6 +211,10 @@ OutputPS IcvfxWarpPS(OutputVS IN)
 	#endif
 
 	#if OVERLAY_UNDER
+		#if UV_LIGHT_CARDS
+			OutColor.xyz = GetUVLightCard(OutColor.xyz, IN.UV_Chromakey.xy);
+		#endif
+	
 		OutColor.xyz = ComposeIcvfxOverlay(OutColor.xyz, OverlayUV);
 	#endif
 	
@@ -214,6 +224,10 @@ OutputPS IcvfxWarpPS(OutputVS IN)
 	#endif
 
 	#if OVERLAY_OVER
+		#if UV_LIGHT_CARDS
+			OutColor.xyz = GetUVLightCard(OutColor.xyz, IN.UV_Chromakey.xy);
+		#endif
+
 		OutColor.xyz = ComposeIcvfxOverlay(OutColor.xyz, OverlayUV);
 	#endif
 
diff --git a/Engine/Plugins/Runtime/nDisplay/Shaders/Private/MPCDIUtils.ush b/Engine/Plugins/Runtime/nDisplay/Shaders/Private/MPCDIUtils.ush
--- a/Engine/Plugins/Runtime/nDisplay/Shaders/Private/MPCDIUtils.ush
+++ b/Engine/Plugins/Runtime/nDisplay/Shaders/Private/MPCDIUtils.ush
@@ -19,6 +19,7 @@ Texture2D OverlayAlphaTexture;
 Texture2D InnerCameraTexture;
 Texture2D ChromakeyCameraTexture;
 Texture2D ChromakeyMarkerTexture;
+Texture2D UVLightCardMapTexture;
 
 SamplerState InputSampler;
 SamplerState WarpMapSampler;
@@ -29,6 +30,7 @@ SamplerState OverlayAlphaSampler;
 SamplerState InnerCameraSampler;
 SamplerState ChromakeyCameraSampler;
 SamplerState ChromakeyMarkerSampler;
+SamplerState UVLightCardMapSampler;
 
 float4x4 MeshToStageProjectionMatrix;
 float4x4 ViewportTextureProjectionMatrix;
diff --git a/Engine/Plugins/Runtime/nDisplay/Shaders/Private/MeshProjectionShaders.usf b/Engine/Plugins/Runtime/nDisplay/Shaders/Private/MeshProjectionShaders.usf
--- a/Engine/Plugins/Runtime/nDisplay/Shaders/Private/MeshProjectionShaders.usf
+++ b/Engine/Plugins/Runtime/nDisplay/Shaders/Private/MeshProjectionShaders.usf
@@ -63,6 +63,40 @@ void PerspectiveVS(FVertexFactoryInput Input, out FMeshProjectionVSToPS Output)
 	Output.Position = mul(Output.SavedWorldPosition, ResolvedView.TranslatedWorldToClip);
 }
 
+uint ProjectionParameters_UVProjectionIndex;
+float ProjectionParameters_UVProjectionPlaneSize;
+float ProjectionParameters_UVProjectionPlaneDistance;
+
+void UVProjectionVS(FVertexFactoryInput Input, out FMeshProjectionVSToPS Output)
+{
+	ResolvedView = ResolveView();
+
+	const FVertexFactoryIntermediates VFIntermediates = GetVertexFactoryIntermediates(Input);
+	Output.SavedWorldPosition = VertexFactoryGetWorldPosition(Input, VFIntermediates);
+	const float3x3 TangentToLocal = VertexFactoryGetTangentToLocal(Input, VFIntermediates);
+
+	const FMaterialVertexParameters VertexParameters = GetMaterialVertexParameters(Input, VFIntermediates, Output.SavedWorldPosition.xyz, TangentToLocal);
+	Output.SavedWorldPosition.xyz += GetMaterialWorldPositionOffset(VertexParameters);
+
+	Output.FactoryInterpolants = VertexFactoryGetInterpolantsVSToPS(Input, VFIntermediates, VertexParameters);
+
+	float2 UV = float2(0, 0);
+	
+#if NUM_TEX_COORD_INTERPOLATORS
+	int UVIndex = min(ProjectionParameters_UVProjectionIndex, NUM_TEX_COORD_INTERPOLATORS - 1);
+	
+	UV = GetUV(Output.FactoryInterpolants, UVIndex);
+	
+	const float3 PlanePosition = float3(ProjectionParameters_UVProjectionPlaneSize * (UV.x - 0.5),ProjectionParameters_UVProjectionPlaneSize * (0.5 - UV.y), ProjectionParameters_UVProjectionPlaneDistance);
+	
+	Output.Azimuth = UVIndex;
+	Output.Position = mul(float4(PlanePosition, 1.0), ResolvedView.ViewToClip);
+#else
+	Output.Azimuth = 0;
+	Output.Position = mul(Output.SavedWorldPosition, ResolvedView.TranslatedWorldToClip);
+#endif
+}
+
 /** Gets a scalar value that indicates how much the current pixel diverges from where it should be from the projected vertices */
 float GetProjectionDivergence(FMeshProjectionVSToPS Input)
 {
diff --git a/Engine/Plugins/Runtime/nDisplay/Shaders/Private/UVLightCardShaders.usf b/Engine/Plugins/Runtime/nDisplay/Shaders/Private/UVLightCardShaders.usf
new file mode 100644
--- /dev/null
+++ b/Engine/Plugins/Runtime/nDisplay/Shaders/Private/UVLightCardShaders.usf
@@ -0,0 +1,48 @@
+// Copyright Epic Games, Inc. All Rights Reserved.
+
+#include "/Engine/Private/Common.ush"
+#include "/Engine/Generated/Material.ush"
+#include "/Engine/Generated/VertexFactory.ush"
+
+struct FUVLightCardVSToPS
+{
+	FVertexFactoryInterpolantsVSToPS FactoryInterpolants;
+	float4 Position : SV_POSITION;
+	float4 SavedWorldPosition : POSITION1;
+};
+
+struct OutputPS
+{
+	float4 Color : SV_Target0;
+};
+
+
+void MainVS(FVertexFactoryInput Input, out FUVLightCardVSToPS Output)
+{
+	ResolvedView = ResolveView();
+
+	const FVertexFactoryIntermediates VFIntermediates = GetVertexFactoryIntermediates(Input);
+	Output.SavedWorldPosition = VertexFactoryGetWorldPosition(Input, VFIntermediates);
+	const float3x3 TangentToLocal = VertexFactoryGetTangentToLocal(Input, VFIntermediates);
+
+	const FMaterialVertexParameters VertexParameters = GetMaterialVertexParameters(Input, VFIntermediates, Output.SavedWorldPosition.xyz, TangentToLocal);
+	Output.SavedWorldPosition.xyz += GetMaterialWorldPositionOffset(VertexParameters);
+
+	Output.FactoryInterpolants = VertexFactoryGetInterpolantsVSToPS(Input, VFIntermediates, VertexParameters);
+	
+	Output.Position = mul(Output.SavedWorldPosition, ResolvedView.TranslatedWorldToClip);
+}
+
+void MainPS(FUVLightCardVSToPS Input OPTIONAL_IsFrontFace, out float4 OutputColor : SV_Target0)
+{
+	FMaterialPixelParameters MaterialParameters = GetMaterialPixelParameters(Input.FactoryInterpolants, Input.Position);
+	FPixelMaterialInputs PixelMaterialInputs;
+	CalcMaterialParameters(MaterialParameters, PixelMaterialInputs, Input.Position, bIsFrontFace);
+
+	GetMaterialCoverageAndClipping(MaterialParameters, PixelMaterialInputs);
+	
+	const float3 Emissive = GetMaterialEmissive(PixelMaterialInputs);
+	const float Opacity = GetMaterialOpacity(PixelMaterialInputs);
+	
+	OutputColor = float4(Emissive, Opacity);
+}
diff --git a/Engine/Plugins/Runtime/nDisplay/Source/DisplayCluster/Private/Game/EngineClasses/Scene/DisplayClusterLightCardActor.cpp b/Engine/Plugins/Runtime/nDisplay/Source/DisplayCluster/Private/Game/EngineClasses/Scene/DisplayClusterLightCardActor.cpp
--- a/Engine/Plugins/Runtime/nDisplay/Source/DisplayCluster/Private/Game/EngineClasses/Scene/DisplayClusterLightCardActor.cpp
+++ b/Engine/Plugins/Runtime/nDisplay/Source/DisplayCluster/Private/Game/EngineClasses/Scene/DisplayClusterLightCardActor.cpp
@@ -27,6 +27,8 @@ static FAutoConsoleVariableRef CVarDisplayClusterLightCardPolygonTextureSize(
 );
 
 const FRotator ADisplayClusterLightCardActor::PlaneMeshRotation = FRotator(0.0f, -90.0f, 90.0f);
+const float ADisplayClusterLightCardActor::UVPlaneDefaultSize = 200.0f;
+const float ADisplayClusterLightCardActor::UVPlaneDefaultDistance = 100.0f;
 
 ADisplayClusterLightCardActor::ADisplayClusterLightCardActor(const FObjectInitializer& ObjectInitializer)
 	: Super(ObjectInitializer)
@@ -119,6 +121,7 @@ void ADisplayClusterLightCardActor::Tick(float DeltaSeconds)
 
 	UpdateLightCardTransform();
 	UpdateLightCardMaterialInstance();
+	UpdateLightCardVisibility();
 }
 
 #if WITH_EDITOR
@@ -133,7 +136,8 @@ void ADisplayClusterLightCardActor::PostEditChangeProperty(FPropertyChangedEvent
 		PropertyChangedEvent.Property->GetFName() == GET_MEMBER_NAME_CHECKED(ADisplayClusterLightCardActor, Latitude) ||
 		PropertyChangedEvent.Property->GetFName() == GET_MEMBER_NAME_CHECKED(ADisplayClusterLightCardActor, Spin) ||
 		PropertyChangedEvent.Property->GetFName() == GET_MEMBER_NAME_CHECKED(ADisplayClusterLightCardActor, RadialOffset) ||
-		PropertyChangedEvent.Property->GetFName() == GET_MEMBER_NAME_CHECKED(ADisplayClusterLightCardActor, Scale)))
+		PropertyChangedEvent.Property->GetFName() == GET_MEMBER_NAME_CHECKED(ADisplayClusterLightCardActor, Scale) ||
+		PropertyChangedEvent.Property->GetFName() == GET_MEMBER_NAME_CHECKED(ADisplayClusterLightCardActor, bIsUVLightCard)))
 	{
 		UpdateLightCardTransform();
 	}
@@ -159,6 +163,12 @@ void ADisplayClusterLightCardActor::PostEditChangeProperty(FPropertyChangedEvent
 	{
 		UpdateLightCardMaterialInstance();
 	}
+
+	if (PropertyChangedEvent.Property && (
+		PropertyChangedEvent.Property->GetFName() == GET_MEMBER_NAME_CHECKED(ADisplayClusterLightCardActor, bIsUVLightCard)))
+	{
+		UpdateLightCardVisibility();
+	}
 }
 
 #endif
@@ -203,8 +213,24 @@ FBox ADisplayClusterLightCardActor::GetLightCardBounds(bool bLocalSpace) const
 
 void ADisplayClusterLightCardActor::UpdateLightCardTransform()
 {
-	MainSpringArmComponent->TargetArmLength = DistanceFromCenter + RadialOffset;
-	MainSpringArmComponent->SetRelativeRotation(FRotator(-Latitude, Longitude, 0.0));
+	// If the light card is in UV space, set the spring arm's trasform to be zero, effectively removing it from the transform hierarchy
+	// This allows the light card to be positioned with the actor's cartesian coordinates instead of longitude and latitude
+	if (bIsUVLightCard)
+	{
+		MainSpringArmComponent->TargetArmLength = 0.0f;
+		MainSpringArmComponent->SetRelativeRotation(FRotator(0.0, 180.0, 0.0));
+
+		// Set world location and rotation such that the light card is always projected onto a YZ plane a distance of UVPlaneDefaultDistance from the world origin, facing in the -X direction
+		// This ensures that when the UV light cards are rendered to the light card map, they are always positioned and oriented correctly regardless of the stages location and rotation.
+		LightCardTransformerComponent->SetWorldLocation(FVector(UVPlaneDefaultDistance, -UVPlaneDefaultSize * (0.5 - UVCoordinates.X), UVPlaneDefaultSize * (0.5 - UVCoordinates.Y)));
+		LightCardTransformerComponent->SetWorldRotation(FVector(-1, 0, 0).Rotation());
+		LightCardTransformerComponent->SetWorldScale3D(FVector::OneVector);
+	}
+	else
+	{
+		MainSpringArmComponent->TargetArmLength = DistanceFromCenter + RadialOffset;
+		MainSpringArmComponent->SetRelativeRotation(FRotator(-Latitude, Longitude, 0.0));
+	}
 
 	FRotator LightCardOrientation = FRotator(-Pitch, Yaw, Spin);
 
@@ -352,6 +378,25 @@ void ADisplayClusterLightCardActor::UpdatePolygonTexture()
 #endif // WITH_OPENCV
 }
 
+void ADisplayClusterLightCardActor::UpdateLightCardVisibility()
+{
+	const bool bShouldBeVisible = !bIsUVLightCard || bIsProxy;
+	const bool bIsVisible = !IsHidden();
+	if (bIsVisible != bShouldBeVisible)
+	{
+		SetActorHiddenInGame(!bShouldBeVisible);
+	}
+
+#if WITH_EDITORONLY_DATA
+	const bool bShouldHideFromEditor = !bShouldBeVisible;
+	if (bShouldHideFromEditor != bHiddenEdLevel)
+	{
+		bHiddenEdLevel = bShouldHideFromEditor;
+		bHiddenEdLayer = bShouldHideFromEditor;
+		MarkComponentsRenderStateDirty();
+	}
+#endif
+}
 
 UStaticMesh* ADisplayClusterLightCardActor::GetStaticMesh() const
 {
diff --git a/Engine/Plugins/Runtime/nDisplay/Source/DisplayCluster/Private/Render/Viewport/DisplayClusterViewportManager.cpp b/Engine/Plugins/Runtime/nDisplay/Source/DisplayCluster/Private/Render/Viewport/DisplayClusterViewportManager.cpp
--- a/Engine/Plugins/Runtime/nDisplay/Source/DisplayCluster/Private/Render/Viewport/DisplayClusterViewportManager.cpp
+++ b/Engine/Plugins/Runtime/nDisplay/Source/DisplayCluster/Private/Render/Viewport/DisplayClusterViewportManager.cpp
@@ -12,6 +12,7 @@
 #include "Render/Viewport/DisplayClusterViewport.h"
 #include "Render/Viewport/DisplayClusterViewportProxy.h"
 #include "Render/Viewport/DisplayClusterViewport_CustomPostProcessSettings.h"
+#include "Render/Viewport/LightCard/DisplayClusterViewportLightCardManager.h"
 #include "Render/Viewport/RenderFrame/DisplayClusterRenderFrameManager.h"
 #include "Render/Viewport/RenderTarget/DisplayClusterRenderTargetManager.h"
 #include "Render/Viewport/RenderTarget/DisplayClusterRenderTargetResource.h"
@@ -63,6 +64,7 @@ FDisplayClusterViewportManager::FDisplayClusterViewportManager()
 
 	RenderTargetManager = MakeShared<FDisplayClusterRenderTargetManager, ESPMode::ThreadSafe>(ViewportManagerProxy);
 	PostProcessManager  = MakeShared<FDisplayClusterViewportPostProcessManager, ESPMode::ThreadSafe>(*this);
+	LightCardManager = MakeShared<FDisplayClusterViewportLightCardManager, ESPMode::ThreadSafe>(*this);
 
 	// initialize proxy
 	ViewportManagerProxy->Initialize(*this);
@@ -142,6 +144,11 @@ void FDisplayClusterViewportManager::StartScene(UWorld* InWorld)
 	{
 		PostProcessManager->HandleStartScene();
 	}
+
+	if (LightCardManager.IsValid())
+	{
+		LightCardManager->HandleStartScene();
+	}
 }
 
 void FDisplayClusterViewportManager::EndScene()
@@ -168,6 +175,11 @@ void FDisplayClusterViewportManager::ResetScene()
 {
 	check(IsInGameThread());
 
+	if (LightCardManager.IsValid())
+	{
+		LightCardManager->ResetScene();
+	}
+
 	for (FDisplayClusterViewport* Viewport : Viewports)
 	{
 		if (Viewport)
@@ -309,6 +321,11 @@ bool FDisplayClusterViewportManager::UpdateConfiguration(EDisplayClusterRenderFr
 			ResetScene();
 		}
 
+		if (LightCardManager.IsValid())
+		{
+			LightCardManager->UpdateConfiguration();
+		}
+
 		if (InPreviewSettings == nullptr)
 		{
 			if (Configuration->UpdateConfiguration(InRenderMode, InClusterNodeId))
@@ -481,6 +498,11 @@ bool FDisplayClusterViewportManager::BeginNewFrame(FViewport* InViewport, UWorld
 		PostProcessManager->HandleBeginNewFrame(OutRenderFrame);
 	}
 
+	if (LightCardManager.IsValid())
+	{
+		LightCardManager->HandleBeginNewFrame();
+	}
+
 	return true;
 }
 
@@ -690,7 +712,9 @@ void FDisplayClusterViewportManager::ConfigureViewFamily(const FDisplayClusterRe
 
 void FDisplayClusterViewportManager::RenderFrame(FViewport* InViewport)
 {
+	PreRenderFrame();
 	ViewportManagerProxy->ImplRenderFrame(InViewport);
+	PostRenderFrame();
 }
 
 bool FDisplayClusterViewportManager::CreateViewport(const FString& InViewportId, const class UDisplayClusterConfigurationViewport* ConfigurationViewport)
@@ -925,3 +949,19 @@ void FDisplayClusterViewportManager::AddReferencedObjects(FReferenceCollector& C
 		}
 	}
 }
+
+void FDisplayClusterViewportManager::PreRenderFrame()
+{
+	if (LightCardManager.IsValid())
+	{
+		LightCardManager->PreRenderFrame();
+	}
+}
+
+void FDisplayClusterViewportManager::PostRenderFrame()
+{
+	if (LightCardManager.IsValid())
+	{
+		LightCardManager->PostRenderFrame();
+	}
+}
\ No newline at end of file
diff --git a/Engine/Plugins/Runtime/nDisplay/Source/DisplayCluster/Private/Render/Viewport/DisplayClusterViewportManager.h b/Engine/Plugins/Runtime/nDisplay/Source/DisplayCluster/Private/Render/Viewport/DisplayClusterViewportManager.h
--- a/Engine/Plugins/Runtime/nDisplay/Source/DisplayCluster/Private/Render/Viewport/DisplayClusterViewportManager.h
+++ b/Engine/Plugins/Runtime/nDisplay/Source/DisplayCluster/Private/Render/Viewport/DisplayClusterViewportManager.h
@@ -20,6 +20,7 @@ class FDisplayClusterRenderFrameManager;
 class FDisplayClusterRenderFrame; 
 class FDisplayClusterViewportManagerProxy;
 
+class IDisplayClusterViewportLightCardManager;
 class IDisplayClusterProjectionPolicy;
 
 class  UDisplayClusterConfigurationViewport;
@@ -90,6 +91,8 @@ public:
 		return TArrayView<IDisplayClusterViewport*>((IDisplayClusterViewport**)(ClusterNodeViewports.GetData()), ClusterNodeViewports.Num());
 	}
 
+	virtual TSharedPtr<IDisplayClusterViewportLightCardManager, ESPMode::ThreadSafe> GetLightCardManager() const override { return LightCardManager; }
+
 	virtual void MarkComponentGeometryDirty(const FName InComponentName = NAME_None) override;
 
 	virtual void AddReferencedObjects(FReferenceCollector& Collector) override;
@@ -137,12 +140,19 @@ private:
 	void HandleViewportRTTChanges(const TArray<FDisplayClusterViewport_Context>& PrevContexts, const TArray<FDisplayClusterViewport_Context>& Contexts);
 	void ImplUpdateClusterNodeViewports(const EDisplayClusterRenderFrameMode InRenderMode, const FString& InClusterNodeId);
 
+	/** Called on the game thread just before a frame is rendered */
+	void PreRenderFrame();
+
+	/** Called on the game thread just after a frame is rendered */
+	void PostRenderFrame();
+
 protected:
 	friend FDisplayClusterViewportManagerProxy;
 	friend FDisplayClusterViewportConfiguration;
 
 	TSharedPtr<FDisplayClusterRenderTargetManager, ESPMode::ThreadSafe>        RenderTargetManager;
 	TSharedPtr<FDisplayClusterViewportPostProcessManager, ESPMode::ThreadSafe> PostProcessManager;
+	TSharedPtr<IDisplayClusterViewportLightCardManager, ESPMode::ThreadSafe>   LightCardManager;
 
 public:
 	TUniquePtr<FDisplayClusterViewportConfiguration> Configuration;
diff --git a/Engine/Plugins/Runtime/nDisplay/Source/DisplayCluster/Private/Render/Viewport/DisplayClusterViewportManagerProxy.cpp b/Engine/Plugins/Runtime/nDisplay/Source/DisplayCluster/Private/Render/Viewport/DisplayClusterViewportManagerProxy.cpp
--- a/Engine/Plugins/Runtime/nDisplay/Source/DisplayCluster/Private/Render/Viewport/DisplayClusterViewportManagerProxy.cpp
+++ b/Engine/Plugins/Runtime/nDisplay/Source/DisplayCluster/Private/Render/Viewport/DisplayClusterViewportManagerProxy.cpp
@@ -13,6 +13,7 @@
 #include "Render/Projection/IDisplayClusterProjectionPolicy.h"
 
 #include "Render/Viewport/DisplayClusterViewportProxy.h"
+#include "Render/Viewport/IDisplayClusterViewportLightCardManager.h"
 #include "Render/Viewport/RenderTarget/DisplayClusterRenderTargetManager.h"
 #include "Render/Viewport/RenderTarget/DisplayClusterRenderTargetResource.h"
 #include "Render/Viewport/Postprocess/DisplayClusterViewportPostProcessManager.h"
@@ -98,6 +99,7 @@ void FDisplayClusterViewportManagerProxy::Initialize(FDisplayClusterViewportMana
 {
 	RenderTargetManager = InViewportManager.RenderTargetManager;
 	PostProcessManager = InViewportManager.PostProcessManager;
+	LightCardManager = InViewportManager.LightCardManager;
 }
 
 void FDisplayClusterViewportManagerProxy::DeleteResource_RenderThread(FDisplayClusterViewportResource* InDeletedResourcePtr)
@@ -135,6 +137,10 @@ void FDisplayClusterViewportManagerProxy::ImplSafeRelease()
 {
 	check(IsInGameThread());
 
+	// Light card manager contains a preview world, which needs to be cleaned up on game thread. Manually reset the pointer
+	// here to ensure that it isn't cleaned up on render thread when the proxy deletes itself.
+	LightCardManager.Reset();
+
 	// Remove viewport manager proxy on render_thread
 	ENQUEUE_RENDER_COMMAND(DeleteDisplayClusterViewportManagerProxy)(
 		[ViewportManagerProxy = this](FRHICommandListImmediate& RHICmdList)
@@ -235,6 +241,11 @@ void FDisplayClusterViewportManagerProxy::ImplRenderFrame(FViewport* InViewport)
 			ViewportManagerProxy->PostProcessManager->HandleRenderFrameSetup_RenderThread(RHICmdList, ViewportManagerProxy);
 		}
 
+		if (ViewportManagerProxy->LightCardManager.IsValid())
+		{
+			ViewportManagerProxy->LightCardManager->RenderLightCardMap_RenderThread(RHICmdList);
+		}
+
 		bool bWarpBlendEnabled = ViewportManagerProxy->RenderFrameSettings.bAllowWarpBlend && CVarWarpBlendEnabled.GetValueOnRenderThread() != 0;
 
 		// mGPU not used for in-editor rendering
diff --git a/Engine/Plugins/Runtime/nDisplay/Source/DisplayCluster/Private/Render/Viewport/DisplayClusterViewportManagerProxy.h b/Engine/Plugins/Runtime/nDisplay/Source/DisplayCluster/Private/Render/Viewport/DisplayClusterViewportManagerProxy.h
--- a/Engine/Plugins/Runtime/nDisplay/Source/DisplayCluster/Private/Render/Viewport/DisplayClusterViewportManagerProxy.h
+++ b/Engine/Plugins/Runtime/nDisplay/Source/DisplayCluster/Private/Render/Viewport/DisplayClusterViewportManagerProxy.h
@@ -14,6 +14,7 @@ class FDisplayClusterRenderTargetManager;
 class FDisplayClusterViewportPostProcessManager;
 class FDisplayClusterViewportManager;
 class FDisplayClusterViewportResource;
+class IDisplayClusterViewportLightCardManager;
 class IDisplayClusterProjectionPolicy;
 
 class FViewport;
@@ -47,6 +48,8 @@ public:
 	virtual bool GetFrameTargets_RenderThread(TArray<FRHITexture2D*>& OutFrameResources, TArray<FIntPoint>& OutTargetOffsets, TArray<FRHITexture2D*>* OutAdditionalFrameResources = nullptr) const override;
 	virtual bool ResolveFrameTargetToBackBuffer_RenderThread(FRHICommandListImmediate& RHICmdList, const uint32 InContextNum, const int32 DestArrayIndex, FRHITexture2D* DestTexture, FVector2D WindowSize) const override;
 
+	virtual TSharedPtr<IDisplayClusterViewportLightCardManager, ESPMode::ThreadSafe> GetLightCardManager_RenderThread() const override { return LightCardManager; }
+
 	// internal use only
 	void DeleteResource_RenderThread(FDisplayClusterViewportResource* InDeletedResourcePtr);
 	void Initialize(FDisplayClusterViewportManager& InViewportManager);
@@ -82,6 +85,7 @@ private:
 private:
 	TSharedPtr<FDisplayClusterRenderTargetManager, ESPMode::ThreadSafe>        RenderTargetManager;
 	TSharedPtr<FDisplayClusterViewportPostProcessManager, ESPMode::ThreadSafe> PostProcessManager;
+	TSharedPtr<IDisplayClusterViewportLightCardManager, ESPMode::ThreadSafe>   LightCardManager;
 
 	FDisplayClusterRenderFrameSettings RenderFrameSettings;
 
diff --git a/Engine/Plugins/Runtime/nDisplay/Source/DisplayCluster/Private/Render/Viewport/LightCard/DisplayClusterViewportLightCardManager.cpp b/Engine/Plugins/Runtime/nDisplay/Source/DisplayCluster/Private/Render/Viewport/LightCard/DisplayClusterViewportLightCardManager.cpp
new file mode 100644
--- /dev/null
+++ b/Engine/Plugins/Runtime/nDisplay/Source/DisplayCluster/Private/Render/Viewport/LightCard/DisplayClusterViewportLightCardManager.cpp
@@ -0,0 +1,171 @@
+// Copyright Epic Games, Inc. All Rights Reserved.
+
+#include "DisplayClusterViewportLightCardManager.h"
+
+#include "DisplayClusterLightCardActor.h"
+#include "Blueprints/DisplayClusterBlueprintLib.h"
+
+#include "IDisplayClusterShaders.h"
+
+#include "Render/Viewport/DisplayClusterViewportManager.h"
+#include "Render/Viewport/RenderFrame/DisplayClusterRenderFrame.h"
+
+#include "PreviewScene.h"
+
+/** Console variable used to control the size of the UV light card map texture */
+static TAutoConsoleVariable<int32> CVarUVLightCardTextureSize(
+	TEXT("nDisplay.render.uvlightcards.UVTextureSize"),
+	512,
+	TEXT("The size of the texture UV light cards are rendered to."));
+
+void FDisplayClusterLightCardMap::InitDynamicRHI()
+{
+	ETextureCreateFlags CreateFlags = TexCreate_Dynamic;
+	CreateFlags |= TexCreate_MultiGPUGraphIgnore;
+
+	const FRHITextureCreateDesc Desc = FRHITextureCreateDesc::Create2D(TEXT("DisplayClusterLightCardMap"))
+		.SetExtent(GetSizeX(), GetSizeY())
+		.SetFormat(PF_FloatRGBA)
+		.SetFlags(CreateFlags | ETextureCreateFlags::RenderTargetable | ETextureCreateFlags::ShaderResource)
+		.SetInitialState(ERHIAccess::SRVMask)
+		.SetClearValue(FClearValueBinding::Transparent);
+
+	RenderTargetTextureRHI = TextureRHI = RHICreateTexture(Desc);
+}
+
+FDisplayClusterViewportLightCardManager::FDisplayClusterViewportLightCardManager(FDisplayClusterViewportManager& InViewportManager)
+	: ViewportManager(InViewportManager)
+{
+}
+
+FDisplayClusterViewportLightCardManager::~FDisplayClusterViewportLightCardManager()
+{
+	ResetScene();
+	ReleaseUVLightCardMap();
+	PreviewScene.Reset();
+}
+
+void FDisplayClusterViewportLightCardManager::UpdateConfiguration()
+{
+	ResetScene();
+
+	ADisplayClusterRootActor* RootActor = ViewportManager.GetRootActor();
+	TSet<ADisplayClusterLightCardActor*> LightCards;
+	UDisplayClusterBlueprintLib::FindLightCardsForRootActor(RootActor, LightCards);
+
+	for (ADisplayClusterLightCardActor* LightCard : LightCards)
+	{
+		if (LightCard->bIsUVLightCard)
+		{
+			UVLightCards.Add(LightCard);
+		}
+	}
+}
+
+void FDisplayClusterViewportLightCardManager::ResetScene()
+{
+	if (PreviewScene.IsValid() && LoadedPrimitiveComponents.Num())
+	{
+		for (UPrimitiveComponent* PrimitiveComponent : LoadedPrimitiveComponents)
+		{
+			PreviewScene->GetScene()->RemovePrimitive(PrimitiveComponent);
+		}
+	}
+
+	LoadedPrimitiveComponents.Empty();
+	UVLightCards.Empty();
+}
+
+void FDisplayClusterViewportLightCardManager::HandleStartScene()
+{
+	PreviewScene = MakeShared<FPreviewScene>(FPreviewScene::ConstructionValues()
+		.SetEditor(false)
+		.SetCreatePhysicsScene(false)
+		.SetCreateDefaultLighting(false));
+
+	UpdateConfiguration();
+}
+
+void FDisplayClusterViewportLightCardManager::HandleBeginNewFrame()
+{
+	InitializeUVLightCardMap();
+}
+
+void FDisplayClusterViewportLightCardManager::PreRenderFrame()
+{
+	bool bLoadedPrimitives = false;
+	for (ADisplayClusterLightCardActor* LightCard : UVLightCards)
+	{
+		UStaticMeshComponent* LightCardMeshComp = LightCard->GetLightCardMeshComponent();
+		if (LightCardMeshComp && LightCardMeshComp->SceneProxy == nullptr)
+		{
+			PreviewScene->GetScene()->AddPrimitive(LightCardMeshComp);
+			LoadedPrimitiveComponents.Add(LightCardMeshComp);
+
+			bLoadedPrimitives = true;
+		}
+	}
+
+	ENQUEUE_RENDER_COMMAND(InitializeLightCardRenderFrame)(
+		[this, bLoadedPrimitives](FRHICommandListImmediate& RHICmdList)
+	{
+		bHasUVLightCards = bLoadedPrimitives;
+	});
+}
+
+void FDisplayClusterViewportLightCardManager::PostRenderFrame()
+{
+	for (UPrimitiveComponent* LoadedComponent : LoadedPrimitiveComponents)
+	{
+		PreviewScene->GetScene()->RemovePrimitive(LoadedComponent);
+	}
+
+	LoadedPrimitiveComponents.Empty();
+}
+
+void FDisplayClusterViewportLightCardManager::RenderLightCardMap_RenderThread(FRHICommandListImmediate& RHICmdList)
+{
+	if (PreviewScene.IsValid() && bHasUVLightCards)
+	{
+		IDisplayClusterShaders& ShadersAPI = IDisplayClusterShaders::Get();
+		ShadersAPI.RenderPreprocess_UVLightCards(RHICmdList, PreviewScene->GetScene(), UVLightCardMapProxy, ADisplayClusterLightCardActor::UVPlaneDefaultSize);
+	}
+}
+
+void FDisplayClusterViewportLightCardManager::InitializeUVLightCardMap()
+{
+	const uint32 LightCardTextureSize = CVarUVLightCardTextureSize.GetValueOnAnyThread();
+	if (UVLightCardMap && UVLightCardMap->GetSizeX() != LightCardTextureSize)
+	{
+		ReleaseUVLightCardMap();
+	}
+
+	if (UVLightCardMap == nullptr)
+	{
+		UVLightCardMap = new FDisplayClusterLightCardMap(LightCardTextureSize);
+
+		ENQUEUE_RENDER_COMMAND(InitializeLightCardResources)(
+			[this, InUVLightCardMap=UVLightCardMap](FRHICommandListImmediate& RHICmdList)
+		{
+			// Store a copy of the texture's pointer on the render thread and initialize the texture's resources
+			UVLightCardMapProxy = InUVLightCardMap;
+			UVLightCardMapProxy->InitResource();
+		});
+	}
+}
+
+void FDisplayClusterViewportLightCardManager::ReleaseUVLightCardMap()
+{
+	if (UVLightCardMap)
+	{
+		UVLightCardMap = nullptr;
+
+		ENQUEUE_RENDER_COMMAND(DeleteLightCardResources)(
+			[this](FRHICommandListImmediate& RHICmdList)
+		{
+			// Release the texture's resources and delete the texture object from the rendering thread
+			UVLightCardMapProxy->ReleaseResource();
+			delete UVLightCardMapProxy;
+		});
+	}
+}
\ No newline at end of file
diff --git a/Engine/Plugins/Runtime/nDisplay/Source/DisplayCluster/Private/Render/Viewport/LightCard/DisplayClusterViewportLightCardManager.h b/Engine/Plugins/Runtime/nDisplay/Source/DisplayCluster/Private/Render/Viewport/LightCard/DisplayClusterViewportLightCardManager.h
new file mode 100644
--- /dev/null
+++ b/Engine/Plugins/Runtime/nDisplay/Source/DisplayCluster/Private/Render/Viewport/LightCard/DisplayClusterViewportLightCardManager.h
@@ -0,0 +1,83 @@
+// Copyright Epic Games, Inc. All Rights Reserved.
+
+#pragma once
+
+#include "CoreMinimal.h"
+#include "Render/Viewport/IDisplayClusterViewportLightCardManager.h"
+
+#include "RenderResource.h"
+#include "UnrealClient.h"
+
+class ADisplayClusterLightCardActor;
+class FDisplayClusterViewportManager;
+class FPreviewScene;
+
+/** A render targetable texture resource used to render the UV light cards to */
+class FDisplayClusterLightCardMap : public FTexture, public FRenderTarget
+{
+public:
+	FDisplayClusterLightCardMap(uint32 InSize)
+		: Size(InSize)
+	{ }
+
+	virtual uint32 GetSizeX() const override { return Size; }
+	virtual uint32 GetSizeY() const override { return Size; }
+	virtual FIntPoint GetSizeXY() const override { return FIntPoint(Size, Size); }
+
+	virtual void InitDynamicRHI() override;
+
+	virtual FString GetFriendlyName() const override { return TEXT("DisplayClusterLightCardMap"); }
+
+private:
+	uint32 Size;
+};
+
+/** Manages the rendering of UV light cards for the viewport manager */
+class FDisplayClusterViewportLightCardManager : public IDisplayClusterViewportLightCardManager
+{
+public:
+	FDisplayClusterViewportLightCardManager(FDisplayClusterViewportManager& InViewportManager);
+	virtual ~FDisplayClusterViewportLightCardManager();
+
+	//~ Begin IDisplayClusterViewportLightCardManager interface
+	virtual void UpdateConfiguration() override;
+	virtual void ResetScene() override;
+	virtual void HandleStartScene() override;
+	virtual void HandleBeginNewFrame() override;
+	virtual void PreRenderFrame() override;
+	virtual void PostRenderFrame() override;
+	virtual FRHITexture* GetUVLightCardMap() const override { return UVLightCardMap->GetTextureRHI(); }
+	virtual FRHITexture* GetUVLightCardMap_RenderThread() const override { return UVLightCardMapProxy->GetTextureRHI(); }
+	virtual bool HasUVLightCards_RenderThread() const override { return bHasUVLightCards; }
+	virtual void RenderLightCardMap_RenderThread(FRHICommandListImmediate& RHICmdList) override;
+	//~ End IDisplayClusterViewportLightCardManager interface
+
+private:
+	/** Initializes the UV light card map texture */
+	void InitializeUVLightCardMap();
+
+	/** Releases the UV light card map texture */
+	void ReleaseUVLightCardMap();
+
+private:
+	/** A reference to the owning viewport manager */
+	FDisplayClusterViewportManager& ViewportManager;
+
+	/** The preview scene the UV light card proxies live in for rendering */
+	TSharedPtr<FPreviewScene> PreviewScene = nullptr;
+
+	/** The list of UV light card actors that are referenced by the root actor */
+	TArray<ADisplayClusterLightCardActor*> UVLightCards;
+
+	/** A list of primitive components that have been added to the preview scene for rendering in the current frame */
+	TArray<UPrimitiveComponent*> LoadedPrimitiveComponents;
+
+	/** The render target to which the UV light card map is rendered */
+	FDisplayClusterLightCardMap* UVLightCardMap = nullptr;
+
+	/** The render thread copy of the pointer to the UV ligth card map */
+	FDisplayClusterLightCardMap* UVLightCardMapProxy = nullptr;
+
+	/** A render thread flag that indicates the light card manager has UV light cards to render */
+	bool bHasUVLightCards = false;
+};
\ No newline at end of file
diff --git a/Engine/Plugins/Runtime/nDisplay/Source/DisplayCluster/Public/DisplayClusterLightCardActor.h b/Engine/Plugins/Runtime/nDisplay/Source/DisplayCluster/Public/DisplayClusterLightCardActor.h
--- a/Engine/Plugins/Runtime/nDisplay/Source/DisplayCluster/Public/DisplayClusterLightCardActor.h
+++ b/Engine/Plugins/Runtime/nDisplay/Source/DisplayCluster/Public/DisplayClusterLightCardActor.h
@@ -65,6 +65,12 @@ public:
 	/** The rotation used to orient the plane mesh used for the light card so that its normal points radially inwards */
 	static const FRotator PlaneMeshRotation;
 
+	/** The default size of the porjection plane UV light cards are rendered to */
+	static const float UVPlaneDefaultSize;
+
+	/** The default distance from the view of the projection plane UV light cards are rendered to */
+	static const float UVPlaneDefaultDistance;
+
 public:
 	ADisplayClusterLightCardActor(const FObjectInitializer& ObjectInitializer);
 
@@ -75,7 +81,7 @@ public:
 	virtual void PostEditChangeProperty(FPropertyChangedEvent& PropertyChangedEvent) override;
 	virtual bool ShouldTickIfViewportsOnly() const override { return true; }
 #endif
-	
+
 	/**
 	 * Gets the transform in world space of the light card component
 	 * @param bIgnoreSpinYawPitch - If the light card component's spin, yaw, and pitch should be ignored when computing the transform
@@ -85,6 +91,9 @@ public:
 	/** Gets the object oriented bounding box of the light card component */
 	FBox GetLightCardBounds(bool bLocalSpace = false) const;
 
+	/** Gets the light card mesh component */
+	UStaticMeshComponent* GetLightCardMeshComponent() const { return LightCardComponent.Get(); }
+
 	/** Returns the current static mesh used by this light card */
 	UStaticMesh* GetStaticMesh() const;
 
@@ -106,35 +115,42 @@ public:
 	/** Updates the polygon texture from the polygon points */
 	void UpdatePolygonTexture();
 
+	/** Updates the light card actor visibility */
+	void UpdateLightCardVisibility();
+
 public:
 
 	/** Radius of light card polar coordinates. Does not include the effect of RadialOffset */
-	UPROPERTY(EditAnywhere, BlueprintReadWrite, Category = "Orientation")
+	UPROPERTY(EditAnywhere, BlueprintReadWrite, Category = "Orientation", meta = (EditCondition = "!bIsUVLightCard", HideEditConditionToggle, EditConditionHides))
 	double DistanceFromCenter;
 
 	/** Related to the Azimuth of light card polar coordinates */
-	UPROPERTY(EditAnywhere, BlueprintReadWrite, Category = "Orientation", meta = (UIMin = 0, ClampMin = 0, UIMax = 360, ClampMax = 360))
+	UPROPERTY(EditAnywhere, BlueprintReadWrite, Category = "Orientation", meta = (UIMin = 0, ClampMin = 0, UIMax = 360, ClampMax = 360, EditCondition = "!bIsUVLightCard", HideEditConditionToggle, EditConditionHides))
 	double Longitude;
 
 	/** Related to the Elevation of light card polar coordinates */
-	UPROPERTY(EditAnywhere, BlueprintReadWrite, Category = "Orientation", meta = (UIMin = -90, ClampMin = -90, UIMax = 90, ClampMax = 90))
+	UPROPERTY(EditAnywhere, BlueprintReadWrite, Category = "Orientation", meta = (UIMin = -90, ClampMin = -90, UIMax = 90, ClampMax = 90, EditCondition = "!bIsUVLightCard", HideEditConditionToggle, EditConditionHides))
 	double Latitude;
 
+	/** The UV coordinates of the light card, if it is in UV space */
+	UPROPERTY(EditAnywhere, BlueprintReadWrite, Category = "Orientation", meta = (DisplayName = "UV Coodinates", EditCondition = "bIsUVLightCard", HideEditConditionToggle, EditConditionHides))
+	FVector2D UVCoordinates = FVector2D(0.5, 0.5);
+
 	/** Roll rotation of light card around its plane axis */
-	UPROPERTY(EditAnywhere, BlueprintReadWrite, Category = "Orientation")
+	UPROPERTY(EditAnywhere, BlueprintReadWrite, Category = "Orientation", meta = (UIMin = -360, ClampMin = -360, UIMax = 360, ClampMax = 360))
 	double Spin;
 
-	UPROPERTY(EditAnywhere, BlueprintReadWrite, Category = "Orientation")
+	UPROPERTY(EditAnywhere, BlueprintReadWrite, Category = "Orientation", meta = (UIMin = -360, ClampMin = -360, UIMax = 360, ClampMax = 360, EditCondition = "!bIsUVLightCard", HideEditConditionToggle, EditConditionHides))
 	double Pitch;
 
-	UPROPERTY(EditAnywhere, BlueprintReadWrite, Category = "Orientation")
+	UPROPERTY(EditAnywhere, BlueprintReadWrite, Category = "Orientation", meta = (UIMin = -360, ClampMin = -360, UIMax = 360, ClampMax = 360, EditCondition = "!bIsUVLightCard", HideEditConditionToggle, EditConditionHides))
 	double Yaw;
 
 	UPROPERTY(EditAnywhere, BlueprintReadWrite, Category = "Orientation")
 	FVector2D Scale;
 
 	/** Used by the flush constraint to offset the location of the light card form the wall */
-	UPROPERTY(EditAnywhere, BlueprintReadWrite, Category = "Orientation")
+	UPROPERTY(EditAnywhere, BlueprintReadWrite, Category = "Orientation", meta = (EditCondition = "!bIsUVLightCard", HideEditConditionToggle, EditConditionHides))
 	double RadialOffset;
 
 	UPROPERTY(EditAnywhere, BlueprintReadWrite, Category = "Appearance")
@@ -177,6 +193,10 @@ public:
 	UPROPERTY()
 	bool bLockToOwningRootActor = true;
 
+	/** Indicates if the light card exists in 3D space or in UV space */
+	UPROPERTY()
+	bool bIsUVLightCard = false;
+
 	/** Used to flag this light card as a proxy of a "real" light card. Used by the LightCard Editor */
 	UPROPERTY(Transient)
 	bool bIsProxy = false;
diff --git a/Engine/Plugins/Runtime/nDisplay/Source/DisplayCluster/Public/Render/Viewport/IDisplayClusterViewportLightCardManager.h b/Engine/Plugins/Runtime/nDisplay/Source/DisplayCluster/Public/Render/Viewport/IDisplayClusterViewportLightCardManager.h
new file mode 100644
--- /dev/null
+++ b/Engine/Plugins/Runtime/nDisplay/Source/DisplayCluster/Public/Render/Viewport/IDisplayClusterViewportLightCardManager.h
@@ -0,0 +1,41 @@
+// Copyright Epic Games, Inc. All Rights Reserved.
+
+#pragma once
+
+class FRHITexture;
+class FRHICommandListImmediate;
+
+/** Interface for a manager which manages the rendering of UV light cards for the viewport manager */
+class DISPLAYCLUSTER_API IDisplayClusterViewportLightCardManager
+{
+public:
+	/** Updates the UV light card scene based on the current root actor configuration */
+	virtual void UpdateConfiguration() = 0;
+
+	/** Resets the UV light card scene */
+	virtual void ResetScene() = 0;
+
+	/** Raised when a new scene is being started, initializes the UV light card scene */
+	virtual void HandleStartScene() = 0;
+
+	/** Raised when a new frame begins, initializes any frame resources needed */
+	virtual void HandleBeginNewFrame() = 0;
+
+	/** Raised just before a render frame pass is performed */
+	virtual void PreRenderFrame() = 0;
+
+	/** Raised just after a render frame pass is performed */
+	virtual void PostRenderFrame() = 0;
+
+	/** Gets the UV light card map texture */
+	virtual FRHITexture* GetUVLightCardMap() const = 0;
+
+	/** Gets the render thread copy of the UV light card map texture */
+	virtual FRHITexture* GetUVLightCardMap_RenderThread() const = 0;
+
+	/** Gets whether there are any UV light cards to render to the viewports */
+	virtual bool HasUVLightCards_RenderThread() const = 0;
+
+	/** Renders the UV light card scene to the UV light card map texture */
+	virtual void RenderLightCardMap_RenderThread(FRHICommandListImmediate& RHICmdList) = 0;
+};
\ No newline at end of file
diff --git a/Engine/Plugins/Runtime/nDisplay/Source/DisplayCluster/Public/Render/Viewport/IDisplayClusterViewportManager.h b/Engine/Plugins/Runtime/nDisplay/Source/DisplayCluster/Public/Render/Viewport/IDisplayClusterViewportManager.h
--- a/Engine/Plugins/Runtime/nDisplay/Source/DisplayCluster/Public/Render/Viewport/IDisplayClusterViewportManager.h
+++ b/Engine/Plugins/Runtime/nDisplay/Source/DisplayCluster/Public/Render/Viewport/IDisplayClusterViewportManager.h
@@ -15,6 +15,7 @@ class FSceneViewFamilyContext;
 class ADisplayClusterRootActor;
 class UDisplayClusterConfigurationViewport;
 class IDisplayClusterViewportManagerProxy;
+class IDisplayClusterViewportLightCardManager;
 class FReferenceCollector;
 
 class DISPLAYCLUSTER_API IDisplayClusterViewportManager
@@ -149,6 +150,12 @@ public:
 	*/
 	virtual const TArrayView<IDisplayClusterViewport*> GetViewports() const = 0;
 
+	/**
+	* Return the light card manager, used to manage and render UV light cards
+	* [Game thread func]
+	*/
+	virtual TSharedPtr<IDisplayClusterViewportLightCardManager, ESPMode::ThreadSafe> GetLightCardManager() const = 0;
+
 	/**
 	* Mark the geometry of the referenced component(s) as dirty (ProceduralMesh, etc)
 	* [Game thread func]
diff --git a/Engine/Plugins/Runtime/nDisplay/Source/DisplayCluster/Public/Render/Viewport/IDisplayClusterViewportManagerProxy.h b/Engine/Plugins/Runtime/nDisplay/Source/DisplayCluster/Public/Render/Viewport/IDisplayClusterViewportManagerProxy.h
--- a/Engine/Plugins/Runtime/nDisplay/Source/DisplayCluster/Public/Render/Viewport/IDisplayClusterViewportManagerProxy.h
+++ b/Engine/Plugins/Runtime/nDisplay/Source/DisplayCluster/Public/Render/Viewport/IDisplayClusterViewportManagerProxy.h
@@ -5,6 +5,7 @@
 #include "CoreMinimal.h"
 
 #include "Render/Viewport/IDisplayClusterViewportProxy.h"
+#include "Render/Viewport/IDisplayClusterViewportLightCardManager.h"
 
 class DISPLAYCLUSTER_API IDisplayClusterViewportManagerProxy
 {
@@ -65,5 +66,11 @@ public:
 	* @return - true if success
 	*/
 	virtual bool ResolveFrameTargetToBackBuffer_RenderThread(FRHICommandListImmediate& RHICmdList, const uint32 InContextNum, const int32 DestArrayIndex, FRHITexture2D* DstBackBuffer, FVector2D WindowSize) const = 0;
+
+	/**
+	* Return the light card manager, used to manage and render UV light cards
+	* [Rendering thread func]
+	*/
+	virtual TSharedPtr<IDisplayClusterViewportLightCardManager, ESPMode::ThreadSafe> GetLightCardManager_RenderThread() const = 0;
 };
 
diff --git a/Engine/Plugins/Runtime/nDisplay/Source/DisplayClusterLightCardEditor/Private/DisplayClusterLightCardEditorCommands.cpp b/Engine/Plugins/Runtime/nDisplay/Source/DisplayClusterLightCardEditor/Private/DisplayClusterLightCardEditorCommands.cpp
--- a/Engine/Plugins/Runtime/nDisplay/Source/DisplayClusterLightCardEditor/Private/DisplayClusterLightCardEditorCommands.cpp
+++ b/Engine/Plugins/Runtime/nDisplay/Source/DisplayClusterLightCardEditor/Private/DisplayClusterLightCardEditorCommands.cpp
@@ -12,6 +12,7 @@ void FDisplayClusterLightCardEditorCommands::RegisterCommands()
 	UI_COMMAND(PerspectiveProjection, "Perspective", "A perspective projection from the stage's view origin", EUserInterfaceActionType::Button, FInputChord());
 	UI_COMMAND(OrthographicProjection, "Orthographic", "An orthographic projection from the stage's view origin", EUserInterfaceActionType::Button, FInputChord());
 	UI_COMMAND(AzimuthalProjection, "Azimuthal", "An azimuthal hemisphere projection from the stage's view origin", EUserInterfaceActionType::Button, FInputChord());
+	UI_COMMAND(UVProjection, "UV", "A UV projection that flattens the stage's meshes based on their UV coordinates", EUserInterfaceActionType::Button, FInputChord());
 
 	UI_COMMAND(ViewOrientationTop, "Top", "Orient the view to look at the top of the stage", EUserInterfaceActionType::Button, FInputChord());
 	UI_COMMAND(ViewOrientationBottom, "Bottom", "Orient the view look at the bottom of the stage", EUserInterfaceActionType::Button, FInputChord());
diff --git a/Engine/Plugins/Runtime/nDisplay/Source/DisplayClusterLightCardEditor/Private/DisplayClusterLightCardEditorCommands.h b/Engine/Plugins/Runtime/nDisplay/Source/DisplayClusterLightCardEditor/Private/DisplayClusterLightCardEditorCommands.h
--- a/Engine/Plugins/Runtime/nDisplay/Source/DisplayClusterLightCardEditor/Private/DisplayClusterLightCardEditorCommands.h
+++ b/Engine/Plugins/Runtime/nDisplay/Source/DisplayClusterLightCardEditor/Private/DisplayClusterLightCardEditorCommands.h
@@ -28,6 +28,7 @@ public:
 	TSharedPtr<FUICommandInfo> PerspectiveProjection;
 	TSharedPtr<FUICommandInfo> OrthographicProjection;
 	TSharedPtr<FUICommandInfo> AzimuthalProjection;
+	TSharedPtr<FUICommandInfo> UVProjection;
 
 	TSharedPtr<FUICommandInfo> ViewOrientationTop;
 	TSharedPtr<FUICommandInfo> ViewOrientationBottom;
diff --git a/Engine/Plugins/Runtime/nDisplay/Source/DisplayClusterLightCardEditor/Private/DisplayClusterLightCardEditorStyle.h b/Engine/Plugins/Runtime/nDisplay/Source/DisplayClusterLightCardEditor/Private/DisplayClusterLightCardEditorStyle.h
--- a/Engine/Plugins/Runtime/nDisplay/Source/DisplayClusterLightCardEditor/Private/DisplayClusterLightCardEditorStyle.h
+++ b/Engine/Plugins/Runtime/nDisplay/Source/DisplayClusterLightCardEditor/Private/DisplayClusterLightCardEditorStyle.h
@@ -36,6 +36,8 @@ public:
 		Set("DisplayClusterLightCardEditor.ViewOrientationFront", new CORE_IMAGE_BRUSH_SVG("Starship/Common/ViewFront", Icon16x16));
 		Set("DisplayClusterLightCardEditor.ViewOrientationBack", new CORE_IMAGE_BRUSH_SVG("Starship/Common/ViewBack", Icon16x16));
 
+		Set("UVLightCard", new CORE_IMAGE_BRUSH_SVG("Starship/Common/SetDrawUVs", Icon16x16));
+
 		FSlateStyleRegistry::RegisterSlateStyle(*this);
 	}
 
diff --git a/Engine/Plugins/Runtime/nDisplay/Source/DisplayClusterLightCardEditor/Private/SDisplayClusterLightCardEditor.cpp b/Engine/Plugins/Runtime/nDisplay/Source/DisplayClusterLightCardEditor/Private/SDisplayClusterLightCardEditor.cpp
--- a/Engine/Plugins/Runtime/nDisplay/Source/DisplayClusterLightCardEditor/Private/SDisplayClusterLightCardEditor.cpp
+++ b/Engine/Plugins/Runtime/nDisplay/Source/DisplayClusterLightCardEditor/Private/SDisplayClusterLightCardEditor.cpp
@@ -243,6 +243,15 @@ ADisplayClusterLightCardActor* SDisplayClusterLightCardEditor::SpawnLightCard()
 
 	NewLightCard->SetActorLabel(NewLightCard->GetName());
 
+	if (ViewportView.IsValid())
+	{
+		const EDisplayClusterMeshProjectionType ProjectionMode = ViewportView->GetLightCardEditorViewportClient()->GetProjectionMode();
+		if (ProjectionMode == EDisplayClusterMeshProjectionType::UV)
+		{
+			NewLightCard->bIsUVLightCard = true;
+		}
+	}
+
 	TArray<ADisplayClusterLightCardActor*> LightCards { NewLightCard } ;
 	AddLightCardsToActor(LightCards);
 
diff --git a/Engine/Plugins/Runtime/nDisplay/Source/DisplayClusterLightCardEditor/Private/SDisplayClusterLightCardList.cpp b/Engine/Plugins/Runtime/nDisplay/Source/DisplayClusterLightCardEditor/Private/SDisplayClusterLightCardList.cpp
--- a/Engine/Plugins/Runtime/nDisplay/Source/DisplayClusterLightCardEditor/Private/SDisplayClusterLightCardList.cpp
+++ b/Engine/Plugins/Runtime/nDisplay/Source/DisplayClusterLightCardEditor/Private/SDisplayClusterLightCardList.cpp
@@ -114,7 +114,6 @@ private:
 
 	const FSlateBrush* GetLightCardIcon() const
 	{
-		FString ActorName = TEXT("");
 		if (LightCardTreeItem.IsValid())
 		{
 			TSharedPtr<SDisplayClusterLightCardList::FLightCardTreeItem> LightCardTreeItemPin = LightCardTreeItem.Pin();
@@ -124,7 +123,14 @@ private:
 			}
 			else if (LightCardTreeItemPin->LightCardActor.IsValid())
 			{
-				return FClassIconFinder::FindIconForActor(LightCardTreeItemPin->LightCardActor.Get());
+				if (LightCardTreeItemPin->LightCardActor->bIsUVLightCard)
+				{
+					return FDisplayClusterLightCardEditorStyle::Get().GetBrush(TEXT("UVLightCard"));
+				}
+				else
+				{
+					return FClassIconFinder::FindIconForActor(LightCardTreeItemPin->LightCardActor.Get());
+				}
 			}
 		}
 
diff --git a/Engine/Plugins/Runtime/nDisplay/Source/DisplayClusterLightCardEditor/Private/Viewport/DisplayClusterLightCardEditorViewportClient.cpp b/Engine/Plugins/Runtime/nDisplay/Source/DisplayClusterLightCardEditor/Private/Viewport/DisplayClusterLightCardEditorViewportClient.cpp
--- a/Engine/Plugins/Runtime/nDisplay/Source/DisplayClusterLightCardEditor/Private/Viewport/DisplayClusterLightCardEditorViewportClient.cpp
+++ b/Engine/Plugins/Runtime/nDisplay/Source/DisplayClusterLightCardEditor/Private/Viewport/DisplayClusterLightCardEditorViewportClient.cpp
@@ -15,6 +15,8 @@
 #include "DisplayClusterRootActor.h"
 #include "DisplayClusterLightCardActor.h"
 #include "DisplayClusterLightCardEditorLog.h"
+#include "IDisplayClusterScenePreview.h"
+#include "SDisplayClusterLightCardEditor.h"
 
 #include "AudioDevice.h"
 #include "CameraController.h"
@@ -32,7 +34,6 @@
 #include "EngineModule.h"
 #include "Engine/Canvas.h"
 #include "Engine/TextureRenderTarget2D.h"
-#include "IDisplayClusterScenePreview.h"
 #include "ImageUtils.h"
 #include "Kismet/KismetMathLibrary.h"
 #include "Kismet/GameplayStatics.h"
@@ -44,7 +45,6 @@
 #include "RayTracingDebugVisualizationMenuCommands.h"
 #include "Renderer/Private/SceneRendering.h"
 #include "ScopedTransaction.h"
-#include "SDisplayClusterLightCardEditor.h"
 #include "Slate/SceneViewport.h"
 #include "UnrealEdGlobals.h"
 #include "UnrealWidget.h"
@@ -617,10 +617,20 @@ void FDisplayClusterLightCardEditorViewportClient::Draw(FViewport* InViewport, F
 
 	if (!bDisableCustomRenderer)
 	{
-		FDisplayClusterScenePreviewRenderSettings RenderSettings;
-		RenderSettings.RenderType = FDisplayClusterScenePreviewRenderSettings::ERenderType::Color;
+		FDisplayClusterMeshProjectionRenderSettings RenderSettings;
+		RenderSettings.RenderType = EDisplayClusterMeshProjectionOutput::Color;
 		RenderSettings.EngineShowFlags = EngineShowFlags;
 		RenderSettings.ProjectionType = ProjectionMode;
+		RenderSettings.PrimitiveFilter.ShouldRenderPrimitiveDelegate = FDisplayClusterMeshProjectionPrimitiveFilter::FPrimitiveFilter::CreateSP(this, &FDisplayClusterLightCardEditorViewportClient::ShouldRenderPrimitive);
+		RenderSettings.PrimitiveFilter.ShouldApplyProjectionDelegate = FDisplayClusterMeshProjectionPrimitiveFilter::FPrimitiveFilter::CreateSP(this, &FDisplayClusterLightCardEditorViewportClient::ShouldApplyProjectionToPrimitive);
+
+		if (ProjectionMode == EDisplayClusterMeshProjectionType::UV)
+		{
+			RenderSettings.ProjectionTypeSettings.UVProjectionIndex = 1;
+			RenderSettings.ProjectionTypeSettings.UVProjectionPlaneSize = ADisplayClusterLightCardActor::UVPlaneDefaultSize;
+			RenderSettings.ProjectionTypeSettings.UVProjectionPlaneDistance = ADisplayClusterLightCardActor::UVPlaneDefaultDistance;
+		}
+
 		GetSceneViewInitOptions(RenderSettings.ViewInitOptions);
 
 		IDisplayClusterScenePreview::Get().Render(PreviewRendererId, RenderSettings, *Canvas);
@@ -677,8 +687,11 @@ void FDisplayClusterLightCardEditorViewportClient::Draw(FViewport* InViewport, F
 	// Axes indicators
 	if (bDrawAxes && !ViewFamily.EngineShowFlags.Game && !GLevelEditorModeTools().IsViewportUIHidden() && !IsVisualizeCalibrationMaterialEnabled())
 	{
-		// TODO: Figure out how we want the axes widget to be drawn
-		DrawAxes(Viewport, Canvas);
+		// Don't draw the usual 3D axes if the projection mode is UV
+		if (ProjectionMode != EDisplayClusterMeshProjectionType::UV)
+		{
+			DrawAxes(Viewport, Canvas);
+		}
 	}
 
 	// NOTE: DebugCanvasObject will be created by UDebugDrawService::Draw() if it doesn't already exist.
@@ -708,7 +721,8 @@ void FDisplayClusterLightCardEditorViewportClient::Draw(FViewport* InViewport, F
 
 void FDisplayClusterLightCardEditorViewportClient::Draw(const FSceneView* View, FPrimitiveDrawInterface* PDI)
 {
-	if (SelectedLightCards.Num())
+	const bool bIsUVProjection = ProjectionMode == EDisplayClusterMeshProjectionType::UV;
+	if (LastSelectedLightCard.IsValid() && LastSelectedLightCard->bIsUVLightCard == bIsUVProjection)
 	{
 		// Project the editor widget's world position into projection space so that it renders at the appropriate screen location.
 		// This needs to be computed on the render thread using the render thread's scene view, which will be behind the game thread's scene view
@@ -821,7 +835,14 @@ bool FDisplayClusterLightCardEditorViewportClient::InputWidgetDelta(FViewport* I
 			switch (EditorWidget->GetWidgetMode())
 			{
 			case FDisplayClusterLightCardEditorWidget::EWidgetMode::WM_Translate:
-				MoveSelectedLightCards(InViewport, CurrentAxis);
+				if (ProjectionMode == EDisplayClusterMeshProjectionType::UV)
+				{
+					MoveSelectedUVLightCards(InViewport, CurrentAxis);
+				}
+				else
+				{
+					MoveSelectedLightCards(InViewport, CurrentAxis);
+				}
 				break;
 
 			case FDisplayClusterLightCardEditorWidget::EWidgetMode::WM_RotateZ:
@@ -883,15 +904,15 @@ void FDisplayClusterLightCardEditorViewportClient::TrackingStarted(
 
 		if (RenderViewportType != LVT_Perspective)
 		{
-			// For orthogonal projections, PixelToWorld does not return the view origin or a direction from the view origin. Use TraceScreenRay
-			// to find a useful direction away from the view origin to use
-			const FVector ViewOrigin = View->ViewLocation;
-
-			Direction = TraceScreenRay(Origin, Direction, ViewOrigin);
-			Origin = ViewOrigin;
+			// For orthogonal projections, the drag widget offset should store the origin offset instead of the direction offset,
+			// since the direction offset is always the same regardless of which pixel has been clicked while the origin moves
+			DragWidgetOffset = Origin - FPlane::PointPlaneProject(CachedEditorWidgetWorldTransform.GetTranslation(), FPlane(Origin, Direction));
+		}
+		else
+		{
+			DragWidgetOffset = Direction - (CachedEditorWidgetWorldTransform.GetTranslation() - Origin).GetSafeNormal();
 		}
 
-		DragWidgetOffset = Direction - (CachedEditorWidgetWorldTransform.GetTranslation() - Origin).GetSafeNormal();
 		LastWidgetMousePos = MousePos;
 	}
 
@@ -923,24 +944,40 @@ void FDisplayClusterLightCardEditorViewportClient::CalculateNormalAndPositionInD
 	FVector& OutRelativeNormal, 
 	double InDesiredDistanceFromFlush) const
 {
-	const FVector Position = InViewOrigin + InDirection; // We fabricate a position in the right direction
+	if (ProjectionMode == EDisplayClusterMeshProjectionType::UV)
+	{
+		// In UV projection mode, all relevant geometry is projected onto the UV plane, so compute the world position and normal
+		// by performing a ray-plane intersection on the UV projection plane
 
-	float Distance;
+		const float UVProjectionPlaneDistance = ADisplayClusterLightCardActor::UVPlaneDefaultDistance;
 
-	// We find the normal and distance from origin
-	if (Position.Z < InViewOrigin.Z)
-	{
-		SouthNormalMap.GetNormalAndDistanceAtPosition(Position, OutRelativeNormal, Distance);
+		const FPlane UVProjectionPlane(InViewOrigin + FVector::ForwardVector * UVProjectionPlaneDistance, -FVector::ForwardVector);
+		const FVector PlaneIntersection = FMath::RayPlaneIntersection(InViewOrigin, InDirection, UVProjectionPlane);
+
+		OutRelativeNormal = UVProjectionPlane.GetNormal();
+		OutWorldPosition = PlaneIntersection;
 	}
 	else
 	{
-		NorthNormalMap.GetNormalAndDistanceAtPosition(Position, OutRelativeNormal, Distance);
-	}
+		const FVector Position = InViewOrigin + InDirection; // We fabricate a position in the right direction
+
+		float Distance;
+
+		// We find the normal and distance from origin
+		if (Position.Z < InViewOrigin.Z)
+		{
+			SouthNormalMap.GetNormalAndDistanceAtPosition(Position, OutRelativeNormal, Distance);
+		}
+		else
+		{
+			NorthNormalMap.GetNormalAndDistanceAtPosition(Position, OutRelativeNormal, Distance);
+		}
 
-	Distance = CalculateFinalLightCardDistance(Distance, InDesiredDistanceFromFlush);
+		Distance = CalculateFinalLightCardDistance(Distance, InDesiredDistanceFromFlush);
 
-	// Calculate world position
-	OutWorldPosition = InViewOrigin + Distance * InDirection;
+		// Calculate world position
+		OutWorldPosition = InViewOrigin + Distance * InDirection;
+	}
 };
 
 
@@ -1047,7 +1084,7 @@ void FDisplayClusterLightCardEditorViewportClient::CreateDrawnLightCard(const TA
 		CalculateNormalAndPositionInDirection(ViewOrigin, LightCardDirection, LightCardLocation, LightCardRelativeNormal);
 
 		const FMatrix RotEffector = FRotationMatrix::MakeFromX(-LightCardRelativeNormal);
-		const FMatrix RotArm = FRotationMatrix::MakeFromX(LightCardDirection);
+		const FMatrix RotArm = ProjectionMode != EDisplayClusterMeshProjectionType::UV ? FRotationMatrix::MakeFromX(LightCardDirection) : FMatrix::Identity;
 		const FRotator TotalRotator = (RotEffector * RotArm).Rotator();
 
 		LightCardPlaneAxisX = TotalRotator.RotateVector(-FVector::XAxisVector); // Same as Normal
@@ -1055,8 +1092,6 @@ void FDisplayClusterLightCardEditorViewportClient::CreateDrawnLightCard(const TA
 		LightCardPlaneAxisZ = TotalRotator.RotateVector( FVector::ZAxisVector);
 	}
 
-	const FSphericalCoordinates LightCardCoords(LightCardLocation - ViewOrigin);
-
 	//
 	// Project mouse positions to light card plane.
 	//
@@ -1197,7 +1232,20 @@ void FDisplayClusterLightCardEditorViewportClient::CreateDrawnLightCard(const TA
 	}
 
 	// Update position
-	MoveLightCardTo(*LightCard, LightCardCoords);
+	if (ProjectionMode == EDisplayClusterMeshProjectionType::UV)
+	{
+		const float UVProjectionPlaneSize = ADisplayClusterLightCardActor::UVPlaneDefaultSize;
+
+		const FVector DesiredLocation = LightCardLocation - ViewOrigin;
+		const FVector2D DesiredUVLocation = FVector2D(DesiredLocation.Y / UVProjectionPlaneSize + 0.5f, 0.5f - DesiredLocation.Z / UVProjectionPlaneSize);
+
+		LightCard->UVCoordinates = DesiredUVLocation;
+	}
+	else
+	{
+		const FSphericalCoordinates LightCardCoords(LightCardLocation - ViewOrigin);
+		MoveLightCardTo(*LightCard, LightCardCoords);
+	}
 
 #endif // WITH_OPENCV
 }
@@ -1597,7 +1645,18 @@ void FDisplayClusterLightCardEditorViewportClient::UpdatePreviewActor(ADisplayCl
 
 			for (const FLightCardProxy& LightCardProxy : LightCardProxies)
 			{
-				IDisplayClusterScenePreview::Get().AddActorToRenderer(PreviewRendererId, LightCardProxy.Proxy.Get());
+				IDisplayClusterScenePreview::Get().AddActorToRenderer(PreviewRendererId, LightCardProxy.Proxy.Get(), [this](const UPrimitiveComponent* PrimitiveComponent)
+				{
+					// Always add the light card mesh component to the renderer's scene even if it is marked hidden in game, since UV light cards will purposefully
+					// hide the light card mesh since it isn't supposed to exist in 3D space. The light card mesh will be appropriately filtered when the scene is
+					// rendered based on the projection mode
+					if (PrimitiveComponent->GetFName() == TEXT("LightCard"))
+					{
+						return true;
+					}
+
+					return !PrimitiveComponent->bHiddenInGame;
+				});
 			}
 
 			Finalize();
@@ -1716,7 +1775,7 @@ void FDisplayClusterLightCardEditorViewportClient::SetProjectionMode(EDisplayClu
 	ProjectionMode = InProjectionMode;
 	RenderViewportType = InViewportType;
 
-	if (ProjectionMode == EDisplayClusterMeshProjectionType::Linear)
+	if (ProjectionMode == EDisplayClusterMeshProjectionType::Linear || ProjectionMode == EDisplayClusterMeshProjectionType::UV)
 	{
 		// TODO: Do we want to cache the perspective rotation and restore it when the user switches back?
 		SetViewRotation(FVector::ForwardVector.Rotation());
@@ -1817,7 +1876,15 @@ void FDisplayClusterLightCardEditorViewportClient::MoveLightCardTo(ADisplayClust
 void FDisplayClusterLightCardEditorViewportClient::CenterLightCardInView(ADisplayClusterLightCardActor& LightCard)
 {
 	VerifyAndFixLightCardOrigin(&LightCard);
-	MoveLightCardTo(LightCard, FSphericalCoordinates(GetViewRotation().RotateVector(FVector::ForwardVector)));
+
+	if (LightCard.bIsUVLightCard)
+	{
+		LightCard.UVCoordinates = FVector2D(0.5, 0.5);
+	}
+	else
+	{
+		MoveLightCardTo(LightCard, FSphericalCoordinates(GetViewRotation().RotateVector(FVector::ForwardVector)));
+	}
 
 	// If this is a proxy light, propagate to its counterpart in the level.
 	if (LightCard.bIsProxy)
@@ -1855,8 +1922,8 @@ void FDisplayClusterLightCardEditorViewportClient::GetSceneViewInitOptions(FScen
 
 	FViewportCameraTransform& ViewTransform = GetViewTransform();
 
-	ViewInitOptions.ViewLocation = ViewTransform.GetLocation();
-	ViewInitOptions.ViewRotation = ViewTransform.GetRotation();
+	ViewInitOptions.ViewLocation = ProjectionMode == EDisplayClusterMeshProjectionType::UV ? FVector::ZeroVector : ViewTransform.GetLocation();
+	ViewInitOptions.ViewRotation = ProjectionMode == EDisplayClusterMeshProjectionType::UV ? FRotator::ZeroRotator : ViewTransform.GetRotation();
 	ViewInitOptions.ViewOrigin = ViewInitOptions.ViewLocation;
 
 	FIntPoint ViewportSize = Viewport->GetSizeXY();
@@ -2118,12 +2185,14 @@ void FDisplayClusterLightCardEditorViewportClient::SelectLightCard(ADisplayClust
 		}
 
 		SelectedLightCards.Empty();
+		LastSelectedLightCard = nullptr;
 	}
 
 	if (Actor)
 	{
 		SelectedLightCards.Add(Actor);
 		UpdatedActors.Add(Actor);
+		LastSelectedLightCard = Actor;
 	}
 
 	for (AActor* UpdatedActor : UpdatedActors)
@@ -2194,6 +2263,7 @@ void FDisplayClusterLightCardEditorViewportClient::PropagateLightCardTransform(A
 		TryChangeProperty(GET_MEMBER_NAME_CHECKED(ADisplayClusterLightCardActor, Pitch));
 		TryChangeProperty(GET_MEMBER_NAME_CHECKED(ADisplayClusterLightCardActor, Yaw));
 		TryChangeProperty(GET_MEMBER_NAME_CHECKED(ADisplayClusterLightCardActor, Scale));
+		TryChangeProperty(GET_MEMBER_NAME_CHECKED(ADisplayClusterLightCardActor, UVCoordinates));
 		
 		// Snapshot the changed properties so multi-user can update while dragging.
 		if (ChangedProperties.Num() > 0)
@@ -2266,9 +2336,7 @@ void FDisplayClusterLightCardEditorViewportClient::MoveSelectedLightCards(FViewp
 
 	FSceneView* View = CalcSceneView(&ViewFamily);
 
-	const TWeakObjectPtr<ADisplayClusterLightCardActor>& LastSelectedLightCard = SelectedLightCards.Last();
-
-	if (LastSelectedLightCard.IsValid())
+	if (LastSelectedLightCard.IsValid() && !LastSelectedLightCard->bIsUVLightCard)
 	{
 		const bool bUseDeltaRotation = (CurrentAxis == EAxisList::Type::XYZ) || (CurrentAxis == EAxisList::Type::Y);
 
@@ -2284,7 +2352,7 @@ void FDisplayClusterLightCardEditorViewportClient::MoveSelectedLightCards(FViewp
 
 		for (const TWeakObjectPtr<ADisplayClusterLightCardActor>& LightCard : SelectedLightCards)
 		{
-			if (!LightCard.IsValid())
+			if (!LightCard.IsValid() || LightCard->bIsUVLightCard)
 			{
 				continue;
 			}
@@ -2383,27 +2451,6 @@ void FDisplayClusterLightCardEditorViewportClient::MoveSelectedLightCards(FViewp
 
 void FDisplayClusterLightCardEditorViewportClient::MoveLightCardsToPixel(const FIntPoint& PixelPos, const TArray<TWeakObjectPtr<ADisplayClusterLightCardActor>>& InLightCards)
 {
-	// First, find the average position of all selected light cards
-	FSphericalCoordinates AverageCoords = FSphericalCoordinates();
-	int32 NumLightCards = 0;
-
-	for (const TWeakObjectPtr<ADisplayClusterLightCardActor>& LightCard : InLightCards)
-	{
-		if (LightCard.IsValid())
-		{
-			const FSphericalCoordinates LightCardCoords = GetLightCardCoordinates(LightCard.Get());
-
-			AverageCoords = AverageCoords + LightCardCoords;
-			++NumLightCards;
-		}
-	}
-
-	AverageCoords.Radius /= NumLightCards;
-	AverageCoords.Azimuth /= NumLightCards;
-	AverageCoords.Inclination /= NumLightCards;
-	AverageCoords.Conform();
-
-	// Now, find the desired coordinates based on the mouse position
 	FSceneViewFamilyContext ViewFamily(FSceneViewFamily::ConstructionValues(
 		Viewport,
 		GetScene(),
@@ -2428,20 +2475,79 @@ void FDisplayClusterLightCardEditorViewportClient::MoveLightCardsToPixel(const F
 		Origin = ViewOrigin;
 	}
 
-	// Compute desired coordinates (radius doesn't matter here since we will use the flush constraint on the light cards after moving them)
-	const FSphericalCoordinates DesiredCoords(Direction * 100.0f);
-	const FSphericalCoordinates DeltaCoords = DesiredCoords - AverageCoords;
+	if (ProjectionMode == EDisplayClusterMeshProjectionType::UV)
+	{
+		// Find the average position of all selected light cards. This group average is what is moved to the specified pixel
+		FVector2D AverageUVCoords = FVector2D::ZeroVector;
+		int32 NumLightCards = 0;
+
+		for (const TWeakObjectPtr<ADisplayClusterLightCardActor>& LightCard : InLightCards)
+		{
+			if (LightCard.IsValid() && LightCard->bIsUVLightCard)
+			{
+				AverageUVCoords += LightCard->UVCoordinates;
+				++NumLightCards;
+			}
+		}
+
+		AverageUVCoords /= NumLightCards;
 
-	// Update each light card with the delta coordinates; the flush constraint is applied by MoveLightCardTo, ensuring the light card is always flush to screens
-	for (const TWeakObjectPtr<ADisplayClusterLightCardActor>& LightCard : InLightCards)
+		// Compute the desired coordinates by projecting the specified screen ray onto the UV projection plane
+		const float UVProjectionPlaneSize = ADisplayClusterLightCardActor::UVPlaneDefaultSize;
+		const float UVProjectionPlaneDistance = ADisplayClusterLightCardActor::UVPlaneDefaultDistance;
+
+		const FPlane UVProjectionPlane(FVector::ForwardVector * UVProjectionPlaneDistance, -FVector::ForwardVector);
+		const FVector PlaneIntersection = FMath::RayPlaneIntersection(FVector::ZeroVector, Direction, UVProjectionPlane);
+		const FVector2D DesiredUVCoords = FVector2D(PlaneIntersection.Y / UVProjectionPlaneSize + 0.5f, 0.5f - PlaneIntersection.Z / UVProjectionPlaneSize);
+
+		const FVector2D DeltaUVCoords = DesiredUVCoords - AverageUVCoords;
+
+		for (const TWeakObjectPtr<ADisplayClusterLightCardActor>& LightCard : InLightCards)
+		{
+			if (LightCard.IsValid() && LightCard->bIsUVLightCard)
+			{
+				LightCard->UVCoordinates += DeltaUVCoords;
+				PropagateLightCardTransform(LightCard.Get());
+			}
+		}
+	}
+	else
 	{
-		if (LightCard.IsValid())
+		// Find the average position of all selected light cards. This group average is what is moved to the specified pixel
+		FSphericalCoordinates AverageCoords = FSphericalCoordinates();
+		int32 NumLightCards = 0;
+
+		for (const TWeakObjectPtr<ADisplayClusterLightCardActor>& LightCard : InLightCards)
 		{
-			const FSphericalCoordinates LightCardCoords = GetLightCardCoordinates(LightCard.Get());
-			const FSphericalCoordinates NewCoords = LightCardCoords + DeltaCoords;
+			if (LightCard.IsValid() && !LightCard->bIsUVLightCard)
+			{
+				const FSphericalCoordinates LightCardCoords = GetLightCardCoordinates(LightCard.Get());
 
-			MoveLightCardTo(*LightCard.Get(), NewCoords);
-			PropagateLightCardTransform(LightCard.Get());
+				AverageCoords = AverageCoords + LightCardCoords;
+				++NumLightCards;
+			}
+		}
+
+		AverageCoords.Radius /= NumLightCards;
+		AverageCoords.Azimuth /= NumLightCards;
+		AverageCoords.Inclination /= NumLightCards;
+		AverageCoords.Conform();
+
+		// Compute desired coordinates (radius doesn't matter here since we will use the flush constraint on the light cards after moving them)
+		const FSphericalCoordinates DesiredCoords(Direction * 100.0f);
+		const FSphericalCoordinates DeltaCoords = DesiredCoords - AverageCoords;
+
+		// Update each light card with the delta coordinates; the flush constraint is applied by MoveLightCardTo, ensuring the light card is always flush to screens
+		for (const TWeakObjectPtr<ADisplayClusterLightCardActor>& LightCard : InLightCards)
+		{
+			if (LightCard.IsValid() && !LightCard->bIsUVLightCard)
+			{
+				const FSphericalCoordinates LightCardCoords = GetLightCardCoordinates(LightCard.Get());
+				const FSphericalCoordinates NewCoords = LightCardCoords + DeltaCoords;
+
+				MoveLightCardTo(*LightCard.Get(), NewCoords);
+				PropagateLightCardTransform(LightCard.Get());
+			}
 		}
 	}
 }
@@ -2500,11 +2606,13 @@ FDisplayClusterLightCardEditorViewportClient::FSphericalCoordinates FDisplayClus
 
 		const FVector ViewOrigin = View->ViewLocation;
 
-		Direction = TraceScreenRay(Origin, Direction, ViewOrigin);
+		Direction = TraceScreenRay(Origin - DragWidgetOffset, Direction, ViewOrigin);
 		Origin = ViewOrigin;
 	}
-
-	Direction = (Direction - DragWidgetOffset).GetSafeNormal();
+	else
+	{
+		Direction = (Direction - DragWidgetOffset).GetSafeNormal();
+	}
 
 	const FVector LocalDirection = LightCard->GetActorRotation().RotateVector(Direction);
 	const FVector LightCardLocation = LightCard->GetLightCardTransform().GetTranslation() - Origin;
@@ -2550,9 +2658,69 @@ FDisplayClusterLightCardEditorViewportClient::FSphericalCoordinates FDisplayClus
 	return DeltaCoords;
 }
 
+void FDisplayClusterLightCardEditorViewportClient::MoveSelectedUVLightCards(FViewport* InViewport, EAxisList::Type CurrentAxis)
+{
+	if (LastSelectedLightCard.IsValid() && LastSelectedLightCard->bIsUVLightCard)
+	{
+		const FVector2D DeltaUV = GetUVLightCardTranslationDelta(InViewport, LastSelectedLightCard.Get(), CurrentAxis);
+		for (const TWeakObjectPtr<ADisplayClusterLightCardActor>& LightCard : SelectedLightCards)
+		{
+			if (!LightCard.IsValid() || !LightCard->bIsUVLightCard)
+			{
+				continue;
+			}
+
+			LightCard->UVCoordinates += DeltaUV;
+
+			PropagateLightCardTransform(LightCard.Get());
+		}
+	}
+}
+
+FVector2D FDisplayClusterLightCardEditorViewportClient::GetUVLightCardTranslationDelta(FViewport* InViewport, ADisplayClusterLightCardActor* LightCard, EAxisList::Type CurrentAxis)
+{
+	FIntPoint MousePos;
+	InViewport->GetMousePos(MousePos);
+
+	FSceneViewFamilyContext ViewFamily(FSceneViewFamily::ConstructionValues(
+		InViewport,
+		GetScene(),
+		EngineShowFlags)
+		.SetRealtimeUpdate(IsRealtime()));
+	FSceneView* View = CalcSceneView(&ViewFamily);
+
+	FVector Origin;
+	FVector Direction;
+	PixelToWorld(*View, MousePos, Origin, Direction);
+
+	const float UVProjectionPlaneSize = ADisplayClusterLightCardActor::UVPlaneDefaultSize;
+	const float UVProjectionPlaneDistance = ADisplayClusterLightCardActor::UVPlaneDefaultDistance;
+
+	const FVector ViewOrigin = View->ViewMatrices.GetViewOrigin();
+	const FPlane UVProjectionPlane(ViewOrigin + FVector::ForwardVector * UVProjectionPlaneDistance, -FVector::ForwardVector);
+	const FVector PlaneIntersection = FMath::RayPlaneIntersection(Origin, Direction, UVProjectionPlane);
+
+	const FVector DesiredLocation = ((PlaneIntersection - DragWidgetOffset) - ViewOrigin);
+	const FVector2D DesiredUVLocation = FVector2D(DesiredLocation.Y / UVProjectionPlaneSize + 0.5f, 0.5f - DesiredLocation.Z / UVProjectionPlaneSize);
+
+	const FVector2D UVDelta = DesiredUVLocation - LightCard->UVCoordinates;
+
+	FVector2D UVAxis = FVector2D::ZeroVector;
+	if (CurrentAxis & EAxisList::Type::X)
+	{
+		UVAxis += FVector2D(1.0, 0.0);
+	}
+
+	if (CurrentAxis & EAxisList::Type::Y)
+	{
+		UVAxis += FVector2D(0.0, 1.0);
+	}
+
+	return UVDelta * UVAxis;
+}
+
 void FDisplayClusterLightCardEditorViewportClient::ScaleSelectedLightCards(FViewport* InViewport, EAxisList::Type CurrentAxis)
 {
-	const TWeakObjectPtr<ADisplayClusterLightCardActor>& LastSelectedLightCard = SelectedLightCards.Last();
 	if (LastSelectedLightCard.IsValid())
 	{
 		const FVector2D DeltaScale = GetLightCardScaleDelta(InViewport, LastSelectedLightCard.Get(), CurrentAxis);
@@ -2639,7 +2807,6 @@ FVector2D FDisplayClusterLightCardEditorViewportClient::GetLightCardScaleDelta(F
 
 void FDisplayClusterLightCardEditorViewportClient::SpinSelectedLightCards(FViewport* InViewport)
 {
-	const TWeakObjectPtr<ADisplayClusterLightCardActor>& LastSelectedLightCard = SelectedLightCards.Last();
 	if (LastSelectedLightCard.IsValid())
 	{
 		const double DeltaSpin = GetLightCardSpinDelta(InViewport, LastSelectedLightCard.Get());
@@ -2747,44 +2914,58 @@ FVector FDisplayClusterLightCardEditorViewportClient::TraceScreenRay(const FVect
 
 	FVector Direction = FVector::ZeroVector;
 
-	// First, trace against the stage actor to see if the screen ray hits it; if so, simply return the direction from the view origin to this hit point
-	FVector HitLocation = FVector::ZeroVector;
-	if (TraceStage(RayStart, RayEnd, HitLocation))
+	if (ProjectionMode == EDisplayClusterMeshProjectionType::UV)
 	{
-		Direction = (HitLocation - ViewOrigin).GetSafeNormal();
+		// In the UV projection mode, all stage geometry has been projected onto a UV projection plane, so perform a ray trace against that plane to find the 
+		// screen ray direction
+		const float UVProjectionPlaneDistance = ADisplayClusterLightCardActor::UVPlaneDefaultDistance;
+
+		const FPlane UVProjectionPlane(ViewOrigin + FVector::ForwardVector * UVProjectionPlaneDistance, -FVector::ForwardVector);
+		const FVector PlaneIntersection = FMath::RayPlaneIntersection(OrthogonalOrigin, OrthogonalDirection, UVProjectionPlane);
+
+		Direction = (PlaneIntersection - ViewOrigin).GetSafeNormal();
 	}
 	else
 	{
-		// If we didn't hit any stage geometry, try to trace against the normal map mesh. Procedural meshes does not appear to handle inward pointing normals correctly,
-		// so we need to reverse the trace start and end locations to get a useful hit result
+		// First, trace against the stage actor to see if the screen ray hits it; if so, simply return the direction from the view origin to this hit point
+		FVector HitLocation = FVector::ZeroVector;
+		if (TraceStage(RayStart, RayEnd, HitLocation))
+		{
+			Direction = (HitLocation - ViewOrigin).GetSafeNormal();
+		}
+		else
+		{
+			// If we didn't hit any stage geometry, try to trace against the normal map mesh. Procedural meshes does not appear to handle inward pointing normals correctly,
+			// so we need to reverse the trace start and end locations to get a useful hit result
 
-		UWorld* PreviewWorld = PreviewScene->GetWorld();
-		check(PreviewWorld);
+			UWorld* PreviewWorld = PreviewScene->GetWorld();
+			check(PreviewWorld);
 
-		FCollisionQueryParams TraceParams(SCENE_QUERY_STAT(DisplayClusterStageTrace), true);
+			FCollisionQueryParams TraceParams(SCENE_QUERY_STAT(DisplayClusterStageTrace), true);
 
-		for (const FLightCardProxy& ProxyRef : LightCardProxies)
-		{
-			if (ProxyRef.Proxy.IsValid())
+			for (const FLightCardProxy& ProxyRef : LightCardProxies)
 			{
-				TraceParams.AddIgnoredActor(ProxyRef.Proxy.Get());
+				if (ProxyRef.Proxy.IsValid())
+				{
+					TraceParams.AddIgnoredActor(ProxyRef.Proxy.Get());
+				}
 			}
-		}
 
-		TraceParams.AddIgnoredActor(RootActorProxy.Get());
+			TraceParams.AddIgnoredActor(RootActorProxy.Get());
 
-		FHitResult HitResult;
-		if (PreviewWorld->LineTraceSingleByObjectType(HitResult, RayEnd, RayStart, FCollisionObjectQueryParams(FCollisionObjectQueryParams::InitType::AllObjects), TraceParams))
-		{
-			HitLocation = HitResult.Location;
-			Direction = (HitLocation - ViewOrigin).GetSafeNormal();
-		}
-		else
-		{
-			// If the screen ray does not hit the stage or the normal map mesh, then simply use the closest point on the ray to the view origin
-			const FVector ClosestPoint = OrthogonalOrigin + ((ViewOrigin - OrthogonalOrigin) | OrthogonalDirection) * OrthogonalDirection;
+			FHitResult HitResult;
+			if (PreviewWorld->LineTraceSingleByObjectType(HitResult, RayEnd, RayStart, FCollisionObjectQueryParams(FCollisionObjectQueryParams::InitType::AllObjects), TraceParams))
+			{
+				HitLocation = HitResult.Location;
+				Direction = (HitLocation - ViewOrigin).GetSafeNormal();
+			}
+			else
+			{
+				// If the screen ray does not hit the stage or the normal map mesh, then simply use the closest point on the ray to the view origin
+				const FVector ClosestPoint = OrthogonalOrigin + ((ViewOrigin - OrthogonalOrigin) | OrthogonalDirection) * OrthogonalDirection;
 
-			Direction = (ClosestPoint - ViewOrigin).GetSafeNormal();
+				Direction = (ClosestPoint - ViewOrigin).GetSafeNormal();
+			}
 		}
 	}
 
@@ -2934,39 +3115,44 @@ bool FDisplayClusterLightCardEditorViewportClient::WorldToScreenDirection(const
 	return false;
 }
 
-bool FDisplayClusterLightCardEditorViewportClient::CalcEditorWidgetTransform(FTransform& WidgetTransformBeforeMapProjection)
+bool FDisplayClusterLightCardEditorViewportClient::CalcEditorWidgetTransform(FTransform& WidgetTransform)
 {
 	if (!SelectedLightCards.Num())
 	{
 		return false;
 	}
 
-	TWeakObjectPtr<ADisplayClusterLightCardActor> LastSelected = SelectedLightCards.Last();
-
-	if (!LastSelected.IsValid())
+	if (!LastSelectedLightCard.IsValid())
 	{
 		return false;
 	}
 
-	FVector LightCardPosition = LastSelected->GetLightCardTransform().GetTranslation();
+	FVector LightCardPosition = LastSelectedLightCard->GetLightCardTransform().GetTranslation();
 
-	WidgetTransformBeforeMapProjection = FTransform(FRotator::ZeroRotator, LightCardPosition, FVector::OneVector);
+	WidgetTransform = FTransform(FRotator::ZeroRotator, LightCardPosition, FVector::OneVector);
 
 	FQuat WidgetOrientation;
 	if (EditorWidget->GetWidgetMode() == FDisplayClusterLightCardEditorWidget::EWidgetMode::WM_Translate)
 	{
-		// The translation widget should be oriented to show the x axis pointing in the longitudinal direction and the y axis pointing in the latitudinal direction
-		const FVector ProjectionOrigin = ProjectionOriginComponent.IsValid() ? ProjectionOriginComponent->GetComponentLocation() : FVector::ZeroVector;
-		const FVector RadialVector = (LightCardPosition - ProjectionOrigin).GetSafeNormal();
-		const FVector AzimuthalVector = (FVector::ZAxisVector ^ RadialVector).GetSafeNormal();
-		const FVector InclinationVector = RadialVector ^ AzimuthalVector;
+		if (ProjectionMode == EDisplayClusterMeshProjectionType::UV)
+		{
+			WidgetOrientation = FMatrix(FVector::YAxisVector, FVector::ZAxisVector, FVector::XAxisVector, FVector::ZeroVector).ToQuat();
+		}
+		else
+		{
+			// The translation widget should be oriented to show the x axis pointing in the longitudinal direction and the y axis pointing in the latitudinal direction
+			const FVector ProjectionOrigin = ProjectionOriginComponent.IsValid() ? ProjectionOriginComponent->GetComponentLocation() : FVector::ZeroVector;
+			const FVector RadialVector = (LightCardPosition - ProjectionOrigin).GetSafeNormal();
+			const FVector AzimuthalVector = (FVector::ZAxisVector ^ RadialVector).GetSafeNormal();
+			const FVector InclinationVector = RadialVector ^ AzimuthalVector;
 
-		WidgetOrientation = FMatrix(AzimuthalVector, InclinationVector, RadialVector, FVector::ZeroVector).ToQuat();
+			WidgetOrientation = FMatrix(AzimuthalVector, InclinationVector, RadialVector, FVector::ZeroVector).ToQuat();
+		}
 	}
 	else
 	{
 		// Otherwise, orient the widget to match the light card's local orientation (spin, pitch, yaw)
-		const FQuat LightCardOrientation = LastSelected->GetLightCardTransform().GetRotation();
+		const FQuat LightCardOrientation = LastSelectedLightCard->GetLightCardTransform().GetRotation();
 		const FVector Normal = LightCardOrientation.RotateVector(FVector::XAxisVector);
 		const FVector Tangent = LightCardOrientation.RotateVector(FVector::YAxisVector);
 		const FVector Binormal = LightCardOrientation.RotateVector(FVector::ZAxisVector);
@@ -2975,7 +3161,7 @@ bool FDisplayClusterLightCardEditorViewportClient::CalcEditorWidgetTransform(FTr
 		WidgetOrientation = FMatrix(-Tangent, Binormal, -Normal, FVector::ZeroVector).ToQuat();
 	}
 
-	WidgetTransformBeforeMapProjection.SetRotation(WidgetOrientation);
+	WidgetTransform.SetRotation(WidgetOrientation);
 
 	return true;
 }
@@ -2984,13 +3170,13 @@ void FDisplayClusterLightCardEditorViewportClient::RenderNormalMap(FNormalMap& N
 {
 	// Only render primitive components from the stage actor for the normal map
 	FDisplayClusterMeshProjectionPrimitiveFilter PrimitiveFilter;
-	PrimitiveFilter.PrimitiveFilterDelegate = FDisplayClusterMeshProjectionPrimitiveFilter::FPrimitiveFilter::CreateLambda([this](const UPrimitiveComponent* PrimitiveComponent)
+	PrimitiveFilter.ShouldRenderPrimitiveDelegate = FDisplayClusterMeshProjectionPrimitiveFilter::FPrimitiveFilter::CreateLambda([this](const UPrimitiveComponent* PrimitiveComponent)
 	{
 		return PrimitiveComponent->GetOwner() == RootActorProxy;
 	});
 
-	FDisplayClusterScenePreviewRenderSettings RenderSettings;
-	RenderSettings.RenderType = FDisplayClusterScenePreviewRenderSettings::ERenderType::Normals;
+	FDisplayClusterMeshProjectionRenderSettings RenderSettings;
+	RenderSettings.RenderType = EDisplayClusterMeshProjectionOutput::Normals;
 	RenderSettings.EngineShowFlags = EngineShowFlags;
 	RenderSettings.ProjectionType = EDisplayClusterMeshProjectionType::Azimuthal;
 	RenderSettings.PrimitiveFilter = PrimitiveFilter;
@@ -3081,13 +3267,14 @@ bool FDisplayClusterLightCardEditorViewportClient::IsLocationCloseToEdge(const F
 
 void FDisplayClusterLightCardEditorViewportClient::ResetFOVs()
 {
-	constexpr int32 MaxFOVs = 2;
+	constexpr int32 MaxFOVs = 3;
 	if (ProjectionFOVs.Num() < MaxFOVs)
 	{
 		ProjectionFOVs.AddDefaulted(MaxFOVs - ProjectionFOVs.Num());
 	}
 	ProjectionFOVs[static_cast<int32>(EDisplayClusterMeshProjectionType::Linear)] = 90.0f;
 	ProjectionFOVs[static_cast<int32>(EDisplayClusterMeshProjectionType::Azimuthal)] = 130.0f;
+	ProjectionFOVs[static_cast<int32>(EDisplayClusterMeshProjectionType::UV)] = 45.0f;
 }
 
 void FDisplayClusterLightCardEditorViewportClient::EnterDrawingLightCardMode()
@@ -3108,5 +3295,33 @@ void FDisplayClusterLightCardEditorViewportClient::ExitDrawingLightCardMode()
 	}
 }
 
+bool FDisplayClusterLightCardEditorViewportClient::ShouldRenderPrimitive(const UPrimitiveComponent* PrimitiveComponent)
+{
+	const bool bIsUVProjection = ProjectionMode == EDisplayClusterMeshProjectionType::UV;
+	if (ADisplayClusterLightCardActor* LightCard = Cast<ADisplayClusterLightCardActor>(PrimitiveComponent->GetOwner()))
+	{
+		// Only render the UV light cards when in UV projection mode, and only render non-UV light cards in any other projection mode
+		return bIsUVProjection ? LightCard->bIsUVLightCard : !LightCard->bIsUVLightCard;
+	}
+	else
+	{
+		return true;
+	}
+}
+
+bool FDisplayClusterLightCardEditorViewportClient::ShouldApplyProjectionToPrimitive(const UPrimitiveComponent* PrimitiveComponent)
+{
+	const bool bIsUVProjection = ProjectionMode == EDisplayClusterMeshProjectionType::UV;
+	if (ADisplayClusterLightCardActor* LightCard = Cast<ADisplayClusterLightCardActor>(PrimitiveComponent->GetOwner()))
+	{
+		// When in UV projection mode, don't render the UV light cards using the UV projection, render them linearly
+		if (bIsUVProjection && LightCard->bIsUVLightCard)
+		{
+			return false;
+		}
+	}
+
+	return true;
+}
 
 #undef LOCTEXT_NAMESPACE
diff --git a/Engine/Plugins/Runtime/nDisplay/Source/DisplayClusterLightCardEditor/Private/Viewport/DisplayClusterLightCardEditorViewportClient.h b/Engine/Plugins/Runtime/nDisplay/Source/DisplayClusterLightCardEditor/Private/Viewport/DisplayClusterLightCardEditorViewportClient.h
--- a/Engine/Plugins/Runtime/nDisplay/Source/DisplayClusterLightCardEditor/Private/Viewport/DisplayClusterLightCardEditorViewportClient.h
+++ b/Engine/Plugins/Runtime/nDisplay/Source/DisplayClusterLightCardEditor/Private/Viewport/DisplayClusterLightCardEditorViewportClient.h
@@ -267,6 +267,12 @@ private:
 	/** Determines the appropriate delta in spherical coordinates needed to move the specified light card to the mouse's location */
 	FSphericalCoordinates GetLightCardTranslationDelta(FViewport* InViewport, ADisplayClusterLightCardActor* LightCard, EAxisList::Type CurrentAxis);
 
+	/** Moves the currently selected UV light cards */
+	void MoveSelectedUVLightCards(FViewport* InViewport, EAxisList::Type CurrentAxis);
+
+	/** Determines the appropriate delta in UV coordinates needed to move the specified light card to the mouse's location */
+	FVector2D GetUVLightCardTranslationDelta(FViewport* InViewport, ADisplayClusterLightCardActor* LightCard, EAxisList::Type CurrentAxis);
+
 	/** Scales the currently selected light cards */
 	void ScaleSelectedLightCards(FViewport* InViewport, EAxisList::Type CurrentAxis);
 
@@ -307,7 +313,7 @@ private:
 	bool WorldToScreenDirection(const FSceneView& View, const FVector& WorldPos, const FVector& WorldDirection, FVector2D& OutScreenDir);
 
 	/** Calculates the world transform to render the editor widget with to align it with the selected light card */
-	bool CalcEditorWidgetTransform(FTransform& WidgetTransformBeforeMapProjection);
+	bool CalcEditorWidgetTransform(FTransform& WidgetTransform);
 	
 	/** Renders the viewport's normal map and stores the texture data to be used later */
 	void RenderNormalMap(FNormalMap& NormalMap, const FVector& NormalMapDirection);
@@ -330,6 +336,12 @@ private:
 	/** Calculates the relative normal vector and world position in the specified direction from the given view origin */
 	void CalculateNormalAndPositionInDirection(const FVector& InViewOrigin, const FVector& InDirection, FVector& OutWorldLocation, FVector& OutRelativeNormal, double InDesiredDistanceFromFlush = 0.) const;
 
+	/** Callback passed into the mesh projection renderer to filter which primitives are drawn */
+	bool ShouldRenderPrimitive(const UPrimitiveComponent* PrimitiveComponent);
+
+	/** Callback passed into the mesh projection renderer to filter which primitives drawn with the current projection and which are drawn linearly */
+	bool ShouldApplyProjectionToPrimitive(const UPrimitiveComponent* PrimitiveComponent);
+
 private:
 	TWeakPtr<FSceneViewport> SceneViewportPtr;
 	TWeakPtr<SDisplayClusterLightCardEditorViewport> LightCardEditorViewportPtr;
@@ -356,6 +368,9 @@ private:
 	TArray<FLightCardProxy> LightCardProxies;
 	TArray<TWeakObjectPtr<ADisplayClusterLightCardActor>> SelectedLightCards;
 
+	/** The light card in the selected light card list that was selected last */
+	TWeakObjectPtr<ADisplayClusterLightCardActor> LastSelectedLightCard = nullptr;
+
 	/** Light card preview actors being dropped on the scene */
 	TArray<TWeakObjectPtr<ADisplayClusterLightCardActor>> DropPreviewLightCards;
 	
diff --git a/Engine/Plugins/Runtime/nDisplay/Source/DisplayClusterLightCardEditor/Private/Viewport/DisplayClusterLightcardEditorViewport.cpp b/Engine/Plugins/Runtime/nDisplay/Source/DisplayClusterLightCardEditor/Private/Viewport/DisplayClusterLightcardEditorViewport.cpp
--- a/Engine/Plugins/Runtime/nDisplay/Source/DisplayClusterLightCardEditor/Private/Viewport/DisplayClusterLightcardEditorViewport.cpp
+++ b/Engine/Plugins/Runtime/nDisplay/Source/DisplayClusterLightCardEditor/Private/Viewport/DisplayClusterLightcardEditorViewport.cpp
@@ -127,6 +127,10 @@ public:
 			case EDisplayClusterMeshProjectionType::Azimuthal:
 				Label = LOCTEXT("ProjectionwMenuTitle_Azimuthal", "Azimuthal");
 				break;
+
+			case EDisplayClusterMeshProjectionType::UV:
+				Label = LOCTEXT("ProjectionwMenuTitle_UV", "UV");
+				break;
 			}
 		}
 
@@ -143,6 +147,7 @@ public:
 		ViewMenuBuilder.AddMenuEntry(FDisplayClusterLightCardEditorCommands::Get().PerspectiveProjection);
 		ViewMenuBuilder.AddMenuEntry(FDisplayClusterLightCardEditorCommands::Get().OrthographicProjection);
 		ViewMenuBuilder.AddMenuEntry(FDisplayClusterLightCardEditorCommands::Get().AzimuthalProjection);
+		ViewMenuBuilder.AddMenuEntry(FDisplayClusterLightCardEditorCommands::Get().UVProjection);
 
 		return ViewMenuBuilder.MakeWidget();
 	}
@@ -406,6 +411,12 @@ void SDisplayClusterLightCardEditorViewport::BindCommands()
 			FCanExecuteAction(),
 			FIsActionChecked::CreateSP(this, &SDisplayClusterLightCardEditorViewport::IsProjectionModeSelected, EDisplayClusterMeshProjectionType::Azimuthal, ELevelViewportType::LVT_Perspective));
 
+		CommandList->MapAction(
+			FDisplayClusterLightCardEditorCommands::Get().UVProjection,
+			FExecuteAction::CreateSP(this, &SDisplayClusterLightCardEditorViewport::SetProjectionMode, EDisplayClusterMeshProjectionType::UV, ELevelViewportType::LVT_OrthoFreelook),
+			FCanExecuteAction(),
+			FIsActionChecked::CreateSP(this, &SDisplayClusterLightCardEditorViewport::IsProjectionModeSelected, EDisplayClusterMeshProjectionType::UV, ELevelViewportType::LVT_OrthoFreelook));
+
 		CommandList->MapAction(
 			Commands.ViewOrientationTop,
 			FExecuteAction::CreateSP(this, &SDisplayClusterLightCardEditorViewport::SetViewDirection, ViewDirectionTop));
diff --git a/Engine/Plugins/Runtime/nDisplay/Source/DisplayClusterLightCardEditorShaders/Private/DisplayClusterMeshProjectionRenderer.cpp b/Engine/Plugins/Runtime/nDisplay/Source/DisplayClusterLightCardEditorShaders/Private/DisplayClusterMeshProjectionRenderer.cpp
--- a/Engine/Plugins/Runtime/nDisplay/Source/DisplayClusterLightCardEditorShaders/Private/DisplayClusterMeshProjectionRenderer.cpp
+++ b/Engine/Plugins/Runtime/nDisplay/Source/DisplayClusterLightCardEditorShaders/Private/DisplayClusterMeshProjectionRenderer.cpp
@@ -18,6 +18,147 @@
 #include "Materials/Material.h"
 
 
+//////////////////////////////////////////////////////////////////////////
+// Mesh Pass Processor
+
+BEGIN_SHADER_PARAMETER_STRUCT(FMeshProjectionPassParameters, )
+	SHADER_PARAMETER_STRUCT_REF(FViewUniformShaderParameters, View)
+	SHADER_PARAMETER_RDG_UNIFORM_BUFFER(FInstanceCullingGlobalUniforms, InstanceCulling)
+	SHADER_PARAMETER_STRUCT_INCLUDE(FInstanceCullingDrawParams, InstanceCullingDrawParams)
+	RENDER_TARGET_BINDING_SLOTS()
+END_SHADER_PARAMETER_STRUCT()
+
+class FMeshProjectionShaderElementData : public FMeshMaterialShaderElementData
+{
+public:
+	FDisplayClusterMeshProjectionTypeSettings ProjectionTypeSettings;
+};
+
+/** Base class for all mesh pass processors used by the projection renderer.
+ * 
+ *  NOTE: Due to C++ dependent naming, any class derived from this base class must access this class's
+ *  member variables and functions through this-> or FMeshProjectionPassProcessorBase<...>::, to ensure
+ *  that the compiler can properly locate the variables and functions once the class templates have been generated
+ */
+template<typename VertexType, typename PixelType, typename ShaderElementDataType = FMeshProjectionShaderElementData>
+class FMeshProjectionPassProcessorBase : public FMeshPassProcessor
+{
+public:
+	FMeshProjectionPassProcessorBase(const FScene* InScene,
+		const FSceneView* InView,
+		FMeshPassDrawListContext* InDrawListContext,
+		const FDisplayClusterMeshProjectionRenderSettings& InRenderSettings)
+		: FMeshPassProcessor(InScene, GMaxRHIFeatureLevel, InView, InDrawListContext)
+		, ProjectionTypeSettings(InRenderSettings.ProjectionTypeSettings)
+		, StencilValue(0)
+	{
+		DrawRenderState.SetDepthStencilState(TStaticDepthStencilState<>::GetRHI());
+		DrawRenderState.SetBlendState(TStaticBlendState<>::GetRHI());
+	}
+
+	virtual void AddMeshBatch(const FMeshBatch& RESTRICT MeshBatch, uint64 BatchElementMask, const FPrimitiveSceneProxy* RESTRICT PrimitiveSceneProxy, int32 StaticMeshId = -1) override final
+	{
+		const FMaterial* Material;
+		const FMaterialRenderProxy* MaterialRenderProxy;
+		GetMeshBatchMaterial(MeshBatch, Material, MaterialRenderProxy);
+
+		check(Material);
+		check(MaterialRenderProxy);
+
+		if (!CanDrawMeshBatch(MeshBatch, PrimitiveSceneProxy, Material))
+		{
+			return;
+		}
+
+		const FVertexFactory* VertexFactory = MeshBatch.VertexFactory;
+
+		TMeshProcessorShaders<VertexType, PixelType> PassShaders;
+
+		FMaterialShaderTypes ShaderTypes;
+		ShaderTypes.AddShaderType<VertexType>();
+		ShaderTypes.AddShaderType<PixelType>();
+
+		FMaterialShaders Shaders;
+		if (!Material->TryGetShaders(ShaderTypes, VertexFactory->GetType(), Shaders))
+		{
+			return;
+		}
+
+		Shaders.TryGetVertexShader(PassShaders.VertexShader);
+		Shaders.TryGetPixelShader(PassShaders.PixelShader);
+
+		FMeshDrawingPolicyOverrideSettings OverrideSettings = GetMeshOverrideSettings(MeshBatch);
+		ERasterizerFillMode MeshFillMode = ComputeMeshFillMode(MeshBatch, *Material, OverrideSettings);
+		ERasterizerCullMode MeshCullMode = ComputeMeshCullMode(MeshBatch, *Material, OverrideSettings);
+
+		ShaderElementDataType ShaderElementData = CreateShaderElementData(MeshBatch, PrimitiveSceneProxy, StaticMeshId);
+		FMeshDrawCommandSortKey SortKey = CreateMeshSortKey(MeshBatch, PrimitiveSceneProxy, *Material, PassShaders.VertexShader.GetShader(), PassShaders.PixelShader.GetShader());
+
+		DrawRenderState.SetStencilRef(StencilValue);
+
+		BuildMeshDrawCommands(
+			MeshBatch,
+			BatchElementMask,
+			PrimitiveSceneProxy,
+			*MaterialRenderProxy,
+			*Material,
+			DrawRenderState,
+			PassShaders,
+			MeshFillMode,
+			MeshCullMode,
+			SortKey,
+			EMeshPassFeatures::Default,
+			ShaderElementData);
+	}
+
+	void SetStencilValue(uint32 InStencilValue)
+	{
+		StencilValue = InStencilValue;
+	}
+
+protected:
+	virtual bool CanDrawMeshBatch(const FMeshBatch& RESTRICT MeshBatch, const FPrimitiveSceneProxy* RESTRICT PrimitiveSceneProxy, const FMaterial* Material)
+	{
+		return true;
+	}
+
+	virtual void GetMeshBatchMaterial(const FMeshBatch& RESTRICT MeshBatch, const FMaterial*& OutMaterial, const FMaterialRenderProxy*& OutMaterialRenderProxy)
+	{
+		const FMaterialRenderProxy* FallbackMaterialRenderProxy = nullptr;
+		OutMaterial = &MeshBatch.MaterialRenderProxy->GetMaterialWithFallback(FeatureLevel, FallbackMaterialRenderProxy);
+		OutMaterialRenderProxy = FallbackMaterialRenderProxy ? FallbackMaterialRenderProxy : MeshBatch.MaterialRenderProxy;
+	}
+
+	virtual FMeshPassProcessor::FMeshDrawingPolicyOverrideSettings GetMeshOverrideSettings(const FMeshBatch& RESTRICT MeshBatch)
+	{
+		return ComputeMeshOverrideSettings(MeshBatch);
+	}
+
+	virtual ShaderElementDataType CreateShaderElementData(const FMeshBatch& RESTRICT MeshBatch, const FPrimitiveSceneProxy* RESTRICT PrimitiveSceneProxy, int32 StaticMeshId)
+	{
+		ShaderElementDataType ShaderElementData;
+		ShaderElementData.InitializeMeshMaterialData(ViewIfDynamicMeshCommand, PrimitiveSceneProxy, MeshBatch, StaticMeshId, false);
+		ShaderElementData.ProjectionTypeSettings = ProjectionTypeSettings;
+
+		return ShaderElementData;
+	}
+
+	virtual FMeshDrawCommandSortKey CreateMeshSortKey(const FMeshBatch& RESTRICT MeshBatch,
+		const FPrimitiveSceneProxy* RESTRICT PrimitiveSceneProxy,
+		const FMaterial& Material,
+		const FMeshMaterialShader* VertexShader,
+		const FMeshMaterialShader* PixelShader)
+	{
+		return CalculateMeshStaticSortKey(VertexShader, PixelShader);
+	}
+
+protected:
+	FMeshPassProcessorRenderState DrawRenderState;
+	FDisplayClusterMeshProjectionTypeSettings ProjectionTypeSettings;
+	uint32 StencilValue;
+};
+
+
 //////////////////////////////////////////////////////////////////////////
 // Base Render Pass
 
@@ -32,6 +173,9 @@ public:
 		: FMeshMaterialShader(Initializer)
 	{
 		PassUniformBuffer.Bind(Initializer.ParameterMap, FSceneTextureUniformParameters::StaticStructMetadata.GetShaderVariableName());
+		UVProjectionIndex.Bind(Initializer.ParameterMap, TEXT("ProjectionParameters_UVProjectionIndex"), SPF_Optional);
+		UVProjectionPlaneSize.Bind(Initializer.ParameterMap, TEXT("ProjectionParameters_UVProjectionPlaneSize"), SPF_Optional);
+		UVProjectionPlaneDistance.Bind(Initializer.ParameterMap, TEXT("ProjectionParameters_UVProjectionPlaneDistance"), SPF_Optional);
 	}
 
 	static bool ShouldCompilePermutation(const FMeshMaterialShaderPermutationParameters& Parameters)
@@ -40,21 +184,39 @@ public:
 			&& IsFeatureLevelSupported(Parameters.Platform, ERHIFeatureLevel::SM5)
 			&& Parameters.VertexFactoryType == FindVertexFactoryType(TEXT("FLocalVertexFactory"));
 	}
+
+	void GetShaderBindings(
+		const FScene* Scene,
+		ERHIFeatureLevel::Type FeatureLevel,
+		const FPrimitiveSceneProxy* PrimitiveSceneProxy,
+		const FMaterialRenderProxy& MaterialRenderProxy,
+		const FMaterial& Material,
+		const FMeshPassProcessorRenderState& DrawRenderState,
+		const FMeshProjectionShaderElementData& ShaderElementData,
+		FMeshDrawSingleShaderBindings& ShaderBindings) const
+	{
+		FMeshMaterialShader::GetShaderBindings(Scene, FeatureLevel, PrimitiveSceneProxy, MaterialRenderProxy, Material, DrawRenderState, ShaderElementData, ShaderBindings);
+
+		ShaderBindings.Add(UVProjectionIndex, ShaderElementData.ProjectionTypeSettings.UVProjectionIndex);
+		ShaderBindings.Add(UVProjectionPlaneSize, ShaderElementData.ProjectionTypeSettings.UVProjectionPlaneSize);
+		ShaderBindings.Add(UVProjectionPlaneDistance, ShaderElementData.ProjectionTypeSettings.UVProjectionPlaneDistance);
+	}
+
+private:
+	LAYOUT_FIELD(FShaderParameter, UVProjectionIndex);
+	LAYOUT_FIELD(FShaderParameter, UVProjectionPlaneSize);
+	LAYOUT_FIELD(FShaderParameter, UVProjectionPlaneDistance);
 };
 
 using FMeshPerspectiveProjectionVS = FMeshProjectionVS<EDisplayClusterMeshProjectionType::Linear>;
 using FMeshAzimuthalProjectionVS = FMeshProjectionVS<EDisplayClusterMeshProjectionType::Azimuthal>;
+using FMeshUVProjectionVS = FMeshProjectionVS<EDisplayClusterMeshProjectionType::UV>;
 
 IMPLEMENT_MATERIAL_SHADER_TYPE(template<>, FMeshPerspectiveProjectionVS, TEXT("/Plugin/nDisplay/Private/MeshProjectionShaders.usf"), TEXT("PerspectiveVS"), SF_Vertex);
 IMPLEMENT_MATERIAL_SHADER_TYPE(template<>, FMeshAzimuthalProjectionVS, TEXT("/Plugin/nDisplay/Private/MeshProjectionShaders.usf"), TEXT("AzimuthalVS"), SF_Vertex);
+IMPLEMENT_MATERIAL_SHADER_TYPE(template<>, FMeshUVProjectionVS, TEXT("/Plugin/nDisplay/Private/MeshProjectionShaders.usf"), TEXT("UVProjectionVS"), SF_Vertex);
 
-enum EMeshProjectionShaderOutput
-{
-	Default,
-	Normals
-};
-
-template<EMeshProjectionShaderOutput OutputType>
+template<EDisplayClusterMeshProjectionOutput OutputType>
 class FMeshProjectionPS : public FMeshMaterialShader
 {
 public:
@@ -75,109 +237,57 @@ public:
 	}
 };
 
-using FMeshProjectionDefaultPS = FMeshProjectionPS<EMeshProjectionShaderOutput::Default>;
-using FMeshProjectionNormalPS = FMeshProjectionPS<EMeshProjectionShaderOutput::Normals>;
+using FMeshProjectionColorPS = FMeshProjectionPS<EDisplayClusterMeshProjectionOutput::Color>;
+using FMeshProjectionNormalPS = FMeshProjectionPS<EDisplayClusterMeshProjectionOutput::Normals>;
 
-IMPLEMENT_MATERIAL_SHADER_TYPE(template<>, FMeshProjectionDefaultPS, TEXT("/Plugin/nDisplay/Private/MeshProjectionShaders.usf"), TEXT("MainPS"), SF_Pixel);
+IMPLEMENT_MATERIAL_SHADER_TYPE(template<>, FMeshProjectionColorPS, TEXT("/Plugin/nDisplay/Private/MeshProjectionShaders.usf"), TEXT("MainPS"), SF_Pixel);
 IMPLEMENT_MATERIAL_SHADER_TYPE(template<>, FMeshProjectionNormalPS, TEXT("/Plugin/nDisplay/Private/MeshProjectionShaders.usf"), TEXT("NormalPS"), SF_Pixel);
 
-BEGIN_SHADER_PARAMETER_STRUCT(FMeshProjectionPassParameters, )
-	SHADER_PARAMETER_STRUCT_REF(FViewUniformShaderParameters, View)
-	SHADER_PARAMETER_RDG_UNIFORM_BUFFER(FInstanceCullingGlobalUniforms, InstanceCulling)
-	SHADER_PARAMETER_STRUCT_INCLUDE(FInstanceCullingDrawParams, InstanceCullingDrawParams)
-	RENDER_TARGET_BINDING_SLOTS()
-END_SHADER_PARAMETER_STRUCT()
+FMeshProjectionPassParameters* CreateProjectionPassParameters(FRDGBuilder& GraphBuilder, const FViewInfo* View, const FDisplayClusterMeshProjectionRenderSettings& RenderSettings)
+{
+	FMeshProjectionPassParameters* PassParameters = GraphBuilder.AllocParameters<FMeshProjectionPassParameters>();
+	PassParameters->View = View->ViewUniformBuffer;
+	PassParameters->InstanceCulling = FInstanceCullingContext::CreateDummyInstanceCullingUniformBuffer(GraphBuilder);
 
-template<EDisplayClusterMeshProjectionType ProjectionType, EMeshProjectionShaderOutput OutputType = EMeshProjectionShaderOutput::Default> 
-class FLightCardEditorMeshPassProcessor : public FMeshPassProcessor
+	return PassParameters;
+}
+
+template<EDisplayClusterMeshProjectionType ProjectionType, EDisplayClusterMeshProjectionOutput OutputType = EDisplayClusterMeshProjectionOutput::Color>
+class FMeshProjectionPassProcessor : public FMeshProjectionPassProcessorBase<FMeshProjectionVS<ProjectionType>, FMeshProjectionPS<OutputType>>
 {
+	using Super = FMeshProjectionPassProcessorBase<FMeshProjectionVS<ProjectionType>, FMeshProjectionPS<OutputType>>;
+
 public:
-	FLightCardEditorMeshPassProcessor(const FScene* InScene, const FSceneView* InView, FMeshPassDrawListContext* InDrawListContext, bool bIsTranslucencyPass = false)
-		: FMeshPassProcessor(InScene, GMaxRHIFeatureLevel, InView, InDrawListContext)
+	FMeshProjectionPassProcessor(const FScene* InScene,
+		const FSceneView* InView,
+		FMeshPassDrawListContext* InDrawListContext,
+		const FDisplayClusterMeshProjectionRenderSettings& InRenderSettings,
+		bool bIsTranslucencyPass = false)
+		: Super(InScene, InView, InDrawListContext, InRenderSettings)
 		, bTranslucencyPass(bIsTranslucencyPass)
 		, bIgnoreTranslucency(false)
 	{
-
 		if (bTranslucencyPass)
 		{
-			DrawRenderState.SetDepthStencilState(TStaticDepthStencilState<false, CF_DepthNearOrEqual>::GetRHI());
-			DrawRenderState.SetBlendState(TStaticBlendState<CW_RGBA, BO_Add, BF_SourceAlpha, BF_InverseSourceAlpha, BO_Add, BF_Zero, BF_One>::GetRHI());
+			this->DrawRenderState.SetDepthStencilState(TStaticDepthStencilState<false, CF_DepthNearOrEqual>::GetRHI());
+			this->DrawRenderState.SetBlendState(TStaticBlendState<CW_RGBA, BO_Add, BF_SourceAlpha, BF_InverseSourceAlpha, BO_Add, BF_Zero, BF_One>::GetRHI());
 		}
 		else
 		{
-			DrawRenderState.SetDepthStencilState(TStaticDepthStencilState<true, CF_DepthNearOrEqual, true, CF_Always, SO_Keep, SO_Keep, SO_Replace>::GetRHI());
-			DrawRenderState.SetBlendState(TStaticBlendStateWriteMask<CW_RGBA, CW_RGBA, CW_RGBA, CW_RGBA>::GetRHI());
+			this->DrawRenderState.SetDepthStencilState(TStaticDepthStencilState<true, CF_DepthNearOrEqual, true, CF_Always, SO_Keep, SO_Keep, SO_Replace>::GetRHI());
+			this->DrawRenderState.SetBlendState(TStaticBlendStateWriteMask<CW_RGBA, CW_RGBA, CW_RGBA, CW_RGBA>::GetRHI());
 		}
 	}
 
-	virtual void AddMeshBatch(const FMeshBatch& RESTRICT MeshBatch, uint64 BatchElementMask, const FPrimitiveSceneProxy* RESTRICT PrimitiveSceneProxy, int32 StaticMeshId = -1) override final
-	{
-		const FMaterialRenderProxy* FallbackMaterialRenderProxy = nullptr;
-		const FMaterial& Material = MeshBatch.MaterialRenderProxy->GetMaterialWithFallback(FeatureLevel, FallbackMaterialRenderProxy);
-		const FMaterialRenderProxy& MaterialRenderProxy = FallbackMaterialRenderProxy ? *FallbackMaterialRenderProxy : *MeshBatch.MaterialRenderProxy;
-
-		if (!CanDrawMeshBatch(MeshBatch, Material))
-		{
-			return;
-		}
-
-		const FVertexFactory* VertexFactory = MeshBatch.VertexFactory;
-
-		TMeshProcessorShaders<FMeshProjectionVS<ProjectionType>, FMeshProjectionPS<OutputType>> PassShaders;
-
-		FMaterialShaderTypes ShaderTypes;
-		ShaderTypes.AddShaderType<FMeshProjectionVS<ProjectionType>>();
-		ShaderTypes.AddShaderType<FMeshProjectionPS<OutputType>>();
-
-		FMaterialShaders Shaders;
-		if (!Material.TryGetShaders(ShaderTypes, VertexFactory->GetType(), Shaders))
-		{
-			return;
-		}
-
-		Shaders.TryGetVertexShader(PassShaders.VertexShader);
-		Shaders.TryGetPixelShader(PassShaders.PixelShader);
-
-		const FMeshDrawingPolicyOverrideSettings OverrideSettings = ComputeMeshOverrideSettings(MeshBatch);
-		ERasterizerFillMode MeshFillMode = ComputeMeshFillMode(MeshBatch, Material, OverrideSettings);
-		ERasterizerCullMode MeshCullMode = ComputeMeshCullMode(MeshBatch, Material, OverrideSettings);
-
-		FMeshMaterialShaderElementData ShaderElementData;
-		ShaderElementData.InitializeMeshMaterialData(ViewIfDynamicMeshCommand, PrimitiveSceneProxy, MeshBatch, StaticMeshId, false);
-
-		FMeshDrawCommandSortKey SortKey = CreateMeshSortKey(MeshBatch, PrimitiveSceneProxy, Material, PassShaders.VertexShader.GetShader(), PassShaders.PixelShader.GetShader());
-
-		DrawRenderState.SetStencilRef(StencilValue);
-
-		BuildMeshDrawCommands(
-			MeshBatch,
-			BatchElementMask,
-			PrimitiveSceneProxy,
-			MaterialRenderProxy,
-			Material,
-			DrawRenderState,
-			PassShaders,
-			MeshFillMode,
-			MeshCullMode,
-			SortKey,
-			EMeshPassFeatures::Default,
-			ShaderElementData);
-	}
-
-	void SetStencilValue(uint32 InStencilValue)
-	{
-		StencilValue = InStencilValue;
-	}
-
 	void SetIgnoreTranslucency(bool bInIgnoreTranslucency)
 	{
 		bIgnoreTranslucency = bInIgnoreTranslucency;
 	}
 
-private:
-	bool CanDrawMeshBatch(const FMeshBatch& RESTRICT MeshBatch, const FMaterial& Material)
+protected:
+	virtual bool CanDrawMeshBatch(const FMeshBatch& RESTRICT MeshBatch, const FPrimitiveSceneProxy* RESTRICT PrimitiveSceneProxy, const FMaterial* Material) override
 	{
-		const bool bIsTranslucent = IsTranslucentBlendMode(Material.GetBlendMode());
+		const bool bIsTranslucent = IsTranslucentBlendMode(Material->GetBlendMode());
 
 		if (bTranslucencyPass)
 		{
@@ -189,11 +299,24 @@ private:
 		}
 	}
 
-	FMeshDrawCommandSortKey CreateMeshSortKey(const FMeshBatch& RESTRICT MeshBatch,
+	virtual FMeshPassProcessor::FMeshDrawingPolicyOverrideSettings GetMeshOverrideSettings(const FMeshBatch& RESTRICT MeshBatch) override
+	{
+		FMeshPassProcessor::FMeshDrawingPolicyOverrideSettings OverrideSettings = FMeshPassProcessor::ComputeMeshOverrideSettings(MeshBatch);
+
+		// Don't cull backfacing geometry if the projection mode is UV, since there is no guarantee that the UV coordinates are in the correct winding order
+		if (ProjectionType == EDisplayClusterMeshProjectionType::UV)
+		{
+			OverrideSettings.MeshOverrideFlags |= EDrawingPolicyOverrideFlags::TwoSided;
+		}
+
+		return OverrideSettings;
+	}
+
+	virtual FMeshDrawCommandSortKey CreateMeshSortKey(const FMeshBatch& RESTRICT MeshBatch,
 		const FPrimitiveSceneProxy* RESTRICT PrimitiveSceneProxy,
 		const FMaterial& Material,
-		const FMeshMaterialShader* VertexShader, 
-		const FMeshMaterialShader* PixelShader)
+		const FMeshMaterialShader* VertexShader,
+		const FMeshMaterialShader* PixelShader) override
 	{
 		FMeshDrawCommandSortKey SortKey = FMeshDrawCommandSortKey::Default;
 
@@ -224,8 +347,6 @@ private:
 	}
 
 private:
-	FMeshPassProcessorRenderState DrawRenderState;
-	uint32 StencilValue;
 	bool bTranslucencyPass;
 	bool bIgnoreTranslucency;
 };
@@ -234,14 +355,9 @@ private:
 //////////////////////////////////////////////////////////////////////////
 // Hit Proxy Render Pass
 
-class FMeshProjectionHitProxyShaderElementData : public FMeshMaterialShaderElementData
+class FMeshProjectionHitProxyShaderElementData : public FMeshProjectionShaderElementData
 {
 public:
-	FMeshProjectionHitProxyShaderElementData(FHitProxyId InBatchHitProxyId)
-		: BatchHitProxyId(InBatchHitProxyId)
-	{
-	}
-
 	FHitProxyId BatchHitProxyId;
 };
 
@@ -291,17 +407,21 @@ private:
 IMPLEMENT_MATERIAL_SHADER_TYPE(, FMeshProjectionHitProxyPS, TEXT("/Plugin/nDisplay/Private/MeshProjectionHitProxy.usf"), TEXT("Main") ,SF_Pixel);
 
 template<EDisplayClusterMeshProjectionType ProjectionType> 
-class FLightCardEditorHitProxyMeshPassProcessor : public FMeshPassProcessor
+class FMeshProjectionHitProxyPassProcessor : public FMeshProjectionPassProcessorBase<FMeshProjectionVS<ProjectionType>, FMeshProjectionHitProxyPS, FMeshProjectionHitProxyShaderElementData>
 {
+	using Super = FMeshProjectionPassProcessorBase<FMeshProjectionVS<ProjectionType>, FMeshProjectionHitProxyPS, FMeshProjectionHitProxyShaderElementData>;
+
 public:
-	FLightCardEditorHitProxyMeshPassProcessor(const FScene* InScene, const FSceneView* InView, FMeshPassDrawListContext* InDrawListContext)
-		: FMeshPassProcessor(InScene, GMaxRHIFeatureLevel, InView, InDrawListContext)
+	FMeshProjectionHitProxyPassProcessor(const FScene* InScene,
+		const FSceneView* InView,
+		FMeshPassDrawListContext* InDrawListContext,
+		const FDisplayClusterMeshProjectionRenderSettings& InRenderSettings)
+		: Super(InScene, InView, InDrawListContext, InRenderSettings)
 	{
-		DrawRenderState.SetDepthStencilState(TStaticDepthStencilState<true, CF_DepthNearOrEqual>::GetRHI());
-		DrawRenderState.SetBlendState(TStaticBlendState<>::GetRHI());
 	}
 
-	virtual void AddMeshBatch(const FMeshBatch& RESTRICT MeshBatch, uint64 BatchElementMask, const FPrimitiveSceneProxy* RESTRICT PrimitiveSceneProxy, int32 StaticMeshId = -1) override final
+protected:
+	virtual bool CanDrawMeshBatch(const FMeshBatch& RESTRICT MeshBatch, const FPrimitiveSceneProxy* RESTRICT PrimitiveSceneProxy, const FMaterial* Material) override
 	{
 		const bool bDrawMeshBatch = MeshBatch.bUseForMaterial
 			&& MeshBatch.BatchHitProxyId != FHitProxyId::InvisibleHitProxyId
@@ -309,66 +429,38 @@ public:
 			&& PrimitiveSceneProxy
 			&& PrimitiveSceneProxy->IsSelectable();
 
-		if (bDrawMeshBatch)
+		return bDrawMeshBatch;
+	}
+
+	virtual void GetMeshBatchMaterial(const FMeshBatch& RESTRICT MeshBatch, const FMaterial*& OutMaterial, const FMaterialRenderProxy*& OutMaterialRenderProxy) override
+	{
+		OutMaterialRenderProxy = MeshBatch.MaterialRenderProxy;
+		OutMaterial = OutMaterialRenderProxy->GetMaterialNoFallback(this->FeatureLevel);
+		if (OutMaterial && OutMaterial->GetRenderingThreadShaderMap())
 		{
-			const FMaterialRenderProxy* MaterialRenderProxy = MeshBatch.MaterialRenderProxy;
-			const FMaterial* Material = MaterialRenderProxy->GetMaterialNoFallback(FeatureLevel);
-			if (Material && Material->GetRenderingThreadShaderMap())
+			if (OutMaterial->WritesEveryPixel() && !OutMaterial->IsTwoSided() && !OutMaterial->MaterialModifiesMeshPosition_RenderThread())
 			{
-				if (Material->WritesEveryPixel() && !Material->IsTwoSided() && !Material->MaterialModifiesMeshPosition_RenderThread())
-				{
-					// Default material doesn't handle masked, and doesn't have the correct bIsTwoSided setting.
-					MaterialRenderProxy = UMaterial::GetDefaultMaterial(MD_Surface)->GetRenderProxy();
-					check(MaterialRenderProxy);
-					Material = MaterialRenderProxy->GetMaterialNoFallback(FeatureLevel);
-					check(Material);
-				}
-
-				const FVertexFactory* VertexFactory = MeshBatch.VertexFactory;
-
-				TMeshProcessorShaders<FMeshProjectionVS<ProjectionType>, FMeshProjectionHitProxyPS> PassShaders;
-
-				FMaterialShaderTypes ShaderTypes;
-				ShaderTypes.AddShaderType<FMeshProjectionVS<ProjectionType>>();
-				ShaderTypes.AddShaderType<FMeshProjectionHitProxyPS>();
-
-				FMaterialShaders Shaders;
-				if (!Material->TryGetShaders(ShaderTypes, VertexFactory->GetType(), Shaders))
-				{
-					return;
-				}
-
-				Shaders.TryGetVertexShader(PassShaders.VertexShader);
-				Shaders.TryGetPixelShader(PassShaders.PixelShader);
-
-				const FMeshDrawingPolicyOverrideSettings OverrideSettings = ComputeMeshOverrideSettings(MeshBatch);
-				ERasterizerFillMode MeshFillMode = ComputeMeshFillMode(MeshBatch, *Material, OverrideSettings);
-				ERasterizerCullMode MeshCullMode = ComputeMeshCullMode(MeshBatch, *Material, OverrideSettings);
-
-				FMeshProjectionHitProxyShaderElementData ShaderElementData(MeshBatch.BatchHitProxyId);
-				ShaderElementData.InitializeMeshMaterialData(ViewIfDynamicMeshCommand, PrimitiveSceneProxy, MeshBatch, StaticMeshId, false);
-
-				const FMeshDrawCommandSortKey SortKey = CalculateMeshStaticSortKey(PassShaders.VertexShader, PassShaders.PixelShader);
-
-				BuildMeshDrawCommands(
-					MeshBatch,
-					BatchElementMask,
-					PrimitiveSceneProxy,
-					*MaterialRenderProxy,
-					*Material,
-					DrawRenderState,
-					PassShaders,
-					MeshFillMode,
-					MeshCullMode,
-					SortKey,
-					EMeshPassFeatures::Default,
-					ShaderElementData);
+				// Default material doesn't handle masked, and doesn't have the correct bIsTwoSided setting.
+				OutMaterialRenderProxy = UMaterial::GetDefaultMaterial(MD_Surface)->GetRenderProxy();
+				OutMaterial = OutMaterialRenderProxy->GetMaterialNoFallback(this->FeatureLevel);
 			}
 		}
+		else
+		{
+			OutMaterialRenderProxy = UMaterial::GetDefaultMaterial(MD_Surface)->GetRenderProxy();
+			OutMaterial = OutMaterialRenderProxy->GetMaterialNoFallback(this->FeatureLevel);
+		}
 	}
 
-private:
-	FMeshPassProcessorRenderState DrawRenderState;
+	virtual FMeshProjectionHitProxyShaderElementData CreateShaderElementData(const FMeshBatch& RESTRICT MeshBatch, const FPrimitiveSceneProxy* RESTRICT PrimitiveSceneProxy, int32 StaticMeshId)
+	{
+		FMeshProjectionHitProxyShaderElementData ShaderElementData;
+		ShaderElementData.InitializeMeshMaterialData(this->ViewIfDynamicMeshCommand, PrimitiveSceneProxy, MeshBatch, StaticMeshId, false);
+		ShaderElementData.ProjectionTypeSettings = this->ProjectionTypeSettings;
+		ShaderElementData.BatchHitProxyId = MeshBatch.BatchHitProxyId;
+
+		return ShaderElementData;
+	}
 };
 
 //////////////////////////////////////////////////////////////////////////
@@ -524,18 +616,25 @@ public:
 IMPLEMENT_GLOBAL_SHADER(FMeshProjectionSelectionOutlinePS, "/Plugin/nDisplay/Private/MeshProjectionSelectionOutline.usf", "Main", SF_Pixel);
 
 template<EDisplayClusterMeshProjectionType ProjectionType> 
-class FLightCardEditorSelectionPassProcessor : public FMeshPassProcessor
+class FMeshProjectionSelectionPassProcessor : public FMeshProjectionPassProcessorBase<FMeshProjectionVS<ProjectionType>, FMeshProjectionColorPS>
 {
+	using Super = FMeshProjectionPassProcessorBase<FMeshProjectionVS<ProjectionType>, FMeshProjectionColorPS>;
+
 public:
-	FLightCardEditorSelectionPassProcessor(const FScene* InScene, const FSceneView* InView, FMeshPassDrawListContext* InDrawListContext)
-		: FMeshPassProcessor(InScene, GMaxRHIFeatureLevel, InView, InDrawListContext)
-		, StencilValue(1)
+	FMeshProjectionSelectionPassProcessor(const FScene* InScene,
+		const FSceneView* InView,
+		FMeshPassDrawListContext* InDrawListContext,
+		const FDisplayClusterMeshProjectionRenderSettings& InRenderSettings)
+		: Super(InScene, InView, InDrawListContext, InRenderSettings)
 	{
-		DrawRenderState.SetDepthStencilState(TStaticDepthStencilState<true, CF_DepthNearOrEqual, true, CF_Always, SO_Keep, SO_Keep, SO_Replace>::GetRHI());
-		DrawRenderState.SetBlendState(TStaticBlendStateWriteMask<CW_NONE, CW_NONE, CW_NONE, CW_NONE>::GetRHI());
+		this->DrawRenderState.SetDepthStencilState(TStaticDepthStencilState<true, CF_DepthNearOrEqual, true, CF_Always, SO_Keep, SO_Keep, SO_Replace>::GetRHI());
+		this->DrawRenderState.SetBlendState(TStaticBlendStateWriteMask<CW_NONE, CW_NONE, CW_NONE, CW_NONE>::GetRHI());
+
+		this->StencilValue = 1;
 	}
 
-	virtual void AddMeshBatch(const FMeshBatch& RESTRICT MeshBatch, uint64 BatchElementMask, const FPrimitiveSceneProxy* RESTRICT PrimitiveSceneProxy, int32 StaticMeshId = -1) override final
+protected:
+	virtual bool CanDrawMeshBatch(const FMeshBatch& RESTRICT MeshBatch, const FPrimitiveSceneProxy* RESTRICT PrimitiveSceneProxy, const FMaterial* Material) override
 	{
 		const bool bDrawMeshBatch = MeshBatch.bUseForMaterial
 			&& MeshBatch.bUseSelectionOutline
@@ -543,69 +642,36 @@ public:
 			&& PrimitiveSceneProxy->WantsSelectionOutline()
 			&& (PrimitiveSceneProxy->IsSelected() || PrimitiveSceneProxy->IsHovered());
 
-		if (bDrawMeshBatch)
+		return bDrawMeshBatch;
+	}
+
+	virtual void GetMeshBatchMaterial(const FMeshBatch& RESTRICT MeshBatch, const FMaterial*& OutMaterial, const FMaterialRenderProxy*& OutMaterialRenderProxy) override
+	{
+		OutMaterialRenderProxy = MeshBatch.MaterialRenderProxy;
+		OutMaterial = OutMaterialRenderProxy->GetMaterialNoFallback(this->FeatureLevel);
+		if (OutMaterial && OutMaterial->GetRenderingThreadShaderMap())
 		{
-			const FMaterialRenderProxy* MaterialRenderProxy = MeshBatch.MaterialRenderProxy;
-			const FMaterial* Material = MaterialRenderProxy->GetMaterialNoFallback(FeatureLevel);
-			if (Material && Material->GetRenderingThreadShaderMap())
+			if (OutMaterial->WritesEveryPixel() && !OutMaterial->IsTwoSided() && !OutMaterial->MaterialModifiesMeshPosition_RenderThread())
 			{
-				if (Material->WritesEveryPixel() && !Material->IsTwoSided() && !Material->MaterialModifiesMeshPosition_RenderThread())
-				{
-					// Default material doesn't handle masked, and doesn't have the correct bIsTwoSided setting.
-					MaterialRenderProxy = UMaterial::GetDefaultMaterial(MD_Surface)->GetRenderProxy();
-					check(MaterialRenderProxy);
-					Material = MaterialRenderProxy->GetMaterialNoFallback(FeatureLevel);
-					check(Material);
-				}
-				
-				const FVertexFactory* VertexFactory = MeshBatch.VertexFactory;
-
-				TMeshProcessorShaders<FMeshProjectionVS<ProjectionType>, FMeshProjectionDefaultPS> PassShaders;
-
-				FMaterialShaderTypes ShaderTypes;
-				ShaderTypes.AddShaderType<FMeshProjectionVS<ProjectionType>>();
-				ShaderTypes.AddShaderType<FMeshProjectionDefaultPS>();
-
-				FMaterialShaders Shaders;
-				if (!Material->TryGetShaders(ShaderTypes, VertexFactory->GetType(), Shaders))
-				{
-					return;
-				}
-
-				Shaders.TryGetVertexShader(PassShaders.VertexShader);
-				Shaders.TryGetPixelShader(PassShaders.PixelShader);
-
-				const FMeshDrawingPolicyOverrideSettings OverrideSettings = ComputeMeshOverrideSettings(MeshBatch);
-				ERasterizerFillMode MeshFillMode = ComputeMeshFillMode(MeshBatch, *Material, OverrideSettings);
-				ERasterizerCullMode MeshCullMode = CM_None;
-
-				DrawRenderState.SetStencilRef(StencilValue);
-
-				FMeshMaterialShaderElementData ShaderElementData;
-				ShaderElementData.InitializeMeshMaterialData(ViewIfDynamicMeshCommand, PrimitiveSceneProxy, MeshBatch, StaticMeshId, false);
-
-				FMeshDrawCommandSortKey SortKey{};
-
-				BuildMeshDrawCommands(
-					MeshBatch,
-					BatchElementMask,
-					PrimitiveSceneProxy,
-					*MaterialRenderProxy,
-					*Material,
-					DrawRenderState,
-					PassShaders,
-					MeshFillMode,
-					MeshCullMode,
-					SortKey,
-					EMeshPassFeatures::Default,
-					ShaderElementData);
+				// Default material doesn't handle masked, and doesn't have the correct bIsTwoSided setting.
+				OutMaterialRenderProxy = UMaterial::GetDefaultMaterial(MD_Surface)->GetRenderProxy();
+				OutMaterial = OutMaterialRenderProxy->GetMaterialNoFallback(this->FeatureLevel);
 			}
 		}
+		else
+		{
+			OutMaterialRenderProxy = UMaterial::GetDefaultMaterial(MD_Surface)->GetRenderProxy();
+			OutMaterial = OutMaterialRenderProxy->GetMaterialNoFallback(this->FeatureLevel);
+		}
 	}
 
-private:
-	FMeshPassProcessorRenderState DrawRenderState;
-	uint32 StencilValue;
+	virtual FMeshPassProcessor::FMeshDrawingPolicyOverrideSettings GetMeshOverrideSettings(const FMeshBatch& RESTRICT MeshBatch) override
+	{
+		FMeshPassProcessor::FMeshDrawingPolicyOverrideSettings OverrideSettings = FMeshPassProcessor::ComputeMeshOverrideSettings(MeshBatch);
+		OverrideSettings.MeshOverrideFlags |= EDrawingPolicyOverrideFlags::TwoSided;
+
+		return OverrideSettings;
+	}
 };
 
 namespace
@@ -676,17 +742,24 @@ private:
 	FHitProxyConsumer* HitProxyConsumer;
 };
 
-bool FDisplayClusterMeshProjectionPrimitiveFilter::IsPrimitiveComponentFiltered(const UPrimitiveComponent* InPrimitiveComponent) const
+bool FDisplayClusterMeshProjectionPrimitiveFilter::ShouldRenderPrimitive(const UPrimitiveComponent* InPrimitiveComponent) const
 {
-	if (PrimitiveFilterDelegate.IsBound())
+	if (ShouldRenderPrimitiveDelegate.IsBound())
 	{
-		if (!PrimitiveFilterDelegate.Execute(InPrimitiveComponent))
-		{
-			return true;
-		}
+		return ShouldRenderPrimitiveDelegate.Execute(InPrimitiveComponent);
 	}
 
-	return false;
+	return true;
+}
+
+bool FDisplayClusterMeshProjectionPrimitiveFilter::ShouldApplyProjection(const UPrimitiveComponent* InPrimitiveComponent) const
+{
+	if (ShouldApplyProjectionDelegate.IsBound())
+	{
+		return ShouldApplyProjectionDelegate.Execute(InPrimitiveComponent);
+	}
+
+	return true;
 }
 
 FVector FDisplayClusterMeshProjectionTransform::ProjectPosition(const FVector& WorldPosition) const
@@ -822,7 +895,7 @@ void FDisplayClusterMeshProjectionRenderer::ClearScene()
 	PrimitiveComponents.Empty();
 }
 
-void FDisplayClusterMeshProjectionRenderer::Render(FCanvas* Canvas, FSceneInterface* Scene, const FSceneViewInitOptions& ViewInitOptions, const FEngineShowFlags& EngineShowFlags, EDisplayClusterMeshProjectionType ProjectionType)
+void FDisplayClusterMeshProjectionRenderer::Render(FCanvas* Canvas, FSceneInterface* Scene, const FDisplayClusterMeshProjectionRenderSettings& RenderSettings)
 {
 	Canvas->Flush_GameThread();
 	FRenderTarget* RenderTarget = Canvas->GetRenderTarget();
@@ -830,14 +903,14 @@ void FDisplayClusterMeshProjectionRenderer::Render(FCanvas* Canvas, FSceneInterf
 	FHitProxyConsumer* HitProxyConsumer = Canvas->GetHitProxyConsumer();
 
 	ENQUEUE_RENDER_COMMAND(FDrawProjectedMeshes)(
-		[RenderTarget, Scene, ViewInitOptions, EngineShowFlags, bIsHitTesting, HitProxyConsumer, ProjectionType, this](FRHICommandListImmediate& RHICmdList)
+		[RenderTarget, Scene, RenderSettings, bIsHitTesting, HitProxyConsumer, this](FRHICommandListImmediate& RHICmdList)
 		{
 			FRDGBuilder GraphBuilder(RHICmdList);
 
 			FSceneViewFamilyContext ViewFamily(FSceneViewFamily::ConstructionValues(
 				RenderTarget,
 				Scene,
-				EngineShowFlags)
+				RenderSettings.EngineShowFlags)
 				.SetTime(FGameTime::GetTimeSinceAppStart())
 				.SetGammaCorrection(1.0f));
 
@@ -855,7 +928,7 @@ void FDisplayClusterMeshProjectionRenderer::Render(FCanvas* Canvas, FSceneInterf
 
 			FScenePrimitiveRenderingContextScopeHelper ScenePrimitiveRenderingContextScopeHelper(GetRendererModule().BeginScenePrimitiveRendering(GraphBuilder, &ViewFamily));
 
-			FSceneViewInitOptions NewInitOptions(ViewInitOptions);
+			FSceneViewInitOptions NewInitOptions(RenderSettings.ViewInitOptions);
 			NewInitOptions.ViewFamily = &ViewFamily;
 
 			GetRendererModule().CreateAndInitSingleView(RHICmdList, &ViewFamily, &NewInitOptions);
@@ -864,115 +937,24 @@ void FDisplayClusterMeshProjectionRenderer::Render(FCanvas* Canvas, FSceneInterf
 			FRDGTextureRef OutputTexture = GraphBuilder.RegisterExternalTexture(CreateRenderTarget(RenderTarget->GetRenderTargetTexture(), TEXT("ViewRenderTarget")));
 			FRenderTargetBinding OutputRenderTargetBinding(OutputTexture, ERenderTargetLoadAction::ELoad);
 
-			if (ViewFamily.EngineShowFlags.HitProxies)
+			switch (RenderSettings.RenderType)
 			{
-				FRDGTextureDesc Desc(FRDGTextureDesc::Create2D(OutputTexture->Desc.Extent, PF_B8G8R8A8, FClearValueBinding::Black, TexCreate_RenderTargetable | TexCreate_ShaderResource));
-				FRDGTextureRef HitProxyTexture = GraphBuilder.CreateTexture(Desc, TEXT("DisplayClusterMeshProjection.HitProxyTexture"));
-
-				const FRDGTextureDesc DepthDesc = FRDGTextureDesc::Create2D(OutputTexture->Desc.Extent, PF_DepthStencil, FClearValueBinding::DepthFar, TexCreate_DepthStencilTargetable | TexCreate_ShaderResource);
-				FRDGTextureRef HitProxyDepthTexture = GraphBuilder.CreateTexture(DepthDesc, TEXT("DisplayClusterMeshProjection.HitProxyDepthTexture"));
-
-				FRenderTargetBinding HitProxyRenderTargetBinding(HitProxyTexture, ERenderTargetLoadAction::EClear);
-				FDepthStencilBinding HitProxyDepthStencilBinding(HitProxyDepthTexture, ERenderTargetLoadAction::EClear, ERenderTargetLoadAction::EClear, FExclusiveDepthStencil::DepthWrite_StencilWrite);
-
-				AddHitProxyRenderPass(GraphBuilder, View, ProjectionType, HitProxyRenderTargetBinding, HitProxyDepthStencilBinding);
-
-				// Copy the hit proxy buffer to the viewport family's render target
-				{
-					FCopyRectPS::FParameters* ScreenPassParameters = GraphBuilder.AllocParameters<FCopyRectPS::FParameters>();
-					ScreenPassParameters->InputTexture = HitProxyTexture;
-					ScreenPassParameters->InputSampler = TStaticSamplerState<>::GetRHI();
-					ScreenPassParameters->RenderTargets[0] = OutputRenderTargetBinding;
-
-					FGlobalShaderMap* GlobalShaderMap = GetGlobalShaderMap(GMaxRHIFeatureLevel);
-					TShaderMapRef<FScreenPassVS> ScreenPassVS(GlobalShaderMap);
-					TShaderMapRef<FCopyRectPS> CopyPixelShader(GlobalShaderMap);
-
-					FRHIBlendState* DefaultBlendState = FScreenPassPipelineState::FDefaultBlendState::GetRHI();
-					const FScreenPassTextureViewport RegionViewport(OutputRenderTargetBinding.GetTexture());
-
-					GraphBuilder.AddPass(
-						RDG_EVENT_NAME("MeshProjectionRenderer::CopyHitProxyTexture"),
-						ScreenPassParameters,
-						ERDGPassFlags::Raster,
-						[View, ScreenPassVS, CopyPixelShader, RegionViewport, ScreenPassParameters, DefaultBlendState](FRHICommandList& RHICmdList)
-					{
-						DrawScreenPass(
-							RHICmdList,
-							*View,
-							RegionViewport,
-							RegionViewport,
-							FScreenPassPipelineState(ScreenPassVS, CopyPixelShader, DefaultBlendState),
-							[&](FRHICommandList&)
-						{
-							SetShaderParameters(RHICmdList, CopyPixelShader, CopyPixelShader.GetPixelShader(), *ScreenPassParameters);
-						});
-					});
-				}
-			}
-			else
-			{
-				FRDGTextureRef ColorTexture = GraphBuilder.CreateTexture(OutputTexture->Desc, TEXT("DisplayClusterMeshProjection.ColorTexture"));
-
-				const FRDGTextureDesc DepthDesc = FRDGTextureDesc::Create2D(OutputTexture->Desc.Extent, PF_DepthStencil, FClearValueBinding::DepthFar, TexCreate_DepthStencilTargetable | TexCreate_ShaderResource);
-				FRDGTextureRef DepthTexture = GraphBuilder.CreateTexture(DepthDesc, TEXT("DisplayClusterMeshProjection.DepthTexture"));
-
-				FRenderTargetBinding ColorRenderTargetBinding(ColorTexture, ERenderTargetLoadAction::EClear);
-				FDepthStencilBinding DepthStencilBinding(DepthTexture, ERenderTargetLoadAction::EClear, ERenderTargetLoadAction::ENoAction, FExclusiveDepthStencil::DepthWrite_StencilNop);
-
-				AddBaseRenderPass(GraphBuilder, View, ProjectionType, ColorRenderTargetBinding, DepthStencilBinding);
-
-				ColorRenderTargetBinding.SetLoadAction(ERenderTargetLoadAction::ELoad);
-				DepthStencilBinding.SetDepthLoadAction(ERenderTargetLoadAction::ELoad);
-
-				AddTranslucencyRenderPass(GraphBuilder, View, ProjectionType, ColorRenderTargetBinding, DepthStencilBinding);
-
-#if WITH_EDITOR
-				if (EngineShowFlags.SelectionOutline)
+			case EDisplayClusterMeshProjectionOutput::Color:
+				if (ViewFamily.EngineShowFlags.HitProxies)
 				{
-					const FRDGTextureDesc SelectionDepthDesc = FRDGTextureDesc::Create2D(OutputTexture->Desc.Extent, PF_DepthStencil, FClearValueBinding::DepthFar, TexCreate_DepthStencilTargetable | TexCreate_ShaderResource);
-					FRDGTextureRef SelectionDepthTexture = GraphBuilder.CreateTexture(SelectionDepthDesc, TEXT("DisplayClusterMeshProjection.SelectionDepthTexture"));
-					FDepthStencilBinding SelectionDepthStencilBinding(SelectionDepthTexture, ERenderTargetLoadAction::EClear, ERenderTargetLoadAction::EClear, FExclusiveDepthStencil::DepthWrite_StencilWrite);
-
-					AddSelectionDepthRenderPass(GraphBuilder, View, ProjectionType, SelectionDepthStencilBinding);
-					AddSelectionOutlineScreenPass(GraphBuilder, View, OutputRenderTargetBinding, ColorTexture, DepthTexture, SelectionDepthTexture);
+					RenderHitProxyOutput(GraphBuilder, View, RenderSettings, OutputRenderTargetBinding, HitProxyConsumer);
 				}
 				else
-#endif
-				// Copy the scene color to the output render target
 				{
-					FCopyRectPS::FParameters* ScreenPassParameters = GraphBuilder.AllocParameters<FCopyRectPS::FParameters>();
-					ScreenPassParameters->InputTexture = ColorTexture;
-					ScreenPassParameters->InputSampler = TStaticSamplerState<>::GetRHI();
-					ScreenPassParameters->RenderTargets[0] = OutputRenderTargetBinding;
-
-					FGlobalShaderMap* GlobalShaderMap = GetGlobalShaderMap(GMaxRHIFeatureLevel);
-					TShaderMapRef<FScreenPassVS> ScreenPassVS(GlobalShaderMap);
-					TShaderMapRef<FCopyRectPS> CopyPixelShader(GlobalShaderMap);
-
-					FRHIBlendState* DefaultBlendState = FScreenPassPipelineState::FDefaultBlendState::GetRHI();
-					const FScreenPassTextureViewport RegionViewport(OutputRenderTargetBinding.GetTexture());
-
-					GraphBuilder.AddPass(
-						RDG_EVENT_NAME("MeshProjectionRenderer::CopyColorTexture"),
-						ScreenPassParameters,
-						ERDGPassFlags::Raster,
-						[View, ScreenPassVS, CopyPixelShader, RegionViewport, ScreenPassParameters, DefaultBlendState](FRHICommandList& RHICmdList)
-						{
-							DrawScreenPass(
-								RHICmdList,
-								*View,
-								RegionViewport,
-								RegionViewport,
-								FScreenPassPipelineState(ScreenPassVS, CopyPixelShader, DefaultBlendState),
-								[&](FRHICommandList&)
-								{
-									SetShaderParameters(RHICmdList, CopyPixelShader, CopyPixelShader.GetPixelShader(), *ScreenPassParameters);
-						});
-					});
+					RenderColorOutput(GraphBuilder, View, RenderSettings, OutputRenderTargetBinding);
 				}
+				break;
+
+			case EDisplayClusterMeshProjectionOutput::Normals:
+				RenderNormalsOutput(GraphBuilder, View, RenderSettings, OutputRenderTargetBinding);
+				break;
 			}
-			
+
 			FMeshProjectionElementCollector ElementCollector(HitProxyConsumer);
 			if (RenderSimpleElementsDelegate.IsBound())
 			{
@@ -984,63 +966,141 @@ void FDisplayClusterMeshProjectionRenderer::Render(FCanvas* Canvas, FSceneInterf
 		});
 }
 
-void FDisplayClusterMeshProjectionRenderer::RenderNormals(FCanvas* Canvas,
-	FSceneInterface* Scene,
-	const FSceneViewInitOptions& ViewInitOptions,
-	const FEngineShowFlags& EngineShowFlags,
-	EDisplayClusterMeshProjectionType  ProjectionType,
-	FDisplayClusterMeshProjectionPrimitiveFilter* PrimitiveFilter)
+void FDisplayClusterMeshProjectionRenderer::RenderColorOutput(FRDGBuilder& GraphBuilder,
+	const FViewInfo* View,
+	const FDisplayClusterMeshProjectionRenderSettings& RenderSettings,
+	FRenderTargetBinding& OutputRenderTargetBinding)
 {
-	Canvas->Flush_GameThread();
-	FRenderTarget* RenderTarget = Canvas->GetRenderTarget();
+	FRDGTextureRef ColorTexture = GraphBuilder.CreateTexture(OutputRenderTargetBinding.GetTexture()->Desc, TEXT("DisplayClusterMeshProjection.ColorTexture"));
 
-	ENQUEUE_RENDER_COMMAND(FDrawProjectedMeshes)(
-		[RenderTarget, Scene, ViewInitOptions, EngineShowFlags, ProjectionType, PrimitiveFilter, this](FRHICommandListImmediate& RHICmdList)
-		{
-			FRDGBuilder GraphBuilder(RHICmdList);
+	const FRDGTextureDesc DepthDesc = FRDGTextureDesc::Create2D(OutputRenderTargetBinding.GetTexture()->Desc.Extent, PF_DepthStencil, FClearValueBinding::DepthFar, TexCreate_DepthStencilTargetable | TexCreate_ShaderResource);
+	FRDGTextureRef DepthTexture = GraphBuilder.CreateTexture(DepthDesc, TEXT("DisplayClusterMeshProjection.DepthTexture"));
 
-			FSceneViewFamilyContext ViewFamily(FSceneViewFamily::ConstructionValues(
-				RenderTarget,
-				Scene,
-				EngineShowFlags)
-				.SetTime(FGameTime::GetTimeSinceAppStart())
-				.SetGammaCorrection(1.0f));
+	FRenderTargetBinding ColorRenderTargetBinding(ColorTexture, ERenderTargetLoadAction::EClear);
+	FDepthStencilBinding DepthStencilBinding(DepthTexture, ERenderTargetLoadAction::EClear, ERenderTargetLoadAction::ENoAction, FExclusiveDepthStencil::DepthWrite_StencilNop);
 
-			if (Scene)
-			{
-				Scene->IncrementFrameNumber();
-				ViewFamily.FrameNumber = Scene->GetFrameNumber();
-			}
-			else
+	AddBaseRenderPass(GraphBuilder, View, RenderSettings, ColorRenderTargetBinding, DepthStencilBinding);
+
+	ColorRenderTargetBinding.SetLoadAction(ERenderTargetLoadAction::ELoad);
+	DepthStencilBinding.SetDepthLoadAction(ERenderTargetLoadAction::ELoad);
+
+	AddTranslucencyRenderPass(GraphBuilder, View, RenderSettings, ColorRenderTargetBinding, DepthStencilBinding);
+
+#if WITH_EDITOR
+	if (RenderSettings.EngineShowFlags.SelectionOutline)
+	{
+		const FRDGTextureDesc SelectionDepthDesc = FRDGTextureDesc::Create2D(OutputRenderTargetBinding.GetTexture()->Desc.Extent, PF_DepthStencil, FClearValueBinding::DepthFar, TexCreate_DepthStencilTargetable | TexCreate_ShaderResource);
+		FRDGTextureRef SelectionDepthTexture = GraphBuilder.CreateTexture(SelectionDepthDesc, TEXT("DisplayClusterMeshProjection.SelectionDepthTexture"));
+		FDepthStencilBinding SelectionDepthStencilBinding(SelectionDepthTexture, ERenderTargetLoadAction::EClear, ERenderTargetLoadAction::EClear, FExclusiveDepthStencil::DepthWrite_StencilWrite);
+
+		AddSelectionDepthRenderPass(GraphBuilder, View, RenderSettings, SelectionDepthStencilBinding);
+		AddSelectionOutlineScreenPass(GraphBuilder, View, OutputRenderTargetBinding, ColorTexture, DepthTexture, SelectionDepthTexture);
+	}
+	else
+#endif
+	// Copy the scene color to the output render target
+	{
+		FCopyRectPS::FParameters* ScreenPassParameters = GraphBuilder.AllocParameters<FCopyRectPS::FParameters>();
+		ScreenPassParameters->InputTexture = ColorTexture;
+		ScreenPassParameters->InputSampler = TStaticSamplerState<>::GetRHI();
+		ScreenPassParameters->RenderTargets[0] = OutputRenderTargetBinding;
+
+		FGlobalShaderMap* GlobalShaderMap = GetGlobalShaderMap(GMaxRHIFeatureLevel);
+		TShaderMapRef<FScreenPassVS> ScreenPassVS(GlobalShaderMap);
+		TShaderMapRef<FCopyRectPS> CopyPixelShader(GlobalShaderMap);
+
+		FRHIBlendState* DefaultBlendState = FScreenPassPipelineState::FDefaultBlendState::GetRHI();
+		const FScreenPassTextureViewport RegionViewport(OutputRenderTargetBinding.GetTexture());
+
+		GraphBuilder.AddPass(
+			RDG_EVENT_NAME("MeshProjectionRenderer::CopyColorTexture"),
+			ScreenPassParameters,
+			ERDGPassFlags::Raster,
+			[View, ScreenPassVS, CopyPixelShader, RegionViewport, ScreenPassParameters, DefaultBlendState](FRHICommandList& RHICmdList)
+		{
+			DrawScreenPass(
+				RHICmdList,
+				*View,
+				RegionViewport,
+				RegionViewport,
+				FScreenPassPipelineState(ScreenPassVS, CopyPixelShader, DefaultBlendState),
+				[&](FRHICommandList&)
 			{
-				ViewFamily.FrameNumber = GFrameNumber;
-			}
+				SetShaderParameters(RHICmdList, CopyPixelShader, CopyPixelShader.GetPixelShader(), *ScreenPassParameters);
+			});
+		});
+	}
+}
 
-			FScenePrimitiveRenderingContextScopeHelper ScenePrimitiveRenderingContextScopeHelper(GetRendererModule().BeginScenePrimitiveRendering(GraphBuilder, &ViewFamily));
+void FDisplayClusterMeshProjectionRenderer::RenderHitProxyOutput(FRDGBuilder& GraphBuilder,
+	const FViewInfo* View,
+	const FDisplayClusterMeshProjectionRenderSettings& RenderSettings,
+	FRenderTargetBinding& OutputRenderTargetBinding,
+	FHitProxyConsumer* HitProxyConsumer)
+{
 
-			FSceneViewInitOptions NewInitOptions(ViewInitOptions);
-			NewInitOptions.ViewFamily = &ViewFamily;
 
-			GetRendererModule().CreateAndInitSingleView(RHICmdList, &ViewFamily, &NewInitOptions);
-			FViewInfo* View = (FViewInfo*)ViewFamily.Views[0];
+	FRDGTextureDesc Desc(FRDGTextureDesc::Create2D(OutputRenderTargetBinding.GetTexture()->Desc.Extent, PF_B8G8R8A8, FClearValueBinding::Black, TexCreate_RenderTargetable | TexCreate_ShaderResource));
+	FRDGTextureRef HitProxyTexture = GraphBuilder.CreateTexture(Desc, TEXT("DisplayClusterMeshProjection.HitProxyTexture"));
 
-			FRDGTextureRef OutputTexture = GraphBuilder.RegisterExternalTexture(CreateRenderTarget(RenderTarget->GetRenderTargetTexture(), TEXT("ViewRenderTarget")));
-			FRenderTargetBinding OutputRenderTargetBinding(OutputTexture, ERenderTargetLoadAction::ELoad);
+	const FRDGTextureDesc DepthDesc = FRDGTextureDesc::Create2D(OutputRenderTargetBinding.GetTexture()->Desc.Extent, PF_DepthStencil, FClearValueBinding::DepthFar, TexCreate_DepthStencilTargetable | TexCreate_ShaderResource);
+	FRDGTextureRef HitProxyDepthTexture = GraphBuilder.CreateTexture(DepthDesc, TEXT("DisplayClusterMeshProjection.HitProxyDepthTexture"));
+
+	FRenderTargetBinding HitProxyRenderTargetBinding(HitProxyTexture, ERenderTargetLoadAction::EClear);
+	FDepthStencilBinding HitProxyDepthStencilBinding(HitProxyDepthTexture, ERenderTargetLoadAction::EClear, ERenderTargetLoadAction::EClear, FExclusiveDepthStencil::DepthWrite_StencilWrite);
 
-			const FRDGTextureDesc NormalsDesc(FRDGTextureDesc::Create2D(OutputTexture->Desc.Extent, OutputTexture->Desc.Format, FClearValueBinding::Black, TexCreate_RenderTargetable | TexCreate_ShaderResource));
-			FRDGTextureRef NormalsTexture = GraphBuilder.CreateTexture(NormalsDesc, TEXT("DisplayClusterMeshProjection.NormalsTexture"));
+	AddHitProxyRenderPass(GraphBuilder, View, RenderSettings, HitProxyRenderTargetBinding, HitProxyDepthStencilBinding);
 
-			const FRDGTextureDesc DepthDesc = FRDGTextureDesc::Create2D(OutputTexture->Desc.Extent, PF_DepthStencil, FClearValueBinding::DepthFar, TexCreate_DepthStencilTargetable | TexCreate_ShaderResource);
-			FRDGTextureRef NormalsDepthTexture = GraphBuilder.CreateTexture(DepthDesc, TEXT("DisplayClusterMeshProjection.NormalsDepthTexture"));
+	// Copy the hit proxy buffer to the viewport family's render target
+	{
+		FCopyRectPS::FParameters* ScreenPassParameters = GraphBuilder.AllocParameters<FCopyRectPS::FParameters>();
+		ScreenPassParameters->InputTexture = HitProxyTexture;
+		ScreenPassParameters->InputSampler = TStaticSamplerState<>::GetRHI();
+		ScreenPassParameters->RenderTargets[0] = OutputRenderTargetBinding;
 
-			FRenderTargetBinding NormalsRenderTargetBinding(NormalsTexture, ERenderTargetLoadAction::EClear);
-			FDepthStencilBinding NormalsDepthStencilBinding(NormalsDepthTexture, ERenderTargetLoadAction::EClear, ERenderTargetLoadAction::EClear, FExclusiveDepthStencil::DepthWrite_StencilWrite);
+		FGlobalShaderMap* GlobalShaderMap = GetGlobalShaderMap(GMaxRHIFeatureLevel);
+		TShaderMapRef<FScreenPassVS> ScreenPassVS(GlobalShaderMap);
+		TShaderMapRef<FCopyRectPS> CopyPixelShader(GlobalShaderMap);
 
-			AddNormalsRenderPass(GraphBuilder, View, ProjectionType, PrimitiveFilter, NormalsRenderTargetBinding, NormalsDepthStencilBinding);
-			AddNormalsFilterPass(GraphBuilder, View, OutputRenderTargetBinding, NormalsTexture, NormalsDepthTexture);
+		FRHIBlendState* DefaultBlendState = FScreenPassPipelineState::FDefaultBlendState::GetRHI();
+		const FScreenPassTextureViewport RegionViewport(OutputRenderTargetBinding.GetTexture());
 
-			GraphBuilder.Execute();
+		GraphBuilder.AddPass(
+			RDG_EVENT_NAME("MeshProjectionRenderer::CopyHitProxyTexture"),
+			ScreenPassParameters,
+			ERDGPassFlags::Raster,
+			[View, ScreenPassVS, CopyPixelShader, RegionViewport, ScreenPassParameters, DefaultBlendState](FRHICommandList& RHICmdList)
+		{
+			DrawScreenPass(
+				RHICmdList,
+				*View,
+				RegionViewport,
+				RegionViewport,
+				FScreenPassPipelineState(ScreenPassVS, CopyPixelShader, DefaultBlendState),
+				[&](FRHICommandList&)
+			{
+				SetShaderParameters(RHICmdList, CopyPixelShader, CopyPixelShader.GetPixelShader(), *ScreenPassParameters);
+			});
 		});
+	}
+}
+
+void FDisplayClusterMeshProjectionRenderer::RenderNormalsOutput(FRDGBuilder& GraphBuilder,
+	const FViewInfo* View,
+	const FDisplayClusterMeshProjectionRenderSettings& RenderSettings,
+	FRenderTargetBinding& OutputRenderTargetBinding)
+{
+	const FRDGTextureDesc NormalsDesc(FRDGTextureDesc::Create2D(OutputRenderTargetBinding.GetTexture()->Desc.Extent, OutputRenderTargetBinding.GetTexture()->Desc.Format, FClearValueBinding::Black, TexCreate_RenderTargetable | TexCreate_ShaderResource));
+	FRDGTextureRef NormalsTexture = GraphBuilder.CreateTexture(NormalsDesc, TEXT("DisplayClusterMeshProjection.NormalsTexture"));
+
+	const FRDGTextureDesc DepthDesc = FRDGTextureDesc::Create2D(OutputRenderTargetBinding.GetTexture()->Desc.Extent, PF_DepthStencil, FClearValueBinding::DepthFar, TexCreate_DepthStencilTargetable | TexCreate_ShaderResource);
+	FRDGTextureRef NormalsDepthTexture = GraphBuilder.CreateTexture(DepthDesc, TEXT("DisplayClusterMeshProjection.NormalsDepthTexture"));
+
+	FRenderTargetBinding NormalsRenderTargetBinding(NormalsTexture, ERenderTargetLoadAction::EClear);
+	FDepthStencilBinding NormalsDepthStencilBinding(NormalsDepthTexture, ERenderTargetLoadAction::EClear, ERenderTargetLoadAction::EClear, FExclusiveDepthStencil::DepthWrite_StencilWrite);
+
+	AddNormalsRenderPass(GraphBuilder, View, RenderSettings, NormalsRenderTargetBinding, NormalsDepthStencilBinding);
+	AddNormalsFilterPass(GraphBuilder, View, OutputRenderTargetBinding, NormalsTexture, NormalsDepthTexture);
 }
 
 // Helper macros that generate appropriate switch cases for each projection type the mesh projection renderer supports, allowing new projection types to be easily added
@@ -1052,91 +1112,94 @@ void FDisplayClusterMeshProjectionRenderer::RenderNormals(FCanvas* Canvas,
 	{ \
 	PROJECTION_TYPE_CASE(FuncName, EDisplayClusterMeshProjectionType::Azimuthal, __VA_ARGS__) \
 	PROJECTION_TYPE_CASE(FuncName, EDisplayClusterMeshProjectionType::Linear, __VA_ARGS__) \
+	PROJECTION_TYPE_CASE(FuncName, EDisplayClusterMeshProjectionType::UV, __VA_ARGS__) \
 	}
 
 void FDisplayClusterMeshProjectionRenderer::AddBaseRenderPass(FRDGBuilder& GraphBuilder,
-	const FViewInfo* View, EDisplayClusterMeshProjectionType ProjectionType,
+	const FViewInfo* View,
+	const FDisplayClusterMeshProjectionRenderSettings& RenderSettings,
 	FRenderTargetBinding& OutputRenderTargetBinding,
 	FDepthStencilBinding& OutputDepthStencilBinding)
 {
-	FMeshProjectionPassParameters* MeshPassParameters = GraphBuilder.AllocParameters<FMeshProjectionPassParameters>();
-	MeshPassParameters->View = View->ViewUniformBuffer;
-	MeshPassParameters->InstanceCulling = FInstanceCullingContext::CreateDummyInstanceCullingUniformBuffer(GraphBuilder);
+	FMeshProjectionPassParameters* MeshPassParameters = CreateProjectionPassParameters(GraphBuilder, View, RenderSettings);
 	MeshPassParameters->RenderTargets[0] = OutputRenderTargetBinding;
 	MeshPassParameters->RenderTargets.DepthStencil = OutputDepthStencilBinding;
 
-	GraphBuilder.AddPass(RDG_EVENT_NAME("MeshProjectionRenderer::Base"), MeshPassParameters, ERDGPassFlags::Raster | ERDGPassFlags::NeverCull, [View, ProjectionType, this](FRHICommandList& RHICmdList)
-	{
-		FIntRect ViewRect = View->UnscaledViewRect;
-		RHICmdList.SetViewport(ViewRect.Min.X, ViewRect.Min.Y, 0.0f, ViewRect.Max.X, ViewRect.Max.Y, 1.0f);
+	GraphBuilder.AddPass(RDG_EVENT_NAME("MeshProjectionRenderer::Base"),
+		MeshPassParameters,
+		ERDGPassFlags::Raster | ERDGPassFlags::NeverCull,
+		[View, &RenderSettings, this](FRHICommandList& RHICmdList)
+		{
+			FIntRect ViewRect = View->UnscaledViewRect;
+			RHICmdList.SetViewport(ViewRect.Min.X, ViewRect.Min.Y, 0.0f, ViewRect.Max.X, ViewRect.Max.Y, 1.0f);
 
-		SWITCH_ON_PROJECTION_TYPE(RenderPrimitives_RenderThread, ProjectionType, View, RHICmdList, false);
-	});
+			SWITCH_ON_PROJECTION_TYPE(RenderPrimitives_RenderThread, RenderSettings.ProjectionType, View, RHICmdList, RenderSettings, false);
+		});
 }
 
 void FDisplayClusterMeshProjectionRenderer::AddTranslucencyRenderPass(FRDGBuilder& GraphBuilder,
-	const FViewInfo* View, EDisplayClusterMeshProjectionType ProjectionType,
+	const FViewInfo* View,
+	const FDisplayClusterMeshProjectionRenderSettings& RenderSettings,
 	FRenderTargetBinding& OutputRenderTargetBinding,
 	FDepthStencilBinding& OutputDepthStencilBinding)
 {
-	FMeshProjectionPassParameters* MeshPassParameters = GraphBuilder.AllocParameters<FMeshProjectionPassParameters>();
-	MeshPassParameters->View = View->ViewUniformBuffer;
-	MeshPassParameters->InstanceCulling = FInstanceCullingContext::CreateDummyInstanceCullingUniformBuffer(GraphBuilder);
+	FMeshProjectionPassParameters* MeshPassParameters = CreateProjectionPassParameters(GraphBuilder, View, RenderSettings);
 	MeshPassParameters->RenderTargets[0] = OutputRenderTargetBinding;
 	MeshPassParameters->RenderTargets.DepthStencil = OutputDepthStencilBinding;
 
-	GraphBuilder.AddPass(RDG_EVENT_NAME("MeshProjectionRenderer::Translucency"), MeshPassParameters, ERDGPassFlags::Raster | ERDGPassFlags::NeverCull, [View, ProjectionType, this](FRHICommandList& RHICmdList)
-	{
-		FIntRect ViewRect = View->UnscaledViewRect;
-		RHICmdList.SetViewport(ViewRect.Min.X, ViewRect.Min.Y, 0.0f, ViewRect.Max.X, ViewRect.Max.Y, 1.0f);
+	GraphBuilder.AddPass(RDG_EVENT_NAME("MeshProjectionRenderer::Translucency"),
+		MeshPassParameters,
+		ERDGPassFlags::Raster | ERDGPassFlags::NeverCull,
+		[View, &RenderSettings, this](FRHICommandList& RHICmdList)
+		{
+			FIntRect ViewRect = View->UnscaledViewRect;
+			RHICmdList.SetViewport(ViewRect.Min.X, ViewRect.Min.Y, 0.0f, ViewRect.Max.X, ViewRect.Max.Y, 1.0f);
 
-		SWITCH_ON_PROJECTION_TYPE(RenderPrimitives_RenderThread, ProjectionType, View, RHICmdList, true);
-	});
+			SWITCH_ON_PROJECTION_TYPE(RenderPrimitives_RenderThread, RenderSettings.ProjectionType, View, RHICmdList, RenderSettings, true);
+		});
 }
 
 void FDisplayClusterMeshProjectionRenderer::AddHitProxyRenderPass(FRDGBuilder& GraphBuilder,
 	const FViewInfo* View,
-	EDisplayClusterMeshProjectionType ProjectionType,
+	const FDisplayClusterMeshProjectionRenderSettings& RenderSettings,
 	FRenderTargetBinding& OutputRenderTargetBinding,
 	FDepthStencilBinding& OutputDepthStencilBinding)
 {
-	FMeshProjectionPassParameters* MeshPassParameters = GraphBuilder.AllocParameters<FMeshProjectionPassParameters>();
-	MeshPassParameters->View = View->ViewUniformBuffer;
-	MeshPassParameters->InstanceCulling = FInstanceCullingContext::CreateDummyInstanceCullingUniformBuffer(GraphBuilder);
+	FMeshProjectionPassParameters* MeshPassParameters = CreateProjectionPassParameters(GraphBuilder, View, RenderSettings);
 	MeshPassParameters->RenderTargets[0] = OutputRenderTargetBinding;
 	MeshPassParameters->RenderTargets.DepthStencil = OutputDepthStencilBinding;
 
-	GraphBuilder.AddPass(RDG_EVENT_NAME("MeshProjectionRenderer::HitProxies"), MeshPassParameters, ERDGPassFlags::Raster | ERDGPassFlags::NeverCull, [View, ProjectionType, this](FRHICommandList& RHICmdList)
-	{
-		FIntRect ViewRect = View->UnscaledViewRect;
-		RHICmdList.SetViewport(ViewRect.Min.X, ViewRect.Min.Y, 0.0f, ViewRect.Max.X, ViewRect.Max.Y, 1.0f);
+	GraphBuilder.AddPass(RDG_EVENT_NAME("MeshProjectionRenderer::HitProxies"),
+		MeshPassParameters,
+		ERDGPassFlags::Raster | ERDGPassFlags::NeverCull,
+		[View, &RenderSettings, this](FRHICommandList& RHICmdList)
+		{
+			FIntRect ViewRect = View->UnscaledViewRect;
+			RHICmdList.SetViewport(ViewRect.Min.X, ViewRect.Min.Y, 0.0f, ViewRect.Max.X, ViewRect.Max.Y, 1.0f);
 
-		SWITCH_ON_PROJECTION_TYPE(RenderHitProxies_RenderThread, ProjectionType, View, RHICmdList);
-	});
+			SWITCH_ON_PROJECTION_TYPE(RenderHitProxies_RenderThread, RenderSettings.ProjectionType, View, RHICmdList, RenderSettings);
+		});
 }
 
 void FDisplayClusterMeshProjectionRenderer::AddNormalsRenderPass(FRDGBuilder& GraphBuilder,
 	const FViewInfo* View,
-	EDisplayClusterMeshProjectionType ProjectionType,
-	FDisplayClusterMeshProjectionPrimitiveFilter* PrimitiveFilter,
+	const FDisplayClusterMeshProjectionRenderSettings& RenderSettings,
 	FRenderTargetBinding& OutputRenderTargetBinding,
 	FDepthStencilBinding& OutputDepthStencilBinding)
 {
-	FMeshProjectionPassParameters* MeshPassParameters = GraphBuilder.AllocParameters<FMeshProjectionPassParameters>();
-	MeshPassParameters->View = View->ViewUniformBuffer;
-	MeshPassParameters->InstanceCulling = FInstanceCullingContext::CreateDummyInstanceCullingUniformBuffer(GraphBuilder);
+	FMeshProjectionPassParameters* MeshPassParameters = CreateProjectionPassParameters(GraphBuilder, View, RenderSettings);
 	MeshPassParameters->RenderTargets[0] = OutputRenderTargetBinding;
 	MeshPassParameters->RenderTargets.DepthStencil = OutputDepthStencilBinding;
 
 	GraphBuilder.AddPass(RDG_EVENT_NAME("MeshProjectionRenderer::Normals"),
 		MeshPassParameters,
 		ERDGPassFlags::Raster | ERDGPassFlags::NeverCull,
-		[View, ProjectionType, PrimitiveFilter, this](FRHICommandList& RHICmdList)
+		[View, &RenderSettings, this](FRHICommandList& RHICmdList)
 		{
 			FIntRect ViewRect = View->UnscaledViewRect;
 			RHICmdList.SetViewport(ViewRect.Min.X, ViewRect.Min.Y, 0.0f, ViewRect.Max.X, ViewRect.Max.Y, 1.0f);
 
-			SWITCH_ON_PROJECTION_TYPE(RenderNormals_RenderThread, ProjectionType, View, RHICmdList, PrimitiveFilter);
+			SWITCH_ON_PROJECTION_TYPE(RenderNormals_RenderThread, RenderSettings.ProjectionType, View, RHICmdList, RenderSettings);
 		});
 }
 
@@ -1337,21 +1400,22 @@ void FDisplayClusterMeshProjectionRenderer::AddNormalsFilterPass(FRDGBuilder& Gr
 
 #if WITH_EDITOR
 void FDisplayClusterMeshProjectionRenderer::AddSelectionDepthRenderPass(FRDGBuilder& GraphBuilder, 
-		const FViewInfo* View,
-		EDisplayClusterMeshProjectionType ProjectionType,
-		FDepthStencilBinding& OutputDepthStencilBinding)
+	const FViewInfo* View,
+	const FDisplayClusterMeshProjectionRenderSettings& RenderSettings,
+	FDepthStencilBinding& OutputDepthStencilBinding)
 {
-	FMeshProjectionPassParameters* SelectionPassParameters = GraphBuilder.AllocParameters<FMeshProjectionPassParameters>();
-	SelectionPassParameters->View = View->ViewUniformBuffer;
-	SelectionPassParameters->InstanceCulling = FInstanceCullingContext::CreateDummyInstanceCullingUniformBuffer(GraphBuilder);
+	FMeshProjectionPassParameters* SelectionPassParameters = CreateProjectionPassParameters(GraphBuilder, View, RenderSettings);
 	SelectionPassParameters->RenderTargets.DepthStencil = OutputDepthStencilBinding;
 
-	GraphBuilder.AddPass(RDG_EVENT_NAME("MeshProjectionRenderer::SelectionDepth"), SelectionPassParameters, ERDGPassFlags::Raster | ERDGPassFlags::NeverCull, [View, ProjectionType, this](FRHICommandList& RHICmdList)
+	GraphBuilder.AddPass(RDG_EVENT_NAME("MeshProjectionRenderer::SelectionDepth"),
+		SelectionPassParameters,
+		ERDGPassFlags::Raster | ERDGPassFlags::NeverCull,
+		[View, &RenderSettings, this](FRHICommandList& RHICmdList)
 	{
 		FIntRect ViewRect = View->UnscaledViewRect;
 		RHICmdList.SetViewport(ViewRect.Min.X, ViewRect.Min.Y, 0.0f, ViewRect.Max.X, ViewRect.Max.Y, 1.0f);
 
-		SWITCH_ON_PROJECTION_TYPE(RenderSelection_RenderThread, ProjectionType, View, RHICmdList);
+		SWITCH_ON_PROJECTION_TYPE(RenderSelection_RenderThread, RenderSettings.ProjectionType, View, RHICmdList, RenderSettings);
 	});
 }
 
@@ -1438,94 +1502,111 @@ void FDisplayClusterMeshProjectionRenderer::AddSimpleElementPass(FRDGBuilder& Gr
 }
 
 template<EDisplayClusterMeshProjectionType ProjectionType>
-void FDisplayClusterMeshProjectionRenderer::RenderPrimitives_RenderThread(const FSceneView* View, FRHICommandList& RHICmdList, bool bTranslucencyPass)
+void FDisplayClusterMeshProjectionRenderer::RenderPrimitives_RenderThread(const FSceneView* View,
+	FRHICommandList& RHICmdList,
+	const FDisplayClusterMeshProjectionRenderSettings& RenderSettings,
+	bool bTranslucencyPass)
 {
-	DrawDynamicMeshPass(*View, RHICmdList, [View, bTranslucencyPass, this](FDynamicPassMeshDrawListContext* DynamicMeshPassContext)
+	DrawDynamicMeshPass(*View, RHICmdList, [View, &RenderSettings, bTranslucencyPass, this](FDynamicPassMeshDrawListContext* DynamicMeshPassContext)
 	{
-		TArray<FPrimitiveSceneProxy*> PrimitiveSceneProxies;
-		for (const TWeakObjectPtr<UPrimitiveComponent>& PrimitiveComponent : PrimitiveComponents)
-		{
-			if (PrimitiveComponent.IsValid() && PrimitiveComponent->SceneProxy)
-			{
-				PrimitiveSceneProxies.Add(PrimitiveComponent->SceneProxy);
-			}
-		}
+		TArray<FSceneProxyElement> PrimitiveSceneProxies;
+		GetSceneProxies(PrimitiveSceneProxies, RenderSettings);
 
-		FLightCardEditorMeshPassProcessor<ProjectionType> MeshProcessor(nullptr, View, DynamicMeshPassContext, bTranslucencyPass);
+		FMeshProjectionPassProcessor<ProjectionType> ProjectionMeshProcessor(nullptr, View, DynamicMeshPassContext, RenderSettings, bTranslucencyPass);
+		FMeshProjectionPassProcessor<EDisplayClusterMeshProjectionType::Linear> LinearMeshProcessor(nullptr, View, DynamicMeshPassContext, RenderSettings, bTranslucencyPass);
 
-		for (FPrimitiveSceneProxy* PrimitiveProxy : PrimitiveSceneProxies)
+		for (const FSceneProxyElement& PrimitiveProxyElement : PrimitiveSceneProxies)
 		{
+			FPrimitiveSceneProxy* PrimitiveProxy = PrimitiveProxyElement.PrimitiveSceneProxy;
 			if (const FMeshBatch* MeshBatch = PrimitiveProxy->GetPrimitiveSceneInfo()->GetMeshBatch(PrimitiveProxy->GetPrimitiveSceneInfo()->StaticMeshes.Num() - 1))
 			{
 				MeshBatch->MaterialRenderProxy->UpdateUniformExpressionCacheIfNeeded(View->GetFeatureLevel());
 
 				const uint64 BatchElementMask = ~0ull;
-				MeshProcessor.AddMeshBatch(*MeshBatch, BatchElementMask, PrimitiveProxy);
+
+				if (PrimitiveProxyElement.bApplyProjection)
+				{
+					ProjectionMeshProcessor.AddMeshBatch(*MeshBatch, BatchElementMask, PrimitiveProxy);
+				}
+				else
+				{
+					LinearMeshProcessor.AddMeshBatch(*MeshBatch, BatchElementMask, PrimitiveProxy);
+				}
 			}
 		}
 	});
 }
 
 template<EDisplayClusterMeshProjectionType ProjectionType>
-void FDisplayClusterMeshProjectionRenderer::RenderHitProxies_RenderThread(const FSceneView* View, FRHICommandList& RHICmdList)
+void FDisplayClusterMeshProjectionRenderer::RenderHitProxies_RenderThread(const FSceneView* View,
+	FRHICommandList& RHICmdList,
+	const FDisplayClusterMeshProjectionRenderSettings& RenderSettings)
 {
-	DrawDynamicMeshPass(*View, RHICmdList, [View, this](FDynamicPassMeshDrawListContext* DynamicMeshPassContext)
+	DrawDynamicMeshPass(*View, RHICmdList, [View, &RenderSettings, this](FDynamicPassMeshDrawListContext* DynamicMeshPassContext)
 	{
-		TArray<FPrimitiveSceneProxy*> PrimitiveSceneProxies;
-		for (const TWeakObjectPtr<UPrimitiveComponent>& PrimitiveComponent : PrimitiveComponents)
-		{
-			if (PrimitiveComponent.IsValid() && PrimitiveComponent->SceneProxy)
-			{
-				PrimitiveSceneProxies.Add(PrimitiveComponent->SceneProxy);
-			}
-		}
+		TArray<FSceneProxyElement> PrimitiveSceneProxies;
+		GetSceneProxies(PrimitiveSceneProxies, RenderSettings);
 
-		FLightCardEditorHitProxyMeshPassProcessor<ProjectionType> MeshProcessor(nullptr, View, DynamicMeshPassContext);
+		FMeshProjectionHitProxyPassProcessor<ProjectionType> ProjectionMeshProcessor(nullptr, View, DynamicMeshPassContext, RenderSettings);
+		FMeshProjectionHitProxyPassProcessor<EDisplayClusterMeshProjectionType::Linear> LinearMeshProcessor(nullptr, View, DynamicMeshPassContext, RenderSettings);
 
-		for (FPrimitiveSceneProxy* PrimitiveProxy : PrimitiveSceneProxies)
+		for (const FSceneProxyElement& PrimitiveProxyElement : PrimitiveSceneProxies)
 		{
+			FPrimitiveSceneProxy* PrimitiveProxy = PrimitiveProxyElement.PrimitiveSceneProxy;
 			if (const FMeshBatch* MeshBatch = PrimitiveProxy->GetPrimitiveSceneInfo()->GetMeshBatch(PrimitiveProxy->GetPrimitiveSceneInfo()->StaticMeshes.Num() - 1))
 			{
 				MeshBatch->MaterialRenderProxy->UpdateUniformExpressionCacheIfNeeded(View->GetFeatureLevel());
 
 				const uint64 BatchElementMask = ~0ull;
-				MeshProcessor.AddMeshBatch(*MeshBatch, BatchElementMask, PrimitiveProxy);
+
+				if (PrimitiveProxyElement.bApplyProjection)
+				{
+					ProjectionMeshProcessor.AddMeshBatch(*MeshBatch, BatchElementMask, PrimitiveProxy);
+				}
+				else
+				{
+					LinearMeshProcessor.AddMeshBatch(*MeshBatch, BatchElementMask, PrimitiveProxy);
+				}
 			}
 		}
 	});
 }
 
 template<EDisplayClusterMeshProjectionType ProjectionType>
-void FDisplayClusterMeshProjectionRenderer::RenderNormals_RenderThread(const FSceneView* View, FRHICommandList& RHICmdList, FDisplayClusterMeshProjectionPrimitiveFilter* PrimitiveFilter)
+void FDisplayClusterMeshProjectionRenderer::RenderNormals_RenderThread(const FSceneView* View,
+	FRHICommandList& RHICmdList,
+	const FDisplayClusterMeshProjectionRenderSettings& RenderSettings)
 {
-	DrawDynamicMeshPass(*View, RHICmdList, [View, PrimitiveFilter, this](FDynamicPassMeshDrawListContext* DynamicMeshPassContext)
+	DrawDynamicMeshPass(*View, RHICmdList, [View, &RenderSettings, this](FDynamicPassMeshDrawListContext* DynamicMeshPassContext)
 	{
-		TArray<FPrimitiveSceneProxy*> PrimitiveSceneProxies;
-		for (const TWeakObjectPtr<UPrimitiveComponent>& PrimitiveComponent : PrimitiveComponents)
-		{
-			if (PrimitiveComponent.IsValid() && PrimitiveComponent->SceneProxy)
-			{
-				if (PrimitiveFilter && PrimitiveFilter->IsPrimitiveComponentFiltered(PrimitiveComponent.Get()))
-				{
-					continue;
-				}
+		TArray<FSceneProxyElement> PrimitiveSceneProxies;
+		GetSceneProxies(PrimitiveSceneProxies, RenderSettings);
 
-				PrimitiveSceneProxies.Add(PrimitiveComponent->SceneProxy);
-			}
-		}
+		FMeshProjectionPassProcessor<ProjectionType, EDisplayClusterMeshProjectionOutput::Normals> ProjectionMeshProcessor(nullptr, View, DynamicMeshPassContext, RenderSettings);
+		ProjectionMeshProcessor.SetStencilValue(1);
+		ProjectionMeshProcessor.SetIgnoreTranslucency(true);
 
-		FLightCardEditorMeshPassProcessor<ProjectionType, EMeshProjectionShaderOutput::Normals> MeshProcessor(nullptr, View, DynamicMeshPassContext);
-		MeshProcessor.SetStencilValue(1);
-		MeshProcessor.SetIgnoreTranslucency(true);
+		FMeshProjectionPassProcessor<EDisplayClusterMeshProjectionType::Linear, EDisplayClusterMeshProjectionOutput::Normals> LinearMeshProcessor(nullptr, View, DynamicMeshPassContext, RenderSettings);
+		LinearMeshProcessor.SetStencilValue(1);
+		LinearMeshProcessor.SetIgnoreTranslucency(true);
 
-		for (FPrimitiveSceneProxy* PrimitiveProxy : PrimitiveSceneProxies)
+		for (const FSceneProxyElement& PrimitiveProxyElement : PrimitiveSceneProxies)
 		{
+			FPrimitiveSceneProxy* PrimitiveProxy = PrimitiveProxyElement.PrimitiveSceneProxy;
 			if (const FMeshBatch* MeshBatch = PrimitiveProxy->GetPrimitiveSceneInfo()->GetMeshBatch(PrimitiveProxy->GetPrimitiveSceneInfo()->StaticMeshes.Num() - 1))
 			{
 				MeshBatch->MaterialRenderProxy->UpdateUniformExpressionCacheIfNeeded(View->GetFeatureLevel());
 
 				const uint64 BatchElementMask = ~0ull;
-				MeshProcessor.AddMeshBatch(*MeshBatch, BatchElementMask, PrimitiveProxy);
+
+				if (PrimitiveProxyElement.bApplyProjection)
+				{
+					ProjectionMeshProcessor.AddMeshBatch(*MeshBatch, BatchElementMask, PrimitiveProxy);
+				}
+				else
+				{
+					LinearMeshProcessor.AddMeshBatch(*MeshBatch, BatchElementMask, PrimitiveProxy);
+				}
 			}
 		}
 	});
@@ -1533,29 +1614,38 @@ void FDisplayClusterMeshProjectionRenderer::RenderNormals_RenderThread(const FSc
 
 #if WITH_EDITOR
 template<EDisplayClusterMeshProjectionType ProjectionType>
-void FDisplayClusterMeshProjectionRenderer::RenderSelection_RenderThread(const FSceneView* View, FRHICommandList& RHICmdList)
+void FDisplayClusterMeshProjectionRenderer::RenderSelection_RenderThread(const FSceneView* View,
+	FRHICommandList& RHICmdList,
+	const FDisplayClusterMeshProjectionRenderSettings& RenderSettings)
 {
-	DrawDynamicMeshPass(*View, RHICmdList, [View, this](FDynamicPassMeshDrawListContext* DynamicMeshPassContext)
+	DrawDynamicMeshPass(*View, RHICmdList, [View, &RenderSettings, this](FDynamicPassMeshDrawListContext* DynamicMeshPassContext)
 	{
-		TArray<FPrimitiveSceneProxy*> PrimitiveSceneProxies;
-		for (const TWeakObjectPtr<UPrimitiveComponent>& PrimitiveComponent : PrimitiveComponents)
+		TArray<FSceneProxyElement> PrimitiveSceneProxies;
+		GetSceneProxies(PrimitiveSceneProxies, RenderSettings, [](const UPrimitiveComponent* PrimitiveComponent)
 		{
-			if (PrimitiveComponent.IsValid() && PrimitiveComponent->SceneProxy && PrimitiveComponent->SceneProxy->IsSelected())
-			{
-				PrimitiveSceneProxies.Add(PrimitiveComponent->SceneProxy);
-			}
-		}
+			return PrimitiveComponent->SceneProxy->IsSelected();
+		});
 
-		FLightCardEditorSelectionPassProcessor<ProjectionType> MeshProcessor(nullptr, View, DynamicMeshPassContext);
+		FMeshProjectionSelectionPassProcessor<ProjectionType> ProjectionMeshProcessor(nullptr, View, DynamicMeshPassContext, RenderSettings);
+		FMeshProjectionSelectionPassProcessor<EDisplayClusterMeshProjectionType::Linear> LinearMeshProcessor(nullptr, View, DynamicMeshPassContext, RenderSettings);
 
-		for (FPrimitiveSceneProxy* PrimitiveProxy : PrimitiveSceneProxies)
+		for (const FSceneProxyElement& PrimitiveProxyElement : PrimitiveSceneProxies)
 		{
+			FPrimitiveSceneProxy* PrimitiveProxy = PrimitiveProxyElement.PrimitiveSceneProxy;
 			if (const FMeshBatch* MeshBatch = PrimitiveProxy->GetPrimitiveSceneInfo()->GetMeshBatch(PrimitiveProxy->GetPrimitiveSceneInfo()->StaticMeshes.Num() - 1))
 			{
 				MeshBatch->MaterialRenderProxy->UpdateUniformExpressionCacheIfNeeded(View->GetFeatureLevel());
 
 				const uint64 BatchElementMask = ~0ull;
-				MeshProcessor.AddMeshBatch(*MeshBatch, BatchElementMask, PrimitiveProxy);
+
+				if (PrimitiveProxyElement.bApplyProjection)
+				{
+					ProjectionMeshProcessor.AddMeshBatch(*MeshBatch, BatchElementMask, PrimitiveProxy);
+				}
+				else
+				{
+					LinearMeshProcessor.AddMeshBatch(*MeshBatch, BatchElementMask, PrimitiveProxy);
+				}
 			}
 		}
 	});
@@ -1572,16 +1662,37 @@ bool FDisplayClusterMeshProjectionRenderer::IsPrimitiveComponentSelected(const U
 	return false;
 }
 
+void FDisplayClusterMeshProjectionRenderer::GetSceneProxies(TArray<FSceneProxyElement>& OutSceneProxyElements,
+	const FDisplayClusterMeshProjectionRenderSettings& RenderSettings,
+	const TFunctionRef<bool(const UPrimitiveComponent*)> PrimitiveFilter)
+{
+	for (const TWeakObjectPtr<UPrimitiveComponent>& PrimitiveComponent : PrimitiveComponents)
+	{
+		if (PrimitiveComponent.IsValid() && PrimitiveComponent->SceneProxy)
+		{
+			if (PrimitiveFilter(PrimitiveComponent.Get()) && RenderSettings.PrimitiveFilter.ShouldRenderPrimitive(PrimitiveComponent.Get()))
+			{
+				FSceneProxyElement SceneProxyElement;
+				SceneProxyElement.PrimitiveSceneProxy = PrimitiveComponent->SceneProxy;
+				SceneProxyElement.bApplyProjection = RenderSettings.PrimitiveFilter.ShouldApplyProjection(PrimitiveComponent.Get());
+
+				OutSceneProxyElements.Add(SceneProxyElement);
+			}
+		}
+	}
+}
+
 // Explicit template specializations for each projection type
 #define PROJECTION_TYPE_TEMPLATE_SPECIALIZATION(FuncName, ProjectionType, ...) template void FDisplayClusterMeshProjectionRenderer::FuncName<ProjectionType>(__VA_ARGS__);
 #define CREATE_PROJECTION_TYPE_TEMPLATE_SPECIALIZATIONS(FuncName, ...) \
 	PROJECTION_TYPE_TEMPLATE_SPECIALIZATION(FuncName, EDisplayClusterMeshProjectionType::Linear, __VA_ARGS__) \
-	PROJECTION_TYPE_TEMPLATE_SPECIALIZATION(FuncName, EDisplayClusterMeshProjectionType::Azimuthal, __VA_ARGS__)
+	PROJECTION_TYPE_TEMPLATE_SPECIALIZATION(FuncName, EDisplayClusterMeshProjectionType::Azimuthal, __VA_ARGS__) \
+	PROJECTION_TYPE_TEMPLATE_SPECIALIZATION(FuncName, EDisplayClusterMeshProjectionType::UV, __VA_ARGS__)
 
-CREATE_PROJECTION_TYPE_TEMPLATE_SPECIALIZATIONS(RenderPrimitives_RenderThread, const FSceneView* View, FRHICommandList& RHICmdList, bool bTranslucencyPass)
-CREATE_PROJECTION_TYPE_TEMPLATE_SPECIALIZATIONS(RenderHitProxies_RenderThread, const FSceneView* View, FRHICommandList& RHICmdList)
-CREATE_PROJECTION_TYPE_TEMPLATE_SPECIALIZATIONS(RenderNormals_RenderThread, const FSceneView* View, FRHICommandList& RHICmdList, FDisplayClusterMeshProjectionPrimitiveFilter* PrimitiveFilter)
+CREATE_PROJECTION_TYPE_TEMPLATE_SPECIALIZATIONS(RenderPrimitives_RenderThread, const FSceneView* View, FRHICommandList& RHICmdList, const FDisplayClusterMeshProjectionRenderSettings& RenderSettings, bool bTranslucencyPass)
+CREATE_PROJECTION_TYPE_TEMPLATE_SPECIALIZATIONS(RenderHitProxies_RenderThread, const FSceneView* View, FRHICommandList& RHICmdList, const FDisplayClusterMeshProjectionRenderSettings& RenderSettings)
+CREATE_PROJECTION_TYPE_TEMPLATE_SPECIALIZATIONS(RenderNormals_RenderThread, const FSceneView* View, FRHICommandList& RHICmdList, const FDisplayClusterMeshProjectionRenderSettings& RenderSettings)
 
 #if WITH_EDITOR
-CREATE_PROJECTION_TYPE_TEMPLATE_SPECIALIZATIONS(RenderSelection_RenderThread, const FSceneView* View, FRHICommandList& RHICmdList)
+CREATE_PROJECTION_TYPE_TEMPLATE_SPECIALIZATIONS(RenderSelection_RenderThread, const FSceneView* View, FRHICommandList& RHICmdList, const FDisplayClusterMeshProjectionRenderSettings& RenderSettings)
 #endif
\ No newline at end of file
diff --git a/Engine/Plugins/Runtime/nDisplay/Source/DisplayClusterLightCardEditorShaders/Public/DisplayClusterMeshProjectionRenderer.h b/Engine/Plugins/Runtime/nDisplay/Source/DisplayClusterLightCardEditorShaders/Public/DisplayClusterMeshProjectionRenderer.h
--- a/Engine/Plugins/Runtime/nDisplay/Source/DisplayClusterLightCardEditorShaders/Public/DisplayClusterMeshProjectionRenderer.h
+++ b/Engine/Plugins/Runtime/nDisplay/Source/DisplayClusterLightCardEditorShaders/Public/DisplayClusterMeshProjectionRenderer.h
@@ -4,6 +4,7 @@
 
 #include "CoreMinimal.h"
 #include "RenderGraphBuilder.h"
+#include "SceneView.h"
 
 class AActor;
 class UPrimitiveComponent;
@@ -11,7 +12,7 @@ class FSceneView;
 class FViewInfo;
 class FPrimitiveDrawInterface;
 class FSimpleElementCollector;
-struct FSceneViewInitOptions;
+class FMeshProjectionPassParameters;
 struct FEngineShowFlags;
 
 /** Indicates which kind of projection is used by the renderer */
@@ -21,19 +22,81 @@ enum DISPLAYCLUSTERLIGHTCARDEDITORSHADERS_API EDisplayClusterMeshProjectionType
 	Linear,
 
 	/** Non-linear spherical projection based on the azimuthal equidistant map projection */
-	Azimuthal
+	Azimuthal,
+
+	/** Projection that positions vertices based on their UV coordinates */
+	UV
+};
+
+/** Indicates the quantity that is output to the canvas by the renderer */
+enum DISPLAYCLUSTERLIGHTCARDEDITORSHADERS_API EDisplayClusterMeshProjectionOutput
+{
+	/** Outputs the emissive color of the rendered primitives */
+	Color,
+
+	/** Outputs the normals and depth of the rendered primitives */
+	Normals
 };
 
 /** A filter that allows specific primitive components to be filtered from a render pass */
 class DISPLAYCLUSTERLIGHTCARDEDITORSHADERS_API FDisplayClusterMeshProjectionPrimitiveFilter
 {
 public:
-	/** A delegate that returns true if the primitive component should be included in the render pass */
 	DECLARE_DELEGATE_RetVal_OneParam(bool, FPrimitiveFilter, const UPrimitiveComponent*);
-	FPrimitiveFilter PrimitiveFilterDelegate;
+
+
+	/** A delegate that returns true if the primitive component should be included in the render pass */
+	FPrimitiveFilter ShouldRenderPrimitiveDelegate;
+
+	/** A delegate that returns true if the primitive component should be rendered using the render pass's projection type */
+	FPrimitiveFilter ShouldApplyProjectionDelegate;
 
 	/** Gets whether a primitive component should be filtered out of the render pass or not */
-	bool IsPrimitiveComponentFiltered(const UPrimitiveComponent* InPrimitiveComponent) const;
+	bool ShouldRenderPrimitive(const UPrimitiveComponent* InPrimitiveComponent) const;
+
+	/** Gets whether a primitive component should be rendered using the current projection type or not */
+	bool ShouldApplyProjection(const UPrimitiveComponent* InPrimitiveComponent) const;
+};
+
+/** Settings for specific mesh projection types  */
+struct DISPLAYCLUSTERLIGHTCARDEDITORSHADERS_API FDisplayClusterMeshProjectionTypeSettings
+{
+	/** The index of the UV to use when performing a UV projection */
+	uint32 UVProjectionIndex = 0;
+
+	/** The size of the plane the UVs are projected to, in view space */
+	float UVProjectionPlaneSize = 100.0f;
+
+	/** The distance from the view origin of the plane the UVs are projected to, in view space */
+	float UVProjectionPlaneDistance = 100.0f;
+};
+
+/** Settings for producing a single render. */
+struct DISPLAYCLUSTERLIGHTCARDEDITORSHADERS_API FDisplayClusterMeshProjectionRenderSettings
+{
+	/** Camera setup options for the render. */
+	FSceneViewInitOptions ViewInitOptions;
+
+	/**
+	 * Flags controlling renderer features to enable/disable for this preview.
+	 * Only flags supported by FDisplayClusterMeshProjectionRenderer will have an impact on the render output.
+	 */
+	FEngineShowFlags EngineShowFlags = FEngineShowFlags(ESFIM_Editor);
+
+	/** Type of projection to use for the renderer. */
+	EDisplayClusterMeshProjectionType ProjectionType = EDisplayClusterMeshProjectionType::Azimuthal;
+
+	/** Settings used by the projection type when rendering */
+	FDisplayClusterMeshProjectionTypeSettings ProjectionTypeSettings;
+
+	/** The output type to use for the renderer. */
+	EDisplayClusterMeshProjectionOutput RenderType = EDisplayClusterMeshProjectionOutput::Color;
+
+	/**
+	 * Optional filter to prevent specific primitives from being rendered.
+	 * This only applies when RenderType is ERenderType::Normals.
+	 */
+	FDisplayClusterMeshProjectionPrimitiveFilter PrimitiveFilter;
 };
 
 /** A transform that can be passed around to project and unprojection positions for a specific projection type */
@@ -85,43 +148,53 @@ public:
 	void ClearScene();
 
 	/** Renders its list of primitive components to the specified canvas using the desired projection type. Can be called from the game thread */
-	void Render(FCanvas* Canvas, FSceneInterface* Scene, const FSceneViewInitOptions& ViewInitOptions, const FEngineShowFlags& EngineShowFlags, EDisplayClusterMeshProjectionType  ProjectionType);
-
-	/** Renders the normals and depth of the scene's primitives and performs blurring to create a continuous normal map */
-	void RenderNormals(FCanvas* Canvas,
-		FSceneInterface* Scene,
-		const FSceneViewInitOptions& ViewInitOptions,
-		const FEngineShowFlags& EngineShowFlags,
-		EDisplayClusterMeshProjectionType ProjectionType, 
-		FDisplayClusterMeshProjectionPrimitiveFilter* PrimitiveFilter = nullptr);
+	void Render(FCanvas* Canvas, FSceneInterface* Scene, const FDisplayClusterMeshProjectionRenderSettings& RenderSettings);
 
 private:
+	/** Constructs the necessary render passes for the default output of the rendered primitives */
+	void RenderColorOutput(FRDGBuilder& GraphBuilder,
+		const FViewInfo* View,
+		const FDisplayClusterMeshProjectionRenderSettings& RenderSettings,
+		FRenderTargetBinding& OutputRenderTargetBinding);
+
+	/** Constructs the necessary render passes for the hit proxy output of the rendered primitives */
+	void RenderHitProxyOutput(FRDGBuilder& GraphBuilder,
+		const FViewInfo* View,
+		const FDisplayClusterMeshProjectionRenderSettings& RenderSettings,
+		FRenderTargetBinding& OutputRenderTargetBinding,
+		FHitProxyConsumer* HitProxyConsumer);
+
+	/** Constructs the necessary render passes for the normals/depth output of the rendered primitives */
+	void RenderNormalsOutput(FRDGBuilder& GraphBuilder,
+		const FViewInfo* View,
+		const FDisplayClusterMeshProjectionRenderSettings& RenderSettings,
+		FRenderTargetBinding& OutputRenderTargetBinding);
+
 	/** Adds a pass to perform the base render for the mesh projection. */
 	void AddBaseRenderPass(FRDGBuilder& GraphBuilder, 
 		const FViewInfo* View,
-		EDisplayClusterMeshProjectionType ProjectionType,
+		const FDisplayClusterMeshProjectionRenderSettings& RenderSettings,
 		FRenderTargetBinding& OutputRenderTargetBinding,
 		FDepthStencilBinding& OutputDepthStencilBinding);
 
 	/** Adds a pass to perform the translucency render for the mesh projection. */
 	void AddTranslucencyRenderPass(FRDGBuilder& GraphBuilder, 
 		const FViewInfo* View,
-		EDisplayClusterMeshProjectionType ProjectionType,
+		const FDisplayClusterMeshProjectionRenderSettings& RenderSettings,
 		FRenderTargetBinding& OutputRenderTargetBinding,
 		FDepthStencilBinding& OutputDepthStencilBinding);
 
 	/** Adds a pass to perform the hit proxy render for the mesh projection */
 	void AddHitProxyRenderPass(FRDGBuilder& GraphBuilder, 
 		const FViewInfo* View,
-		EDisplayClusterMeshProjectionType ProjectionType,
+		const FDisplayClusterMeshProjectionRenderSettings& RenderSettings,
 		FRenderTargetBinding& OutputRenderTargetBinding,
 		FDepthStencilBinding& OutputDepthStencilBinding);
 
 	/** Adds a pass to perform a render of the primitives' normals for the mesh projection */
 	void AddNormalsRenderPass(FRDGBuilder& GraphBuilder,
 		const FViewInfo* View,
-		EDisplayClusterMeshProjectionType ProjectionType,
-		FDisplayClusterMeshProjectionPrimitiveFilter* PrimitiveFilter,
+		const FDisplayClusterMeshProjectionRenderSettings& RenderSettings,
 		FRenderTargetBinding& OutputRenderTargetBinding,
 		FDepthStencilBinding& OutputDepthStencilBinding);
 
@@ -137,7 +210,7 @@ private:
 	/** Adds a pass to perform the depth render for any selected primitives for the mesh projection */
 	void AddSelectionDepthRenderPass(FRDGBuilder& GraphBuilder, 
 		const FViewInfo* View,
-		EDisplayClusterMeshProjectionType ProjectionType,
+		const FDisplayClusterMeshProjectionRenderSettings& RenderSettings,
 		FDepthStencilBinding& OutputDepthStencilBinding);
 
 	/** Adds a pass to perform the post process selection outline for the mesh projection */
@@ -157,25 +230,47 @@ private:
 
 	/** Renders the list of primitive components using the appropriate mesh pass processor given by the template parameter */
 	template<EDisplayClusterMeshProjectionType ProjectionType>
-	void RenderPrimitives_RenderThread(const FSceneView* View, FRHICommandList& RHICmdList, bool bTranslucencyPass);
+	void RenderPrimitives_RenderThread(const FSceneView* View,
+		FRHICommandList& RHICmdList,
+		const FDisplayClusterMeshProjectionRenderSettings& RenderSettings,
+		bool bTranslucencyPass);
 
 	/** Renders the hit proxies of the list of primitive components using the appropriate mesh pass processor given by the template parameter */
 	template<EDisplayClusterMeshProjectionType ProjectionType>
-	void RenderHitProxies_RenderThread(const FSceneView* View, FRHICommandList& RHICmdList);
+	void RenderHitProxies_RenderThread(const FSceneView* View,
+		FRHICommandList& RHICmdList,
+		const FDisplayClusterMeshProjectionRenderSettings& RenderSettings);
 
 	/** Renders the normals of the list of primitive components using the appropriate mesh pass processor given by the template parameter */
 	template<EDisplayClusterMeshProjectionType ProjectionType>
-	void RenderNormals_RenderThread(const FSceneView* View, FRHICommandList& RHICmdList, FDisplayClusterMeshProjectionPrimitiveFilter* PrimitiveFilter);
+	void RenderNormals_RenderThread(const FSceneView* View,
+		FRHICommandList& RHICmdList,
+		const FDisplayClusterMeshProjectionRenderSettings& RenderSettings);
 
 #if WITH_EDITOR
 	/** Renders the list of primitive components using the appropriate mesh pass processor given by the template parameter */
 	template<EDisplayClusterMeshProjectionType ProjectionType>
-	void RenderSelection_RenderThread(const FSceneView* View, FRHICommandList& RHICmdList);
+	void RenderSelection_RenderThread(const FSceneView* View,
+		FRHICommandList& RHICmdList,
+		const FDisplayClusterMeshProjectionRenderSettings& RenderSettings);
 #endif
 
+
 	/** Callback used to determine if a primitive component should be rendered with a selection outline */
 	bool IsPrimitiveComponentSelected(const UPrimitiveComponent* InPrimitiveComponent);
 
+	/** Data struct that contains the scene proxy and the render configuration needed to render the scene's primitive components */
+	struct FSceneProxyElement
+	{
+		FPrimitiveSceneProxy* PrimitiveSceneProxy;
+		bool bApplyProjection;
+	};
+
+	/** Collects all valid scene proxies to be rendered from the scene's primitive componets */
+	void GetSceneProxies(TArray<FSceneProxyElement>& OutSceneProxyElements,
+		const FDisplayClusterMeshProjectionRenderSettings& RenderSettings,
+		const TFunctionRef<bool(const UPrimitiveComponent*)> PrimitiveFilter = [](const UPrimitiveComponent*) { return true; });
+
 public:
 	/** Delegate to determine if an actor should be rendered as selected */
 	DECLARE_DELEGATE_RetVal_OneParam(bool, FSelection, const AActor*);
diff --git a/Engine/Plugins/Runtime/nDisplay/Source/DisplayClusterProjection/Private/Policy/MPCDI/DisplayClusterProjectionMPCDIPolicy.cpp b/Engine/Plugins/Runtime/nDisplay/Source/DisplayClusterProjection/Private/Policy/MPCDI/DisplayClusterProjectionMPCDIPolicy.cpp
--- a/Engine/Plugins/Runtime/nDisplay/Source/DisplayClusterProjection/Private/Policy/MPCDI/DisplayClusterProjectionMPCDIPolicy.cpp
+++ b/Engine/Plugins/Runtime/nDisplay/Source/DisplayClusterProjection/Private/Policy/MPCDI/DisplayClusterProjectionMPCDIPolicy.cpp
@@ -317,6 +317,11 @@ void FDisplayClusterProjectionMPCDIPolicy::ApplyWarpBlend_RenderThread(FRHIComma
 					}
 				}
 
+				if (RefViewportManagerProxy.GetLightCardManager_RenderThread().IsValid() && RefViewportManagerProxy.GetLightCardManager_RenderThread()->HasUVLightCards_RenderThread())
+				{
+					ShaderICVFX.UVLightCardMap = RefViewportManagerProxy.GetLightCardManager_RenderThread()->GetUVLightCardMap_RenderThread();
+				}
+
 				// Initialize shader input data
 				FDisplayClusterShaderParameters_WarpBlend WarpBlendParameters;
 
diff --git a/Engine/Plugins/Runtime/nDisplay/Source/DisplayClusterScenePreview/Private/DisplayClusterScenePreviewModule.cpp b/Engine/Plugins/Runtime/nDisplay/Source/DisplayClusterScenePreview/Private/DisplayClusterScenePreviewModule.cpp
--- a/Engine/Plugins/Runtime/nDisplay/Source/DisplayClusterScenePreview/Private/DisplayClusterScenePreviewModule.cpp
+++ b/Engine/Plugins/Runtime/nDisplay/Source/DisplayClusterScenePreview/Private/DisplayClusterScenePreviewModule.cpp
@@ -99,6 +99,17 @@ bool FDisplayClusterScenePreviewModule::AddActorToRenderer(int32 RendererId, AAc
 	return false;
 }
 
+bool FDisplayClusterScenePreviewModule::AddActorToRenderer(int32 RendererId, AActor* Actor, const TFunctionRef<bool(const UPrimitiveComponent*)>& PrimitiveFilter)
+{
+	if (FRendererConfig* Config = RendererConfigs.Find(RendererId))
+	{
+		Config->Renderer->AddActor(Actor, PrimitiveFilter);
+		return true;
+	}
+
+	return false;
+}
+
 bool FDisplayClusterScenePreviewModule::RemoveActorFromRenderer(int32 RendererId, AActor* Actor)
 {
 	if (FRendererConfig* Config = RendererConfigs.Find(RendererId))
@@ -143,7 +154,7 @@ bool FDisplayClusterScenePreviewModule::SetRendererRenderSimpleElementsDelegate(
 	return false;
 }
 
-bool FDisplayClusterScenePreviewModule::Render(int32 RendererId, FDisplayClusterScenePreviewRenderSettings& RenderSettings, FCanvas& Canvas)
+bool FDisplayClusterScenePreviewModule::Render(int32 RendererId, FDisplayClusterMeshProjectionRenderSettings& RenderSettings, FCanvas& Canvas)
 {
 	if (FRendererConfig* Config = RendererConfigs.Find(RendererId))
 	{
@@ -153,7 +164,7 @@ bool FDisplayClusterScenePreviewModule::Render(int32 RendererId, FDisplayCluster
 	return false;
 }
 
-bool FDisplayClusterScenePreviewModule::RenderQueued(int32 RendererId, FDisplayClusterScenePreviewRenderSettings& RenderSettings, const FIntPoint& Size, FRenderResultDelegate ResultDelegate)
+bool FDisplayClusterScenePreviewModule::RenderQueued(int32 RendererId, FDisplayClusterMeshProjectionRenderSettings& RenderSettings, const FIntPoint& Size, FRenderResultDelegate ResultDelegate)
 {
 	if (FRendererConfig* Config = RendererConfigs.Find(RendererId))
 	{
@@ -215,7 +226,7 @@ void FDisplayClusterScenePreviewModule::InternalSetRendererRootActor(FRendererCo
 	RegisterOrUnregisterGlobalActorEvents();
 }
 
-bool FDisplayClusterScenePreviewModule::InternalRenderImmediate(FRendererConfig& RendererConfig, FDisplayClusterScenePreviewRenderSettings& RenderSettings, FCanvas& Canvas)
+bool FDisplayClusterScenePreviewModule::InternalRenderImmediate(FRendererConfig& RendererConfig, FDisplayClusterMeshProjectionRenderSettings& RenderSettings, FCanvas& Canvas)
 {
 	UWorld* World = RendererConfig.RootActor.IsValid() ? RendererConfig.RootActor->GetWorld() : nullptr;
 	if (!World)
@@ -228,31 +239,8 @@ bool FDisplayClusterScenePreviewModule::InternalRenderImmediate(FRendererConfig&
 		AutoPopulateScene(RendererConfig);
 	}
 
-	switch (RenderSettings.RenderType)
-	{
-	case FDisplayClusterScenePreviewRenderSettings::ERenderType::Color:
-		RendererConfig.Renderer->Render(
-			&Canvas,
-			World->Scene,
-			RenderSettings.ViewInitOptions,
-			RenderSettings.EngineShowFlags,
-			RenderSettings.ProjectionType
-		);
-		return true;
-
-	case FDisplayClusterScenePreviewRenderSettings::ERenderType::Normals:
-		RendererConfig.Renderer->RenderNormals(
-			&Canvas,
-			World->Scene,
-			RenderSettings.ViewInitOptions,
-			RenderSettings.EngineShowFlags,
-			RenderSettings.ProjectionType,
-			&RenderSettings.PrimitiveFilter
-		);
-		return true;
-	}
-
-	return false;
+	RendererConfig.Renderer->Render(&Canvas, World->Scene, RenderSettings);
+	return true;
 }
 
 void FDisplayClusterScenePreviewModule::RegisterOrUnregisterGlobalActorEvents()
diff --git a/Engine/Plugins/Runtime/nDisplay/Source/DisplayClusterScenePreview/Private/DisplayClusterScenePreviewModule.h b/Engine/Plugins/Runtime/nDisplay/Source/DisplayClusterScenePreview/Private/DisplayClusterScenePreviewModule.h
--- a/Engine/Plugins/Runtime/nDisplay/Source/DisplayClusterScenePreview/Private/DisplayClusterScenePreviewModule.h
+++ b/Engine/Plugins/Runtime/nDisplay/Source/DisplayClusterScenePreview/Private/DisplayClusterScenePreviewModule.h
@@ -30,12 +30,13 @@ public:
 	virtual bool SetRendererRootActor(int32 RendererId, ADisplayClusterRootActor* Actor, bool bAutoUpdateLightcards = false) override;
 	virtual ADisplayClusterRootActor* GetRendererRootActor(int32 RendererId) override;
 	virtual bool AddActorToRenderer(int32 RendererId, AActor* Actor) override;
+	virtual bool AddActorToRenderer(int32 RendererId, AActor* Actor, const TFunctionRef<bool(const UPrimitiveComponent*)>& PrimitiveFilter) override;
 	virtual bool RemoveActorFromRenderer(int32 RendererId, AActor* Actor) override;
 	virtual bool ClearRendererScene(int32 RendererId) override;
 	virtual bool SetRendererActorSelectedDelegate(int32 RendererId, FDisplayClusterMeshProjectionRenderer::FSelection ActorSelectedDelegate) override;
 	virtual bool SetRendererRenderSimpleElementsDelegate(int32 RendererId, FDisplayClusterMeshProjectionRenderer::FSimpleElementPass RenderSimpleElementsDelegate) override;
-	virtual bool Render(int32 RendererId, FDisplayClusterScenePreviewRenderSettings& RenderSettings, FCanvas& Canvas) override;
-	virtual bool RenderQueued(int32 RendererId, FDisplayClusterScenePreviewRenderSettings& RenderSettings, const FIntPoint& Size, FRenderResultDelegate ResultDelegate) override;
+	virtual bool Render(int32 RendererId, FDisplayClusterMeshProjectionRenderSettings& RenderSettings, FCanvas& Canvas) override;
+	virtual bool RenderQueued(int32 RendererId, FDisplayClusterMeshProjectionRenderSettings& RenderSettings, const FIntPoint& Size, FRenderResultDelegate ResultDelegate) override;
 	//~ End IDisplayClusterScenePreview Interface
 
 private:
@@ -69,7 +70,7 @@ private:
 	{
 		FPreviewRenderJob() {}
 
-		FPreviewRenderJob(int32 RendererId, const FDisplayClusterScenePreviewRenderSettings& Settings, const FIntPoint& Size,
+		FPreviewRenderJob(int32 RendererId, const FDisplayClusterMeshProjectionRenderSettings& Settings, const FIntPoint& Size,
 			FRenderResultDelegate ResultDelegate)
 			: RendererId(RendererId), Settings(Settings), Size(Size), ResultDelegate(ResultDelegate)
 		{
@@ -79,7 +80,7 @@ private:
 		int32 RendererId;
 
 		/** The settings to use for the render. */
-		FDisplayClusterScenePreviewRenderSettings Settings;
+		FDisplayClusterMeshProjectionRenderSettings Settings;
 
 		/** The size of the image to render. */
 		FIntPoint Size;
@@ -95,7 +96,7 @@ private:
 	void InternalSetRendererRootActor(FRendererConfig& RendererConfig, ADisplayClusterRootActor* Actor, bool bAutoUpdateLightcards);
 
 	/** Immediately render with the given renderer config and settings to the given canvas. */
-	bool InternalRenderImmediate(FRendererConfig& RendererConfig, FDisplayClusterScenePreviewRenderSettings& RenderSettings, FCanvas& Canvas);
+	bool InternalRenderImmediate(FRendererConfig& RendererConfig, FDisplayClusterMeshProjectionRenderSettings& RenderSettings, FCanvas& Canvas);
 
 	/** Check if any of the tracked root actors are set to auto-update their lightcards and register/unregister event listeners accordingly. */
 	void RegisterOrUnregisterGlobalActorEvents();
diff --git a/Engine/Plugins/Runtime/nDisplay/Source/DisplayClusterScenePreview/Public/IDisplayClusterScenePreview.h b/Engine/Plugins/Runtime/nDisplay/Source/DisplayClusterScenePreview/Public/IDisplayClusterScenePreview.h
--- a/Engine/Plugins/Runtime/nDisplay/Source/DisplayClusterScenePreview/Public/IDisplayClusterScenePreview.h
+++ b/Engine/Plugins/Runtime/nDisplay/Source/DisplayClusterScenePreview/Public/IDisplayClusterScenePreview.h
@@ -6,54 +6,13 @@
 #include "Modules/ModuleInterface.h"
 #include "Modules/ModuleManager.h"
 
-#include "SceneView.h"
 #include "ShowFlags.h"
 #include "DisplayClusterMeshProjectionRenderer.h"
 
 class ADisplayClusterRootActor;
-struct FSceneViewInitOptions;
 
 DECLARE_DELEGATE_OneParam(FRenderResultDelegate, FRenderTarget&);
 
-/** Settings for producing a single render. */
-struct FDisplayClusterScenePreviewRenderSettings
-{
-	FDisplayClusterScenePreviewRenderSettings()
-		: EngineShowFlags(ESFIM_Editor)
-	{}
-
-	/** Possible render types. */
-	enum ERenderType
-	{
-		/** Render the scene's colors. */
-		Color,
-
-		/** Render the scene's normals. */
-		Normals
-	};
-
-	/** The type of render to perform. */
-	ERenderType RenderType = ERenderType::Color;
-
-	/** Camera setup options for the render. */
-	FSceneViewInitOptions ViewInitOptions;
-
-	/**
-	 * Flags controlling renderer features to enable/disable for this preview.
-	 * Only flags supported by FDisplayClusterMeshProjectionRenderer will have an impact on the render output.
-	 */
-	FEngineShowFlags EngineShowFlags;
-
-	/** Type of projection to use for the renderer. */
-	EDisplayClusterMeshProjectionType ProjectionType = EDisplayClusterMeshProjectionType::Azimuthal;
-
-	/**
-	 * Optional filter to prevent specific primitives from being rendered.
-	 * This only applies when RenderType is ERenderType::Normals.
-	 */
-	FDisplayClusterMeshProjectionPrimitiveFilter PrimitiveFilter;
-};
-
 /**
  * Interface for module containing tools for generation Display Cluster scene previews.
  */
@@ -134,6 +93,15 @@ public:
 	 */
 	virtual bool AddActorToRenderer(int32 RendererId, AActor* Actor) = 0;
 
+	/**
+	 * Add an actor to the preview scene for a renderer.
+	 *
+	 * @param RendererId The ID of the renderer as returned from CreateRenderer.
+	 * @param Actor The actor to add to the renderer's scene.
+	 * @param PrimitiveFilter A custom filter used to determine which of the actor's primitives are added to the renderer's scene.
+	 */
+	virtual bool AddActorToRenderer(int32 RendererId, AActor* Actor, const TFunctionRef<bool(const UPrimitiveComponent*)>& PrimitiveFilter) = 0;
+
 	/**
 	 * Removes an actor from the preview scene for a renderer.
 	 *
@@ -174,7 +142,7 @@ public:
 	 * @param RenderSettings Settings controlling how the scene will be rendered.
 	 * @param Canvas The canvas to render to.
 	 */
-	virtual bool Render(int32 RendererId, FDisplayClusterScenePreviewRenderSettings& RenderSettings, FCanvas& Canvas) = 0;
+	virtual bool Render(int32 RendererId, FDisplayClusterMeshProjectionRenderSettings& RenderSettings, FCanvas& Canvas) = 0;
 
 	/**
 	 * Queue a preview to be rendered in the future.
@@ -184,5 +152,5 @@ public:
 	 * @param Size The size of the image to produce. Note that whenever this changes for a given renderer, the underlying RenderTarget will be resized, which has a performance cost.
 	 * @param ResultDelegate The delegate to call when the render is complete. It will be passed a FRenderTarget containing the rendered preview.
 	 */
-	virtual bool RenderQueued(int32 RendererId, FDisplayClusterScenePreviewRenderSettings& RenderSettings, const FIntPoint& Size, FRenderResultDelegate ResultDelegate) = 0;
+	virtual bool RenderQueued(int32 RendererId, FDisplayClusterMeshProjectionRenderSettings& RenderSettings, const FIntPoint& Size, FRenderResultDelegate ResultDelegate) = 0;
 };
\ No newline at end of file
diff --git a/Engine/Plugins/Runtime/nDisplay/Source/DisplayClusterShaders/Private/DisplayClusterShadersModule.cpp b/Engine/Plugins/Runtime/nDisplay/Source/DisplayClusterShaders/Private/DisplayClusterShadersModule.cpp
--- a/Engine/Plugins/Runtime/nDisplay/Source/DisplayClusterShaders/Private/DisplayClusterShadersModule.cpp
+++ b/Engine/Plugins/Runtime/nDisplay/Source/DisplayClusterShaders/Private/DisplayClusterShadersModule.cpp
@@ -6,6 +6,7 @@
 #include "Misc/Paths.h"
 #include "ShaderCore.h"
 
+#include "Shaders/DisplayClusterShadersPreprocess_UVLightCards.h"
 #include "Shaders/DisplayClusterShadersPostprocess_Blur.h"
 #include "Shaders/DisplayClusterShadersPostprocess_OutputRemap.h"
 #include "Shaders/DisplayClusterShadersWarpblend_ICVFX.h"
@@ -41,6 +42,11 @@ bool FDisplayClusterShadersModule::RenderWarpBlend_ICVFX(FRHICommandListImmediat
 	return FDisplayClusterShadersWarpblend_ICVFX::RenderWarpBlend_ICVFX(RHICmdList, InWarpBlendParameters, InICVFXParameters);
 }
 
+bool FDisplayClusterShadersModule::RenderPreprocess_UVLightCards(FRHICommandListImmediate& RHICmdList, FSceneInterface* InScene, FRenderTarget* InRenderTarget, float ProjectionPlaneSize) const
+{
+	return FDisplayClusterShadersPreprocess_UVLightCards::RenderPreprocess_UVLightCards(RHICmdList, InScene, InRenderTarget, ProjectionPlaneSize);
+}
+
 bool FDisplayClusterShadersModule::RenderPostprocess_OutputRemap(FRHICommandListImmediate& RHICmdList, FRHITexture2D* InSourceTexture, FRHITexture2D* InRenderTargetableDestTexture, const IDisplayClusterRender_MeshComponentProxy& MeshProxy) const
 {
 	return FDisplayClusterShadersPostprocess_OutputRemap::RenderPostprocess_OutputRemap(RHICmdList, InSourceTexture, InRenderTargetableDestTexture, MeshProxy);
diff --git a/Engine/Plugins/Runtime/nDisplay/Source/DisplayClusterShaders/Private/DisplayClusterShadersModule.h b/Engine/Plugins/Runtime/nDisplay/Source/DisplayClusterShaders/Private/DisplayClusterShadersModule.h
--- a/Engine/Plugins/Runtime/nDisplay/Source/DisplayClusterShaders/Private/DisplayClusterShadersModule.h
+++ b/Engine/Plugins/Runtime/nDisplay/Source/DisplayClusterShaders/Private/DisplayClusterShadersModule.h
@@ -19,6 +19,7 @@ public:
 public:
 	virtual bool RenderWarpBlend_MPCDI(FRHICommandListImmediate& RHICmdList, const FDisplayClusterShaderParameters_WarpBlend& InWarpBlendParameters) const override;
 	virtual bool RenderWarpBlend_ICVFX(FRHICommandListImmediate& RHICmdList, const FDisplayClusterShaderParameters_WarpBlend& InWarpBlendParameters, const FDisplayClusterShaderParameters_ICVFX& InICVFXParameters) const override;
+	virtual bool RenderPreprocess_UVLightCards(FRHICommandListImmediate& RHICmdList, FSceneInterface* InScene, FRenderTarget* InRenderTarget, float ProjectionPlaneSize) const override;
 	virtual bool RenderPostprocess_OutputRemap(FRHICommandListImmediate& RHICmdList, FRHITexture2D* InSourceTexture, FRHITexture2D* InRenderTargetableDestTexture, const IDisplayClusterRender_MeshComponentProxy& MeshProxy) const override;
 	virtual bool RenderPostprocess_Blur(FRHICommandListImmediate& RHICmdList, FRHITexture2D* InSourceTexture, FRHITexture2D* InRenderTargetableDestTexture, const FDisplayClusterShaderParameters_PostprocessBlur& InSettings) const override;
 	virtual bool GenerateMips(FRHICommandListImmediate& RHICmdList, FRHITexture2D* InOutMipsTexture, const FDisplayClusterShaderParameters_GenerateMips& InSettings) const override;
diff --git a/Engine/Plugins/Runtime/nDisplay/Source/DisplayClusterShaders/Private/Shaders/DisplayClusterShadersPreprocess_UVLightCards.cpp b/Engine/Plugins/Runtime/nDisplay/Source/DisplayClusterShaders/Private/Shaders/DisplayClusterShadersPreprocess_UVLightCards.cpp
new file mode 100644
--- /dev/null
+++ b/Engine/Plugins/Runtime/nDisplay/Source/DisplayClusterShaders/Private/Shaders/DisplayClusterShadersPreprocess_UVLightCards.cpp
@@ -0,0 +1,242 @@
+// Copyright Epic Games, Inc. All Rights Reserved.
+
+#include "DisplayClusterShadersPreprocess_UVLightCards.h"
+
+#include "Components/PrimitiveComponent.h"
+#include "EngineModule.h"
+#include "EngineUtils.h"
+#include "InstanceCulling/InstanceCullingContext.h"
+#include "MeshMaterialShader.h"
+#include "MeshPassProcessor.h"
+#include "PrimitiveSceneInfo.h"
+#include "PrimitiveSceneProxy.h"
+#include "RenderGraphBuilder.h"
+#include "RenderGraphResources.h"
+#include "Renderer/Private/ScenePrivate.h"
+#include "SceneRenderTargetParameters.h"
+#include "MeshPassProcessor.inl"
+
+class FUVLightCardVS : public FMeshMaterialShader
+{
+public:
+	DECLARE_SHADER_TYPE(FUVLightCardVS, MeshMaterial);
+
+	FUVLightCardVS() { }
+	FUVLightCardVS(const ShaderMetaType::CompiledShaderInitializerType& Initializer)
+		: FMeshMaterialShader(Initializer)
+	{
+		PassUniformBuffer.Bind(Initializer.ParameterMap, FSceneTextureUniformParameters::StaticStructMetadata.GetShaderVariableName());
+	}
+	              
+	static bool ShouldCompilePermutation(const FMeshMaterialShaderPermutationParameters& Parameters)
+	{
+		return EnumHasAllFlags(Parameters.Flags, EShaderPermutationFlags::HasEditorOnlyData)
+			&& IsFeatureLevelSupported(Parameters.Platform, ERHIFeatureLevel::SM5)
+			&& Parameters.VertexFactoryType == FindVertexFactoryType(TEXT("FLocalVertexFactory"));
+	}
+};
+
+IMPLEMENT_MATERIAL_SHADER_TYPE(, FUVLightCardVS, TEXT("/Plugin/nDisplay/Private/UVLightCardShaders.usf"), TEXT("MainVS"), SF_Vertex);
+
+class FUVLightCardPS : public FMeshMaterialShader
+{
+public:
+	DECLARE_SHADER_TYPE(FUVLightCardPS, MeshMaterial);
+
+	FUVLightCardPS() { }
+	FUVLightCardPS(const ShaderMetaType::CompiledShaderInitializerType& Initializer)
+		: FMeshMaterialShader(Initializer)
+	{
+		PassUniformBuffer.Bind(Initializer.ParameterMap, FSceneTextureUniformParameters::StaticStructMetadata.GetShaderVariableName());
+	}
+
+	static bool ShouldCompilePermutation(const FMeshMaterialShaderPermutationParameters& Parameters)
+	{
+		return EnumHasAllFlags(Parameters.Flags, EShaderPermutationFlags::HasEditorOnlyData)
+			&& IsFeatureLevelSupported(Parameters.Platform, ERHIFeatureLevel::SM5)
+			&& Parameters.VertexFactoryType == FindVertexFactoryType(TEXT("FLocalVertexFactory"));
+	}
+};
+
+IMPLEMENT_MATERIAL_SHADER_TYPE(, FUVLightCardPS, TEXT("/Plugin/nDisplay/Private/UVLightCardShaders.usf"), TEXT("MainPS"), SF_Pixel);
+
+class FUVLightCardPassProcessor : public FMeshPassProcessor
+{
+public:
+	FUVLightCardPassProcessor(const FSceneView* InView, FMeshPassDrawListContext* InDrawListContext)
+		: FMeshPassProcessor(nullptr, GMaxRHIFeatureLevel, InView, InDrawListContext)
+	{
+		DrawRenderState.SetDepthStencilState(TStaticDepthStencilState<false, CF_Always>::GetRHI());
+		DrawRenderState.SetBlendState(TStaticBlendState<CW_RGBA, BO_Add, BF_SourceAlpha, BF_InverseSourceAlpha, BO_Add, BF_One, BF_InverseSourceAlpha>::GetRHI());
+	}
+
+	virtual void AddMeshBatch(const FMeshBatch& RESTRICT MeshBatch, uint64 BatchElementMask, const FPrimitiveSceneProxy* RESTRICT PrimitiveSceneProxy, int32 StaticMeshId = -1) override final
+	{
+		const FMaterialRenderProxy* FallbackMaterialRenderProxy = nullptr;
+		const FMaterial& Material = MeshBatch.MaterialRenderProxy->GetMaterialWithFallback(FeatureLevel, FallbackMaterialRenderProxy);
+		const FMaterialRenderProxy& MaterialRenderProxy = FallbackMaterialRenderProxy ? *FallbackMaterialRenderProxy : *MeshBatch.MaterialRenderProxy;
+
+		const FVertexFactory* VertexFactory = MeshBatch.VertexFactory;
+
+		TMeshProcessorShaders<FUVLightCardVS, FUVLightCardPS> PassShaders;
+
+		FMaterialShaderTypes ShaderTypes;
+		ShaderTypes.AddShaderType<FUVLightCardVS>();
+		ShaderTypes.AddShaderType<FUVLightCardPS>();
+
+		FMaterialShaders Shaders;
+		if (!Material.TryGetShaders(ShaderTypes, VertexFactory->GetType(), Shaders))
+		{
+			return;
+		}
+
+		Shaders.TryGetVertexShader(PassShaders.VertexShader);
+		Shaders.TryGetPixelShader(PassShaders.PixelShader);
+
+		FMeshDrawingPolicyOverrideSettings OverrideSettings = ComputeMeshOverrideSettings(MeshBatch);
+		ERasterizerFillMode MeshFillMode = ComputeMeshFillMode(MeshBatch, Material, OverrideSettings);
+		ERasterizerCullMode MeshCullMode = ComputeMeshCullMode(MeshBatch, Material, OverrideSettings);
+
+		FMeshMaterialShaderElementData ShaderElementData;
+		ShaderElementData.InitializeMeshMaterialData(ViewIfDynamicMeshCommand, PrimitiveSceneProxy, MeshBatch, StaticMeshId, false);
+
+		FMeshDrawCommandSortKey SortKey = CalculateMeshStaticSortKey(PassShaders.VertexShader.GetShader(), PassShaders.PixelShader.GetShader());
+
+		BuildMeshDrawCommands(
+			MeshBatch,
+			BatchElementMask,
+			PrimitiveSceneProxy,
+			MaterialRenderProxy,
+			Material,
+			DrawRenderState,
+			PassShaders,
+			MeshFillMode,
+			MeshCullMode,
+			SortKey,
+			EMeshPassFeatures::Default,
+			ShaderElementData);
+	}
+
+private:
+	FMeshPassProcessorRenderState DrawRenderState;
+};
+
+BEGIN_SHADER_PARAMETER_STRUCT(FUVLightCardPassParameters, )
+	SHADER_PARAMETER_STRUCT_REF(FViewUniformShaderParameters, View)
+	SHADER_PARAMETER_RDG_UNIFORM_BUFFER(FInstanceCullingGlobalUniforms, InstanceCulling)
+	SHADER_PARAMETER_STRUCT_INCLUDE(FInstanceCullingDrawParams, InstanceCullingDrawParams)
+	RENDER_TARGET_BINDING_SLOTS()
+END_SHADER_PARAMETER_STRUCT()
+
+DECLARE_GPU_STAT_NAMED(nDisplay_UVLightCards_Render, TEXT("nDisplay UVLightCards::Render"));
+
+bool FDisplayClusterShadersPreprocess_UVLightCards::RenderPreprocess_UVLightCards(FRHICommandListImmediate& RHICmdList, FSceneInterface* InScene, FRenderTarget* InRenderTarget, float ProjectionPlaneSize)
+{
+	check(IsInRenderingThread());
+
+	if (InRenderTarget == nullptr || InScene->GetRenderScene()->PrimitiveSceneProxies.IsEmpty())
+	{
+		return false;
+	}
+
+	SCOPED_GPU_STAT(RHICmdList, nDisplay_UVLightCards_Render);
+	SCOPED_DRAW_EVENT(RHICmdList, nDisplay_UVLightCards_Render);
+
+	FRDGBuilder GraphBuilder(RHICmdList);
+
+	FEngineShowFlags EngineShowFlags(ESFIM_Game);
+
+	FSceneViewFamilyContext ViewFamily(FSceneViewFamily::ConstructionValues(
+		InRenderTarget,
+		InScene,
+		EngineShowFlags)
+		.SetTime(FGameTime::GetTimeSinceAppStart())
+		.SetGammaCorrection(1.0f));
+
+	FScenePrimitiveRenderingContextScopeHelper ScenePrimitiveRenderingContextScopeHelper(GetRendererModule().BeginScenePrimitiveRendering(GraphBuilder, &ViewFamily));
+
+	FSceneViewInitOptions ViewInitOptions;
+
+	ViewInitOptions.ViewFamily = &ViewFamily;
+	ViewInitOptions.ViewLocation = FVector::ZeroVector;
+	ViewInitOptions.ViewRotation = FRotator::ZeroRotator;
+	ViewInitOptions.ViewOrigin = ViewInitOptions.ViewLocation;
+
+	ViewInitOptions.SetViewRectangle(FIntRect(FIntPoint(0, 0), InRenderTarget->GetSizeXY()));
+
+	ViewInitOptions.ViewRotationMatrix = FInverseRotationMatrix(ViewInitOptions.ViewRotation) * FMatrix(
+		FPlane(0, 0, 1, 0),
+		FPlane(1, 0, 0, 0),
+		FPlane(0, 1, 0, 0),
+		FPlane(0, 0, 0, 1));
+
+	const float ZScale = 0.5f / HALF_WORLD_MAX;
+	const float ZOffset = HALF_WORLD_MAX;
+	const float OrthoSize = 0.5f * ProjectionPlaneSize;
+
+	if ((bool)ERHIZBuffer::IsInverted)
+	{
+		ViewInitOptions.ProjectionMatrix = FReversedZOrthoMatrix(
+			OrthoSize,
+			OrthoSize,
+			ZScale,
+			ZOffset
+		);
+	}
+	else
+	{
+		ViewInitOptions.ProjectionMatrix = FOrthoMatrix(
+			OrthoSize,
+			OrthoSize,
+			ZScale,
+			ZOffset
+		);
+	}
+
+	ViewInitOptions.BackgroundColor = FLinearColor::Transparent;
+
+	GetRendererModule().CreateAndInitSingleView(RHICmdList, &ViewFamily, &ViewInitOptions);
+	const FSceneView* View = ViewFamily.Views[0];
+
+	FRDGTextureRef OutputTexture = GraphBuilder.RegisterExternalTexture(CreateRenderTarget(InRenderTarget->GetRenderTargetTexture(), TEXT("UVLightCardRenderTarget")));
+	FRenderTargetBinding OutputRenderTargetBinding(OutputTexture, ERenderTargetLoadAction::EClear);
+
+	FUVLightCardPassParameters* PassParameters = GraphBuilder.AllocParameters<FUVLightCardPassParameters>();
+	PassParameters->View = View->ViewUniformBuffer;
+	PassParameters->InstanceCulling = FInstanceCullingContext::CreateDummyInstanceCullingUniformBuffer(GraphBuilder);
+	PassParameters->RenderTargets[0] = OutputRenderTargetBinding;
+
+	GraphBuilder.AddPass(RDG_EVENT_NAME("DisplayClusterUVLightCards::Render"),
+		PassParameters,
+		ERDGPassFlags::Raster | ERDGPassFlags::NeverCull,
+		[View, InScene](FRHICommandList& RHICmdList)
+		{
+			FIntRect ViewRect = View->UnscaledViewRect;
+			RHICmdList.SetViewport(ViewRect.Min.X, ViewRect.Min.Y, 0.0f, ViewRect.Max.X, ViewRect.Max.Y, 1.0f);
+
+			DrawDynamicMeshPass(*View, RHICmdList, [View, InScene](FDynamicPassMeshDrawListContext* DynamicMeshPassContext)
+			{
+				TArray<FPrimitiveSceneProxy*> SceneProxies;
+				for (FPrimitiveSceneProxy* SceneProxy : InScene->GetRenderScene()->PrimitiveSceneProxies)
+				{
+					SceneProxies.Add(SceneProxy);
+				}
+
+				FUVLightCardPassProcessor MeshPassProcessor(View, DynamicMeshPassContext);
+				for (FPrimitiveSceneProxy* SceneProxy : SceneProxies)
+				{
+					if (const FMeshBatch* MeshBatch = SceneProxy->GetPrimitiveSceneInfo()->GetMeshBatch(SceneProxy->GetPrimitiveSceneInfo()->StaticMeshes.Num() - 1))
+					{
+						MeshBatch->MaterialRenderProxy->UpdateUniformExpressionCacheIfNeeded(View->GetFeatureLevel());
+
+						const uint64 BatchElementMask = ~0ull;
+						MeshPassProcessor.AddMeshBatch(*MeshBatch, BatchElementMask, SceneProxy);
+					}
+				}
+			});
+		});
+
+	GraphBuilder.Execute();
+
+	return true;
+}
\ No newline at end of file
diff --git a/Engine/Plugins/Runtime/nDisplay/Source/DisplayClusterShaders/Private/Shaders/DisplayClusterShadersPreprocess_UVLightCards.h b/Engine/Plugins/Runtime/nDisplay/Source/DisplayClusterShaders/Private/Shaders/DisplayClusterShadersPreprocess_UVLightCards.h
new file mode 100644
--- /dev/null
+++ b/Engine/Plugins/Runtime/nDisplay/Source/DisplayClusterShaders/Private/Shaders/DisplayClusterShadersPreprocess_UVLightCards.h
@@ -0,0 +1,16 @@
+// Copyright Epic Games, Inc. All Rights Reserved.
+
+#pragma once
+
+#include "CoreMinimal.h"
+#include "RHI.h"
+#include "RHICommandList.h"
+
+class FSceneInterface;
+class FRenderTarget;
+
+class FDisplayClusterShadersPreprocess_UVLightCards
+{
+public:
+	static bool RenderPreprocess_UVLightCards(FRHICommandListImmediate& RHICmdList, FSceneInterface* InScene, FRenderTarget* InRenderTarget, float ProjectionPlaneSize);
+};
\ No newline at end of file
diff --git a/Engine/Plugins/Runtime/nDisplay/Source/DisplayClusterShaders/Private/Shaders/DisplayClusterShadersWarpblend_ICVFX.cpp b/Engine/Plugins/Runtime/nDisplay/Source/DisplayClusterShaders/Private/Shaders/DisplayClusterShadersWarpblend_ICVFX.cpp
--- a/Engine/Plugins/Runtime/nDisplay/Source/DisplayClusterShaders/Private/Shaders/DisplayClusterShadersWarpblend_ICVFX.cpp
+++ b/Engine/Plugins/Runtime/nDisplay/Source/DisplayClusterShaders/Private/Shaders/DisplayClusterShadersWarpblend_ICVFX.cpp
@@ -73,6 +73,8 @@ namespace IcvfxShaderPermutation
 
 	class FIcvfxShaderMeshWarp         : SHADER_PERMUTATION_BOOL("MESH_WARP");
 
+	class FIcvfxShaderUVLightCards : SHADER_PERMUTATION_BOOL("UV_LIGHT_CARDS");
+
 	using FCommonVSDomain = TShaderPermutationDomain<FIcvfxShaderMeshWarp>;
 
 	using FCommonPSDomain = TShaderPermutationDomain<
@@ -90,7 +92,8 @@ namespace IcvfxShaderPermutation
 
 		FIcvfxShaderAlphaMapBlending,
 		FIcvfxShaderBetaMapBlending,
-		FIcvfxShaderMeshWarp
+		FIcvfxShaderMeshWarp,
+		FIcvfxShaderUVLightCards
 	>;
 	
 	bool ShouldCompileCommonPSPermutation(const FGlobalShaderPermutationParameters& Parameters, const FCommonPSDomain& PermutationVector)
@@ -150,6 +153,11 @@ namespace IcvfxShaderPermutation
 			return false;
 		}
 
+		if (PermutationVector.Get<FIcvfxShaderUVLightCards>() && !PermutationVector.Get<FIcvfxShaderMeshWarp>())
+		{
+			return false;
+		}
+
 		return true;
 	}
 
@@ -177,6 +185,7 @@ BEGIN_SHADER_PARAMETER_STRUCT(FIcvfxPixelShaderParameters, )
 	SHADER_PARAMETER_TEXTURE(Texture2D, InnerCameraTexture)
 	SHADER_PARAMETER_TEXTURE(Texture2D, ChromakeyCameraTexture)
 	SHADER_PARAMETER_TEXTURE(Texture2D, ChromakeyMarkerTexture)
+	SHADER_PARAMETER_TEXTURE(Texture2D, UVLightCardMapTexture)
 
 	SHADER_PARAMETER_SAMPLER(SamplerState, InputSampler)
 	SHADER_PARAMETER_SAMPLER(SamplerState, WarpMapSampler)
@@ -187,6 +196,7 @@ BEGIN_SHADER_PARAMETER_STRUCT(FIcvfxPixelShaderParameters, )
 	SHADER_PARAMETER_SAMPLER(SamplerState, InnerCameraSampler)
 	SHADER_PARAMETER_SAMPLER(SamplerState, ChromakeyCameraSampler)
 	SHADER_PARAMETER_SAMPLER(SamplerState, ChromakeyMarkerSampler)
+	SHADER_PARAMETER_SAMPLER(SamplerState, UVLightCardMapSampler)
 
 	SHADER_PARAMETER(FMatrix44f, ViewportTextureProjectionMatrix)
 	SHADER_PARAMETER(FMatrix44f, OverlayProjectionMatrix)
@@ -633,6 +643,20 @@ public:
 		return false;
 	}
 
+	bool GetUVLightCardParameters(FIcvfxRenderPassData& RenderPassData)
+	{
+		if (ICVFXParameters.UVLightCardMap.IsValid() && RenderPassData.PSPermutationVector.Get<IcvfxShaderPermutation::FIcvfxShaderMeshWarp>())
+		{
+			RenderPassData.PSParameters.UVLightCardMapSampler = TStaticSamplerState<SF_Trilinear, AM_Clamp, AM_Clamp, AM_Clamp>::GetRHI();
+			RenderPassData.PSParameters.UVLightCardMapTexture = ICVFXParameters.UVLightCardMap;
+
+			RenderPassData.PSPermutationVector.Set<IcvfxShaderPermutation::FIcvfxShaderUVLightCards>(true);
+			return true;
+		}
+
+		return false;
+	}
+
 	void InitRenderPass(FIcvfxRenderPassData& RenderPassData)
 	{
 		// Forward input viewport
@@ -658,6 +682,8 @@ public:
 					GetCameraChromakeyMarkerParameters(RenderPassData);
 				}
 			}
+
+			GetUVLightCardParameters(RenderPassData);
 		}
 	}
 
diff --git a/Engine/Plugins/Runtime/nDisplay/Source/DisplayClusterShaders/Public/IDisplayClusterShaders.h b/Engine/Plugins/Runtime/nDisplay/Source/DisplayClusterShaders/Public/IDisplayClusterShaders.h
--- a/Engine/Plugins/Runtime/nDisplay/Source/DisplayClusterShaders/Public/IDisplayClusterShaders.h
+++ b/Engine/Plugins/Runtime/nDisplay/Source/DisplayClusterShaders/Public/IDisplayClusterShaders.h
@@ -17,6 +17,8 @@ struct FDisplayClusterShaderParameters_GenerateMips;
 struct FDisplayClusterShaderParameters_PostprocessBlur;
 struct FDisplayClusterShaderParameters_WarpBlend;
 class IDisplayClusterWarpBlendManager;
+class FSceneInterface;
+class FRenderTarget;
 
 
 class IDisplayClusterShaders : public IModuleInterface
@@ -78,6 +80,18 @@ public:
 	*/
 	virtual bool RenderWarpBlend_ICVFX(FRHICommandListImmediate& RHICmdList, const FDisplayClusterShaderParameters_WarpBlend& InWarpBlendParameters, const FDisplayClusterShaderParameters_ICVFX& InICVFXParameters) const = 0;
 
+	/**
+	* Render UV light cards to texture that can be sampled by viewports to place light cards in UV space
+	*
+	* @param RHICmdList                    - RHI command list
+	* @param InScene                       - The scene the light cards live in
+	* @param InRenderTargetableDestTexture - Destination RTT texture
+	* @param ProjectionPlaneSize           - The size of the plane the UV light cards are projected to
+	*
+	* @return - true if success
+	*/
+	virtual bool RenderPreprocess_UVLightCards(FRHICommandListImmediate& RHICmdList, FSceneInterface* InScene, FRenderTarget* InRenderTarget, float ProjectionPlaneSize) const = 0;
+
 	/**
 	* Render postprocess OutputRemap
 	*
diff --git a/Engine/Plugins/Runtime/nDisplay/Source/DisplayClusterShaders/Public/ShaderParameters/DisplayClusterShaderParameters_ICVFX.h b/Engine/Plugins/Runtime/nDisplay/Source/DisplayClusterShaders/Public/ShaderParameters/DisplayClusterShaderParameters_ICVFX.h
--- a/Engine/Plugins/Runtime/nDisplay/Source/DisplayClusterShaders/Public/ShaderParameters/DisplayClusterShaderParameters_ICVFX.h
+++ b/Engine/Plugins/Runtime/nDisplay/Source/DisplayClusterShaders/Public/ShaderParameters/DisplayClusterShaderParameters_ICVFX.h
@@ -114,6 +114,8 @@ public:
 		Cameras.Empty();
 		Lightcard.Reset();
 		Lightcard_OCIO.Reset();
+
+		UVLightCardMap = nullptr;
 	}
 
 	// Implement copy ref and arrays
@@ -126,6 +128,8 @@ public:
 		Lightcard      = InParameters.Lightcard;
 		Lightcard_OCIO = InParameters.Lightcard_OCIO;
 		LightcardMode = InParameters.LightcardMode;
+
+		UVLightCardMap = InParameters.UVLightCardMap;
 	}
 
 	void CollectRefViewports(TArray<FDisplayClusterShaderParametersICVFX_ViewportResource*>& Dst)
@@ -230,4 +234,7 @@ public:
 	FDisplayClusterShaderParametersICVFX_ViewportResource    Lightcard;
 	FDisplayClusterShaderParametersICVFX_ViewportResource    Lightcard_OCIO;
 	EDisplayClusterShaderParametersICVFX_LightcardRenderMode LightcardMode = EDisplayClusterShaderParametersICVFX_LightcardRenderMode::Under;
+
+	/** Texture containing a UV map of the rendered UV light cards */
+	FTextureRHIRef UVLightCardMap;
 };
diff --git a/Engine/Plugins/VirtualProduction/EpicStageApp/Source/EpicStageApp/Private/StageAppRequest.h b/Engine/Plugins/VirtualProduction/EpicStageApp/Source/EpicStageApp/Private/StageAppRequest.h
--- a/Engine/Plugins/VirtualProduction/EpicStageApp/Source/EpicStageApp/Private/StageAppRequest.h
+++ b/Engine/Plugins/VirtualProduction/EpicStageApp/Source/EpicStageApp/Private/StageAppRequest.h
@@ -7,7 +7,7 @@
 
 #include "StageAppRequest.generated.h"
 
-/** Type of preview render to perform, exposed as a UENUM. Corresponds to FDisplayClusterScenePreviewRenderSettings::ERenderType. */
+/** Type of preview render to perform, exposed as a UENUM. Corresponds to EDisplayClusterMeshProjectionOutput. */
 UENUM()
 enum class ERCWebSocketNDisplayPreviewRenderType : uint8
 {
diff --git a/Engine/Plugins/VirtualProduction/EpicStageApp/Source/EpicStageApp/Private/StageAppRouteHandler.cpp b/Engine/Plugins/VirtualProduction/EpicStageApp/Source/EpicStageApp/Private/StageAppRouteHandler.cpp
--- a/Engine/Plugins/VirtualProduction/EpicStageApp/Source/EpicStageApp/Private/StageAppRouteHandler.cpp
+++ b/Engine/Plugins/VirtualProduction/EpicStageApp/Source/EpicStageApp/Private/StageAppRouteHandler.cpp
@@ -154,16 +154,16 @@ void FStageAppRouteHandler::HandleWebSocketNDisplayPreviewRender(const FRemoteCo
 	}
 
 	// Set up the render
-	FDisplayClusterScenePreviewRenderSettings RenderSettings;
+	FDisplayClusterMeshProjectionRenderSettings RenderSettings;
 
 	switch (Body.RenderType)
 	{
 	case ERCWebSocketNDisplayPreviewRenderType::Color:
-		RenderSettings.RenderType = FDisplayClusterScenePreviewRenderSettings::ERenderType::Color;
+		RenderSettings.RenderType = EDisplayClusterMeshProjectionOutput::Color;
 		break;
 
 	case ERCWebSocketNDisplayPreviewRenderType::Normals:
-		RenderSettings.RenderType = FDisplayClusterScenePreviewRenderSettings::ERenderType::Normals;
+		RenderSettings.RenderType = EDisplayClusterMeshProjectionOutput::Normals;
 		break;
 	}
 
