nDisplay: Fix lightcards not transacting their appearance correctly via MU.

The material in the custom static meshes now have the proper material by default, so Concert not making the MIDs available in the transaction is not affecting this anymore.

diff --git a/Engine/Plugins/Runtime/nDisplay/Source/DisplayCluster/Private/Game/EngineClasses/Scene/DisplayClusterLightCardActor.cpp b/Engine/Plugins/Runtime/nDisplay/Source/DisplayCluster/Private/Game/EngineClasses/Scene/DisplayClusterLightCardActor.cpp
--- a/Engine/Plugins/Runtime/nDisplay/Source/DisplayCluster/Private/Game/EngineClasses/Scene/DisplayClusterLightCardActor.cpp
+++ b/Engine/Plugins/Runtime/nDisplay/Source/DisplayCluster/Private/Game/EngineClasses/Scene/DisplayClusterLightCardActor.cpp
@@ -63,14 +63,12 @@ ADisplayClusterLightCardActor::ADisplayClusterLightCardActor(const FObjectInitia
 	LightCardTransformerComponent->AttachToComponent(MainSpringArmComponent, FAttachmentTransformRules::KeepRelativeTransform);
 
 	static ConstructorHelpers::FObjectFinder<UStaticMesh> PlaneObj(TEXT("/nDisplay/LightCard/SM_LightCardPlane"));
-	static ConstructorHelpers::FObjectFinder<UMaterial> LightCardMatObj(TEXT("/nDisplay/LightCard/M_LightCard"));
 
 	LightCardComponent = CreateDefaultSubobject<UStaticMeshComponent>(TEXT("LightCard"));
 	LightCardComponent->AttachToComponent(LightCardTransformerComponent, FAttachmentTransformRules::KeepRelativeTransform);
 	LightCardComponent->SetCollisionProfileName(TEXT("OverlapAllDynamic"));
 	LightCardComponent->Mobility = EComponentMobility::Movable;
 	LightCardComponent->SetStaticMesh(PlaneObj.Object);
-	LightCardComponent->SetMaterial(0, LightCardMatObj.Object);
 
 	UpdateLightCardTransform();
 }
@@ -85,7 +83,7 @@ void ADisplayClusterLightCardActor::OnConstruction(const FTransform& Transform)
 
 	if (Material && !Material->IsA<UMaterialInstanceDynamic>())
 	{
-		UMaterialInstanceDynamic* LightCardMatInstance = UMaterialInstanceDynamic::Create(Material, this, TEXT("LightCardMID"));
+		UMaterialInstanceDynamic* LightCardMatInstance = UMaterialInstanceDynamic::Create(Material, LightCardComponent, TEXT("LightCardMID"));
 		LightCardComponent->SetMaterial(0, LightCardMatInstance);
 
 		UpdateLightCardMaterialInstance();
