nDisplay: Fixed issue where editor could crash on level change or project close if the UV light card manager was destroyed before render command ran to destroy its resources.

diff --git a/Engine/Plugins/Runtime/nDisplay/Source/DisplayCluster/Private/Render/Viewport/LightCard/DisplayClusterViewportLightCardManager.cpp b/Engine/Plugins/Runtime/nDisplay/Source/DisplayCluster/Private/Render/Viewport/LightCard/DisplayClusterViewportLightCardManager.cpp
--- a/Engine/Plugins/Runtime/nDisplay/Source/DisplayCluster/Private/Render/Viewport/LightCard/DisplayClusterViewportLightCardManager.cpp
+++ b/Engine/Plugins/Runtime/nDisplay/Source/DisplayCluster/Private/Render/Viewport/LightCard/DisplayClusterViewportLightCardManager.cpp
@@ -161,11 +161,12 @@ void FDisplayClusterViewportLightCardManager::ReleaseUVLightCardMap()
 		UVLightCardMap = nullptr;
 
 		ENQUEUE_RENDER_COMMAND(DeleteLightCardResources)(
-			[this](FRHICommandListImmediate& RHICmdList)
+			[InUVLightCardMap=UVLightCardMapProxy](FRHICommandListImmediate& RHICmdList)
 		{
-			// Release the texture's resources and delete the texture object from the rendering thread
-			UVLightCardMapProxy->ReleaseResource();
-			delete UVLightCardMapProxy;
+			// Release the texture's resources and delete the texture object from the rendering thread. Pointer is copied to ensure it is still
+			// valid on the render thread even if the light card manager has been disposed already
+			InUVLightCardMap->ReleaseResource();
+			delete InUVLightCardMap;
 		});
 	}
 }
\ No newline at end of file
