nDisplay MRQ:
 * Added Viewport rendering based on UMoviePipelineDeferredPassBase
   - Multiviewport
   - Tile+ OverlapRatio (Tile>1 only without warp)
   - Camera overscan
   - Postprocess volumes
 * Added nDisplay settings
 * Fixed ND crashes

diff --git a/Engine/Plugins/MovieScene/MovieRenderPipeline/Source/MovieRenderPipelineCore/Private/MoviePipeline.cpp b/Engine/Plugins/MovieScene/MovieRenderPipeline/Source/MovieRenderPipelineCore/Private/MoviePipeline.cpp
--- a/Engine/Plugins/MovieScene/MovieRenderPipeline/Source/MovieRenderPipelineCore/Private/MoviePipeline.cpp
+++ b/Engine/Plugins/MovieScene/MovieRenderPipeline/Source/MovieRenderPipelineCore/Private/MoviePipeline.cpp
@@ -1507,6 +1507,7 @@ void UMoviePipeline::ResolveFilenameFormatArguments(const FString& InFormatStrin
 		Params.ShotOverride = ActiveShotList[InOutputState->ShotIndex];
 		Params.InitializationVersion = ActiveShotList[InOutputState->ShotIndex]->ShotInfo.VersionNumber;
 		Params.CameraIndex = InOutputState->CameraIndex;
+		Params.CameraNameOverride = InOutputState->CameraNameOverride;
 	}
 
 	UMoviePipelineOutputSetting* OutputSetting = GetPipelineMasterConfig()->FindSetting<UMoviePipelineOutputSetting>();
diff --git a/Engine/Plugins/MovieScene/MovieRenderPipeline/Source/MovieRenderPipelineCore/Public/MovieRenderPipelineDataTypes.h b/Engine/Plugins/MovieScene/MovieRenderPipeline/Source/MovieRenderPipelineCore/Public/MovieRenderPipelineDataTypes.h
--- a/Engine/Plugins/MovieScene/MovieRenderPipeline/Source/MovieRenderPipelineCore/Public/MovieRenderPipelineDataTypes.h
+++ b/Engine/Plugins/MovieScene/MovieRenderPipeline/Source/MovieRenderPipelineCore/Public/MovieRenderPipelineDataTypes.h
@@ -651,6 +651,9 @@ public:
 	/** The name of the currently active camera being rendered. May be empty. */
 	FString CameraName;
 
+	/** Name used by the {camera_name} format tag. May be empty */
+	FString CameraNameOverride;
+
 	/** THe name of the currently active shot. May be empty if there is no shot track. */
 	FString ShotName;
 
diff --git a/Engine/Plugins/MovieScene/MovieRenderPipeline/Source/MovieRenderPipelineRenderPasses/Private/MoviePipelineDeferredPasses.cpp b/Engine/Plugins/MovieScene/MovieRenderPipeline/Source/MovieRenderPipelineRenderPasses/Private/MoviePipelineDeferredPasses.cpp
--- a/Engine/Plugins/MovieScene/MovieRenderPipeline/Source/MovieRenderPipelineRenderPasses/Private/MoviePipelineDeferredPasses.cpp
+++ b/Engine/Plugins/MovieScene/MovieRenderPipeline/Source/MovieRenderPipelineRenderPasses/Private/MoviePipelineDeferredPasses.cpp
@@ -62,6 +62,27 @@ UMoviePipelineDeferredPassBase::UMoviePipelineDeferredPassBase()
 	bUse32BitPostProcessMaterials = false;
 }
 
+int32 UMoviePipelineDeferredPassBase::GetNumCamerasToRender() const
+{
+	UMoviePipelineExecutorShot* CurrentShot = GetPipeline()->GetActiveShotList()[GetPipeline()->GetCurrentShotIndex()];
+	UMoviePipelineCameraSetting* CameraSettings = GetPipeline()->FindOrAddSettingForShot<UMoviePipelineCameraSetting>(CurrentShot);
+
+	return CameraSettings->bRenderAllCameras ? CurrentShot->SidecarCameras.Num() : 1;
+}
+
+FString UMoviePipelineDeferredPassBase::GetCameraName(const int32 InCameraIndex) const
+{
+	UMoviePipelineExecutorShot* CurrentShot = GetPipeline()->GetActiveShotList()[GetPipeline()->GetCurrentShotIndex()];
+
+	return CurrentShot->GetCameraName(InCameraIndex);
+}
+
+FString UMoviePipelineDeferredPassBase::GetCameraNameOverride(const int32 InCameraIndex) const
+{
+	// Custom camera name used to override ouput file name param
+	return TEXT("");
+}
+
 void UMoviePipelineDeferredPassBase::MoviePipelineRenderShowFlagOverride(FEngineShowFlags& OutShowFlag)
 {
 	if (bDisableMultisampleEffects)
@@ -129,9 +150,8 @@ void UMoviePipelineDeferredPassBase::SetupImpl(const MoviePipeline::FMoviePipeli
 
 	// Create a view state. Each individual camera, tile, and stencil layer need their own unique state as this includes visual history for anti-aliasing, etc. 
 	UMoviePipelineExecutorShot* CurrentShot = GetPipeline()->GetActiveShotList()[GetPipeline()->GetCurrentShotIndex()];
-	UMoviePipelineCameraSetting* CameraSettings = GetPipeline()->FindOrAddSettingForShot<UMoviePipelineCameraSetting>(CurrentShot);
 	UMoviePipelineHighResSetting* HighResSettings = GetPipeline()->FindOrAddSettingForShot<UMoviePipelineHighResSetting>(CurrentShot);
-	int32 NumCameras = CameraSettings->bRenderAllCameras ? CurrentShot->SidecarCameras.Num() : 1;
+	const int32 NumCameras = GetNumCamerasToRender();
 
 	int32 TotalNumberOfAccumulators = 0;
 	for (int32 CamIndex = 0; CamIndex < NumCameras; CamIndex++)
@@ -383,15 +403,12 @@ void UMoviePipelineDeferredPassBase::GatherOutputPassesImpl(TArray<FMoviePipelin
 	// No super call here because multiple cameras makes this all complicated
 	// Super::GatherOutputPassesImpl(ExpectedRenderPasses);
 
-	UMoviePipelineExecutorShot* CurrentShot = GetPipeline()->GetActiveShotList()[GetPipeline()->GetCurrentShotIndex()];
-	UMoviePipelineCameraSetting* CameraSettings = GetPipeline()->FindOrAddSettingForShot<UMoviePipelineCameraSetting>(CurrentShot);
-	int32 NumCameras = CameraSettings->bRenderAllCameras ? CurrentShot->SidecarCameras.Num() : 1;
-
+	const int32 NumCameras = GetNumCamerasToRender();
 	for (int32 CameraIndex = 0; CameraIndex < NumCameras; CameraIndex++)
 	{
 		FMoviePipelinePassIdentifier PassIdentifierForCurrentCamera;
 		PassIdentifierForCurrentCamera.Name = PassIdentifier.Name;
-		PassIdentifierForCurrentCamera.CameraName = CurrentShot->GetCameraName(CameraIndex);
+		PassIdentifierForCurrentCamera.CameraName = GetCameraName(CameraIndex);
 
 		// Add the default backbuffer
 		if (bRenderMainPass)
@@ -455,15 +472,12 @@ void UMoviePipelineDeferredPassBase::RenderSample_GameThreadImpl(const FMoviePip
 		SurfaceQueue->BlockUntilAnyAvailable();
 	}
 
-	UMoviePipelineExecutorShot* CurrentShot = GetPipeline()->GetActiveShotList()[GetPipeline()->GetCurrentShotIndex()];
-	UMoviePipelineCameraSetting* CameraSettings = GetPipeline()->FindOrAddSettingForShot<UMoviePipelineCameraSetting>(CurrentShot);
-	int32 NumCameras = CameraSettings->bRenderAllCameras ? CurrentShot->SidecarCameras.Num() : 1;
-
+	const int32 NumCameras = GetNumCamerasToRender();
 	for (int32 CameraIndex = 0; CameraIndex < NumCameras; CameraIndex++)
 	{
 		FMoviePipelinePassIdentifier PassIdentifierForCurrentCamera;
 		PassIdentifierForCurrentCamera.Name = PassIdentifier.Name;
-		PassIdentifierForCurrentCamera.CameraName = CurrentShot->GetCameraName(CameraIndex);
+		PassIdentifierForCurrentCamera.CameraName = GetCameraName(CameraIndex);
 
 		// Main Render Pass
 		if (bRenderMainPass)
@@ -471,6 +485,7 @@ void UMoviePipelineDeferredPassBase::RenderSample_GameThreadImpl(const FMoviePip
 			FMoviePipelineRenderPassMetrics InOutSampleState = InSampleState;
 			// InOutSampleState.OutputState.CameraCount = NumCameras;
 			InOutSampleState.OutputState.CameraIndex = CameraIndex;
+			InOutSampleState.OutputState.CameraNameOverride = GetCameraNameOverride(CameraIndex);
 
 			UE::MoviePipeline::FDeferredPassRenderStatePayload Payload;
 			Payload.CameraIndex = CameraIndex;
@@ -527,6 +542,7 @@ void UMoviePipelineDeferredPassBase::RenderSample_GameThreadImpl(const FMoviePip
 		{
 			FMoviePipelineRenderPassMetrics InOutSampleState = InSampleState;
 			InOutSampleState.OutputState.CameraIndex = CameraIndex;
+			InOutSampleState.OutputState.CameraNameOverride = GetCameraNameOverride(CameraIndex);
 
 			struct FStencilValues
 			{
@@ -744,8 +760,7 @@ TFunction<void(TUniquePtr<FImagePixelData>&&)> UMoviePipelineDeferredPassBase::M
 UE::MoviePipeline::FImagePassCameraViewData UMoviePipelineDeferredPassBase::GetCameraInfo(FMoviePipelineRenderPassMetrics& InOutSampleState, IViewCalcPayload* OptPayload) const
 {
 	UMoviePipelineExecutorShot* CurrentShot = GetPipeline()->GetActiveShotList()[GetPipeline()->GetCurrentShotIndex()];
-	UMoviePipelineCameraSetting* CameraSettings = GetPipeline()->FindOrAddSettingForShot<UMoviePipelineCameraSetting>(CurrentShot);
-	int32 NumCameras = CameraSettings->bRenderAllCameras ? CurrentShot->SidecarCameras.Num() : 1;
+	const int32 NumCameras = GetNumCamerasToRender();
 	
 	if (NumCameras == 1)
 	{
@@ -788,8 +803,7 @@ UE::MoviePipeline::FImagePassCameraViewData UMoviePipelineDeferredPassBase::GetC
 void UMoviePipelineDeferredPassBase::BlendPostProcessSettings(FSceneView* InView, FMoviePipelineRenderPassMetrics& InOutSampleState, IViewCalcPayload* OptPayload)
 {
 	UMoviePipelineExecutorShot* CurrentShot = GetPipeline()->GetActiveShotList()[GetPipeline()->GetCurrentShotIndex()];
-	UMoviePipelineCameraSetting* CameraSettings = GetPipeline()->FindOrAddSettingForShot<UMoviePipelineCameraSetting>(CurrentShot);
-	int32 NumCameras = CameraSettings->bRenderAllCameras ? CurrentShot->SidecarCameras.Num() : 1;
+	const int32 NumCameras = GetNumCamerasToRender();
 
 	UCameraComponent* OutCamera = nullptr;
 	FMinimalViewInfo OutViewInfo;
diff --git a/Engine/Plugins/MovieScene/MovieRenderPipeline/Source/MovieRenderPipelineRenderPasses/Private/MoviePipelineImagePassBase.cpp b/Engine/Plugins/MovieScene/MovieRenderPipeline/Source/MovieRenderPipelineRenderPasses/Private/MoviePipelineImagePassBase.cpp
--- a/Engine/Plugins/MovieScene/MovieRenderPipeline/Source/MovieRenderPipelineRenderPasses/Private/MoviePipelineImagePassBase.cpp
+++ b/Engine/Plugins/MovieScene/MovieRenderPipeline/Source/MovieRenderPipelineRenderPasses/Private/MoviePipelineImagePassBase.cpp
@@ -352,9 +352,18 @@ FSceneView* UMoviePipelineImagePassBase::GetSceneViewForSampleState(FSceneViewFa
 	float DofSensorScale = 1.0f;
 
 	// Calculate a Projection Matrix. This code unfortunately ends up similar to, but not quite the same as FMinimalViewInfo::CalculateProjectionMatrixGivenView
+	FMatrix BaseProjMatrix;
+
+	if (CameraInfo.bUseCustomProjectionMatrix)
 	{
-		FMatrix BaseProjMatrix;
+		BaseProjMatrix = CameraInfo.CustomProjectionMatrix;
 
+		// Modify the custom matrix to do an off center projection, with overlap for high-res tiling
+		const bool bOrthographic = false;
+		ModifyProjectionMatrixForTiling(InOutSampleState, bOrthographic, /*InOut*/ BaseProjMatrix, DofSensorScale);
+	}
+	else
+	{
 		if (CameraInfo.ViewInfo.ProjectionMode == ECameraProjectionMode::Orthographic)
 		{
 			const float YScale = 1.0f / CameraInfo.ViewInfo.AspectRatio;
@@ -460,10 +469,9 @@ FSceneView* UMoviePipelineImagePassBase::GetSceneViewForSampleState(FSceneViewFa
 			ModifyProjectionMatrixForTiling(InOutSampleState, bOrthographic, /*InOut*/ BaseProjMatrix, DofSensorScale);
 			// ToDo: Does orthographic support tiling in the same way or do I need to modify the values before creating the ortho view.
 		}
-		
-		// BaseProjMatrix may be perspective or orthographic.
-		ViewInitOptions.ProjectionMatrix = BaseProjMatrix;
 	}
+	// BaseProjMatrix may be perspective or orthographic.
+	ViewInitOptions.ProjectionMatrix = BaseProjMatrix;
 
 	ViewInitOptions.SceneViewStateInterface = GetSceneViewStateInterface(OptPayload);
 	ViewInitOptions.FOV = ViewFOV;
diff --git a/Engine/Plugins/MovieScene/MovieRenderPipeline/Source/MovieRenderPipelineRenderPasses/Public/MoviePipelineDeferredPasses.h b/Engine/Plugins/MovieScene/MovieRenderPipeline/Source/MovieRenderPipelineRenderPasses/Public/MoviePipelineDeferredPasses.h
--- a/Engine/Plugins/MovieScene/MovieRenderPipeline/Source/MovieRenderPipelineRenderPasses/Public/MoviePipelineDeferredPasses.h
+++ b/Engine/Plugins/MovieScene/MovieRenderPipeline/Source/MovieRenderPipelineRenderPasses/Public/MoviePipelineDeferredPasses.h
@@ -64,7 +64,11 @@ protected:
 	// ~FGCObject Interface
 
 	TFunction<void(TUniquePtr<FImagePixelData>&&)> MakeForwardingEndpoint(const FMoviePipelinePassIdentifier InPassIdentifier, const FMoviePipelineRenderPassMetrics& InSampleState);
-	void PostRendererSubmission(const FMoviePipelineRenderPassMetrics& InSampleState, const FMoviePipelinePassIdentifier InPassIdentifier, const int32 InSortingOrder, FCanvas& InCanvas);
+	virtual void PostRendererSubmission(const FMoviePipelineRenderPassMetrics& InSampleState, const FMoviePipelinePassIdentifier InPassIdentifier, const int32 InSortingOrder, FCanvas& InCanvas);
+
+	virtual int32 GetNumCamerasToRender() const;
+	virtual FString GetCameraName(const int32 InCameraIndex) const;
+	virtual FString GetCameraNameOverride(const int32 InCameraIndex) const;
 
 public:
 	/**
diff --git a/Engine/Plugins/MovieScene/MovieRenderPipeline/Source/MovieRenderPipelineRenderPasses/Public/MoviePipelineImagePassBase.h b/Engine/Plugins/MovieScene/MovieRenderPipeline/Source/MovieRenderPipelineRenderPasses/Public/MoviePipelineImagePassBase.h
--- a/Engine/Plugins/MovieScene/MovieRenderPipeline/Source/MovieRenderPipelineRenderPasses/Public/MoviePipelineImagePassBase.h
+++ b/Engine/Plugins/MovieScene/MovieRenderPipeline/Source/MovieRenderPipelineRenderPasses/Public/MoviePipelineImagePassBase.h
@@ -32,11 +32,15 @@ namespace UE
 		{
 			FImagePassCameraViewData()
 				: ViewActor(nullptr)
+				, bUseCustomProjectionMatrix(false)
 			{}
 			
 			FMinimalViewInfo ViewInfo;
 			TMap<FString, FString> FileMetadata;
 			AActor* ViewActor;
+
+			bool bUseCustomProjectionMatrix;
+			FMatrix CustomProjectionMatrix;
 		};
 	}
 }
@@ -107,7 +111,7 @@ protected:
 	// ~FGCObject Interface
 
 	FVector4 CalculatePrinciplePointOffsetForTiling(const FMoviePipelineRenderPassMetrics& InSampleState) const;
-	void ModifyProjectionMatrixForTiling(const FMoviePipelineRenderPassMetrics& InSampleState, const bool bInOrthographic, FMatrix& InOutProjectionMatrix, float& OutDoFSensorScale) const;
+	virtual void ModifyProjectionMatrixForTiling(const FMoviePipelineRenderPassMetrics& InSampleState, const bool bInOrthographic, FMatrix& InOutProjectionMatrix, float& OutDoFSensorScale) const;
 
 protected:
 	virtual void GetViewShowFlags(FEngineShowFlags& OutShowFlag, EViewModeIndex& OutViewModeIndex) const;
diff --git a/Engine/Plugins/Runtime/nDisplay/Source/DisplayCluster/Private/Render/Viewport/Configuration/DisplayClusterViewportConfiguration.cpp b/Engine/Plugins/Runtime/nDisplay/Source/DisplayCluster/Private/Render/Viewport/Configuration/DisplayClusterViewportConfiguration.cpp
--- a/Engine/Plugins/Runtime/nDisplay/Source/DisplayCluster/Private/Render/Viewport/Configuration/DisplayClusterViewportConfiguration.cpp
+++ b/Engine/Plugins/Runtime/nDisplay/Source/DisplayCluster/Private/Render/Viewport/Configuration/DisplayClusterViewportConfiguration.cpp
@@ -65,7 +65,7 @@ ADisplayClusterRootActor* FDisplayClusterViewportConfiguration::GetRootActor() c
 	return nullptr;
 }
 
-bool FDisplayClusterViewportConfiguration::ImplUpdateConfiguration(EDisplayClusterRenderFrameMode InRenderMode, const FString& InClusterNodeId, const FDisplayClusterPreviewSettings* InPreviewSettings)
+bool FDisplayClusterViewportConfiguration::ImplUpdateConfiguration(EDisplayClusterRenderFrameMode InRenderMode, const FString& InClusterNodeId, const FDisplayClusterPreviewSettings* InPreviewSettings, const TArray<FString>* InViewportNames)
 {
 	check(IsInGameThread());
 
@@ -120,7 +120,15 @@ bool FDisplayClusterViewportConfiguration::ImplUpdateConfiguration(EDisplayClust
 				RenderFrameSettings.bIsPreviewRendering = false;
 			}
 
-			ConfigurationBase.Update(InClusterNodeId);
+			if (InViewportNames)
+			{
+				ConfigurationBase.Update(*InViewportNames);
+			}
+			else
+			{
+				ConfigurationBase.Update(InClusterNodeId);
+			}
+
 			ConfigurationICVFX.Update();
 			ConfigurationProjectionPolicy.Update();
 			ConfigurationICVFX.PostUpdate();
@@ -134,7 +142,11 @@ bool FDisplayClusterViewportConfiguration::ImplUpdateConfiguration(EDisplayClust
 
 			ImplUpdateConfigurationVisibility(*RootActor, *ConfigurationData);
 
-			ConfigurationBase.UpdateClusterNodePostProcess(InClusterNodeId, RenderFrameSettings);
+			if (!InClusterNodeId.IsEmpty())
+			{
+				// support postprocess only for per-node render
+				ConfigurationBase.UpdateClusterNodePostProcess(InClusterNodeId, RenderFrameSettings);
+			}
 
 			ImplPostUpdateRenderFrameConfiguration();
 
@@ -147,7 +159,12 @@ bool FDisplayClusterViewportConfiguration::ImplUpdateConfiguration(EDisplayClust
 
 bool FDisplayClusterViewportConfiguration::UpdateConfiguration(EDisplayClusterRenderFrameMode InRenderMode, const FString& InClusterNodeId)
 {
-	return ImplUpdateConfiguration(InRenderMode, InClusterNodeId, nullptr);
+	return ImplUpdateConfiguration(InRenderMode, InClusterNodeId, nullptr, nullptr);
+}
+
+bool FDisplayClusterViewportConfiguration::UpdateCustomConfiguration(EDisplayClusterRenderFrameMode InRenderMode, const TArray<FString>& InViewportNames)
+{
+	return ImplUpdateConfiguration(InRenderMode, TEXT(""), nullptr, &InViewportNames);
 }
 
 #if WITH_EDITOR
@@ -168,7 +185,7 @@ bool FDisplayClusterViewportConfiguration::UpdatePreviewConfiguration(EDisplayCl
 				ConfigurationData->Cluster->GetNodeIds(ClusterNodesIDs);
 				for (const FString& ClusterNodeIdIt : ClusterNodesIDs)
 				{
-					ImplUpdateConfiguration(InRenderMode, ClusterNodeIdIt, &InPreviewSettings);
+					ImplUpdateConfiguration(InRenderMode, ClusterNodeIdIt, &InPreviewSettings, nullptr);
 				}
 
 				// all cluster nodes viewports updated
@@ -179,7 +196,7 @@ bool FDisplayClusterViewportConfiguration::UpdatePreviewConfiguration(EDisplayCl
 		return false;
 	}
 
-	return ImplUpdateConfiguration(InRenderMode, InClusterNodeId, &InPreviewSettings);
+	return ImplUpdateConfiguration(InRenderMode, InClusterNodeId, &InPreviewSettings, nullptr);
 }
 #endif
 
diff --git a/Engine/Plugins/Runtime/nDisplay/Source/DisplayCluster/Private/Render/Viewport/Configuration/DisplayClusterViewportConfiguration.h b/Engine/Plugins/Runtime/nDisplay/Source/DisplayCluster/Private/Render/Viewport/Configuration/DisplayClusterViewportConfiguration.h
--- a/Engine/Plugins/Runtime/nDisplay/Source/DisplayCluster/Private/Render/Viewport/Configuration/DisplayClusterViewportConfiguration.h
+++ b/Engine/Plugins/Runtime/nDisplay/Source/DisplayCluster/Private/Render/Viewport/Configuration/DisplayClusterViewportConfiguration.h
@@ -37,13 +37,14 @@ public:
 	}
 
 	bool UpdateConfiguration(EDisplayClusterRenderFrameMode InRenderMode, const FString& InClusterNodeId);
+	bool UpdateCustomConfiguration(EDisplayClusterRenderFrameMode InRenderMode, const TArray<FString>& InViewportNames);
 
 #if WITH_EDITOR
 	bool UpdatePreviewConfiguration(EDisplayClusterRenderFrameMode InRenderMode, const FString& ClusterNodeId, const FDisplayClusterPreviewSettings& InPreviewSettings);
 #endif
 
 private:
-	bool ImplUpdateConfiguration(EDisplayClusterRenderFrameMode InRenderMode, const FString& InClusterNodeId, const FDisplayClusterPreviewSettings* InPreviewSettings);
+	bool ImplUpdateConfiguration(EDisplayClusterRenderFrameMode InRenderMode, const FString& InClusterNodeId, const FDisplayClusterPreviewSettings* InPreviewSettings, const TArray<FString>* InViewportNames);
 
 	void ImplUpdateRenderFrameConfiguration(const FDisplayClusterConfigurationRenderFrame& InRenderFrameConfiguration);
 	void ImplPostUpdateRenderFrameConfiguration();
diff --git a/Engine/Plugins/Runtime/nDisplay/Source/DisplayCluster/Private/Render/Viewport/Configuration/DisplayClusterViewportConfigurationBase.cpp b/Engine/Plugins/Runtime/nDisplay/Source/DisplayCluster/Private/Render/Viewport/Configuration/DisplayClusterViewportConfigurationBase.cpp
--- a/Engine/Plugins/Runtime/nDisplay/Source/DisplayCluster/Private/Render/Viewport/Configuration/DisplayClusterViewportConfigurationBase.cpp
+++ b/Engine/Plugins/Runtime/nDisplay/Source/DisplayCluster/Private/Render/Viewport/Configuration/DisplayClusterViewportConfigurationBase.cpp
@@ -89,6 +89,54 @@ void FDisplayClusterViewportConfigurationBase::Update(const FString& ClusterNode
 	}
 }
 
+void FDisplayClusterViewportConfigurationBase::Update(const TArray<FString>& InViewportNames)
+{
+	// Collect unused viewports and delete
+	{
+		TArray<FDisplayClusterViewport*> UnusedViewports;
+		for (FDisplayClusterViewport* It : ViewportManager.ImplGetViewports())
+		{
+			// ignore ICVFX internal resources
+			if ((It->GetRenderSettingsICVFX().RuntimeFlags & ViewportRuntime_InternalResource) == 0)
+			{
+				if (InViewportNames.Find(It->GetId()) == INDEX_NONE)
+				{
+					UnusedViewports.Add(It);
+				}
+			}
+		}
+
+		// Delete unused viewports
+		for (FDisplayClusterViewport* DeleteIt : UnusedViewports)
+		{
+			ViewportManager.ImplDeleteViewport(DeleteIt);
+		}
+	}
+
+	// Update and Create new viewports
+	for (const TPair<FString, UDisplayClusterConfigurationClusterNode*>& ClusterNodeConfigurationIt : ConfigurationData.Cluster->Nodes)
+	{
+		if (const UDisplayClusterConfigurationClusterNode* ClusterNodeConfiguration = ClusterNodeConfigurationIt.Value)
+		{
+			for (const TPair<FString, UDisplayClusterConfigurationViewport*>& ViewportIt : ClusterNodeConfiguration->Viewports)
+			{
+				if (ViewportIt.Key.Len() && ViewportIt.Value && InViewportNames.Find(ViewportIt.Key) != INDEX_NONE)
+				{
+					FDisplayClusterViewport* ExistViewport = ViewportManager.ImplFindViewport(ViewportIt.Key);
+					if (ExistViewport)
+					{
+						FDisplayClusterViewportConfigurationBase::UpdateViewportConfiguration(ViewportManager, RootActor, ExistViewport, ViewportIt.Value);
+					}
+					else
+					{
+						ViewportManager.CreateViewport(ViewportIt.Key, ViewportIt.Value);
+					}
+				}
+			}
+		}
+	}
+}
+
 void FDisplayClusterViewportConfigurationBase::UpdateClusterNodePostProcess(const FString& InClusterNodeId, const FDisplayClusterRenderFrameSettings& InRenderFrameSettings)
 {
 	const UDisplayClusterConfigurationClusterNode* ClusterNode = ConfigurationData.Cluster->GetNode(InClusterNodeId);
diff --git a/Engine/Plugins/Runtime/nDisplay/Source/DisplayCluster/Private/Render/Viewport/Configuration/DisplayClusterViewportConfigurationBase.h b/Engine/Plugins/Runtime/nDisplay/Source/DisplayCluster/Private/Render/Viewport/Configuration/DisplayClusterViewportConfigurationBase.h
--- a/Engine/Plugins/Runtime/nDisplay/Source/DisplayCluster/Private/Render/Viewport/Configuration/DisplayClusterViewportConfigurationBase.h
+++ b/Engine/Plugins/Runtime/nDisplay/Source/DisplayCluster/Private/Render/Viewport/Configuration/DisplayClusterViewportConfigurationBase.h
@@ -25,6 +25,7 @@ public:
 
 public:
 	void Update(const FString& ClusterNodeId);
+	void Update(const TArray<FString>& InViewportNames);
 	void UpdateClusterNodePostProcess(const FString& ClusterNodeId, const FDisplayClusterRenderFrameSettings& InRenderFrameSettings);
 
 public:
diff --git a/Engine/Plugins/Runtime/nDisplay/Source/DisplayCluster/Private/Render/Viewport/Containers/ImplDisplayClusterViewport_CustomFrustum.cpp b/Engine/Plugins/Runtime/nDisplay/Source/DisplayCluster/Private/Render/Viewport/Containers/ImplDisplayClusterViewport_CustomFrustum.cpp
--- a/Engine/Plugins/Runtime/nDisplay/Source/DisplayCluster/Private/Render/Viewport/Containers/ImplDisplayClusterViewport_CustomFrustum.cpp
+++ b/Engine/Plugins/Runtime/nDisplay/Source/DisplayCluster/Private/Render/Viewport/Containers/ImplDisplayClusterViewport_CustomFrustum.cpp
@@ -39,6 +39,12 @@ void FImplDisplayClusterViewport_CustomFrustum::Update(FDisplayClusterViewport&
 		return;
 	}
 
+	// Disable CustomFrustum feature from viewport settings
+	if (Viewport.GetRenderSettings().bDisableCustomFrustumFeature)
+	{
+		return;
+	}
+
 	RuntimeSettings.bIsEnabled = true;
 	const FIntPoint Size = InOutRenderTargetRect.Size();
 
diff --git a/Engine/Plugins/Runtime/nDisplay/Source/DisplayCluster/Private/Render/Viewport/Containers/ImplDisplayClusterViewport_Overscan.cpp b/Engine/Plugins/Runtime/nDisplay/Source/DisplayCluster/Private/Render/Viewport/Containers/ImplDisplayClusterViewport_Overscan.cpp
--- a/Engine/Plugins/Runtime/nDisplay/Source/DisplayCluster/Private/Render/Viewport/Containers/ImplDisplayClusterViewport_Overscan.cpp
+++ b/Engine/Plugins/Runtime/nDisplay/Source/DisplayCluster/Private/Render/Viewport/Containers/ImplDisplayClusterViewport_Overscan.cpp
@@ -33,12 +33,18 @@ bool FImplDisplayClusterViewport_Overscan::UpdateProjectionAngles(float& InOutLe
 
 void FImplDisplayClusterViewport_Overscan::Update(FDisplayClusterViewport& Viewport, FIntRect& InOutRenderTargetRect)
 {
-	// Disable overscane feature from console
+	// Disable viewport overscan feature from console
 	if (CVarDisplayClusterRenderOverscanEnable.GetValueOnGameThread() == 0)
 	{
 		return;
 	}
 
+	// Disable viewport overscan feature from settings
+	if (Viewport.GetRenderSettings().bDisableFrustumOverscanFeature)
+	{
+		return;
+	}
+
 	const FIntPoint Size = InOutRenderTargetRect.Size();
 
 	switch (OverscanSettings.Mode)
diff --git a/Engine/Plugins/Runtime/nDisplay/Source/DisplayCluster/Private/Render/Viewport/DisplayClusterViewport.h b/Engine/Plugins/Runtime/nDisplay/Source/DisplayCluster/Private/Render/Viewport/DisplayClusterViewport.h
--- a/Engine/Plugins/Runtime/nDisplay/Source/DisplayCluster/Private/Render/Viewport/DisplayClusterViewport.h
+++ b/Engine/Plugins/Runtime/nDisplay/Source/DisplayCluster/Private/Render/Viewport/DisplayClusterViewport.h
@@ -128,8 +128,6 @@ public:
 	/// ~IDisplayClusterViewport
 	//////////////////////////////////////////////////////
 
-	FMatrix ImplCreateProjectionMatrix(float InLeft, float InRight, float InTop, float InBottom, float ZNear, float ZFar) const;
-
 #if WITH_EDITOR
 	FSceneView* ImplCalcScenePreview(class FSceneViewFamilyContext& InOutViewFamily, uint32 ContextNum);
 	bool    ImplPreview_CalculateStereoViewOffset(const uint32 InContextNum, FRotator& ViewRotation, const float WorldToMeters, FVector& ViewLocation);
diff --git a/Engine/Plugins/Runtime/nDisplay/Source/DisplayCluster/Private/Render/Viewport/DisplayClusterViewportManager.cpp b/Engine/Plugins/Runtime/nDisplay/Source/DisplayCluster/Private/Render/Viewport/DisplayClusterViewportManager.cpp
--- a/Engine/Plugins/Runtime/nDisplay/Source/DisplayCluster/Private/Render/Viewport/DisplayClusterViewportManager.cpp
+++ b/Engine/Plugins/Runtime/nDisplay/Source/DisplayCluster/Private/Render/Viewport/DisplayClusterViewportManager.cpp
@@ -168,6 +168,11 @@ void FDisplayClusterViewportManager::EndScene()
 		PostProcessManager->HandleEndScene();
 	}
 
+	if (LightCardManager.IsValid())
+	{
+		LightCardManager->HandleEndScene();
+	}
+
 	CurrentWorldRef.Reset();
 }
 
@@ -177,7 +182,8 @@ void FDisplayClusterViewportManager::ResetScene()
 
 	if (LightCardManager.IsValid())
 	{
-		LightCardManager->ResetScene();
+		LightCardManager->HandleEndScene();
+		LightCardManager->HandleStartScene();
 	}
 
 	for (FDisplayClusterViewport* Viewport : Viewports)
@@ -307,13 +313,33 @@ bool FDisplayClusterViewportManager::ShouldUseAdditionalFrameTargetableResource(
 	return false;
 }
 
+bool FDisplayClusterViewportManager::UpdateCustomConfiguration(EDisplayClusterRenderFrameMode InRenderMode, const TArray<FString>& InViewportNames, class ADisplayClusterRootActor* InRootActorPtr)
+{
+	ImplUpdateClusterNodeViewports(InRenderMode, TEXT(""));
+
+	if (InRootActorPtr)
+	{
+		const bool bIsRootActorChanged = Configuration->SetRootActor(InRootActorPtr);
+
+		// When the root actor changes, we have to ResetScene() to reinitialize the internal references of the projection policy.
+		if (bIsRootActorChanged)
+		{
+			ResetScene();
+		}
+
+		return Configuration->UpdateCustomConfiguration(InRenderMode, InViewportNames);
+	}
+
+	return false;
+}
+
 bool FDisplayClusterViewportManager::UpdateConfiguration(EDisplayClusterRenderFrameMode InRenderMode, const FString& InClusterNodeId, ADisplayClusterRootActor* InRootActorPtr, const FDisplayClusterPreviewSettings* InPreviewSettings)
 {
 	ImplUpdateClusterNodeViewports(InRenderMode, InClusterNodeId);
 
 	if (InRootActorPtr)
 	{
-		bool bIsRootActorChanged = Configuration->SetRootActor(InRootActorPtr);
+		const bool bIsRootActorChanged = Configuration->SetRootActor(InRootActorPtr);
 
 		// When the root actor changes, we have to ResetScene() to reinitialize the internal references of the projection policy.
 		if (bIsRootActorChanged)
@@ -498,11 +524,6 @@ bool FDisplayClusterViewportManager::BeginNewFrame(FViewport* InViewport, UWorld
 		PostProcessManager->HandleBeginNewFrame(OutRenderFrame);
 	}
 
-	if (LightCardManager.IsValid())
-	{
-		LightCardManager->HandleBeginNewFrame();
-	}
-
 	return true;
 }
 
@@ -712,9 +733,12 @@ void FDisplayClusterViewportManager::ConfigureViewFamily(const FDisplayClusterRe
 
 void FDisplayClusterViewportManager::RenderFrame(FViewport* InViewport)
 {
-	PreRenderFrame();
+	if (LightCardManager.IsValid())
+	{
+		LightCardManager->RenderFrame();
+	}
+
 	ViewportManagerProxy->ImplRenderFrame(InViewport);
-	PostRenderFrame();
 }
 
 bool FDisplayClusterViewportManager::CreateViewport(const FString& InViewportId, const class UDisplayClusterConfigurationViewport* ConfigurationViewport)
@@ -906,6 +930,11 @@ TSharedPtr<IDisplayClusterProjectionPolicy, ESPMode::ThreadSafe> FDisplayCluster
 	return nullptr;
 }
 
+TSharedPtr<IDisplayClusterViewportLightCardManager, ESPMode::ThreadSafe> FDisplayClusterViewportManager::GetLightCardManager() const 
+{
+	return LightCardManager;
+}
+
 void FDisplayClusterViewportManager::MarkComponentGeometryDirty(const FName InComponentName)
 {
 	check(IsInGameThread());
@@ -949,19 +978,3 @@ void FDisplayClusterViewportManager::AddReferencedObjects(FReferenceCollector& C
 		}
 	}
 }
-
-void FDisplayClusterViewportManager::PreRenderFrame()
-{
-	if (LightCardManager.IsValid())
-	{
-		LightCardManager->PreRenderFrame();
-	}
-}
-
-void FDisplayClusterViewportManager::PostRenderFrame()
-{
-	if (LightCardManager.IsValid())
-	{
-		LightCardManager->PostRenderFrame();
-	}
-}
\ No newline at end of file
diff --git a/Engine/Plugins/Runtime/nDisplay/Source/DisplayCluster/Private/Render/Viewport/DisplayClusterViewportManager.h b/Engine/Plugins/Runtime/nDisplay/Source/DisplayCluster/Private/Render/Viewport/DisplayClusterViewportManager.h
--- a/Engine/Plugins/Runtime/nDisplay/Source/DisplayCluster/Private/Render/Viewport/DisplayClusterViewportManager.h
+++ b/Engine/Plugins/Runtime/nDisplay/Source/DisplayCluster/Private/Render/Viewport/DisplayClusterViewportManager.h
@@ -19,7 +19,7 @@ class FDisplayClusterViewportPostProcessManager;
 class FDisplayClusterRenderFrameManager;
 class FDisplayClusterRenderFrame; 
 class FDisplayClusterViewportManagerProxy;
-
+class FDisplayClusterViewportLightCardManager;
 class IDisplayClusterViewportLightCardManager;
 class IDisplayClusterProjectionPolicy;
 
@@ -51,6 +51,7 @@ public:
 	virtual bool IsSceneOpened() const override;
 
 	virtual bool UpdateConfiguration(EDisplayClusterRenderFrameMode InRenderMode, const FString& InClusterNodeId, class ADisplayClusterRootActor* InRootActorPtr, const FDisplayClusterPreviewSettings* InPreviewSettings = nullptr) override;
+	virtual bool UpdateCustomConfiguration(EDisplayClusterRenderFrameMode InRenderMode, const TArray<FString>& InViewportNames, class ADisplayClusterRootActor* InRootActorPtr) override;
 
 	virtual bool BeginNewFrame(FViewport* InViewport, UWorld* InWorld, FDisplayClusterRenderFrame& OutRenderFrame) override;
 	virtual void FinalizeNewFrame() override;
@@ -91,7 +92,7 @@ public:
 		return TArrayView<IDisplayClusterViewport*>((IDisplayClusterViewport**)(ClusterNodeViewports.GetData()), ClusterNodeViewports.Num());
 	}
 
-	virtual TSharedPtr<IDisplayClusterViewportLightCardManager, ESPMode::ThreadSafe> GetLightCardManager() const override { return LightCardManager; }
+	virtual TSharedPtr<IDisplayClusterViewportLightCardManager, ESPMode::ThreadSafe> GetLightCardManager() const override;
 
 	virtual void MarkComponentGeometryDirty(const FName InComponentName = NAME_None) override;
 
@@ -140,19 +141,13 @@ private:
 	void HandleViewportRTTChanges(const TArray<FDisplayClusterViewport_Context>& PrevContexts, const TArray<FDisplayClusterViewport_Context>& Contexts);
 	void ImplUpdateClusterNodeViewports(const EDisplayClusterRenderFrameMode InRenderMode, const FString& InClusterNodeId);
 
-	/** Called on the game thread just before a frame is rendered */
-	void PreRenderFrame();
-
-	/** Called on the game thread just after a frame is rendered */
-	void PostRenderFrame();
-
 protected:
 	friend FDisplayClusterViewportManagerProxy;
 	friend FDisplayClusterViewportConfiguration;
 
 	TSharedPtr<FDisplayClusterRenderTargetManager, ESPMode::ThreadSafe>        RenderTargetManager;
 	TSharedPtr<FDisplayClusterViewportPostProcessManager, ESPMode::ThreadSafe> PostProcessManager;
-	TSharedPtr<IDisplayClusterViewportLightCardManager, ESPMode::ThreadSafe>   LightCardManager;
+	TSharedPtr<FDisplayClusterViewportLightCardManager, ESPMode::ThreadSafe>   LightCardManager;
 
 public:
 	TUniquePtr<FDisplayClusterViewportConfiguration> Configuration;
diff --git a/Engine/Plugins/Runtime/nDisplay/Source/DisplayCluster/Private/Render/Viewport/DisplayClusterViewportManagerProxy.cpp b/Engine/Plugins/Runtime/nDisplay/Source/DisplayCluster/Private/Render/Viewport/DisplayClusterViewportManagerProxy.cpp
--- a/Engine/Plugins/Runtime/nDisplay/Source/DisplayCluster/Private/Render/Viewport/DisplayClusterViewportManagerProxy.cpp
+++ b/Engine/Plugins/Runtime/nDisplay/Source/DisplayCluster/Private/Render/Viewport/DisplayClusterViewportManagerProxy.cpp
@@ -12,8 +12,10 @@
 #include "Render/Projection/IDisplayClusterProjectionPolicyFactory.h"
 #include "Render/Projection/IDisplayClusterProjectionPolicy.h"
 
-#include "Render/Viewport/DisplayClusterViewportProxy.h"
 #include "Render/Viewport/IDisplayClusterViewportLightCardManager.h"
+#include "Render/Viewport/LightCard/DisplayClusterViewportLightCardManager.h"
+
+#include "Render/Viewport/DisplayClusterViewportProxy.h"
 #include "Render/Viewport/RenderTarget/DisplayClusterRenderTargetManager.h"
 #include "Render/Viewport/RenderTarget/DisplayClusterRenderTargetResource.h"
 #include "Render/Viewport/Postprocess/DisplayClusterViewportPostProcessManager.h"
@@ -22,8 +24,6 @@
 #include "Render/Viewport/DisplayClusterViewportStrings.h"
 #include "RHIContext.h"
 
-
-
 // Enable/disable warp&blend
 static TAutoConsoleVariable<int32> CVarWarpBlendEnabled(
 	TEXT("nDisplay.render.WarpBlendEnabled"),
@@ -81,6 +81,11 @@ FDisplayClusterViewportManagerProxy::FDisplayClusterViewportManagerProxy()
 { }
 
 FDisplayClusterViewportManagerProxy::~FDisplayClusterViewportManagerProxy()
+{
+	Release();
+}
+
+void FDisplayClusterViewportManagerProxy::Release()
 {
 	// Delete viewport proxy objects
 	for (FDisplayClusterViewportProxy* ViewportProxy : ViewportProxies)
@@ -90,9 +95,10 @@ FDisplayClusterViewportManagerProxy::~FDisplayClusterViewportManagerProxy()
 			delete ViewportProxy;
 		}
 	}
-
 	ViewportProxies.Empty();
-	ImplUpdateClusterNodeViewportProxies();
+	ClusterNodeViewportProxies.Empty();
+
+	RenderTargetManager.Reset();
 }
 
 void FDisplayClusterViewportManagerProxy::Initialize(FDisplayClusterViewportManager& InViewportManager)
@@ -137,9 +143,20 @@ void FDisplayClusterViewportManagerProxy::ImplSafeRelease()
 {
 	check(IsInGameThread());
 
-	// Light card manager contains a preview world, which needs to be cleaned up on game thread. Manually reset the pointer
-	// here to ensure that it isn't cleaned up on render thread when the proxy deletes itself.
-	LightCardManager.Reset();
+	if (RenderTargetManager.IsValid())
+	{
+		RenderTargetManager->Release();
+	}
+	
+	if (PostProcessManager.IsValid())
+	{
+		PostProcessManager->Release();
+	}
+
+	if (LightCardManager.IsValid())
+	{
+		LightCardManager->Release();
+	}
 
 	// Remove viewport manager proxy on render_thread
 	ENQUEUE_RENDER_COMMAND(DeleteDisplayClusterViewportManagerProxy)(
@@ -241,11 +258,6 @@ void FDisplayClusterViewportManagerProxy::ImplRenderFrame(FViewport* InViewport)
 			ViewportManagerProxy->PostProcessManager->HandleRenderFrameSetup_RenderThread(RHICmdList, ViewportManagerProxy);
 		}
 
-		if (ViewportManagerProxy->LightCardManager.IsValid())
-		{
-			ViewportManagerProxy->LightCardManager->RenderLightCardMap_RenderThread(RHICmdList);
-		}
-
 		bool bWarpBlendEnabled = ViewportManagerProxy->RenderFrameSettings.bAllowWarpBlend && CVarWarpBlendEnabled.GetValueOnRenderThread() != 0;
 
 		// mGPU not used for in-editor rendering
@@ -626,6 +638,12 @@ bool FDisplayClusterViewportManagerProxy::ResolveFrameTargetToBackBuffer_RenderT
 	return false;
 }
 
+TSharedPtr<IDisplayClusterViewportLightCardManager, ESPMode::ThreadSafe> FDisplayClusterViewportManagerProxy::GetLightCardManager_RenderThread() const
+{
+	return LightCardManager;
+}
+
+
 FDisplayClusterViewportProxy* FDisplayClusterViewportManagerProxy::ImplFindViewport_RenderThread(const FString& ViewportId) const
 {
 	check(IsInRenderingThread());
diff --git a/Engine/Plugins/Runtime/nDisplay/Source/DisplayCluster/Private/Render/Viewport/DisplayClusterViewportManagerProxy.h b/Engine/Plugins/Runtime/nDisplay/Source/DisplayCluster/Private/Render/Viewport/DisplayClusterViewportManagerProxy.h
--- a/Engine/Plugins/Runtime/nDisplay/Source/DisplayCluster/Private/Render/Viewport/DisplayClusterViewportManagerProxy.h
+++ b/Engine/Plugins/Runtime/nDisplay/Source/DisplayCluster/Private/Render/Viewport/DisplayClusterViewportManagerProxy.h
@@ -5,7 +5,6 @@
 #include "CoreMinimal.h"
 
 #include "Render/Viewport/IDisplayClusterViewportManagerProxy.h"
-
 #include "Render/Viewport/DisplayClusterViewportProxy.h"
 #include "Render/Viewport/Containers/DisplayClusterViewport_Enums.h"
 #include "Render/Viewport/RenderFrame/DisplayClusterRenderFrameSettings.h"
@@ -14,6 +13,7 @@ class FDisplayClusterRenderTargetManager;
 class FDisplayClusterViewportPostProcessManager;
 class FDisplayClusterViewportManager;
 class FDisplayClusterViewportResource;
+class FDisplayClusterViewportLightCardManager;
 class IDisplayClusterViewportLightCardManager;
 class IDisplayClusterProjectionPolicy;
 
@@ -27,6 +27,8 @@ public:
 	FDisplayClusterViewportManagerProxy();
 	virtual ~FDisplayClusterViewportManagerProxy();
 
+	void Release();
+
 public:
 	void DoCrossGPUTransfers_RenderThread(FRHICommandListImmediate& RHICmdList) const;
 	void UpdateDeferredResources_RenderThread(FRHICommandListImmediate& RHICmdList) const;
@@ -48,7 +50,7 @@ public:
 	virtual bool GetFrameTargets_RenderThread(TArray<FRHITexture2D*>& OutFrameResources, TArray<FIntPoint>& OutTargetOffsets, TArray<FRHITexture2D*>* OutAdditionalFrameResources = nullptr) const override;
 	virtual bool ResolveFrameTargetToBackBuffer_RenderThread(FRHICommandListImmediate& RHICmdList, const uint32 InContextNum, const int32 DestArrayIndex, FRHITexture2D* DestTexture, FVector2D WindowSize) const override;
 
-	virtual TSharedPtr<IDisplayClusterViewportLightCardManager, ESPMode::ThreadSafe> GetLightCardManager_RenderThread() const override { return LightCardManager; }
+	virtual TSharedPtr<IDisplayClusterViewportLightCardManager, ESPMode::ThreadSafe> GetLightCardManager_RenderThread() const override;
 
 	// internal use only
 	void DeleteResource_RenderThread(FDisplayClusterViewportResource* InDeletedResourcePtr);
@@ -85,7 +87,7 @@ private:
 private:
 	TSharedPtr<FDisplayClusterRenderTargetManager, ESPMode::ThreadSafe>        RenderTargetManager;
 	TSharedPtr<FDisplayClusterViewportPostProcessManager, ESPMode::ThreadSafe> PostProcessManager;
-	TSharedPtr<IDisplayClusterViewportLightCardManager, ESPMode::ThreadSafe>   LightCardManager;
+	TSharedPtr<FDisplayClusterViewportLightCardManager, ESPMode::ThreadSafe>   LightCardManager;
 
 	FDisplayClusterRenderFrameSettings RenderFrameSettings;
 
diff --git a/Engine/Plugins/Runtime/nDisplay/Source/DisplayCluster/Private/Render/Viewport/DisplayClusterViewport_Math.cpp b/Engine/Plugins/Runtime/nDisplay/Source/DisplayCluster/Private/Render/Viewport/DisplayClusterViewport_Math.cpp
--- a/Engine/Plugins/Runtime/nDisplay/Source/DisplayCluster/Private/Render/Viewport/DisplayClusterViewport_Math.cpp
+++ b/Engine/Plugins/Runtime/nDisplay/Source/DisplayCluster/Private/Render/Viewport/DisplayClusterViewport_Math.cpp
@@ -138,7 +138,7 @@ void FDisplayClusterViewport::CalculateProjectionMatrix(const uint32 InContextNu
 	GetNonZeroFrustumRange(l, r, n);
 	GetNonZeroFrustumRange(b, t, n);
 
-	Contexts[InContextNum].ProjectionMatrix = ImplCreateProjectionMatrix(l, r, t, b, n, f);
+	Contexts[InContextNum].ProjectionMatrix = IDisplayClusterViewport::MakeProjectionMatrix(l, r, t, b, n, f);
 
 	if (OverscanRendering.UpdateProjectionAngles(l, r, t, b))
 	{
@@ -149,7 +149,7 @@ void FDisplayClusterViewport::CalculateProjectionMatrix(const uint32 InContextNu
 			)
 		{
 			// Use overscan projection matrix
-			Contexts[InContextNum].OverscanProjectionMatrix = ImplCreateProjectionMatrix(l, r, t, b, n, f);
+			Contexts[InContextNum].OverscanProjectionMatrix = IDisplayClusterViewport::MakeProjectionMatrix(l, r, t, b, n, f);
 			return;
 		}
 	}
@@ -158,7 +158,10 @@ void FDisplayClusterViewport::CalculateProjectionMatrix(const uint32 InContextNu
 	OverscanRendering.Disable();
 }
 
-FMatrix FDisplayClusterViewport::ImplCreateProjectionMatrix(float l, float r, float t, float b, float n, float f) const
+///////////////////////////////////////////////////////////////////////////////////////
+// IDisplayClusterViewport
+///////////////////////////////////////////////////////////////////////////////////////
+FMatrix IDisplayClusterViewport::MakeProjectionMatrix(float l, float r, float t, float b, float n, float f)
 {
 	const float mx = 2.f * n / (r - l);
 	const float my = 2.f * n / (t - b);
diff --git a/Engine/Plugins/Runtime/nDisplay/Source/DisplayCluster/Private/Render/Viewport/LightCard/DisplayClusterViewportLightCardManager.cpp b/Engine/Plugins/Runtime/nDisplay/Source/DisplayCluster/Private/Render/Viewport/LightCard/DisplayClusterViewportLightCardManager.cpp
--- a/Engine/Plugins/Runtime/nDisplay/Source/DisplayCluster/Private/Render/Viewport/LightCard/DisplayClusterViewportLightCardManager.cpp
+++ b/Engine/Plugins/Runtime/nDisplay/Source/DisplayCluster/Private/Render/Viewport/LightCard/DisplayClusterViewportLightCardManager.cpp
@@ -18,6 +18,9 @@ static TAutoConsoleVariable<int32> CVarUVLightCardTextureSize(
 	512,
 	TEXT("The size of the texture UV light cards are rendered to."));
 
+//-----------------------------------------------------------------------------------------------------------------
+// FDisplayClusterLightCardMap
+//-----------------------------------------------------------------------------------------------------------------
 void FDisplayClusterLightCardMap::InitDynamicRHI()
 {
 	ETextureCreateFlags CreateFlags = TexCreate_Dynamic;
@@ -33,47 +36,59 @@ void FDisplayClusterLightCardMap::InitDynamicRHI()
 	RenderTargetTextureRHI = TextureRHI = RHICreateTexture(Desc);
 }
 
+//-----------------------------------------------------------------------------------------------------------------
+// FDisplayClusterViewportLightCardManager
+//-----------------------------------------------------------------------------------------------------------------
 FDisplayClusterViewportLightCardManager::FDisplayClusterViewportLightCardManager(FDisplayClusterViewportManager& InViewportManager)
 	: ViewportManager(InViewportManager)
 {
+	ProxyData = MakeShared<FProxyData, ESPMode::ThreadSafe>();
 }
 
 FDisplayClusterViewportLightCardManager::~FDisplayClusterViewportLightCardManager()
 {
-	ResetScene();
-	ReleaseUVLightCardMap();
-	PreviewScene.Reset();
+	Release();
+
+	ProxyData.Reset();
 }
 
-void FDisplayClusterViewportLightCardManager::UpdateConfiguration()
+void FDisplayClusterViewportLightCardManager::Release()
 {
-	ResetScene();
+	// The destructor is usually called from the rendering thread, so Release() must be called first from the game thread.
+	const bool bIsInRenderingThread = IsInRenderingThread();
+	check(!bIsInRenderingThread || (bIsInRenderingThread && PreviewScene.IsValid() == false));
 
-	ADisplayClusterRootActor* RootActor = ViewportManager.GetRootActor();
-	TSet<ADisplayClusterLightCardActor*> LightCards;
-	UDisplayClusterBlueprintLib::FindLightCardsForRootActor(RootActor, LightCards);
+	// Deleting PreviewScene is only called from the game thread
+	PreviewScene.Reset();
 
-	for (ADisplayClusterLightCardActor* LightCard : LightCards)
-	{
-		if (LightCard->bIsUVLightCard)
-		{
-			UVLightCards.Add(LightCard);
-		}
-	}
+	UVLightCards.Empty();
+	ReleaseUVLightCardMap();
 }
 
-void FDisplayClusterViewportLightCardManager::ResetScene()
+FRHITexture* FDisplayClusterViewportLightCardManager::GetUVLightCardMap_RenderThread() const
+{ 
+	check(IsInRenderingThread());
+
+	return ProxyData.IsValid() ? ProxyData->GetUVLightCardMap_RenderThread() : nullptr;
+}
+
+void FDisplayClusterViewportLightCardManager::UpdateConfiguration()
 {
-	if (PreviewScene.IsValid() && LoadedPrimitiveComponents.Num())
+	UVLightCards.Empty();
+
+	if (ADisplayClusterRootActor* RootActorPtr = ViewportManager.GetRootActor())
 	{
-		for (UPrimitiveComponent* PrimitiveComponent : LoadedPrimitiveComponents)
+		TSet<ADisplayClusterLightCardActor*> LightCards;
+		UDisplayClusterBlueprintLib::FindLightCardsForRootActor(RootActorPtr, LightCards);
+
+		for (ADisplayClusterLightCardActor* LightCard : LightCards)
 		{
-			PreviewScene->GetScene()->RemovePrimitive(PrimitiveComponent);
+			if (LightCard->bIsUVLightCard)
+			{
+				UVLightCards.Add(LightCard);
+			}
 		}
 	}
-
-	LoadedPrimitiveComponents.Empty();
-	UVLightCards.Empty();
 }
 
 void FDisplayClusterViewportLightCardManager::HandleStartScene()
@@ -82,91 +97,125 @@ void FDisplayClusterViewportLightCardManager::HandleStartScene()
 		.SetEditor(false)
 		.SetCreatePhysicsScene(false)
 		.SetCreateDefaultLighting(false));
-
-	UpdateConfiguration();
 }
 
-void FDisplayClusterViewportLightCardManager::HandleBeginNewFrame()
+void FDisplayClusterViewportLightCardManager::HandleEndScene()
 {
-	InitializeUVLightCardMap();
+	PreviewScene.Reset();
 }
 
-void FDisplayClusterViewportLightCardManager::PreRenderFrame()
+void FDisplayClusterViewportLightCardManager::RenderFrame()
 {
-	bool bLoadedPrimitives = false;
-	for (ADisplayClusterLightCardActor* LightCard : UVLightCards)
+	if (PreviewScene.IsValid() && ProxyData.IsValid())
 	{
-		UStaticMeshComponent* LightCardMeshComp = LightCard->GetLightCardMeshComponent();
-		if (LightCardMeshComp && LightCardMeshComp->SceneProxy == nullptr)
+		if (FSceneInterface* SceneInterface = PreviewScene->GetScene())
 		{
-			PreviewScene->GetScene()->AddPrimitive(LightCardMeshComp);
-			LoadedPrimitiveComponents.Add(LightCardMeshComp);
-
-			bLoadedPrimitives = true;
+			InitializeUVLightCardMap();
+
+			/** A list of primitive components that have been added to the preview scene for rendering in the current frame */
+			TArray<UPrimitiveComponent*> LoadedPrimitiveComponents;
+
+			bool bLoadedPrimitives = false;
+			for (ADisplayClusterLightCardActor* LightCard : UVLightCards)
+			{
+				UStaticMeshComponent* LightCardMeshComp = LightCard->GetLightCardMeshComponent();
+				if (LightCardMeshComp && LightCardMeshComp->SceneProxy == nullptr)
+				{
+					SceneInterface->AddPrimitive(LightCardMeshComp);
+					LoadedPrimitiveComponents.Add(LightCardMeshComp);
+
+					bLoadedPrimitives = true;
+				}
+			}
+
+			ENQUEUE_RENDER_COMMAND(DisplayClusterViewportLightCardManager_InitializeUVLightCardMap)(
+				[InProxyData = ProxyData, bLoadedPrimitives, SceneInterface](FRHICommandListImmediate& RHICmdList)
+				{
+					InProxyData->RenderLightCardMap_RenderThread(RHICmdList, bLoadedPrimitives, SceneInterface);
+				});
+
+			for (UPrimitiveComponent* LoadedComponent : LoadedPrimitiveComponents)
+			{
+				SceneInterface->RemovePrimitive(LoadedComponent);
+			}
 		}
 	}
-
-	ENQUEUE_RENDER_COMMAND(InitializeLightCardRenderFrame)(
-		[this, bLoadedPrimitives](FRHICommandListImmediate& RHICmdList)
-	{
-		bHasUVLightCards = bLoadedPrimitives;
-	});
 }
 
-void FDisplayClusterViewportLightCardManager::PostRenderFrame()
+void FDisplayClusterViewportLightCardManager::InitializeUVLightCardMap()
 {
-	for (UPrimitiveComponent* LoadedComponent : LoadedPrimitiveComponents)
+	const uint32 LightCardTextureSize = CVarUVLightCardTextureSize.GetValueOnAnyThread();
+	if (UVLightCardMap && UVLightCardMap->GetSizeX() != LightCardTextureSize)
 	{
-		PreviewScene->GetScene()->RemovePrimitive(LoadedComponent);
+		ReleaseUVLightCardMap();
 	}
 
-	LoadedPrimitiveComponents.Empty();
+	if (UVLightCardMap == nullptr && ProxyData.IsValid())
+	{
+		UVLightCardMap = new FDisplayClusterLightCardMap(LightCardTextureSize);
+
+		ENQUEUE_RENDER_COMMAND(DisplayClusterViewportLightCardManager_InitializeUVLightCardMap)(
+			[InProxyData = ProxyData, InUVLightCardMap = UVLightCardMap](FRHICommandListImmediate& RHICmdList)
+			{
+				InProxyData->InitializeUVLightCardMap_RenderThread(InUVLightCardMap);
+			});
+	}
 }
 
-void FDisplayClusterViewportLightCardManager::RenderLightCardMap_RenderThread(FRHICommandListImmediate& RHICmdList)
+void FDisplayClusterViewportLightCardManager::ReleaseUVLightCardMap()
 {
-	if (PreviewScene.IsValid() && bHasUVLightCards)
+	UVLightCardMap = nullptr;
+
+	if (ProxyData.IsValid())
 	{
-		IDisplayClusterShaders& ShadersAPI = IDisplayClusterShaders::Get();
-		ShadersAPI.RenderPreprocess_UVLightCards(RHICmdList, PreviewScene->GetScene(), UVLightCardMapProxy, ADisplayClusterLightCardActor::UVPlaneDefaultSize);
+		ENQUEUE_RENDER_COMMAND(DisplayClusterViewportLightCardManager_ReleaseUVLightCardMap)(
+			[InProxyData = ProxyData](FRHICommandListImmediate& RHICmdList)
+			{
+				InProxyData->ReleaseUVLightCardMap_RenderThread();
+			});
 	}
 }
 
-void FDisplayClusterViewportLightCardManager::InitializeUVLightCardMap()
+//-----------------------------------------------------------------------------------------------------------------
+// FDisplayClusterViewportLightCardManager::FProxyData
+//-----------------------------------------------------------------------------------------------------------------
+FDisplayClusterViewportLightCardManager::FProxyData::~FProxyData()
 {
-	const uint32 LightCardTextureSize = CVarUVLightCardTextureSize.GetValueOnAnyThread();
-	if (UVLightCardMap && UVLightCardMap->GetSizeX() != LightCardTextureSize)
-	{
-		ReleaseUVLightCardMap();
-	}
+	ReleaseUVLightCardMap_RenderThread();
+}
 
+void FDisplayClusterViewportLightCardManager::FProxyData::InitializeUVLightCardMap_RenderThread(FDisplayClusterLightCardMap* InUVLightCardMap)
+{
 	if (UVLightCardMap == nullptr)
 	{
-		UVLightCardMap = new FDisplayClusterLightCardMap(LightCardTextureSize);
-
-		ENQUEUE_RENDER_COMMAND(InitializeLightCardResources)(
-			[this, InUVLightCardMap=UVLightCardMap](FRHICommandListImmediate& RHICmdList)
-		{
-			// Store a copy of the texture's pointer on the render thread and initialize the texture's resources
-			UVLightCardMapProxy = InUVLightCardMap;
-			UVLightCardMapProxy->InitResource();
-		});
+		// Store a copy of the texture's pointer on the render thread and initialize the texture's resources
+		UVLightCardMap = InUVLightCardMap;
+		UVLightCardMap->InitResource();
 	}
 }
 
-void FDisplayClusterViewportLightCardManager::ReleaseUVLightCardMap()
+void FDisplayClusterViewportLightCardManager::FProxyData::ReleaseUVLightCardMap_RenderThread()
 {
+	// Release the texture's resources and delete the texture object from the rendering thread
 	if (UVLightCardMap)
 	{
-		UVLightCardMap = nullptr;
+		UVLightCardMap->ReleaseResource();
 
-		ENQUEUE_RENDER_COMMAND(DeleteLightCardResources)(
-			[InUVLightCardMap=UVLightCardMapProxy](FRHICommandListImmediate& RHICmdList)
-		{
-			// Release the texture's resources and delete the texture object from the rendering thread. Pointer is copied to ensure it is still
-			// valid on the render thread even if the light card manager has been disposed already
-			InUVLightCardMap->ReleaseResource();
-			delete InUVLightCardMap;
-		});
+		delete UVLightCardMap;
+		UVLightCardMap = nullptr;
 	}
-}
\ No newline at end of file
+}
+
+void FDisplayClusterViewportLightCardManager::FProxyData::RenderLightCardMap_RenderThread(FRHICommandListImmediate& RHICmdList, const bool bLoadedPrimitives, FSceneInterface* InSceneInterface)
+{
+	bHasUVLightCards = bLoadedPrimitives;
+
+	IDisplayClusterShaders& ShadersAPI = IDisplayClusterShaders::Get();
+	ShadersAPI.RenderPreprocess_UVLightCards(RHICmdList, InSceneInterface, UVLightCardMap, ADisplayClusterLightCardActor::UVPlaneDefaultSize);
+}
+
+FRHITexture* FDisplayClusterViewportLightCardManager::FProxyData::GetUVLightCardMap_RenderThread() const
+{
+	return (bHasUVLightCards && UVLightCardMap != nullptr) ? UVLightCardMap->GetTextureRHI() : nullptr;
+}
+
diff --git a/Engine/Plugins/Runtime/nDisplay/Source/DisplayCluster/Private/Render/Viewport/LightCard/DisplayClusterViewportLightCardManager.h b/Engine/Plugins/Runtime/nDisplay/Source/DisplayCluster/Private/Render/Viewport/LightCard/DisplayClusterViewportLightCardManager.h
--- a/Engine/Plugins/Runtime/nDisplay/Source/DisplayCluster/Private/Render/Viewport/LightCard/DisplayClusterViewportLightCardManager.h
+++ b/Engine/Plugins/Runtime/nDisplay/Source/DisplayCluster/Private/Render/Viewport/LightCard/DisplayClusterViewportLightCardManager.h
@@ -20,6 +20,8 @@ public:
 		: Size(InSize)
 	{ }
 
+	virtual ~FDisplayClusterLightCardMap() = default;
+
 	virtual uint32 GetSizeX() const override { return Size; }
 	virtual uint32 GetSizeY() const override { return Size; }
 	virtual FIntPoint GetSizeXY() const override { return FIntPoint(Size, Size); }
@@ -39,19 +41,21 @@ public:
 	FDisplayClusterViewportLightCardManager(FDisplayClusterViewportManager& InViewportManager);
 	virtual ~FDisplayClusterViewportLightCardManager();
 
+	void Release();
+
+public:
 	//~ Begin IDisplayClusterViewportLightCardManager interface
-	virtual void UpdateConfiguration() override;
-	virtual void ResetScene() override;
-	virtual void HandleStartScene() override;
-	virtual void HandleBeginNewFrame() override;
-	virtual void PreRenderFrame() override;
-	virtual void PostRenderFrame() override;
-	virtual FRHITexture* GetUVLightCardMap() const override { return UVLightCardMap->GetTextureRHI(); }
-	virtual FRHITexture* GetUVLightCardMap_RenderThread() const override { return UVLightCardMapProxy->GetTextureRHI(); }
-	virtual bool HasUVLightCards_RenderThread() const override { return bHasUVLightCards; }
-	virtual void RenderLightCardMap_RenderThread(FRHICommandListImmediate& RHICmdList) override;
+	virtual FRHITexture* GetUVLightCardMap_RenderThread() const override;
 	//~ End IDisplayClusterViewportLightCardManager interface
 
+public:
+	void UpdateConfiguration();
+
+	void HandleStartScene();
+	void HandleEndScene();
+
+	void RenderFrame();
+
 private:
 	/** Initializes the UV light card map texture */
 	void InitializeUVLightCardMap();
@@ -69,15 +73,27 @@ private:
 	/** The list of UV light card actors that are referenced by the root actor */
 	TArray<ADisplayClusterLightCardActor*> UVLightCards;
 
-	/** A list of primitive components that have been added to the preview scene for rendering in the current frame */
-	TArray<UPrimitiveComponent*> LoadedPrimitiveComponents;
-
 	/** The render target to which the UV light card map is rendered */
 	FDisplayClusterLightCardMap* UVLightCardMap = nullptr;
 
-	/** The render thread copy of the pointer to the UV ligth card map */
-	FDisplayClusterLightCardMap* UVLightCardMapProxy = nullptr;
+	struct FProxyData
+	{
+		~FProxyData();
+
+		void InitializeUVLightCardMap_RenderThread(FDisplayClusterLightCardMap* InUVLightCardMap);
+		void ReleaseUVLightCardMap_RenderThread();
+
+		void RenderLightCardMap_RenderThread(FRHICommandListImmediate& RHICmdList, const bool bLoadedPrimitives, FSceneInterface* InSceneInterface);
+
+		FRHITexture* GetUVLightCardMap_RenderThread() const;
 
-	/** A render thread flag that indicates the light card manager has UV light cards to render */
-	bool bHasUVLightCards = false;
-};
\ No newline at end of file
+	private:
+		/** The render thread copy of the pointer to the UV ligth card map */
+		FDisplayClusterLightCardMap* UVLightCardMap = nullptr;
+
+		/** A render thread flag that indicates the light card manager has UV light cards to render */
+		bool bHasUVLightCards = false;
+	};
+
+	TSharedPtr<FProxyData, ESPMode::ThreadSafe> ProxyData;
+};
diff --git a/Engine/Plugins/Runtime/nDisplay/Source/DisplayCluster/Private/Render/Viewport/Postprocess/DisplayClusterViewportPostProcessManager.cpp b/Engine/Plugins/Runtime/nDisplay/Source/DisplayCluster/Private/Render/Viewport/Postprocess/DisplayClusterViewportPostProcessManager.cpp
--- a/Engine/Plugins/Runtime/nDisplay/Source/DisplayCluster/Private/Render/Viewport/Postprocess/DisplayClusterViewportPostProcessManager.cpp
+++ b/Engine/Plugins/Runtime/nDisplay/Source/DisplayCluster/Private/Render/Viewport/Postprocess/DisplayClusterViewportPostProcessManager.cpp
@@ -66,12 +66,35 @@ static TAutoConsoleVariable<int32> CVarPostprocessFrameAfterWarpBlend(
 	ECVF_RenderThreadSafe
 );
 
+//----------------------------------------------------------------------------------------------------------
+// FDisplayClusterViewportPostProcessManager
+//----------------------------------------------------------------------------------------------------------
 FDisplayClusterViewportPostProcessManager::FDisplayClusterViewportPostProcessManager(FDisplayClusterViewportManager& InViewportManager)
 	: ViewportManager(InViewportManager)
 {
 	OutputRemap = MakeShared<FDisplayClusterViewportPostProcessOutputRemap, ESPMode::ThreadSafe>();
 }
 
+FDisplayClusterViewportPostProcessManager::~FDisplayClusterViewportPostProcessManager()
+{
+	Release();
+}
+
+void FDisplayClusterViewportPostProcessManager::Release()
+{
+	if (IsInGameThread())
+	{
+		Postprocess.Empty();
+		OutputRemap.Reset();
+
+		ENQUEUE_RENDER_COMMAND(DisplayClusterViewportPostProcessManager_Release)(
+			[this](FRHICommandListImmediate& RHICmdList)
+			{
+				PostprocessProxy.Empty();
+			});
+	}
+}
+
 bool FDisplayClusterViewportPostProcessManager::HandleStartScene()
 {
 	check(IsInGameThread());;
diff --git a/Engine/Plugins/Runtime/nDisplay/Source/DisplayCluster/Private/Render/Viewport/Postprocess/DisplayClusterViewportPostProcessManager.h b/Engine/Plugins/Runtime/nDisplay/Source/DisplayCluster/Private/Render/Viewport/Postprocess/DisplayClusterViewportPostProcessManager.h
--- a/Engine/Plugins/Runtime/nDisplay/Source/DisplayCluster/Private/Render/Viewport/Postprocess/DisplayClusterViewportPostProcessManager.h
+++ b/Engine/Plugins/Runtime/nDisplay/Source/DisplayCluster/Private/Render/Viewport/Postprocess/DisplayClusterViewportPostProcessManager.h
@@ -18,7 +18,9 @@ class FDisplayClusterViewportPostProcessManager
 {
 public:
 	FDisplayClusterViewportPostProcessManager(FDisplayClusterViewportManager& InViewportManager);
-	virtual ~FDisplayClusterViewportPostProcessManager() = default;
+	virtual ~FDisplayClusterViewportPostProcessManager();
+
+	void Release();
 
 public:
 	bool IsPostProcessViewBeforeWarpBlendRequired(const TSharedPtr<IDisplayClusterPostProcess, ESPMode::ThreadSafe>& PostprocessInstance) const;
diff --git a/Engine/Plugins/Runtime/nDisplay/Source/DisplayCluster/Private/Render/Viewport/RenderTarget/DisplayClusterRenderTargetManager.cpp b/Engine/Plugins/Runtime/nDisplay/Source/DisplayCluster/Private/Render/Viewport/RenderTarget/DisplayClusterRenderTargetManager.cpp
--- a/Engine/Plugins/Runtime/nDisplay/Source/DisplayCluster/Private/Render/Viewport/RenderTarget/DisplayClusterRenderTargetManager.cpp
+++ b/Engine/Plugins/Runtime/nDisplay/Source/DisplayCluster/Private/Render/Viewport/RenderTarget/DisplayClusterRenderTargetManager.cpp
@@ -26,6 +26,10 @@ namespace DisplayClusterRenderTargetManager
 		case EDisplayClusterViewportCaptureMode::Lightcard:
 			return EPixelFormat::PF_FloatRGBA;
 
+		case EDisplayClusterViewportCaptureMode::MoviePipeline:
+			// Movie pipeline always use PF_FloatRGBA
+			return EPixelFormat::PF_FloatRGBA;
+
 		default:
 			break;
 		}
@@ -55,6 +59,12 @@ FDisplayClusterRenderTargetManager::FDisplayClusterRenderTargetManager(FDisplayC
 
 FDisplayClusterRenderTargetManager::~FDisplayClusterRenderTargetManager()
 {
+	Release();
+}
+
+void FDisplayClusterRenderTargetManager::Release()
+{
+	ResourcesPool->Release();
 }
 
 bool FDisplayClusterRenderTargetManager::AllocateRenderFrameResources(FViewport* InViewport, const FDisplayClusterRenderFrameSettings& InRenderFrameSettings, const TArray<FDisplayClusterViewport*>& InViewports, FDisplayClusterRenderFrame& InOutRenderFrame)
diff --git a/Engine/Plugins/Runtime/nDisplay/Source/DisplayCluster/Private/Render/Viewport/RenderTarget/DisplayClusterRenderTargetManager.h b/Engine/Plugins/Runtime/nDisplay/Source/DisplayCluster/Private/Render/Viewport/RenderTarget/DisplayClusterRenderTargetManager.h
--- a/Engine/Plugins/Runtime/nDisplay/Source/DisplayCluster/Private/Render/Viewport/RenderTarget/DisplayClusterRenderTargetManager.h
+++ b/Engine/Plugins/Runtime/nDisplay/Source/DisplayCluster/Private/Render/Viewport/RenderTarget/DisplayClusterRenderTargetManager.h
@@ -18,6 +18,8 @@ public:
 	FDisplayClusterRenderTargetManager(FDisplayClusterViewportManagerProxy* InViewportManagerProxy);
 	~FDisplayClusterRenderTargetManager();
 
+	void Release();
+
 public:
 	bool AllocateRenderFrameResources(FViewport* InViewport, const FDisplayClusterRenderFrameSettings& InRenderFrameSettings, const TArray<FDisplayClusterViewport*>& InViewports, FDisplayClusterRenderFrame& InOutRenderFrame);
 
diff --git a/Engine/Plugins/Runtime/nDisplay/Source/DisplayCluster/Private/Render/Viewport/RenderTarget/DisplayClusterRenderTargetResourcesPool.cpp b/Engine/Plugins/Runtime/nDisplay/Source/DisplayCluster/Private/Render/Viewport/RenderTarget/DisplayClusterRenderTargetResourcesPool.cpp
--- a/Engine/Plugins/Runtime/nDisplay/Source/DisplayCluster/Private/Render/Viewport/RenderTarget/DisplayClusterRenderTargetResourcesPool.cpp
+++ b/Engine/Plugins/Runtime/nDisplay/Source/DisplayCluster/Private/Render/Viewport/RenderTarget/DisplayClusterRenderTargetResourcesPool.cpp
@@ -83,6 +83,11 @@ FDisplayClusterRenderTargetResourcesPool::FDisplayClusterRenderTargetResourcesPo
 { }
 
 FDisplayClusterRenderTargetResourcesPool::~FDisplayClusterRenderTargetResourcesPool()
+{
+	Release();
+}
+
+void FDisplayClusterRenderTargetResourcesPool::Release()
 {
 	// Release all resources
 	ImplReleaseResources<FDisplayClusterViewportRenderTargetResource>(RenderTargetResources);
diff --git a/Engine/Plugins/Runtime/nDisplay/Source/DisplayCluster/Private/Render/Viewport/RenderTarget/DisplayClusterRenderTargetResourcesPool.h b/Engine/Plugins/Runtime/nDisplay/Source/DisplayCluster/Private/Render/Viewport/RenderTarget/DisplayClusterRenderTargetResourcesPool.h
--- a/Engine/Plugins/Runtime/nDisplay/Source/DisplayCluster/Private/Render/Viewport/RenderTarget/DisplayClusterRenderTargetResourcesPool.h
+++ b/Engine/Plugins/Runtime/nDisplay/Source/DisplayCluster/Private/Render/Viewport/RenderTarget/DisplayClusterRenderTargetResourcesPool.h
@@ -14,6 +14,8 @@ public:
 	FDisplayClusterRenderTargetResourcesPool(FDisplayClusterViewportManagerProxy* InViewportManagerProxy);
 	~FDisplayClusterRenderTargetResourcesPool();
 
+	void Release();
+
 public:
 	bool BeginReallocateResources(const FDisplayClusterRenderFrameSettings& InRenderFrameSettings, class FViewport* InViewport);
 	void FinishReallocateResources();
diff --git a/Engine/Plugins/Runtime/nDisplay/Source/DisplayCluster/Public/Render/Viewport/Containers/DisplayClusterViewport_RenderSettings.h b/Engine/Plugins/Runtime/nDisplay/Source/DisplayCluster/Public/Render/Viewport/Containers/DisplayClusterViewport_RenderSettings.h
--- a/Engine/Plugins/Runtime/nDisplay/Source/DisplayCluster/Public/Render/Viewport/Containers/DisplayClusterViewport_RenderSettings.h
+++ b/Engine/Plugins/Runtime/nDisplay/Source/DisplayCluster/Public/Render/Viewport/Containers/DisplayClusterViewport_RenderSettings.h
@@ -15,6 +15,9 @@ enum class EDisplayClusterViewportCaptureMode : uint8
 	// use hi-res float texture with alpha for compisiting
 	Lightcard,
 	Lightcard_OCIO,
+
+	// Special hi-res mode for movie pipeline
+	MoviePipeline,
 };
 
 class FDisplayClusterViewport_RenderSettings
@@ -47,6 +50,12 @@ public:
 	// Render alpha channel from input texture to warp output
 	bool bWarpBlendRenderAlphaChannel = false;
 
+	// Disable CustomFrustum feature from viewport settings
+	bool bDisableCustomFrustumFeature = false;
+
+	// Disable viewport overscan feature from settings
+	bool bDisableFrustumOverscanFeature = false;
+
 	// Read viewport pixels for preview (this flag is cleared at the end of the frame)
 	bool bPreviewReadPixels = false;
 
diff --git a/Engine/Plugins/Runtime/nDisplay/Source/DisplayCluster/Public/Render/Viewport/IDisplayClusterViewport.h b/Engine/Plugins/Runtime/nDisplay/Source/DisplayCluster/Public/Render/Viewport/IDisplayClusterViewport.h
--- a/Engine/Plugins/Runtime/nDisplay/Source/DisplayCluster/Public/Render/Viewport/IDisplayClusterViewport.h
+++ b/Engine/Plugins/Runtime/nDisplay/Source/DisplayCluster/Public/Render/Viewport/IDisplayClusterViewport.h
@@ -49,4 +49,6 @@ public:
 
 	virtual void SetRenderSettings(const FDisplayClusterViewport_RenderSettings& InRenderSettings) = 0;
 	virtual void SetContexts(TArray<FDisplayClusterViewport_Context>& InContexts) = 0;
+
+	static FMatrix MakeProjectionMatrix(float InLeft, float InRight, float InTop, float InBottom, float ZNear, float ZFar);
 };
diff --git a/Engine/Plugins/Runtime/nDisplay/Source/DisplayCluster/Public/Render/Viewport/IDisplayClusterViewportLightCardManager.h b/Engine/Plugins/Runtime/nDisplay/Source/DisplayCluster/Public/Render/Viewport/IDisplayClusterViewportLightCardManager.h
--- a/Engine/Plugins/Runtime/nDisplay/Source/DisplayCluster/Public/Render/Viewport/IDisplayClusterViewportLightCardManager.h
+++ b/Engine/Plugins/Runtime/nDisplay/Source/DisplayCluster/Public/Render/Viewport/IDisplayClusterViewportLightCardManager.h
@@ -2,6 +2,8 @@
 
 #pragma once
 
+#include "CoreMinimal.h"
+
 class FRHITexture;
 class FRHICommandListImmediate;
 
@@ -9,33 +11,6 @@ class FRHICommandListImmediate;
 class DISPLAYCLUSTER_API IDisplayClusterViewportLightCardManager
 {
 public:
-	/** Updates the UV light card scene based on the current root actor configuration */
-	virtual void UpdateConfiguration() = 0;
-
-	/** Resets the UV light card scene */
-	virtual void ResetScene() = 0;
-
-	/** Raised when a new scene is being started, initializes the UV light card scene */
-	virtual void HandleStartScene() = 0;
-
-	/** Raised when a new frame begins, initializes any frame resources needed */
-	virtual void HandleBeginNewFrame() = 0;
-
-	/** Raised just before a render frame pass is performed */
-	virtual void PreRenderFrame() = 0;
-
-	/** Raised just after a render frame pass is performed */
-	virtual void PostRenderFrame() = 0;
-
-	/** Gets the UV light card map texture */
-	virtual FRHITexture* GetUVLightCardMap() const = 0;
-
 	/** Gets the render thread copy of the UV light card map texture */
 	virtual FRHITexture* GetUVLightCardMap_RenderThread() const = 0;
-
-	/** Gets whether there are any UV light cards to render to the viewports */
-	virtual bool HasUVLightCards_RenderThread() const = 0;
-
-	/** Renders the UV light card scene to the UV light card map texture */
-	virtual void RenderLightCardMap_RenderThread(FRHICommandListImmediate& RHICmdList) = 0;
 };
\ No newline at end of file
diff --git a/Engine/Plugins/Runtime/nDisplay/Source/DisplayCluster/Public/Render/Viewport/IDisplayClusterViewportManager.h b/Engine/Plugins/Runtime/nDisplay/Source/DisplayCluster/Public/Render/Viewport/IDisplayClusterViewportManager.h
--- a/Engine/Plugins/Runtime/nDisplay/Source/DisplayCluster/Public/Render/Viewport/IDisplayClusterViewportManager.h
+++ b/Engine/Plugins/Runtime/nDisplay/Source/DisplayCluster/Public/Render/Viewport/IDisplayClusterViewportManager.h
@@ -51,6 +51,19 @@ public:
 	*/
 	virtual bool UpdateConfiguration(EDisplayClusterRenderFrameMode InRenderMode, const FString& InClusterNodeId, class ADisplayClusterRootActor* InRootActorPtr, const FDisplayClusterPreviewSettings* InPreviewSettings = nullptr) = 0;
 
+	/**
+	* Update\Create\Delete viewports for frame. For rendering outside of cluster nodes
+	* Update ICVFX configuration from root actor components
+	* [Game thread func]
+	*
+	* @param InRenderMode    - Render mode
+	* @param InViewportNames - Viewports names for next frame
+	* @param InRootActorPtr  - reference to RootActor with actual configuration inside
+	*
+	* @return - true, if success
+	*/
+	virtual bool UpdateCustomConfiguration(EDisplayClusterRenderFrameMode InRenderMode, const TArray<FString>& InViewportNames, class ADisplayClusterRootActor* InRootActorPtr) = 0;
+
 	/**
 	* Initialize new frame for all viewports on game thread, and update context, render resources with viewport new settings
 	* And finally build render frame structure and send to render thread proxy viewport objects
diff --git a/Engine/Plugins/Runtime/nDisplay/Source/DisplayClusterMoviePipeline/DisplayClusterMoviePipeline.Build.cs b/Engine/Plugins/Runtime/nDisplay/Source/DisplayClusterMoviePipeline/DisplayClusterMoviePipeline.Build.cs
new file mode 100644
--- /dev/null
+++ b/Engine/Plugins/Runtime/nDisplay/Source/DisplayClusterMoviePipeline/DisplayClusterMoviePipeline.Build.cs
@@ -0,0 +1,33 @@
+// Copyright Epic Games, Inc. All Rights Reserved.
+
+using UnrealBuildTool;
+
+public class DisplayClusterMoviePipeline : ModuleRules
+{
+	public DisplayClusterMoviePipeline(ReadOnlyTargetRules Target) : base(Target)
+	{
+		PrivateDependencyModuleNames.AddRange(
+			new string[] {
+				"Core",
+				"CoreUObject",
+				"Engine",
+				"Json",
+				"DisplayCluster"
+			}
+		);
+
+		PublicDependencyModuleNames.AddRange(
+			new string[] {
+				"MovieRenderPipelineCore",
+				"MovieRenderPipelineRenderPasses",
+				"ImageWriteQueue",
+				"RenderCore",
+				"RHI",
+				"ActorLayerUtilities",
+				"OpenColorIO",
+				"DisplayCluster",
+				"DisplayClusterConfiguration",
+			}
+		);
+	}
+}
diff --git a/Engine/Plugins/Runtime/nDisplay/Source/DisplayClusterMoviePipeline/Private/DisplayClusterMoviePipelineSettings.cpp b/Engine/Plugins/Runtime/nDisplay/Source/DisplayClusterMoviePipeline/Private/DisplayClusterMoviePipelineSettings.cpp
new file mode 100644
--- /dev/null
+++ b/Engine/Plugins/Runtime/nDisplay/Source/DisplayClusterMoviePipeline/Private/DisplayClusterMoviePipelineSettings.cpp
@@ -0,0 +1,68 @@
+// Copyright Epic Games, Inc. All Rights Reserved.
+
+#include "DisplayClusterMoviePipelineSettings.h"
+#include "DisplayClusterConfigurationTypes.h"
+#include "MovieRenderPipelineCoreModule.h"
+#include "Modules/ModuleManager.h"
+
+///////////////////////////////////////////////////////////////////////////////////////////
+// UDisplayClusterMoviePipelineSettings
+///////////////////////////////////////////////////////////////////////////////////////////
+ADisplayClusterRootActor* UDisplayClusterMoviePipelineSettings::GetRootActor(const UWorld* InWorld) const
+{
+	if (InWorld)
+	{
+		for (const TWeakObjectPtr<ADisplayClusterRootActor> RootActorRef : TActorRange<ADisplayClusterRootActor>(InWorld))
+		{
+			if (ADisplayClusterRootActor* RootActorPtr = RootActorRef.Get())
+			{
+				if (!DCRootActor.IsValid() || RootActorPtr->GetFName() == DCRootActor->GetFName())
+				{
+					return RootActorPtr;
+				}
+			}
+		}
+	}
+
+	return nullptr;
+}
+
+bool UDisplayClusterMoviePipelineSettings::GetViewports(const UWorld* InWorld, TArray<FString>& OutViewports) const
+{
+	if (ADisplayClusterRootActor* RootActorPtr = GetRootActor(InWorld))
+	{
+		if (const UDisplayClusterConfigurationData* InConfigurationData = RootActorPtr->GetConfigData())
+		{
+			if (const UDisplayClusterConfigurationCluster* InClusterCfg =  InConfigurationData->Cluster)
+			{
+				for (const TPair<FString, UDisplayClusterConfigurationClusterNode*>& NodeIt : InClusterCfg->Nodes)
+				{
+					if (const UDisplayClusterConfigurationClusterNode* InConfigurationClusterNode = NodeIt.Value)
+					{
+						const FString& InClusterNodeId = NodeIt.Key;
+						for (const TPair<FString, UDisplayClusterConfigurationViewport*>& InConfigurationViewportIt : InConfigurationClusterNode->Viewports)
+						{
+							if (const UDisplayClusterConfigurationViewport* InConfigurationViewport = InConfigurationViewportIt.Value)
+							{
+								if (InConfigurationViewport->bAllowRendering)
+								{
+									const FString& InViewportId = InConfigurationViewportIt.Key;
+									if (bRenderAllViewports || AllowedViewportNamesList.Find(InViewportId) != INDEX_NONE)
+									{
+										OutViewports.Add(InViewportId);
+									}
+								}
+							}
+						}
+					}
+				}
+
+				return OutViewports.Num() > 0;
+			}
+		}
+	}
+
+	return false;
+}
+
+IMPLEMENT_MODULE(FDefaultModuleImpl, DisplayClusterMoviePipeline);
diff --git a/Engine/Plugins/Runtime/nDisplay/Source/DisplayClusterMoviePipeline/Private/DisplayClusterMoviePipelineSettings.h b/Engine/Plugins/Runtime/nDisplay/Source/DisplayClusterMoviePipeline/Private/DisplayClusterMoviePipelineSettings.h
new file mode 100644
--- /dev/null
+++ b/Engine/Plugins/Runtime/nDisplay/Source/DisplayClusterMoviePipeline/Private/DisplayClusterMoviePipelineSettings.h
@@ -0,0 +1,55 @@
+// Copyright Epic Games, Inc. All Rights Reserved.
+
+#pragma once
+
+#include "MoviePipelineSetting.h"
+#include "DisplayClusterRootActor.h"
+
+#include "DisplayClusterMoviePipelineSettings.generated.h"
+
+/**
+ * nDisplay settings for MoviePipeline
+ */
+UCLASS(Blueprintable)
+class DISPLAYCLUSTERMOVIEPIPELINE_API UDisplayClusterMoviePipelineSettings : public UMoviePipelineSetting
+{
+	GENERATED_BODY()
+
+public:
+	/**
+	 * return a pointer to the specified DCRA from the current world
+	 */
+	ADisplayClusterRootActor* GetRootActor(const UWorld* InWorld) const;
+
+	/**
+	* Collect viewport names for rendering
+	* 
+	* @param InWorld       - World for render
+	* @param OutViewports  - [Out] list of viewports to render in MoviewPipeline
+	* 
+	 * return true if is whether it populated 1 or more viewports in the given array.
+	 */
+	bool GetViewports(const UWorld* InWorld, TArray<FString>& OutViewports) const;
+
+#if WITH_EDITOR
+	virtual FText GetDisplayText() const override { return NSLOCTEXT("MovieRenderPipeline", "DisplayClusterSettingDisplayName", "nDisplay"); }
+#endif
+
+protected:
+	virtual bool IsValidOnShots() const override { return true; }
+	virtual bool IsValidOnMaster() const override { return true; }
+
+public:
+	// Reference to Display Cluster Root Actor
+	// If not set, the first available in the scene will be set.
+	UPROPERTY(EditAnywhere, BlueprintReadWrite, Category = "nDisplay", meta = (DisplayName = "DC Root Actor"))
+	TSoftObjectPtr<ADisplayClusterRootActor> DCRootActor;
+
+	// Render all viewports
+	UPROPERTY(EditAnywhere, BlueprintReadWrite, Category = "nDisplay")
+	bool bRenderAllViewports = true;
+
+	// Render only viewports from this list.
+	UPROPERTY(EditAnywhere, BlueprintReadWrite, Category = "nDisplay", meta = (EditCondition = "!bRenderAllViewports"))
+	TArray<FString> AllowedViewportNamesList;
+};
diff --git a/Engine/Plugins/Runtime/nDisplay/Source/DisplayClusterMoviePipeline/Private/DisplayClusterMoviePipelineViewportPass.cpp b/Engine/Plugins/Runtime/nDisplay/Source/DisplayClusterMoviePipeline/Private/DisplayClusterMoviePipelineViewportPass.cpp
new file mode 100644
--- /dev/null
+++ b/Engine/Plugins/Runtime/nDisplay/Source/DisplayClusterMoviePipeline/Private/DisplayClusterMoviePipelineViewportPass.cpp
@@ -0,0 +1,600 @@
+// Copyright Epic Games, Inc. All Rights Reserved.
+
+#include "DisplayClusterMoviePipelineViewportPass.h"
+#include "DisplayClusterMoviePipelineSettings.h"
+
+#include "MoviePipeline.h"
+#include "MoviePipelineHighResSetting.h"
+
+#include "DisplayClusterConfigurationTypes.h"
+
+#include "Components/DisplayClusterCameraComponent.h"
+#include "Render/Projection/IDisplayClusterProjectionPolicy.h"
+#include "Render/Viewport/IDisplayClusterViewportManager.h"
+#include "Render/Viewport/IDisplayClusterViewportManagerProxy.h"
+#include "Render/Viewport/IDisplayClusterViewport.h"
+#include "Render/Viewport/IDisplayClusterViewportProxy.h"
+
+/////////////////////////////////////////////////////////////////////////////////////////////////////////////////////////////////
+namespace DisplayClusterMoviePipelineHelpers
+{
+	static void DisplayClusterWarpBlendImpl_RenderThread(FRHICommandListImmediate& RHICmdList, const FTextureRHIRef& RenderTargetRHI, IDisplayClusterViewportManagerProxy* InViewportManagerProxy, const FString& InViewportId, const FIntPoint& OffsetMin, const FIntPoint& OffsetMax)
+	{
+		if (InViewportManagerProxy)
+		{
+			if (IDisplayClusterViewportProxy* InViewportProxy = InViewportManagerProxy->FindViewport_RenderThread(InViewportId))
+			{
+				const TSharedPtr<IDisplayClusterProjectionPolicy, ESPMode::ThreadSafe>& InPolicy = InViewportProxy->GetProjectionPolicy_RenderThread();
+				if (InPolicy.IsValid() && InPolicy->IsWarpBlendSupported())
+				{
+					TArray<FRHITexture2D*> WarpInputTextures, WarpOutputTextures;
+					TArray<FIntRect> WarpInputRects, WarpOutputRects;
+
+					if (InViewportProxy->GetResourcesWithRects_RenderThread(EDisplayClusterViewportResourceType::InputShaderResource, WarpInputTextures, WarpInputRects))
+					{
+						if (InViewportProxy->GetResourcesWithRects_RenderThread(EDisplayClusterViewportResourceType::AdditionalTargetableResource, WarpOutputTextures, WarpOutputRects))
+						{
+							const FIntRect RTTRect(FIntPoint(0, 0), RenderTargetRHI->GetSizeXY());
+							const FIntRect WarpInputRect = WarpInputRects[0];
+							const FIntRect WarpOutpuRect = WarpOutputRects[0];
+
+							FRHICopyTextureInfo CopyInfo;
+							CopyInfo.Size = FIntVector(
+								RTTRect.Width() - OffsetMin.X - OffsetMax.X,
+								RTTRect.Height() - OffsetMin.Y - OffsetMax.Y,
+								0);
+							{
+								// Copy RTTRect->WarpInputRect
+								CopyInfo.SourcePosition.X = RTTRect.Min.X + OffsetMin.X;
+								CopyInfo.SourcePosition.Y = RTTRect.Min.Y + OffsetMin.Y;
+								CopyInfo.DestPosition.X = WarpInputRect.Min.X;
+								CopyInfo.DestPosition.Y = WarpInputRect.Min.Y;
+
+								TransitionAndCopyTexture(RHICmdList, RenderTargetRHI, WarpInputTextures[0], CopyInfo);
+							}
+
+							// Apply Warp: WarpInputTextures->WarpOutputTextures
+							InPolicy->ApplyWarpBlend_RenderThread(RHICmdList, InViewportProxy);
+
+							{
+								// Copy WarpOutputRect -> RTTRect
+								CopyInfo.SourcePosition.X = WarpOutpuRect.Min.X;
+								CopyInfo.SourcePosition.Y = WarpOutpuRect.Min.Y;
+								CopyInfo.DestPosition.X = RTTRect.Min.X + OffsetMin.X;
+								CopyInfo.DestPosition.Y = RTTRect.Min.Y + OffsetMin.Y;
+
+								TransitionAndCopyTexture(RHICmdList, WarpOutputTextures[0], RenderTargetRHI, CopyInfo);
+							}
+
+							// Immediatelly flush RHI before PostRendererSubmission
+							RHICmdList.ImmediateFlush(EImmediateFlushType::FlushRHIThreadFlushResources);
+						}
+					}
+				}
+			}
+		}
+	}
+};
+using namespace DisplayClusterMoviePipelineHelpers;
+
+/////////////////////////////////////////////////////////////////////////////////////////////////////////////////////////////////
+// UDisplayClusterMoviePipelineViewportPassBase
+/////////////////////////////////////////////////////////////////////////////////////////////////////////////////////////////////
+void UDisplayClusterMoviePipelineViewportPassBase::SetupImpl(const MoviePipeline::FMoviePipelineRenderPassInitSettings& InPassInitSettings)
+{
+	InitializeDisplayCluster();
+
+	Super::SetupImpl(InPassInitSettings);
+}
+
+void UDisplayClusterMoviePipelineViewportPassBase::TeardownImpl()
+{
+	ReleaseDisplayCluster();
+
+	Super::TeardownImpl();
+}
+
+/////////////////////////////////////////////////////////////////////////////////////////////////////////////////////////////////
+UE::MoviePipeline::FImagePassCameraViewData UDisplayClusterMoviePipelineViewportPassBase::GetCameraInfo(FMoviePipelineRenderPassMetrics& InOutSampleState, IViewCalcPayload* OptPayload) const
+{
+	UE::MoviePipeline::FImagePassCameraViewData OutCameraData;
+
+	// This has to come from the main camera for consistency's sake, and it's not a per-camera setting in the editor.
+	OutCameraData.ViewActor = GetPipeline()->GetWorld()->GetFirstPlayerController()->GetViewTarget();
+
+	const int32 InViewportIndex = InOutSampleState.OutputState.CameraIndex;
+	FString ViewportId;
+	if (GetViewportId(InViewportIndex, ViewportId))
+	{
+		if (DCViews.Contains(InViewportIndex) && DCPrevViews.Contains(InViewportIndex))
+		{
+			// We override the current/previous transform
+			OutCameraData.bUseCustomProjectionMatrix = true;
+			OutCameraData.CustomProjectionMatrix = DCViews[InViewportIndex].ProjectionMatrix;
+
+			OutCameraData.ViewInfo.Location = DCViews[InViewportIndex].ViewLocation;
+			OutCameraData.ViewInfo.Rotation = DCViews[InViewportIndex].ViewRotation;
+
+			OutCameraData.ViewInfo.PreviousViewTransform = FTransform(DCPrevViews[InViewportIndex].ViewRotation, DCPrevViews[InViewportIndex].ViewLocation);
+		}
+	}
+
+	return OutCameraData;
+}
+
+void UDisplayClusterMoviePipelineViewportPassBase::BlendPostProcessSettings(FSceneView* InView, FMoviePipelineRenderPassMetrics& InOutSampleState, IViewCalcPayload* OptPayload)
+{
+	const int32 InViewportIndex = InOutSampleState.OutputState.CameraIndex;
+	FString ViewportId;
+	if (GetViewportId(InViewportIndex, ViewportId))
+	{
+		if (DCViews.Contains(InViewportIndex) && DCPrevViews.Contains(InViewportIndex))
+		{
+			// Use DCRA as a camera
+			FVector ViewLocation = DCViews[InViewportIndex].ViewLocation;
+
+			for (IInterface_PostProcessVolume* PPVolume : GetWorld()->PostProcessVolumes)
+			{
+				const FPostProcessVolumeProperties VolumeProperties = PPVolume->GetProperties();
+
+				// Skip any volumes which are disabled
+				if (!VolumeProperties.bIsEnabled)
+				{
+					continue;
+				}
+
+				float LocalWeight = FMath::Clamp(VolumeProperties.BlendWeight, 0.0f, 1.0f);
+
+				if (!VolumeProperties.bIsUnbound)
+				{
+					float DistanceToPoint = 0.0f;
+					PPVolume->EncompassesPoint(ViewLocation, 0.0f, &DistanceToPoint);
+
+					if (DistanceToPoint >= 0 && DistanceToPoint < VolumeProperties.BlendRadius)
+					{
+						LocalWeight *= FMath::Clamp(1.0f - DistanceToPoint / VolumeProperties.BlendRadius, 0.0f, 1.0f);
+					}
+					else
+					{
+						LocalWeight = 0.0f;
+					}
+				}
+
+				InView->OverridePostProcessSettings(*VolumeProperties.Settings, LocalWeight);
+			}
+
+			//@todo: It is necessary to discuss which post-processes from DCRA can or cannot be used for MRQ.
+		}
+	}
+}
+
+int32 UDisplayClusterMoviePipelineViewportPassBase::GetNumCamerasToRender() const
+{
+	return DisplayClusterViewports.Num();
+}
+
+FString UDisplayClusterMoviePipelineViewportPassBase::GetCameraName(const int32 InCameraIndex) const
+{
+	if (InCameraIndex >= 0 && InCameraIndex < DisplayClusterViewports.Num())
+	{
+		return DisplayClusterViewports[InCameraIndex];
+	}
+
+	return TEXT("UndefinedCameraName");
+}
+
+FString UDisplayClusterMoviePipelineViewportPassBase::GetCameraNameOverride(const int32 InCameraIndex) const
+{
+	if (DisplayClusterViewports.IsValidIndex(InCameraIndex))
+	{
+		return DisplayClusterViewports[InCameraIndex];
+	}
+
+	return TEXT("");
+}
+
+void UDisplayClusterMoviePipelineViewportPassBase::ModifyProjectionMatrixForTiling(const FMoviePipelineRenderPassMetrics& InSampleState, const bool bInOrthographic, FMatrix& InOutProjectionMatrix, float& OutDoFSensorScale) const
+{
+	const FMatrix InProjectionMatrix(InOutProjectionMatrix);
+	Super::ModifyProjectionMatrixForTiling(InSampleState, bInOrthographic, InOutProjectionMatrix, OutDoFSensorScale);
+
+	// Support asymmetric projection matrix (Rebuild assym projection matrix for tile + OverlappedPad)
+	if(!bInOrthographic)
+	{
+		const FIntPoint TileCounts = InSampleState.TileCounts;
+		const FIntPoint  TileSize = InSampleState.TileSize;
+
+		const double   TileLeft = InSampleState.OverlappedOffset.X;
+		const double  TileRight = TileLeft + TileSize.X + InSampleState.OverlappedPad.X * 2;
+
+		const double    TileTop = InSampleState.OverlappedOffset.Y;
+		const double TileBottom = TileTop + TileSize.Y + InSampleState.OverlappedPad.Y * 2;
+
+		const double  Width = TileSize.X * TileCounts.X;
+		const double Height = TileSize.Y * TileCounts.Y;
+
+		const float ZNear = GNearClippingPlane;
+
+		// Extract FOV values from matrix
+		const double FOVWidth = (2.f * ZNear) / InProjectionMatrix.M[0][0];
+		const double FOVHeight = (2.f * ZNear) / InProjectionMatrix.M[1][1];
+		const double FOVOffsetX = InProjectionMatrix.M[2][0] * FOVWidth;
+		const double FOVOffsetY = InProjectionMatrix.M[2][1] * FOVHeight;
+
+		const double   FOVLeft = -(FOVOffsetX + FOVWidth) / 2.;
+		const double  FOVRight = -FOVOffsetX - FOVLeft;
+		const double FOVBottom = -(FOVOffsetY + FOVHeight) / 2.;
+		const double    FOVTop = -(FOVOffsetY + FOVBottom);
+
+		const FVector2D UnscaledFOV(FOVRight - FOVLeft, FOVBottom - FOVTop);
+		const FVector2D FOVCenter(FOVLeft + UnscaledFOV.X * 0.5, FOVTop + UnscaledFOV.Y * 0.5);
+
+		// Support camera overscan percentage
+		const FVector2D OverscanFOV(UnscaledFOV * (1.0f + InSampleState.OverscanPercentage));
+		const FVector2D FOVMin(FOVCenter.X - OverscanFOV.X * .5, FOVCenter.Y - OverscanFOV.Y * .5);
+
+		// FOV values for Tile
+		const double TileFOVLeft   = FOVMin.X + OverscanFOV.X * TileLeft / Width;
+		const double TileFOVRight  = FOVMin.X + OverscanFOV.X* TileRight / Width;
+		const double TileFOVTop    = FOVMin.Y + OverscanFOV.Y * TileTop / Height;
+		const double TileFOVBottom = FOVMin.Y + OverscanFOV.Y * TileBottom / Height;
+
+		InOutProjectionMatrix = IDisplayClusterViewport::MakeProjectionMatrix(TileFOVLeft, TileFOVRight, TileFOVTop, TileFOVBottom, ZNear, ZNear);
+	}
+}
+
+FSceneView* UDisplayClusterMoviePipelineViewportPassBase::GetSceneViewForSampleState(FSceneViewFamily* ViewFamily, FMoviePipelineRenderPassMetrics& InOutSampleState, IViewCalcPayload* OptPayload)
+{
+	const int32 InViewportIndex = InOutSampleState.OutputState.CameraIndex;
+	FString ViewportId;
+	if (GetViewportId(InViewportIndex, ViewportId))
+	{
+		UMoviePipelineExecutorShot* CurrentShot = GetPipeline()->GetActiveShotList()[GetPipeline()->GetCurrentShotIndex()];
+		UMoviePipelineHighResSetting* HighResSettings = GetPipeline()->FindOrAddSettingForShot<UMoviePipelineHighResSetting>(CurrentShot);
+
+		bFrameWarpBlend = bEnabledWarpBlend;
+		if (HighResSettings->TileCount > 1 && bFrameWarpBlend)
+		{
+			// WarpBlend isn't supported when using tiled rendering
+			bFrameWarpBlend = false;
+			UE_LOG(LogMovieRenderPipeline, Warning, TEXT("DCRA WarpBlend isn't supported when using tiling! (Viewport '%s')"), *ViewportId);
+		}
+
+		FIntPoint OffsetMin = FIntPoint::ZeroValue;
+		FIntPoint OffsetMax = FIntPoint::ZeroValue;
+		GetViewportCutOffset(InOutSampleState, OffsetMin, OffsetMax);
+
+		const FIntPoint FullDestSize(InOutSampleState.TileSize.X, InOutSampleState.TileSize.Y);
+		const FIntPoint DestSize = FullDestSize - OffsetMin - OffsetMax;
+
+		FDisplayClusterViewInfo NewView, PrevView;
+		if (CalculateDisplayClusterViewport(InOutSampleState, *ViewportId, DestSize, 0, NewView))
+		{
+			// Update view cache
+			if (!DCViews.Contains(InViewportIndex))
+			{
+				DCViews.Emplace(InViewportIndex, NewView);
+			}
+
+			DCPrevViews.Emplace(InViewportIndex, DCViews[InViewportIndex]);
+			DCViews.Emplace(InViewportIndex, NewView);
+		}
+		else
+		{
+			UE_LOG(LogMovieRenderPipeline, Error, TEXT("DCRA render failed for viewport '%s')"), *ViewportId);
+		}
+	}
+	else
+	{
+		UE_LOG(LogMovieRenderPipeline, Error, TEXT("DCRA render failed: undefined viewport index, skipped"));
+	}
+
+	return Super::GetSceneViewForSampleState(ViewFamily, InOutSampleState, OptPayload);
+}
+
+void UDisplayClusterMoviePipelineViewportPassBase::PostRendererSubmission(const FMoviePipelineRenderPassMetrics& InSampleState, const FMoviePipelinePassIdentifier InPassIdentifier, const int32 InSortingOrder, FCanvas& InCanvas)
+{
+	if (bFrameWarpBlend)
+	{
+		const int32 InViewportIndex = InSampleState.OutputState.CameraIndex;
+		FString ViewportId;
+		if (GetViewportId(InViewportIndex, ViewportId))
+		{
+			if (FRenderTarget* RenderTarget = InCanvas.GetRenderTarget())
+			{
+				if (IDisplayClusterViewportManager* ViewportManager = DisplayClusterRootActorPtr->GetViewportManager())
+				{
+					FIntPoint OffsetMin = FIntPoint::ZeroValue;
+					FIntPoint OffsetMax = FIntPoint::ZeroValue;
+					GetViewportCutOffset(InSampleState, OffsetMin, OffsetMax);
+
+					ENQUEUE_RENDER_COMMAND(DisplayClusterMoviePipelineWarp)([
+						InTextureRHIRef = RenderTarget->GetRenderTargetTexture(),
+						ViewportManagerProxy = ViewportManager->GetProxy(),
+						InViewportId = ViewportId,
+						OffsetMin, OffsetMax](FRHICommandListImmediate& RHICmdList) mutable
+					{
+						DisplayClusterWarpBlendImpl_RenderThread(RHICmdList, InTextureRHIRef, ViewportManagerProxy, InViewportId, OffsetMin, OffsetMax);
+					});
+				}
+			}
+		}
+	}
+
+	// Use Deffered render pass
+	Super::PostRendererSubmission(InSampleState, InPassIdentifier, InSortingOrder, InCanvas);
+}
+
+void UDisplayClusterMoviePipelineViewportPassBase::GetViewportCutOffset(const FMoviePipelineRenderPassMetrics& InSampleState, FIntPoint& OutOffsetMin, FIntPoint& OutOffsetMax) const
+{
+	// Difference between the clipped constrained rect and the tile rect
+	OutOffsetMin = FIntPoint::ZeroValue;
+	OutOffsetMax = FIntPoint::ZeroValue;
+
+	APlayerCameraManager* PlayerCameraManager = GetPipeline()->GetWorld()->GetFirstPlayerController()->PlayerCameraManager;
+	if (PlayerCameraManager && PlayerCameraManager->GetCameraCacheView().bConstrainAspectRatio)
+	{
+		const FMinimalViewInfo CameraCache = PlayerCameraManager->GetCameraCacheView();
+
+		// Taking overscan into account.
+		FIntPoint FullOutputSize = UMoviePipelineBlueprintLibrary::GetEffectiveOutputResolution(GetPipeline()->GetPipelineMasterConfig(), GetPipeline()->GetActiveShotList()[GetPipeline()->GetCurrentShotIndex()]);
+
+		float OutputSizeAspectRatio = FullOutputSize.X / (float)FullOutputSize.Y;
+		const FIntPoint ConstrainedFullSize = CameraCache.AspectRatio > OutputSizeAspectRatio ?
+			FIntPoint(FullOutputSize.X, FMath::CeilToInt((double)FullOutputSize.X / (double)CameraCache.AspectRatio)) :
+			FIntPoint(FMath::CeilToInt(CameraCache.AspectRatio * FullOutputSize.Y), FullOutputSize.Y);
+
+		const FIntPoint TileViewMin = InSampleState.OverlappedOffset;
+		const FIntPoint TileViewMax = TileViewMin + InSampleState.BackbufferSize;
+
+		// Camera ratio constrained rect, clipped by the tile rect
+		FIntPoint ConstrainedViewMin = (FullOutputSize - ConstrainedFullSize) / 2;
+		FIntPoint ConstrainedViewMax = ConstrainedViewMin + ConstrainedFullSize;
+		ConstrainedViewMin = FIntPoint(FMath::Clamp(ConstrainedViewMin.X, TileViewMin.X, TileViewMax.X),
+			FMath::Clamp(ConstrainedViewMin.Y, TileViewMin.Y, TileViewMax.Y));
+		ConstrainedViewMax = FIntPoint(FMath::Clamp(ConstrainedViewMax.X, TileViewMin.X, TileViewMax.X),
+			FMath::Clamp(ConstrainedViewMax.Y, TileViewMin.Y, TileViewMax.Y));
+
+		// Difference between the clipped constrained rect and the tile rect
+		OutOffsetMin = ConstrainedViewMin - TileViewMin;
+		OutOffsetMax = TileViewMax - ConstrainedViewMax;
+
+		// Support camera overscan
+		const float BorderRatio = FMath::Clamp((InSampleState.OverscanPercentage / (1.0f + InSampleState.OverscanPercentage)) / 2.0f, 0.0f, 0.5f);
+		const FIntPoint InnerSize = FullOutputSize - OutOffsetMin - OutOffsetMax;
+		const int32 OverscanBorderWidth = FMath::RoundToInt(InnerSize.X* BorderRatio);
+		const int32 OverscanBorderHeight = FMath::RoundToInt(InnerSize.Y * BorderRatio);
+
+		OutOffsetMin.X += OverscanBorderWidth;
+		OutOffsetMin.Y += OverscanBorderHeight;
+		OutOffsetMax.X += OverscanBorderWidth;
+		OutOffsetMax.Y += OverscanBorderHeight;
+	}
+}
+
+ADisplayClusterRootActor* UDisplayClusterMoviePipelineViewportPassBase::GetRootActor(const UDisplayClusterMoviePipelineSettings& InDCSettings)
+{
+	if (DisplayClusterRootActorPtr)
+	{
+		const bool bWorldEqual = DisplayClusterRootActorPtr->GetWorld() == GetWorld();
+		if (!bWorldEqual)
+		{
+			DisplayClusterRootActorPtr = nullptr;
+		}
+	}
+
+	if (DisplayClusterRootActorPtr == nullptr)
+	{
+		DisplayClusterRootActorPtr = InDCSettings.GetRootActor(GetWorld());
+	}
+
+	return DisplayClusterRootActorPtr;
+}
+
+bool UDisplayClusterMoviePipelineViewportPassBase::InitializeDisplayCluster()
+{
+	// Reset all runtime values
+	{
+		DCViews.Empty();
+		DCPrevViews.Empty();
+
+		bFrameInitialized = false;
+		bFrameWarpBlend = false;
+
+		DisplayClusterRootActorPtr = nullptr;
+
+		bDisplayClusterInitialized = false;
+
+		DisplayClusterViewports.Empty();
+	}
+
+	UMoviePipelineExecutorShot* CurrentShot = GetPipeline()->GetActiveShotList()[GetPipeline()->GetCurrentShotIndex()];
+	UDisplayClusterMoviePipelineSettings* DisplayClusterCSettings = GetPipeline()->FindOrAddSettingForShot<UDisplayClusterMoviePipelineSettings>(CurrentShot);
+
+	// Initialize
+	if (DisplayClusterCSettings && GetRootActor(*DisplayClusterCSettings) && DisplayClusterCSettings->GetViewports(GetWorld(), DisplayClusterViewports))
+	{
+		bDisplayClusterInitialized = true;
+
+		return true;
+	}
+
+	UE_LOG(LogMovieRenderPipeline, Error, TEXT("Failed to find nDisplay viewports"));
+
+	return false;
+}
+
+void UDisplayClusterMoviePipelineViewportPassBase::ReleaseDisplayCluster()
+{
+	bDisplayClusterInitialized = false;
+	DisplayClusterViewports.Empty();
+	DCViews.Empty();
+}
+
+bool UDisplayClusterMoviePipelineViewportPassBase::GetViewportId(int32 InViewportIndex, FString& OutViewportId) const
+{
+	if (bDisplayClusterInitialized && InViewportIndex >= 0 && InViewportIndex < DisplayClusterViewports.Num())
+	{
+		OutViewportId = DisplayClusterViewports[InViewportIndex];
+
+		return true;
+	}
+
+	return false;
+}
+
+bool UDisplayClusterMoviePipelineViewportPassBase::CalculateDisplayClusterViewport(const FMoviePipelineRenderPassMetrics& InSampleState, const FString& InViewportId, const FIntPoint& InDestSize, const uint32 InContextNum, FDisplayClusterViewInfo& OutView)
+{
+	const EDisplayClusterRenderFrameMode RenderFrameMode = EDisplayClusterRenderFrameMode::Mono;
+
+	bFrameInitialized = false;
+	if (IDisplayClusterViewportManager* ViewportManager = DisplayClusterRootActorPtr->GetViewportManager())
+	{
+		if (UWorld* InWorld = GetWorld())
+		{
+			const AWorldSettings* WorldSettings = InWorld->GetWorldSettings();
+			const float InWorldToMeters = (WorldSettings) ? WorldSettings->WorldToMeters : 100.f;
+
+			if (UGameViewportClient* GameViewportClient = InWorld->GetGameViewport())
+			{
+				// Update local node viewports (update\create\delete) and build new render frame
+				if (ViewportManager->UpdateCustomConfiguration(RenderFrameMode, DisplayClusterViewports, DisplayClusterRootActorPtr))
+				{
+					// Force custom viewport size and settings:
+					if (IDisplayClusterViewport* ViewportPtr = ViewportManager->FindViewport(InViewportId))
+					{
+						FDisplayClusterViewport_RenderSettings RenderSettings;
+						RenderSettings.Rect = FIntRect(FIntPoint(0, 0), InDestSize);
+						RenderSettings.CaptureMode = EDisplayClusterViewportCaptureMode::MoviePipeline;
+
+						if (!bFrameWarpBlend)
+						{
+							// Create internal resources only when warpblend will be used
+							RenderSettings.bSkipRendering = true;
+						}
+
+						// Disable DCRA frustum modifier features
+						RenderSettings.bDisableCustomFrustumFeature = true;
+						RenderSettings.bDisableFrustumOverscanFeature = true;
+
+						ViewportPtr->SetRenderSettings(RenderSettings);
+
+						FDisplayClusterRenderFrame RenderFrame;
+						if (ViewportManager->BeginNewFrame(GameViewportClient->Viewport, InWorld, RenderFrame))
+						{
+							// InDestSize
+							bFrameInitialized = true;
+
+							if (bFrameWarpBlend)
+							{
+								// Initialize proxy only when warpblend will be used
+								ViewportManager->FinalizeNewFrame();
+							}
+
+							if (!ViewportPtr->GetProjectionPolicy().IsValid())
+							{
+								// ignore viewports with uninitialized prj policy
+								return false;
+							}
+
+							const TArray<FDisplayClusterViewport_Context>& ViewportContexts = ViewportPtr->GetContexts();
+							const FDisplayClusterViewport_Context& ViewportContext = ViewportContexts[InContextNum];
+
+							UDisplayClusterCameraComponent* ViewCamera = nullptr;
+							{
+								// Get camera ID assigned to the viewport
+								const FString& CameraId = ViewportPtr->GetRenderSettings().CameraId;
+
+								// Get camera component assigned to the viewport (or default camera if nothing assigned)
+								ViewCamera = (CameraId.IsEmpty() ?
+									DisplayClusterRootActorPtr->GetDefaultCamera() :
+									DisplayClusterRootActorPtr->GetComponentByName<UDisplayClusterCameraComponent>(CameraId));
+							}
+
+							if (!ViewCamera)
+							{
+								UE_LOG(LogMovieRenderPipeline, Error, TEXT("No camera found for viewport '%s'"), *ViewportPtr->GetId());
+								return false;
+							}
+
+							// Get the actual camera settings
+							const float CfgEyeDist = ViewCamera->GetInterpupillaryDistance();
+							const bool  CfgEyeSwap = ViewCamera->GetSwapEyes();
+							const float CfgNCP = GNearClippingPlane;
+							const EDisplayClusterEyeStereoOffset CfgEyeOffset = ViewCamera->GetStereoOffset();
+
+							// Calculate eye offset considering the world scale
+							const float EyeOffset = CfgEyeDist / 2.f;
+							const float EyeOffsetValues[] = { -EyeOffset, 0.f, EyeOffset };
+
+							enum class EDisplayClusterEyeType : uint8
+							{
+								StereoLeft = 0,
+								Mono = 1,
+								StereoRight = 2,
+								COUNT
+							};
+
+							// Decode current eye type
+							const EDisplayClusterEyeType DCEyeType = (ViewportContexts.Num() > 1)
+								? ((InContextNum == 0) ? EDisplayClusterEyeType::StereoLeft : EDisplayClusterEyeType::StereoRight)
+								: EDisplayClusterEyeType::Mono;
+
+							const int32 EyeIndex = (uint8)DCEyeType;
+
+							float PassOffset = 0.f;
+							float PassOffsetSwap = 0.f;
+
+							switch (DCEyeType)
+							{
+							case EDisplayClusterEyeType::Mono:
+							{
+								// For monoscopic camera let's check if the "force offset" feature is used
+								// * Force left (-1) ==> 0 left eye
+								// * Force right (1) ==> 2 right eye
+								// * Default (0) ==> 1 mono
+								const int32 EyeOffsetIdx = (CfgEyeOffset == EDisplayClusterEyeStereoOffset::None ? 0 : (CfgEyeOffset == EDisplayClusterEyeStereoOffset::Left ? -1 : 1));
+
+								PassOffset = EyeOffsetValues[EyeOffsetIdx + 1];
+								// Eye swap is not available for monoscopic so just save the value
+								PassOffsetSwap = PassOffset;
+							}
+							break;
+
+							default:
+							{
+								// For stereo camera we can only swap eyes if required (no "force offset" allowed)
+								PassOffset = EyeOffsetValues[EyeIndex];
+								PassOffsetSwap = (CfgEyeSwap ? -PassOffset : PassOffset);
+							}
+							break;
+							}
+
+							FVector ViewOffset = FVector::ZeroVector;
+							if (ViewCamera)
+							{
+								// View base location
+								OutView.ViewLocation = ViewCamera->GetComponentLocation();
+								OutView.ViewRotation = ViewCamera->GetComponentRotation();
+								// Apply computed offset to the view location
+								const FQuat EyeQuat = OutView.ViewRotation.Quaternion();
+								ViewOffset = EyeQuat.RotateVector(FVector(0.0f, PassOffsetSwap, 0.0f));
+								OutView.ViewLocation += ViewOffset;
+							}
+
+							// Perform view calculations on a policy side
+							if (ViewportPtr->CalculateView(InContextNum, OutView.ViewLocation, OutView.ViewRotation, ViewOffset, InWorldToMeters, CfgNCP, CfgNCP))
+							{
+								OutView.ProjectionMatrix = FMatrix::Identity;
+
+								return ViewportPtr->GetOwner().IsSceneOpened() && ViewportPtr->GetProjectionMatrix(InContextNum, OutView.ProjectionMatrix);
+							}
+						}
+					}
+				}
+			}
+		}
+	}
+
+	return false;
+}
diff --git a/Engine/Plugins/Runtime/nDisplay/Source/DisplayClusterMoviePipeline/Private/DisplayClusterMoviePipelineViewportPass.h b/Engine/Plugins/Runtime/nDisplay/Source/DisplayClusterMoviePipeline/Private/DisplayClusterMoviePipelineViewportPass.h
new file mode 100644
--- /dev/null
+++ b/Engine/Plugins/Runtime/nDisplay/Source/DisplayClusterMoviePipeline/Private/DisplayClusterMoviePipelineViewportPass.h
@@ -0,0 +1,192 @@
+// Copyright Epic Games, Inc. All Rights Reserved.
+#pragma once
+
+#include "MoviePipelineDeferredPasses.h"
+
+#include "DisplayClusterMoviePipelineViewportPass.generated.h"
+
+class ADisplayClusterRootActor;
+
+/**
+ * nDisplay viewport render pass (Lit)
+ */
+UCLASS(BlueprintType)
+class DISPLAYCLUSTERMOVIEPIPELINE_API UDisplayClusterMoviePipelineViewportPassBase
+	: public UMoviePipelineDeferredPassBase
+{
+	GENERATED_BODY()
+
+	struct FDisplayClusterViewInfo
+	{
+		FVector  ViewLocation;
+		FRotator ViewRotation;
+		FMatrix  ProjectionMatrix;
+	};
+
+public:
+	UDisplayClusterMoviePipelineViewportPassBase()
+		: UMoviePipelineDeferredPassBase(), RenderPassName(TEXT("nDisplayLit"))
+	{
+		PassIdentifier = FMoviePipelinePassIdentifier(RenderPassName);
+	}
+
+	UDisplayClusterMoviePipelineViewportPassBase(const FString& InRenderPassName)
+		: UMoviePipelineDeferredPassBase(), RenderPassName(InRenderPassName)
+	{
+		PassIdentifier = FMoviePipelinePassIdentifier(RenderPassName);
+	}
+
+public:
+	// UMoviePipelineDeferredPassBase
+	virtual void SetupImpl(const MoviePipeline::FMoviePipelineRenderPassInitSettings& InPassInitSettings) override;
+	virtual void TeardownImpl() override;
+
+	virtual FSceneView* GetSceneViewForSampleState(FSceneViewFamily* ViewFamily, FMoviePipelineRenderPassMetrics& InOutSampleState, IViewCalcPayload* OptPayload) override;
+	virtual void ModifyProjectionMatrixForTiling(const FMoviePipelineRenderPassMetrics& InSampleState, const bool bInOrthographic, FMatrix& InOutProjectionMatrix, float& OutDoFSensorScale) const override;
+	virtual void PostRendererSubmission(const FMoviePipelineRenderPassMetrics& InSampleState, const FMoviePipelinePassIdentifier InPassIdentifier, const int32 InSortingOrder, FCanvas& InCanvas) override;
+
+	virtual int32 GetNumCamerasToRender() const override;
+	virtual FString GetCameraName(const int32 InCameraIndex) const override;
+	virtual FString GetCameraNameOverride(const int32 InCameraIndex) const override;
+
+	virtual UE::MoviePipeline::FImagePassCameraViewData GetCameraInfo(FMoviePipelineRenderPassMetrics& InOutSampleState, IViewCalcPayload* OptPayload = nullptr) const override;
+	virtual void BlendPostProcessSettings(FSceneView* InView, FMoviePipelineRenderPassMetrics& InOutSampleState, IViewCalcPayload* OptPayload = nullptr) override;
+	// ~UMoviePipelineDeferredPassBase
+
+#if WITH_EDITOR
+	virtual FText GetDisplayText() const override { return NSLOCTEXT("MovieRenderPipeline", "DCViewportBasePassSetting_DisplayName_Lit", "nDisplay Rendering"); }
+#endif
+
+	virtual int32 GetOutputFileSortingOrder() const override { return 2; }
+	virtual bool bIsEnabledWarpBlend() const { return true; }
+
+protected:
+	bool GetViewportId(int32 InViewportIndex, FString& OutViewportId) const;
+	ADisplayClusterRootActor* GetRootActor(const class UDisplayClusterMoviePipelineSettings& InDCSettings);
+
+	bool CalculateDisplayClusterViewport(const FMoviePipelineRenderPassMetrics& InSampleState, const FString& InViewportId, const FIntPoint& InDestSize, const uint32 InContextNum, FDisplayClusterViewInfo& OutView);
+
+	bool InitializeDisplayCluster();
+	void ReleaseDisplayCluster();
+
+	void GetViewportCutOffset(const FMoviePipelineRenderPassMetrics& InSampleState, FIntPoint& OutOffsetMin, FIntPoint& OutOffsetMax) const;
+
+public:
+	// Allow warp blend for this pass
+	UPROPERTY(EditAnywhere, BlueprintReadWrite, Category = "nDisplay")
+	bool bEnabledWarpBlend = true;
+
+private:
+	// Render pass name
+	const FString RenderPassName;
+
+	// DCRA for current shot
+	ADisplayClusterRootActor* DisplayClusterRootActorPtr = nullptr;
+
+	// Names of viewports to render in the current shot
+	TArray<FString> DisplayClusterViewports;
+
+	// Cache view info for viewports
+	TMap<int32, FDisplayClusterViewInfo> DCViews;
+	TMap<int32, FDisplayClusterViewInfo> DCPrevViews;
+
+	// Runtime flags
+	bool bFrameInitialized = true;
+	bool bFrameWarpBlend = false;
+	bool bDisplayClusterInitialized = false;
+};
+
+/**
+ * nDisplay viewport render pass (Unlit)
+ */
+UCLASS(BlueprintType)
+class UDisplayClusterMoviePipelineViewportPass_Unlit
+	: public UDisplayClusterMoviePipelineViewportPassBase
+{
+	GENERATED_BODY()
+
+public:
+	UDisplayClusterMoviePipelineViewportPass_Unlit() : UDisplayClusterMoviePipelineViewportPassBase(TEXT("nDisplayUnlit"))
+	{ }
+#if WITH_EDITOR
+	virtual FText GetDisplayText() const override { return NSLOCTEXT("MovieRenderPipeline", "DCViewportBasePassSetting_DisplayName_Unlit", "nDisplay Rendering (Unlit)"); }
+#endif
+
+	virtual void GetViewShowFlags(FEngineShowFlags& OutShowFlag, EViewModeIndex& OutViewModeIndex) const override
+	{
+		OutShowFlag = FEngineShowFlags(EShowFlagInitMode::ESFIM_Game);
+		OutViewModeIndex = EViewModeIndex::VMI_Unlit;
+	}
+};
+
+/**
+ * nDisplay viewport render pass (DetailLightingOnly)
+ */
+UCLASS(BlueprintType)
+class UDisplayClusterMoviePipelineViewportPass_DetailLighting
+	: public UDisplayClusterMoviePipelineViewportPassBase
+{
+	GENERATED_BODY()
+
+public:
+	UDisplayClusterMoviePipelineViewportPass_DetailLighting() : UDisplayClusterMoviePipelineViewportPassBase(TEXT("nDisplayDetailLightingOnly"))
+	{ }
+#if WITH_EDITOR
+	virtual FText GetDisplayText() const override { return NSLOCTEXT("MovieRenderPipeline", "DCViewportBasePassSetting_DisplayName_DetailLighting", "nDisplay Rendering (Detail Lighting)"); }
+#endif
+	virtual void GetViewShowFlags(FEngineShowFlags& OutShowFlag, EViewModeIndex& OutViewModeIndex) const override
+	{
+		OutShowFlag = FEngineShowFlags(EShowFlagInitMode::ESFIM_Game);
+		OutShowFlag.SetLightingOnlyOverride(true);
+		OutViewModeIndex = EViewModeIndex::VMI_Lit_DetailLighting;
+	}
+	virtual int32 GetOutputFileSortingOrder() const override { return 2; }
+};
+
+/**
+ * nDisplay viewport render pass (LightingOnly)
+ */
+UCLASS(BlueprintType)
+class UDisplayClusterMoviePipelineViewportPass_LightingOnly
+	: public UDisplayClusterMoviePipelineViewportPassBase
+{
+	GENERATED_BODY()
+
+public:
+	UDisplayClusterMoviePipelineViewportPass_LightingOnly() : UDisplayClusterMoviePipelineViewportPassBase(TEXT("nDisplayLightingOnly"))
+	{ }
+#if WITH_EDITOR
+	virtual FText GetDisplayText() const override { return NSLOCTEXT("MovieRenderPipeline", "DCViewportBasePassSetting_DisplayName_LightingOnly", "nDisplay Rendering (Lighting Only)"); }
+#endif
+	virtual void GetViewShowFlags(FEngineShowFlags& OutShowFlag, EViewModeIndex& OutViewModeIndex) const override
+	{
+		OutShowFlag = FEngineShowFlags(EShowFlagInitMode::ESFIM_Game);
+		OutShowFlag.SetLightingOnlyOverride(true);
+		OutViewModeIndex = EViewModeIndex::VMI_LightingOnly;
+	}
+	virtual int32 GetOutputFileSortingOrder() const override { return 2; }
+};
+
+/**
+ * nDisplay viewport render pass (ReflectionsOnly)
+ */
+UCLASS(BlueprintType)
+class UDisplayClusterMoviePipelineViewportPass_ReflectionsOnly
+	: public UDisplayClusterMoviePipelineViewportPassBase
+{
+	GENERATED_BODY()
+
+public:
+	UDisplayClusterMoviePipelineViewportPass_ReflectionsOnly() : UDisplayClusterMoviePipelineViewportPassBase(TEXT("nDisplayReflectionsOnly"))
+	{ }
+#if WITH_EDITOR
+	virtual FText GetDisplayText() const override { return NSLOCTEXT("MovieRenderPipeline", "DCViewportBasePassSetting_DisplayName_ReflectionsOnly", "nDisplay Rendering (Reflections Only)"); }
+#endif
+	virtual void GetViewShowFlags(FEngineShowFlags& OutShowFlag, EViewModeIndex& OutViewModeIndex) const override
+	{
+		OutShowFlag = FEngineShowFlags(EShowFlagInitMode::ESFIM_Game);
+		OutShowFlag.SetReflectionOverride(true);
+		OutViewModeIndex = EViewModeIndex::VMI_ReflectionOverride;
+	}
+	virtual int32 GetOutputFileSortingOrder() const override { return 2; }
+};
diff --git a/Engine/Plugins/Runtime/nDisplay/Source/DisplayClusterProjection/Private/Policy/MPCDI/DisplayClusterProjectionMPCDIPolicy.cpp b/Engine/Plugins/Runtime/nDisplay/Source/DisplayClusterProjection/Private/Policy/MPCDI/DisplayClusterProjectionMPCDIPolicy.cpp
--- a/Engine/Plugins/Runtime/nDisplay/Source/DisplayClusterProjection/Private/Policy/MPCDI/DisplayClusterProjectionMPCDIPolicy.cpp
+++ b/Engine/Plugins/Runtime/nDisplay/Source/DisplayClusterProjection/Private/Policy/MPCDI/DisplayClusterProjectionMPCDIPolicy.cpp
@@ -317,9 +317,10 @@ void FDisplayClusterProjectionMPCDIPolicy::ApplyWarpBlend_RenderThread(FRHIComma
 					}
 				}
 
-				if (RefViewportManagerProxy.GetLightCardManager_RenderThread().IsValid() && RefViewportManagerProxy.GetLightCardManager_RenderThread()->HasUVLightCards_RenderThread())
+				TSharedPtr<IDisplayClusterViewportLightCardManager, ESPMode::ThreadSafe> LightCardManager = RefViewportManagerProxy.GetLightCardManager_RenderThread();
+				if (LightCardManager.IsValid())
 				{
-					ShaderICVFX.UVLightCardMap = RefViewportManagerProxy.GetLightCardManager_RenderThread()->GetUVLightCardMap_RenderThread();
+					ShaderICVFX.UVLightCardMap = LightCardManager->GetUVLightCardMap_RenderThread();
 				}
 
 				// Initialize shader input data
diff --git a/Engine/Plugins/Runtime/nDisplay/nDisplay.uplugin b/Engine/Plugins/Runtime/nDisplay/nDisplay.uplugin
--- a/Engine/Plugins/Runtime/nDisplay/nDisplay.uplugin
+++ b/Engine/Plugins/Runtime/nDisplay/nDisplay.uplugin
@@ -159,6 +159,11 @@
 				"Win64",
 				"Linux"
 			]
+		},
+		{
+			"Name": "DisplayClusterMoviePipeline",
+			"Type": "UncookedOnly",
+			"LoadingPhase": "Default"
 		}
 	],
 	"Plugins": [
@@ -189,6 +194,10 @@
 		{
 			"Name": "VirtualProductionUtilities",
 			"Enabled": true
+		},
+		{
+			"Name": "MovieRenderPipeline",
+			"Enabled": true
 		}
 	]
 }
