nDisplay: Fix non-unity build with light card manager.

diff --git a/Engine/Plugins/Runtime/nDisplay/Source/DisplayCluster/Private/Render/Viewport/LightCard/DisplayClusterViewportLightCardManager.cpp b/Engine/Plugins/Runtime/nDisplay/Source/DisplayCluster/Private/Render/Viewport/LightCard/DisplayClusterViewportLightCardManager.cpp
--- a/Engine/Plugins/Runtime/nDisplay/Source/DisplayCluster/Private/Render/Viewport/LightCard/DisplayClusterViewportLightCardManager.cpp
+++ b/Engine/Plugins/Runtime/nDisplay/Source/DisplayCluster/Private/Render/Viewport/LightCard/DisplayClusterViewportLightCardManager.cpp
@@ -11,6 +11,7 @@
 #include "Render/Viewport/RenderFrame/DisplayClusterRenderFrame.h"
 
 #include "PreviewScene.h"
+#include "SceneInterface.h"
 
 /** Console variable used to control the size of the UV light card map texture */
 static TAutoConsoleVariable<int32> CVarUVLightCardTextureSize(
diff --git a/Engine/Plugins/Runtime/nDisplay/Source/DisplayCluster/Private/Render/Viewport/LightCard/DisplayClusterViewportLightCardManager.h b/Engine/Plugins/Runtime/nDisplay/Source/DisplayCluster/Private/Render/Viewport/LightCard/DisplayClusterViewportLightCardManager.h
--- a/Engine/Plugins/Runtime/nDisplay/Source/DisplayCluster/Private/Render/Viewport/LightCard/DisplayClusterViewportLightCardManager.h
+++ b/Engine/Plugins/Runtime/nDisplay/Source/DisplayCluster/Private/Render/Viewport/LightCard/DisplayClusterViewportLightCardManager.h
@@ -11,6 +11,7 @@
 class ADisplayClusterLightCardActor;
 class FDisplayClusterViewportManager;
 class FPreviewScene;
+class FSceneInterface;
 
 /** A render targetable texture resource used to render the UV light cards to */
 class FDisplayClusterLightCardMap : public FTexture, public FRenderTarget
