nDisplay: Add lightcard labels feature. Labels can be toggled and scaled for all light cards in the light card editor. The labels will only show up in scene capture for the preview and in 3d space for the wall.
- Add new toolbar to configure light card labels and refactor operator panel extensions to support instance hooks.
- Fix a few crashes under mesh projection renderer.
- Fix a crash when exiting an MU session with the light card editor open.

diff --git a/Engine/Plugins/Runtime/nDisplay/Source/DisplayCluster/DisplayCluster.Build.cs b/Engine/Plugins/Runtime/nDisplay/Source/DisplayCluster/DisplayCluster.Build.cs
--- a/Engine/Plugins/Runtime/nDisplay/Source/DisplayCluster/DisplayCluster.Build.cs
+++ b/Engine/Plugins/Runtime/nDisplay/Source/DisplayCluster/DisplayCluster.Build.cs
@@ -40,6 +40,7 @@ public class DisplayCluster : ModuleRules
 				"Slate",
 				"SlateCore",
 				"Sockets",
+				"UMG"
 			});
 
 		if (Target.bBuildEditor == true)
diff --git a/Engine/Plugins/Runtime/nDisplay/Source/DisplayCluster/Private/Game/EngineClasses/Scene/DisplayClusterLabelComponent.cpp b/Engine/Plugins/Runtime/nDisplay/Source/DisplayCluster/Private/Game/EngineClasses/Scene/DisplayClusterLabelComponent.cpp
new file mode 100644
--- /dev/null
+++ b/Engine/Plugins/Runtime/nDisplay/Source/DisplayCluster/Private/Game/EngineClasses/Scene/DisplayClusterLabelComponent.cpp
@@ -0,0 +1,149 @@
+﻿// Copyright Epic Games, Inc. All Rights Reserved.
+
+#include "Components/DisplayClusterLabelComponent.h"
+
+#include "DisplayClusterLabelWidget.h"
+#include "DisplayClusterLightCardActor.h"
+#include "DisplayClusterRootActor.h"
+
+#include "EngineUtils.h"
+#include "Components/TextRenderComponent.h"
+#include "Components/WidgetComponent.h"
+#include "Kismet/KismetMathLibrary.h"
+
+UDisplayClusterLabelComponent::UDisplayClusterLabelComponent()
+{
+	PrimaryComponentTick.bCanEverTick = true;
+	
+	WidgetComponent = CreateDefaultSubobject<UWidgetComponent>(TEXT("LabelWidget"));
+	
+	// Keeps this visible even if its clipping through the owner.
+	// Requires the owner have a transparent blend mode on its material.
+	WidgetComponent->SetBlendMode(EWidgetBlendMode::Transparent);
+	WidgetComponent->SetTranslucentSortPriority(1.f);
+	WidgetComponent->SetUsingAbsoluteScale(true);
+
+	WidgetClass = TSoftClassPtr<UDisplayClusterLabelWidget>(FSoftClassPath(TEXT("/nDisplay/Widgets/BP_DisplayClusterLabelWidget.BP_DisplayClusterLabelWidget_C")));
+
+	SetRelativeRotation(FRotator(90.f, 0.f, -90.f));
+	SetVisibility(false);
+}
+
+void UDisplayClusterLabelComponent::SetRootActor(ADisplayClusterRootActor* InActor)
+{
+	RootActor = InActor;
+}
+
+void UDisplayClusterLabelComponent::SetWidgetScale(float NewValue)
+{
+	WidgetScale = NewValue;
+}
+
+void UDisplayClusterLabelComponent::TickComponent(float DeltaTime, ELevelTick TickType,
+                                                  FActorComponentTickFunction* ThisTickFunction)
+{
+	Super::TickComponent(DeltaTime, TickType, ThisTickFunction);
+
+	CheckForVisibilityChange();
+	
+	if (IsVisible())
+	{
+		UpdateWidgetComponent();
+	}
+}
+
+void UDisplayClusterLabelComponent::OnRegister()
+{
+	Super::OnRegister();
+
+	CheckForVisibilityChange();
+	
+	WidgetComponent->SetupAttachment(this);
+	WidgetComponent->SetWidgetClass(WidgetClass.LoadSynchronous());
+}
+
+#if WITH_EDITOR
+
+void UDisplayClusterLabelComponent::PostEditChangeProperty(FPropertyChangedEvent& PropertyChangedEvent)
+{
+	Super::PostEditChangeProperty(PropertyChangedEvent);
+
+	if (PropertyChangedEvent.GetPropertyName() == GET_MEMBER_NAME_CHECKED(UDisplayClusterLabelComponent, WidgetClass))
+	{
+		WidgetComponent->SetWidgetClass(WidgetClass.Get());
+	}
+}
+
+#endif
+
+void UDisplayClusterLabelComponent::UpdateWidgetComponent()
+{
+	FString LabelString;
+	if (const ADisplayClusterLightCardActor* LightCardActor = Cast<ADisplayClusterLightCardActor>(GetOwner()))
+	{
+#if WITH_EDITOR
+		LabelString = LightCardActor->GetActorLabel();
+#else
+		LabelString = LightCardActor->GetName();
+#endif
+		
+		if (!RootActor.IsValid())
+		{
+			// The light card editor normally will have set this, but if we're running as -game then it will
+			// be set from the DCRA instead
+			SetRootActor(LightCardActor->GetRootActorOwner().Get());
+		}
+		
+		const USceneComponent* Viewpoint = RootActor.IsValid() ? RootActor->GetCommonViewPoint() : nullptr;
+		if (Viewpoint != nullptr)
+		{
+			const FTransform LightCardTransform = LightCardActor->GetLightCardTransform();
+			const FVector LightCardLocation = LightCardTransform.GetLocation();
+			const FVector DestinationPoint = Viewpoint->GetComponentLocation();
+
+			// Keep the text facing up.
+			FRotator Rotation = UKismetMathLibrary::FindLookAtRotation(LightCardLocation, DestinationPoint);
+			Rotation.Roll = 0.f;
+			
+			WidgetComponent->SetWorldRotation(Rotation);
+
+			// Keep the label the same size regardless of distance to the view origin.
+			const float Distance = (DestinationPoint - LightCardLocation).Length();
+
+			const float BaseLabelScale = 0.0025f;
+			const FVector VectorScale(Distance * BaseLabelScale * WidgetScale);
+			
+			WidgetComponent->SetWorldScale3D(VectorScale);
+		}
+	}
+	else
+	{
+		LabelString = TEXT("Label");
+	}
+
+	if (WidgetComponent->GetWidgetClass() != WidgetClass.Get())
+	{
+		WidgetComponent->SetWidgetClass(WidgetClass.Get());
+		WidgetComponent->InitWidget();
+	}
+
+	if (UDisplayClusterLabelWidget* Widget = Cast<UDisplayClusterLabelWidget>(WidgetComponent->GetWidget()))
+	{
+		Widget->SetLabelText(FText::FromString(LabelString));
+	}
+
+#if WITH_EDITOR
+	WidgetComponent->SetVisibleInSceneCaptureOnly(GEditor != nullptr);
+#endif
+}
+
+void UDisplayClusterLabelComponent::CheckForVisibilityChange()
+{
+	// MU doesn't seem to register nested sub component visibility changes, and OnVisibilityChanged won't
+	// fire on the MU client so we have to manually check here.
+
+	if (GetVisibleFlag() != WidgetComponent->GetVisibleFlag())
+	{
+		WidgetComponent->SetVisibility(GetVisibleFlag());
+	}
+}
diff --git a/Engine/Plugins/Runtime/nDisplay/Source/DisplayCluster/Private/Game/EngineClasses/Scene/DisplayClusterLabelWidget.h b/Engine/Plugins/Runtime/nDisplay/Source/DisplayCluster/Private/Game/EngineClasses/Scene/DisplayClusterLabelWidget.h
new file mode 100644
--- /dev/null
+++ b/Engine/Plugins/Runtime/nDisplay/Source/DisplayCluster/Private/Game/EngineClasses/Scene/DisplayClusterLabelWidget.h
@@ -0,0 +1,22 @@
+﻿// Copyright Epic Games, Inc. All Rights Reserved.
+
+#pragma once
+
+#include "CoreMinimal.h"
+#include "Blueprint/UserWidget.h"
+
+#include "DisplayClusterLabelWidget.generated.h"
+
+/**
+ * A base widget to display on actor labels.
+ */
+UCLASS(Abstract)
+class DISPLAYCLUSTER_API UDisplayClusterLabelWidget : public UUserWidget
+{
+public:
+	GENERATED_BODY()
+
+	UFUNCTION(BlueprintImplementableEvent, Category=Label)
+	void SetLabelText(const FText& InText);
+	
+};
diff --git a/Engine/Plugins/Runtime/nDisplay/Source/DisplayCluster/Private/Game/EngineClasses/Scene/DisplayClusterLightCardActor.cpp b/Engine/Plugins/Runtime/nDisplay/Source/DisplayCluster/Private/Game/EngineClasses/Scene/DisplayClusterLightCardActor.cpp
--- a/Engine/Plugins/Runtime/nDisplay/Source/DisplayCluster/Private/Game/EngineClasses/Scene/DisplayClusterLightCardActor.cpp
+++ b/Engine/Plugins/Runtime/nDisplay/Source/DisplayCluster/Private/Game/EngineClasses/Scene/DisplayClusterLightCardActor.cpp
@@ -1,7 +1,9 @@
 // Copyright Epic Games, Inc. All Rights Reserved.
 
 #include "DisplayClusterLightCardActor.h"
+#include "DisplayClusterRootActor.h"
 
+#include "Components/DisplayClusterLabelComponent.h"
 #include "Components/StaticMeshComponent.h"
 #include "GameFramework/SpringArmComponent.h"
 #include "Materials/MaterialInstanceDynamic.h"
@@ -71,6 +73,9 @@ ADisplayClusterLightCardActor::ADisplayClusterLightCardActor(const FObjectInitia
 	LightCardComponent->SetStaticMesh(PlaneObj.Object);
 
 	UpdateLightCardTransform();
+
+	LabelComponent = CreateOptionalDefaultSubobject<UDisplayClusterLabelComponent>(TEXT("Label"), true);
+	LabelComponent->AttachToComponent(LightCardComponent, FAttachmentTransformRules::KeepRelativeTransform);
 }
 
 void ADisplayClusterLightCardActor::OnConstruction(const FTransform& Transform)
@@ -396,6 +401,41 @@ void ADisplayClusterLightCardActor::UpdateLightCardVisibility()
 #endif
 }
 
+void ADisplayClusterLightCardActor::ShowLightCardLabel(bool bValue, float ScaleValue, ADisplayClusterRootActor* InRootActor)
+{
+#if WITH_EDITOR
+	LabelComponent->Modify(false);
+	LightCardComponent->Modify(false);
+#endif
+	LabelComponent->SetVisibility(bValue, true);
+	LabelComponent->SetRootActor(InRootActor);
+	LabelComponent->SetWidgetScale(ScaleValue);
+	
+	const float LightCardSortPriorityLow = -1.f;
+	
+	if (bValue)
+	{
+		// Change the translucent sort priority of the light card component so it will always be shown behind the
+		// label component. When left at 0 the renderer will choose the sort order which may be wrong.
+		// This does rely on the light card material being translucent as well.
+		SavedTranslucencySortPriority = LightCardComponent->TranslucencySortPriority;
+		LightCardComponent->SetTranslucentSortPriority(LightCardSortPriorityLow);
+	}
+	else if (SavedTranslucencySortPriority.IsSet() && *SavedTranslucencySortPriority != LightCardComponent->TranslucencySortPriority &&
+			LightCardComponent->TranslucencySortPriority == LightCardSortPriorityLow)
+	{
+		// Change the sort order back to the original value and only if the user hasn't modified it since activating the label.
+		// It's highly unlikely this will ever be anything besides 0.
+		LightCardComponent->SetTranslucentSortPriority(*SavedTranslucencySortPriority);
+		SavedTranslucencySortPriority.Reset();
+	}
+}
+
+void ADisplayClusterLightCardActor::SetRootActorOwner(ADisplayClusterRootActor* InRootActor)
+{
+	RootActorOwner = MakeWeakObjectPtr(InRootActor);
+}
+
 UStaticMesh* ADisplayClusterLightCardActor::GetStaticMesh() const
 {
 	return LightCardComponent->GetStaticMesh();
diff --git a/Engine/Plugins/Runtime/nDisplay/Source/DisplayCluster/Private/Game/EngineClasses/Scene/DisplayClusterRootActor.cpp b/Engine/Plugins/Runtime/nDisplay/Source/DisplayCluster/Private/Game/EngineClasses/Scene/DisplayClusterRootActor.cpp
--- a/Engine/Plugins/Runtime/nDisplay/Source/DisplayCluster/Private/Game/EngineClasses/Scene/DisplayClusterRootActor.cpp
+++ b/Engine/Plugins/Runtime/nDisplay/Source/DisplayCluster/Private/Game/EngineClasses/Scene/DisplayClusterRootActor.cpp
@@ -659,6 +659,9 @@ void ADisplayClusterRootActor::UpdateLightCardPositions()
 			ADisplayClusterLightCardActor* LightCardActor = Actor.IsValid() ? Cast<ADisplayClusterLightCardActor>(Actor.Get()) : nullptr;
 			if (LightCardActor)
 			{
+				// Set the owner so the light card actor can look the root actor up in certain situations, like adjusting labels
+				// when running as -game
+				LightCardActor->SetRootActorOwner(this);
 				if (LightCardActor->bLockToOwningRootActor)
 				{
 					LightCardActor->SetActorLocation(Location);
diff --git a/Engine/Plugins/Runtime/nDisplay/Source/DisplayCluster/Public/Components/DisplayClusterLabelComponent.h b/Engine/Plugins/Runtime/nDisplay/Source/DisplayCluster/Public/Components/DisplayClusterLabelComponent.h
new file mode 100644
--- /dev/null
+++ b/Engine/Plugins/Runtime/nDisplay/Source/DisplayCluster/Public/Components/DisplayClusterLabelComponent.h
@@ -0,0 +1,68 @@
+﻿// Copyright Epic Games, Inc. All Rights Reserved.
+
+#pragma once
+
+#include "CoreMinimal.h"
+#include "Components/SceneComponent.h"
+
+#include "DisplayClusterLabelComponent.generated.h"
+
+class UDisplayClusterLabelWidget;
+class ADisplayClusterRootActor;
+class UTextRenderComponent;
+class UUserWidget;
+class UWidgetComponent;
+
+/**
+ * A label component specific to nDisplay. Displays a widget with a consistent scale facing the root actor view origin.
+ * Visible only in scene capture when in the editor or in 3d space on the wall.
+ *
+ * The component needs to be transient so its settings aren't saved, but should still transact over multi-user.
+ */
+UCLASS(transient, meta=(BlueprintSpawnableComponent, DisplayClusterMultiUserInclude))
+class DISPLAYCLUSTER_API UDisplayClusterLabelComponent : public USceneComponent
+{
+	GENERATED_BODY()
+
+public:
+	UDisplayClusterLabelComponent();
+
+	/** Set the root actor this label should use when determining transform */
+	void SetRootActor(ADisplayClusterRootActor* InActor);
+
+	/** Set the widget scale */
+	void SetWidgetScale(float NewValue);
+	
+	virtual void TickComponent(float DeltaTime, ELevelTick TickType,
+	                           FActorComponentTickFunction* ThisTickFunction) override;
+	virtual void OnRegister() override;
+
+#if WITH_EDITOR
+public:
+	virtual void PostEditChangeProperty(FPropertyChangedEvent& PropertyChangedEvent) override;
+#endif
+	
+protected:
+	/** Updates the widget component text and transform */
+	void UpdateWidgetComponent();
+
+	/** Verify the widget visibility matches this component's visibility */
+	void CheckForVisibilityChange();
+	
+private:
+	/** The root actor this label needs to know about */
+	UPROPERTY()
+	TWeakObjectPtr<ADisplayClusterRootActor> RootActor;
+	
+	/** The widget component to display for this label */
+	UPROPERTY(VisibleAnywhere, Category=LabelText)
+	TObjectPtr<UWidgetComponent> WidgetComponent;
+
+	/** The widget class to apply to the widget component */
+	UPROPERTY(EditAnywhere, Category=LabelText)
+	TSoftClassPtr<UDisplayClusterLabelWidget> WidgetClass;
+
+	/** A uniform scale to apply to the text which will keep consistency across distance from the label to the root actor view */
+	UPROPERTY(EditAnywhere, BlueprintReadWrite, Category=Label, meta = (AllowPrivateAccess))
+	float WidgetScale = 1.f;
+};
diff --git a/Engine/Plugins/Runtime/nDisplay/Source/DisplayCluster/Public/DisplayClusterLightCardActor.h b/Engine/Plugins/Runtime/nDisplay/Source/DisplayCluster/Public/DisplayClusterLightCardActor.h
--- a/Engine/Plugins/Runtime/nDisplay/Source/DisplayCluster/Public/DisplayClusterLightCardActor.h
+++ b/Engine/Plugins/Runtime/nDisplay/Source/DisplayCluster/Public/DisplayClusterLightCardActor.h
@@ -7,6 +7,8 @@
 
 #include "DisplayClusterLightCardActor.generated.h"
 
+class ADisplayClusterRootActor;
+class UDisplayClusterLabelComponent;
 class USceneComponent;
 class USpringArmComponent;
 class UStaticMeshComponent;
@@ -118,6 +120,15 @@ public:
 	/** Updates the light card actor visibility */
 	void UpdateLightCardVisibility();
 
+	/** Show or hide the light card label  */
+	void ShowLightCardLabel(bool bValue, float ScaleValue, ADisplayClusterRootActor* InRootActor);
+
+	/** Set the current owner of the light card */
+	void SetRootActorOwner(ADisplayClusterRootActor* InRootActor);
+	
+	/** Return the current owner, providing one was set */
+	TWeakObjectPtr<ADisplayClusterRootActor> GetRootActorOwner() const { return RootActorOwner; }
+	
 public:
 
 	/** Radius of light card polar coordinates. Does not include the effect of RadialOffset */
@@ -221,4 +232,15 @@ protected:
 
 	UPROPERTY(EditDefaultsOnly, BlueprintReadOnly, Category = "Default")
 	TObjectPtr<UStaticMeshComponent> LightCardComponent;
+	
+	UPROPERTY(VisibleAnywhere, transient, BlueprintReadOnly, Category = "Default")
+	TObjectPtr<UDisplayClusterLabelComponent> LabelComponent;
+
+	/** The current owner of the light card */
+	TWeakObjectPtr<ADisplayClusterRootActor> RootActorOwner;
+
+private:
+	/** Stores the user translucency value when labels are displayed */
+	TOptional<int32> SavedTranslucencySortPriority;
+	
 };
\ No newline at end of file
diff --git a/Engine/Plugins/Runtime/nDisplay/Source/DisplayClusterLightCardEditor/Private/DisplayClusterLightCardEditorCommands.cpp b/Engine/Plugins/Runtime/nDisplay/Source/DisplayClusterLightCardEditor/Private/DisplayClusterLightCardEditorCommands.cpp
--- a/Engine/Plugins/Runtime/nDisplay/Source/DisplayClusterLightCardEditor/Private/DisplayClusterLightCardEditorCommands.cpp
+++ b/Engine/Plugins/Runtime/nDisplay/Source/DisplayClusterLightCardEditor/Private/DisplayClusterLightCardEditorCommands.cpp
@@ -29,6 +29,9 @@ void FDisplayClusterLightCardEditorCommands::RegisterCommands()
 	UI_COMMAND(SaveLightCardTemplate, "Save As Template", "Save a template of the light card's appearance", EUserInterfaceActionType::Button, FInputChord(EModifierKey::Control | EModifierKey::Alt, EKeys::S));
 
 	UI_COMMAND(DrawLightCard, "Draw Light Card", "Draw polygon light card on viewport", EUserInterfaceActionType::ToggleButton, FInputChord());
+
+	UI_COMMAND(ToggleAllLabels, "Labels", "Display labels in the preview and wall", EUserInterfaceActionType::ToggleButton, FInputChord());
+	UI_COMMAND(ToggleLightCardLabels, "Light Card Labels", "Display labels on light cards in the preview and wall", EUserInterfaceActionType::ToggleButton, FInputChord());
 }
 
 #undef LOCTEXT_NAMESPACE
diff --git a/Engine/Plugins/Runtime/nDisplay/Source/DisplayClusterLightCardEditor/Private/DisplayClusterLightCardEditorCommands.h b/Engine/Plugins/Runtime/nDisplay/Source/DisplayClusterLightCardEditor/Private/DisplayClusterLightCardEditorCommands.h
--- a/Engine/Plugins/Runtime/nDisplay/Source/DisplayClusterLightCardEditor/Private/DisplayClusterLightCardEditorCommands.h
+++ b/Engine/Plugins/Runtime/nDisplay/Source/DisplayClusterLightCardEditor/Private/DisplayClusterLightCardEditorCommands.h
@@ -45,4 +45,7 @@ public:
 	TSharedPtr<FUICommandInfo> SaveLightCardTemplate;
 	
 	TSharedPtr<FUICommandInfo> DrawLightCard;
+
+	TSharedPtr<FUICommandInfo> ToggleAllLabels;
+	TSharedPtr<FUICommandInfo> ToggleLightCardLabels;
 };
diff --git a/Engine/Plugins/Runtime/nDisplay/Source/DisplayClusterLightCardEditor/Private/DisplayClusterLightCardEditorModule.cpp b/Engine/Plugins/Runtime/nDisplay/Source/DisplayClusterLightCardEditor/Private/DisplayClusterLightCardEditorModule.cpp
--- a/Engine/Plugins/Runtime/nDisplay/Source/DisplayClusterLightCardEditor/Private/DisplayClusterLightCardEditorModule.cpp
+++ b/Engine/Plugins/Runtime/nDisplay/Source/DisplayClusterLightCardEditor/Private/DisplayClusterLightCardEditorModule.cpp
@@ -2,11 +2,15 @@
 
 #include "DisplayClusterLightCardEditorModule.h"
 
+#include "DisplayClusterConfigurationTypes.h"
+#include "DisplayClusterLightCardActor.h"
 #include "DisplayClusterLightCardEditorCommands.h"
+#include "DisplayClusterRootActor.h"
 #include "SDisplayClusterLightCardEditor.h"
 #include "Settings/DisplayClusterLightCardEditorSettings.h"
 
 #include "ISettingsModule.h"
+#include "ScopedTransaction.h"
 
 #define LOCTEXT_NAMESPACE "DisplayClusterLightCardEditor"
 
@@ -26,6 +30,30 @@ void FDisplayClusterLightCardEditorModule::ShutdownModule()
 	FDisplayClusterLightCardEditorCommands::Unregister();
 }
 
+void FDisplayClusterLightCardEditorModule::ShowLabels(const FLabelArgs& InArgs)
+{
+	check(InArgs.RootActor != nullptr);
+	
+	UDisplayClusterLightCardEditorProjectSettings* ProjectSettings = GetMutableDefault<UDisplayClusterLightCardEditorProjectSettings>();
+	
+	ProjectSettings->Modify();
+	ProjectSettings->bDisplayLightCardLabels = InArgs.bVisible;
+	ProjectSettings->LightCardLabelScale = InArgs.Scale;
+
+	if (UDisplayClusterConfigurationData* ConfigData = InArgs.RootActor->GetConfigData())
+	{
+		FDisplayClusterConfigurationICVFX_VisibilityList& RootActorLightCards = ConfigData->StageSettings.Lightcard.ShowOnlyList;
+		for (TSoftObjectPtr<AActor> Actor : RootActorLightCards.Actors)
+		{
+			if (ADisplayClusterLightCardActor* LightCardActor = Cast<ADisplayClusterLightCardActor>(Actor.Get()))
+			{
+				LightCardActor->Modify(false);
+				LightCardActor->ShowLightCardLabel(InArgs.bVisible, InArgs.Scale, InArgs.RootActor);
+			}
+		}
+	}
+}
+
 void FDisplayClusterLightCardEditorModule::RegisterTabSpawners()
 {
 	SDisplayClusterLightCardEditor::RegisterTabSpawner();
@@ -40,10 +68,13 @@ void FDisplayClusterLightCardEditorModule::RegisterSettings()
 {
 	if (ISettingsModule* SettingsModule = FModuleManager::GetModulePtr<ISettingsModule>("Settings"))
 	{
+		UDisplayClusterLightCardEditorProjectSettings* Settings = GetMutableDefault<UDisplayClusterLightCardEditorProjectSettings>();
+		// Needs transactional for undo/redo support in the light card editor
+		Settings->SetFlags(RF_Transactional);
 		SettingsModule->RegisterSettings("Project", "Plugins", "nDisplayLightCardEditor",
 			LOCTEXT("nDisplayLightCardEditorName", "nDisplay Light Card Editor"),
 			LOCTEXT("nDisplayLightCardEditorDescription", "Configure settings for the nDisplay Light Card Editor."),
-			GetMutableDefault<UDisplayClusterLightCardEditorProjectSettings>());
+			Settings);
 	}
 }
 
diff --git a/Engine/Plugins/Runtime/nDisplay/Source/DisplayClusterLightCardEditor/Private/DisplayClusterLightCardEditorModule.h b/Engine/Plugins/Runtime/nDisplay/Source/DisplayClusterLightCardEditor/Private/DisplayClusterLightCardEditorModule.h
--- a/Engine/Plugins/Runtime/nDisplay/Source/DisplayClusterLightCardEditor/Private/DisplayClusterLightCardEditorModule.h
+++ b/Engine/Plugins/Runtime/nDisplay/Source/DisplayClusterLightCardEditor/Private/DisplayClusterLightCardEditorModule.h
@@ -15,7 +15,9 @@ public:
 	//~ IModuleInterface interface
 	virtual void StartupModule() override;
 	virtual void ShutdownModule() override;
-
+	
+	virtual void ShowLabels(const FLabelArgs& InArgs) override;
+	
 private:
 	void RegisterTabSpawners();
 	void UnregisterTabSpawners();
diff --git a/Engine/Plugins/Runtime/nDisplay/Source/DisplayClusterLightCardEditor/Private/DisplayClusterLightCardEditorStyle.h b/Engine/Plugins/Runtime/nDisplay/Source/DisplayClusterLightCardEditor/Private/DisplayClusterLightCardEditorStyle.h
--- a/Engine/Plugins/Runtime/nDisplay/Source/DisplayClusterLightCardEditor/Private/DisplayClusterLightCardEditorStyle.h
+++ b/Engine/Plugins/Runtime/nDisplay/Source/DisplayClusterLightCardEditor/Private/DisplayClusterLightCardEditorStyle.h
@@ -38,6 +38,8 @@ public:
 
 		Set("UVLightCard", new CORE_IMAGE_BRUSH_SVG("Starship/Common/SetDrawUVs", Icon16x16));
 
+		Set("DisplayClusterLightCardEditor.Labels", new CORE_IMAGE_BRUSH_SVG("Starship/AssetIcons/Font_16", Icon16x16));
+		
 		FSlateStyleRegistry::RegisterSlateStyle(*this);
 	}
 
diff --git a/Engine/Plugins/Runtime/nDisplay/Source/DisplayClusterLightCardEditor/Private/SDisplayClusterLightCardEditor.cpp b/Engine/Plugins/Runtime/nDisplay/Source/DisplayClusterLightCardEditor/Private/SDisplayClusterLightCardEditor.cpp
--- a/Engine/Plugins/Runtime/nDisplay/Source/DisplayClusterLightCardEditor/Private/SDisplayClusterLightCardEditor.cpp
+++ b/Engine/Plugins/Runtime/nDisplay/Source/DisplayClusterLightCardEditor/Private/SDisplayClusterLightCardEditor.cpp
@@ -3,6 +3,8 @@
 #include "SDisplayClusterLightCardEditor.h"
 
 #include "DisplayClusterLightCardEditorCommands.h"
+#include "DisplayClusterLightCardEditorStyle.h"
+#include "IDisplayClusterLightCardEditor.h"
 #include "SDisplayClusterLightCardList.h"
 #include "LightCardTemplates/SDisplayClusterLightCardTemplateList.h"
 #include "LightCardTemplates/DisplayClusterLightCardTemplate.h"
@@ -32,6 +34,7 @@
 #include "Selection.h"
 #include "Toolkits/AssetEditorToolkit.h"
 #include "Widgets/Docking/SDockTab.h"
+#include "Widgets/Input/SNumericEntryBox.h"
 #include "Widgets/Workflow/SWizard.h"
 
 #define LOCTEXT_NAMESPACE "SDisplayClusterLightCardEditor"
@@ -46,10 +49,6 @@ void SDisplayClusterLightCardEditor::RegisterTabSpawner()
 		.SetDisplayName(LOCTEXT("TabDisplayName", "Light Cards Editor"))
 		.SetTooltipText(LOCTEXT("TabTooltip", "Editing tools for nDisplay light cards."))
 		.SetMenuType(ETabSpawnerMenuType::Hidden);
-	
-	TSharedPtr<FExtender> ToolbarExtender = MakeShareable(new FExtender);
-	ToolbarExtender->AddToolBarExtension("General", EExtensionHook::After, nullptr, FToolBarExtensionDelegate::CreateStatic(&SDisplayClusterLightCardEditor::ExtendToolbar));
-	IDisplayClusterOperator::Get().GetOperatorToolBarExtensibilityManager()->AddExtender(ToolbarExtender);
 }
 
 void SDisplayClusterLightCardEditor::UnregisterTabSpawner()
@@ -75,11 +74,6 @@ TSharedRef<SDockTab> SDisplayClusterLightCardEditor::SpawnInTab(const FSpawnTabA
 	return MajorTab;
 }
 
-void SDisplayClusterLightCardEditor::ExtendToolbar(FToolBarBuilder& ToolbarBuilder)
-{
-	// TODO: Any toolbar buttons needed for the lightcards editor can be added to the operator panel's toolbar using this toolbar extender
-}
-
 SDisplayClusterLightCardEditor::~SDisplayClusterLightCardEditor()
 {
 	IDisplayClusterOperator::Get().OnActiveRootActorChanged().Remove(ActiveRootActorChangedHandle);
@@ -98,6 +92,18 @@ SDisplayClusterLightCardEditor::~SDisplayClusterLightCardEditor()
 	RemoveCompileDelegates();
 }
 
+void SDisplayClusterLightCardEditor::PostUndo(bool bSuccess)
+{
+	FEditorUndoClient::PostUndo(bSuccess);
+	RefreshLabels();
+}
+
+void SDisplayClusterLightCardEditor::PostRedo(bool bSuccess)
+{
+	FEditorUndoClient::PostRedo(bSuccess);
+	RefreshLabels();
+}
+
 void SDisplayClusterLightCardEditor::Construct(const FArguments& InArgs, const TSharedRef<SDockTab>& MajorTabOwner, const TSharedPtr<SWindow>& WindowOwner)
 {
 	ActiveRootActorChangedHandle = IDisplayClusterOperator::Get().OnActiveRootActorChanged().AddSP(this, &SDisplayClusterLightCardEditor::OnActiveRootActorChanged);
@@ -109,6 +115,7 @@ void SDisplayClusterLightCardEditor::Construct(const FArguments& InArgs, const T
 	OnObjectTransactedHandle = FCoreUObjectDelegates::OnObjectTransacted.AddSP(this, &SDisplayClusterLightCardEditor::OnObjectTransacted);
 
 	BindCommands();
+	RegisterToolbarExtensions();
 	
 	TabManager = FGlobalTabmanager::Get()->NewTabManager(MajorTabOwner);
 
@@ -193,6 +200,10 @@ void SDisplayClusterLightCardEditor::Construct(const FArguments& InArgs, const T
 	];
 
 	BindCompileDelegates();
+
+	RefreshLabels();
+
+	GEditor->RegisterForUndo(this);
 }
 
 void SDisplayClusterLightCardEditor::SelectLightCards(const TArray<ADisplayClusterLightCardActor*>& LightCardsToSelect)
@@ -434,6 +445,8 @@ void SDisplayClusterLightCardEditor::AddLightCardsToActor(TArray<ADisplayCluster
 					return Actor.Get() == LightCard;
 				}))
 			{
+				LightCard->ShowLightCardLabel(ShouldShowLightCardLabels(), *GetLightCardLabelScale(), ActiveRootActor.Get());
+				
 				const TSoftObjectPtr<AActor> LightCardSoftObject(LightCard);
 
 				// Remove any exact paths to this actor. It's possible invalid actors are present if a light card
@@ -575,7 +588,7 @@ void SDisplayClusterLightCardEditor::RemoveLightCards(bool bDeleteLightCardActor
 
 		FDisplayClusterConfigurationICVFX_VisibilityList& RootActorLightCards = ConfigData->StageSettings.Lightcard.ShowOnlyList;
 
-		for (AActor* LightCard : SelectedLightCards)
+		for (ADisplayClusterLightCardActor* LightCard : SelectedLightCards)
 		{
 			RootActorLightCards.Actors.RemoveAll([&](const TSoftObjectPtr<AActor>& Actor)
 				{
@@ -587,6 +600,10 @@ void SDisplayClusterLightCardEditor::RemoveLightCards(bool bDeleteLightCardActor
 				WorldToUse = LightCard->GetWorld();
 				GEditor->SelectActor(LightCard, /*bSelect =*/true, /*bNotifyForActor =*/false, /*bSelectEvenIfHidden =*/true);
 			}
+			else
+			{
+				LightCard->ShowLightCardLabel(false, *GetLightCardLabelScale(), ActiveRootActor.Get());
+			}
 		}
 	}
 
@@ -682,8 +699,68 @@ bool SDisplayClusterLightCardEditor::CanCreateLightCardTemplate() const
 	return SelectedLightCards.Num() == 1;
 }
 
+void SDisplayClusterLightCardEditor::ToggleLightCardLabels()
+{
+	ShowLightCardLabels(!ShouldShowLightCardLabels());
+}
+
+void SDisplayClusterLightCardEditor::ShowLightCardLabels(bool bVisible)
+{
+	if (!GetActiveRootActor().IsValid())
+	{
+		return;
+	}
+	
+	FScopedTransaction Transaction(LOCTEXT("ToggleLightCardLabelsTransactionMessage", "Toggle Light Card Labels"), !GIsTransacting);
+	
+	IDisplayClusterLightCardEditor& LightCardEditorModule = FModuleManager::GetModuleChecked<IDisplayClusterLightCardEditor>(IDisplayClusterLightCardEditor::ModuleName);
+
+	IDisplayClusterLightCardEditor::FLabelArgs Args;
+	Args.bVisible = bVisible;
+	Args.Scale = *GetLightCardLabelScale();
+	Args.RootActor = GetActiveRootActor().Get();
+	
+	LightCardEditorModule.ShowLabels(MoveTemp(Args));
+
+	RefreshPreviewActors(EDisplayClusterLightCardEditorProxyType::LightCards);
+}
+
+bool SDisplayClusterLightCardEditor::ShouldShowLightCardLabels() const
+{
+	return GetDefault<UDisplayClusterLightCardEditorProjectSettings>()->bDisplayLightCardLabels;
+}
+
+TOptional<float> SDisplayClusterLightCardEditor::GetLightCardLabelScale() const
+{
+	return GetDefault<UDisplayClusterLightCardEditorProjectSettings>()->LightCardLabelScale;
+}
+
+void SDisplayClusterLightCardEditor::SetLightCardLabelScale(float NewValue)
+{
+	if (!GetActiveRootActor().IsValid())
+	{
+		return;
+	}
+	
+	FScopedTransaction Transaction(LOCTEXT("ScaleLightCardLabelsTransactionMessage", "Scale Light Card Labels"), !GIsTransacting);
+	
+	IDisplayClusterLightCardEditor& LightCardEditorModule = FModuleManager::GetModuleChecked<IDisplayClusterLightCardEditor>(IDisplayClusterLightCardEditor::ModuleName);
+
+	IDisplayClusterLightCardEditor::FLabelArgs Args;
+	Args.bVisible = ShouldShowLightCardLabels();
+	Args.Scale = NewValue;
+	Args.RootActor = GetActiveRootActor().Get();
+	
+	LightCardEditorModule.ShowLabels(MoveTemp(Args));
+}
+
 void SDisplayClusterLightCardEditor::OnActiveRootActorChanged(ADisplayClusterRootActor* NewRootActor)
 {
+	if (NewRootActor == ActiveRootActor.Get())
+	{
+		return;
+	}
+	
 	RemoveCompileDelegates();
 	
 	// The new root actor pointer could be null, indicating that it was deleted or the user didn't select a valid root actor
@@ -695,6 +772,8 @@ void SDisplayClusterLightCardEditor::OnActiveRootActorChanged(ADisplayClusterRoo
 
 	FCoreUObjectDelegates::OnObjectPropertyChanged.RemoveAll(this);
 	FCoreUObjectDelegates::OnObjectPropertyChanged.AddSP(this, &SDisplayClusterLightCardEditor::OnActorPropertyChanged);
+
+	RefreshLabels();
 }
 
 TSharedRef<SWidget> SDisplayClusterLightCardEditor::CreateLightCardListWidget()
@@ -707,6 +786,50 @@ TSharedRef<SWidget> SDisplayClusterLightCardEditor::CreateViewportWidget()
 	return SAssignNew(ViewportView, SDisplayClusterLightCardEditorViewport, SharedThis(this), CommandList);
 }
 
+TSharedRef<SWidget> SDisplayClusterLightCardEditor::GenerateLabelsMenu()
+{
+	const bool bInShouldCloseWindowAfterMenuSelection = true;
+
+	FMenuBuilder MenuBuilder(bInShouldCloseWindowAfterMenuSelection, CommandList);
+
+	MenuBuilder.BeginSection("Labels", LOCTEXT("LabelsMenuHeader", "Labels"));
+	{
+		MenuBuilder.AddMenuEntry(FDisplayClusterLightCardEditorCommands::Get().ToggleLightCardLabels);
+	
+		MenuBuilder.AddWidget(
+		SNew(SBox)
+		.HAlign(HAlign_Right)
+		[
+			SNew(SBox)
+			.Padding(FMargin(8.0f, 0.0f, 0.0f, 0.0f))
+			.WidthOverride(100.0f)
+			[
+				SNew(SBorder)
+				.BorderImage(FAppStyle::Get().GetBrush("Menu.WidgetBorder"))
+				.Padding(FMargin(1.0f))
+				[
+					SNew(SBox)
+					.MinDesiredWidth(64)
+					[
+						SNew(SNumericEntryBox<float>)
+						.Value(this, &SDisplayClusterLightCardEditor::GetLightCardLabelScale)
+						.OnValueChanged(this, &SDisplayClusterLightCardEditor::SetLightCardLabelScale)
+						.MinValue(0)
+						.MaxValue(FLT_MAX)
+						.MinSliderValue(0)
+						.MaxSliderValue(10)
+						.AllowSpin(true)
+					]
+				]
+			]
+		],
+		LOCTEXT("LightCardLabelScale_Label", "Label Scale"));
+	}
+	MenuBuilder.EndSection();
+	
+	return MenuBuilder.MakeWidget();
+}
+
 TSharedRef<SWidget> SDisplayClusterLightCardEditor::CreateLightCardTemplateWidget()
 {
 	return SAssignNew(LightCardTemplateList, SDisplayClusterLightCardTemplateList, SharedThis(this));
@@ -760,6 +883,35 @@ void SDisplayClusterLightCardEditor::BindCommands()
 		FDisplayClusterLightCardEditorCommands::Get().SaveLightCardTemplate,
 		FExecuteAction::CreateSP(this, &SDisplayClusterLightCardEditor::CreateLightCardTemplate),
 		FCanExecuteAction::CreateSP(this, &SDisplayClusterLightCardEditor::CanCreateLightCardTemplate));
+
+	CommandList->MapAction(
+		FDisplayClusterLightCardEditorCommands::Get().ToggleLightCardLabels,
+		FExecuteAction::CreateSP(this, &SDisplayClusterLightCardEditor::ToggleLightCardLabels),
+		FCanExecuteAction(),
+		FIsActionChecked::CreateSP(this, &SDisplayClusterLightCardEditor::ShouldShowLightCardLabels));
+}
+
+void SDisplayClusterLightCardEditor::RegisterToolbarExtensions()
+{
+	const TSharedPtr<FExtender> ToolbarExtender = MakeShareable(new FExtender);
+	ToolbarExtender->AddToolBarExtension("General", EExtensionHook::After, nullptr,
+		FToolBarExtensionDelegate::CreateSP(this, &SDisplayClusterLightCardEditor::ExtendToolbar));
+	IDisplayClusterOperator::Get().GetOperatorToolBarExtensibilityManager()->AddExtender(ToolbarExtender);
+}
+
+void SDisplayClusterLightCardEditor::ExtendToolbar(FToolBarBuilder& ToolbarBuilder)
+{
+	// TODO: Any toolbar buttons needed for the lightcards editor can be added to the operator panel's toolbar using this toolbar extender
+
+	ToolbarBuilder.AddSeparator();
+	
+	const FUIAction DefaultAction;
+	ToolbarBuilder.AddComboButton(
+		DefaultAction,
+		FOnGetContent::CreateSP(this, &SDisplayClusterLightCardEditor::GenerateLabelsMenu),
+		TAttribute<FText>(),
+		LOCTEXT("Labels_ToolTip", "Configure options for labels"),
+		FSlateIcon(FDisplayClusterLightCardEditorStyle::Get().GetStyleSetName(), "DisplayClusterLightCardEditor.Labels"));
 }
 
 void SDisplayClusterLightCardEditor::RefreshPreviewActors(EDisplayClusterLightCardEditorProxyType ProxyType)
@@ -783,6 +935,11 @@ void SDisplayClusterLightCardEditor::RefreshPreviewActors(EDisplayClusterLightCa
 	BindCompileDelegates();
 }
 
+void SDisplayClusterLightCardEditor::RefreshLabels()
+{
+	ShowLightCardLabels(ShouldShowLightCardLabels());
+}
+
 bool SDisplayClusterLightCardEditor::IsOurObject(UObject* InObject,
 	EDisplayClusterLightCardEditorProxyType& OutProxyType) const
 {
diff --git a/Engine/Plugins/Runtime/nDisplay/Source/DisplayClusterLightCardEditor/Private/SDisplayClusterLightCardEditor.h b/Engine/Plugins/Runtime/nDisplay/Source/DisplayClusterLightCardEditor/Private/SDisplayClusterLightCardEditor.h
--- a/Engine/Plugins/Runtime/nDisplay/Source/DisplayClusterLightCardEditor/Private/SDisplayClusterLightCardEditor.h
+++ b/Engine/Plugins/Runtime/nDisplay/Source/DisplayClusterLightCardEditor/Private/SDisplayClusterLightCardEditor.h
@@ -8,6 +8,8 @@
 
 #include "DisplayClusterLightCardEditorProxyType.h"
 
+#include "EditorUndoClient.h"
+
 class FTabManager;
 class FLayoutExtender;
 class FSpawnTabArgs;
@@ -22,7 +24,7 @@ class ADisplayClusterLightCardActor;
 class UDisplayClusterLightCardTemplate;
 
 /** A panel that can be spawned in a tab that contains all the UI elements that make up the 2D light cards editor */
-class SDisplayClusterLightCardEditor : public SCompoundWidget
+class SDisplayClusterLightCardEditor : public SCompoundWidget, public FEditorUndoClient
 {
 public:
 	/** The name of the tab that the light card editor lives in */
@@ -40,15 +42,17 @@ public:
 	/** Creates a tab with the light card editor inside */
 	static TSharedRef<SDockTab> SpawnInTab(const FSpawnTabArgs& SpawnTabArgs);
 
-	static void ExtendToolbar(FToolBarBuilder& ToolbarBuilder);
-
-
 	SLATE_BEGIN_ARGS(SDisplayClusterLightCardEditor)
 	{}
 	SLATE_END_ARGS()
 
 	~SDisplayClusterLightCardEditor();
 
+	// FEditorUndoClient
+	virtual void PostUndo(bool bSuccess) override;
+	virtual void PostRedo(bool bSuccess) override;
+	// ~FEditorUndoClient
+	
 	void Construct(const FArguments& args, const TSharedRef<SDockTab>& MajorTabOwner, const TSharedPtr<SWindow>& WindowOwner);
 
 	/** The current active root actor for this light card editor */
@@ -123,6 +127,21 @@ public:
 	/** Determines if there are selected light cards that can have templates created */
 	bool CanCreateLightCardTemplate() const;
 
+	/** Toggles the visible status of labels */
+	void ToggleLightCardLabels();
+	
+	/** Display or hide labels for all included light cards */
+	void ShowLightCardLabels(bool bVisible);
+
+	/** If light card labels are currently toggled */
+	bool ShouldShowLightCardLabels() const;
+
+	/** Return the current light card label scale */
+	TOptional<float> GetLightCardLabelScale() const;
+
+	/** Update light card label scale */
+	void SetLightCardLabelScale(float NewValue);
+	
 private:
 	/** Raised when the active Display cluster root actor has been changed in the operator panel */
 	void OnActiveRootActorChanged(ADisplayClusterRootActor* NewRootActor);
@@ -136,12 +155,24 @@ private:
 	/** Create the 3d viewport widget */
 	TSharedRef<SWidget> CreateViewportWidget();
 
+	/** Generate the labels dropdown menu */
+	TSharedRef<SWidget> GenerateLabelsMenu();
+
 	/** Creates the editor's command list and binds commands to it */
 	void BindCommands();
 
+	/** Register any extensions for the toolbar */
+	void RegisterToolbarExtensions();
+
+	/** Extend the toolbar for this local instance */
+	void ExtendToolbar(FToolBarBuilder& ToolbarBuilder);
+
 	/** Refresh all preview actors */
 	void RefreshPreviewActors(EDisplayClusterLightCardEditorProxyType ProxyType = EDisplayClusterLightCardEditorProxyType::All);
 
+	/** Reapply the current label settings */
+	void RefreshLabels();
+	
 	/**
 	 * Check if an object is managed by us
 	 * @param InObject The object to compare
diff --git a/Engine/Plugins/Runtime/nDisplay/Source/DisplayClusterLightCardEditor/Private/Settings/DisplayClusterLightCardEditorSettings.cpp b/Engine/Plugins/Runtime/nDisplay/Source/DisplayClusterLightCardEditor/Private/Settings/DisplayClusterLightCardEditorSettings.cpp
--- a/Engine/Plugins/Runtime/nDisplay/Source/DisplayClusterLightCardEditor/Private/Settings/DisplayClusterLightCardEditorSettings.cpp
+++ b/Engine/Plugins/Runtime/nDisplay/Source/DisplayClusterLightCardEditor/Private/Settings/DisplayClusterLightCardEditorSettings.cpp
@@ -5,4 +5,6 @@
 UDisplayClusterLightCardEditorProjectSettings::UDisplayClusterLightCardEditorProjectSettings()
 {
 	LightCardTemplateDefaultPath.Path = TEXT("/Game/VP/LightCards");
+	LightCardLabelScale = 1.f;
+	bDisplayLightCardLabels = false;
 }
diff --git a/Engine/Plugins/Runtime/nDisplay/Source/DisplayClusterLightCardEditor/Private/Settings/DisplayClusterLightCardEditorSettings.h b/Engine/Plugins/Runtime/nDisplay/Source/DisplayClusterLightCardEditor/Private/Settings/DisplayClusterLightCardEditorSettings.h
--- a/Engine/Plugins/Runtime/nDisplay/Source/DisplayClusterLightCardEditor/Private/Settings/DisplayClusterLightCardEditorSettings.h
+++ b/Engine/Plugins/Runtime/nDisplay/Source/DisplayClusterLightCardEditor/Private/Settings/DisplayClusterLightCardEditorSettings.h
@@ -6,7 +6,7 @@
 #include "DisplayClusterLightCardEditorSettings.generated.h"
 
 
-UCLASS(config = Editor, defaultconfig)
+UCLASS(config = Editor, defaultconfig, meta = (DisplayClusterMultiUserInclude))
 class UDisplayClusterLightCardEditorProjectSettings : public UObject
 {
 	GENERATED_BODY()
@@ -17,4 +17,12 @@ public:
 	/** The default path to save new light card templates */
 	UPROPERTY(config, EditAnywhere, Category = LightCardTemplates)
 	FDirectoryPath LightCardTemplateDefaultPath;
+
+	/** Whether light card labels should be displayed. Handled through the light card editor */
+	UPROPERTY()
+	bool bDisplayLightCardLabels;
+	
+	/** The scale to use for light card labels */
+	UPROPERTY(config, EditAnywhere, Category = LightCardLabels)
+	float LightCardLabelScale;
 };
diff --git a/Engine/Plugins/Runtime/nDisplay/Source/DisplayClusterLightCardEditor/Private/Viewport/DisplayClusterLightCardEditorViewportClient.cpp b/Engine/Plugins/Runtime/nDisplay/Source/DisplayClusterLightCardEditor/Private/Viewport/DisplayClusterLightCardEditorViewportClient.cpp
--- a/Engine/Plugins/Runtime/nDisplay/Source/DisplayClusterLightCardEditor/Private/Viewport/DisplayClusterLightCardEditorViewportClient.cpp
+++ b/Engine/Plugins/Runtime/nDisplay/Source/DisplayClusterLightCardEditor/Private/Viewport/DisplayClusterLightCardEditorViewportClient.cpp
@@ -1512,20 +1512,28 @@ void FDisplayClusterLightCardEditorViewportClient::UpdatePreviewActor(ADisplayCl
 	{
 		UWorld* PreviewWorld = PreviewScene->GetWorld();
 		check(PreviewWorld);
+
+		TWeakObjectPtr<ADisplayClusterRootActor> RootActorPtr (RootActor);
 		
 		// Schedule for the next tick so CDO changes get propagated first in the event of config editor skeleton
 		// regeneration & compiles. nDisplay's custom propagation may have issues if the archetype isn't correct.
 		PreviewWorld->GetTimerManager().SetTimerForNextTick([=]()
 		{			
 			DestroyProxies(ProxyType);
-			RootActor->SubscribeToPostProcessRenderTarget(reinterpret_cast<uint8*>(this));
-			RootActorLevelInstance = RootActor;
+
+			if (!RootActorPtr.IsValid())
+			{
+				return;
+			}
+			
+			RootActorPtr->SubscribeToPostProcessRenderTarget(reinterpret_cast<uint8*>(this));
+			RootActorLevelInstance = RootActorPtr;
 			
 			if (ProxyType == EDisplayClusterLightCardEditorProxyType::All ||
 				ProxyType == EDisplayClusterLightCardEditorProxyType::RootActor)
 			{
 				{
-					FObjectDuplicationParameters DupeActorParameters(RootActor, PreviewWorld->GetCurrentLevel());
+					FObjectDuplicationParameters DupeActorParameters(RootActorPtr.Get(), PreviewWorld->GetCurrentLevel());
 					DupeActorParameters.FlagMask = RF_AllFlags & ~(RF_ArchetypeObject | RF_Transactional); // Keeps archetypes correct in config data.
 					DupeActorParameters.PortFlags = PPF_DuplicateVerbatim;
 			
@@ -1594,7 +1602,7 @@ void FDisplayClusterLightCardEditorViewportClient::UpdatePreviewActor(ADisplayCl
 				ProxyType == EDisplayClusterLightCardEditorProxyType::LightCards)
 			{
 				TSet<ADisplayClusterLightCardActor*> LightCards;
-				UDisplayClusterBlueprintLib::FindLightCardsForRootActor(RootActor, LightCards);
+				UDisplayClusterBlueprintLib::FindLightCardsForRootActor(RootActorPtr.Get(), LightCards);
 				
 				SelectLightCard(nullptr);
 				
@@ -1607,8 +1615,8 @@ void FDisplayClusterLightCardEditorViewportClient::UpdatePreviewActor(ADisplayCl
 					ADisplayClusterLightCardActor* LightCardProxy = CastChecked<ADisplayClusterLightCardActor>(StaticDuplicateObjectEx(DupeActorParameters));
 					PreviewWorld->GetCurrentLevel()->AddLoadedActor(LightCardProxy);
 				
-					LightCardProxy->SetActorLocation(LightCard->GetActorLocation() - RootActor->GetActorLocation());
-					LightCardProxy->SetActorRotation(LightCard->GetActorRotation() - RootActor->GetActorRotation());
+					LightCardProxy->SetActorLocation(LightCard->GetActorLocation() - RootActorPtr->GetActorLocation());
+					LightCardProxy->SetActorRotation(LightCard->GetActorRotation() - RootActorPtr->GetActorRotation());
 					LightCardProxy->PolygonMask = LightCard->PolygonMask;
 					LightCardProxy->bIsProxy = true;
 
diff --git a/Engine/Plugins/Runtime/nDisplay/Source/DisplayClusterLightCardEditor/Private/Viewport/DisplayClusterLightcardEditorViewport.cpp b/Engine/Plugins/Runtime/nDisplay/Source/DisplayClusterLightCardEditor/Private/Viewport/DisplayClusterLightcardEditorViewport.cpp
--- a/Engine/Plugins/Runtime/nDisplay/Source/DisplayClusterLightCardEditor/Private/Viewport/DisplayClusterLightcardEditorViewport.cpp
+++ b/Engine/Plugins/Runtime/nDisplay/Source/DisplayClusterLightCardEditor/Private/Viewport/DisplayClusterLightcardEditorViewport.cpp
@@ -2,21 +2,21 @@
 
 #include "DisplayClusterLightcardEditorViewport.h"
 
+#include "DisplayClusterLightCardEditorCommands.h"
 #include "DisplayClusterLightCardEditorStyle.h"
 #include "DisplayClusterLightCardEditorViewportClient.h"
-#include "DisplayClusterLightCardEditorCommands.h"
 #include "SDisplayClusterLightCardEditor.h"
-#include "LightCardTemplates/DisplayClusterLightCardTemplateDragDropOp.h"
 #include "LightCardTemplates/DisplayClusterLightCardTemplate.h"
+#include "LightCardTemplates/DisplayClusterLightCardTemplateDragDropOp.h"
 
 #include "EditorViewportCommands.h"
-#include "Framework/Commands/GenericCommands.h"
-#include "HAL/PlatformApplicationMisc.h"
+#include "ScopedTransaction.h"
 #include "SEditorViewportToolBarMenu.h"
 #include "STransformViewportToolbar.h"
-#include "Slate/SceneViewport.h"
-#include "ScopedTransaction.h"
+#include "Framework/Commands/GenericCommands.h"
+#include "HAL/PlatformApplicationMisc.h"
 #include "Kismet2/DebuggerCommands.h"
+#include "Slate/SceneViewport.h"
 #include "Styling/AppStyle.h"
 
 #define LOCTEXT_NAMESPACE "DisplayClusterLightcardEditorViewport"
@@ -72,6 +72,16 @@ public:
 					.OnGetMenuContent(this, &SDisplayClusterLightCardEditorViewportToolBar::GenerateViewMenu)
 				]
 				+ SHorizontalBox::Slot()
+				.AutoWidth()
+				.Padding(2.0f, 2.0f)
+				[
+					SNew(SEditorViewportToolbarMenu)
+					.ParentToolBar(SharedThis(this))
+					.Cursor(EMouseCursor::Default)
+					.Label(LOCTEXT("LightCardEditorShowMenuLabel", "Show"))
+					.OnGetMenuContent(this, &SDisplayClusterLightCardEditorViewportToolBar::GenerateShowMenu)
+				]
+				+ SHorizontalBox::Slot()
 				.Padding(3.0f, 1.0f)
 				.HAlign(HAlign_Right)
 				[
@@ -180,7 +190,25 @@ public:
 
 		return MenuBuilder.MakeWidget();
 	}
+	
+	TSharedRef<SWidget> GenerateShowMenu() const
+	{
+		const TSharedPtr<const FUICommandList> CommandList = EditorViewport.IsValid() ? EditorViewport.Pin()->GetCommandList() : nullptr;
+
+		const bool bInShouldCloseWindowAfterMenuSelection = true;
+
+		FMenuBuilder MenuBuilder(bInShouldCloseWindowAfterMenuSelection, CommandList);
+		{
+			MenuBuilder.BeginSection("LightCardEditorShow", LOCTEXT("ShowMenuHeader", "Show Flags"));
+			{
+				MenuBuilder.AddMenuEntry(FDisplayClusterLightCardEditorCommands::Get().ToggleAllLabels);
+			}
+			MenuBuilder.EndSection();
+		}
 
+		return MenuBuilder.MakeWidget();
+	}
+	
 	TSharedRef<SWidget> MakeTransformToolBar()
 	{
 		FSlimHorizontalToolBarBuilder ToolbarBuilder(EditorViewport.Pin()->GetCommandList(), FMultiBoxCustomization::None);
@@ -455,6 +483,12 @@ void SDisplayClusterLightCardEditorViewport::BindCommands()
 			Commands.PasteHere,
 			FExecuteAction::CreateSP(this, &SDisplayClusterLightCardEditorViewport::PasteLightCardsHere),
 			FCanExecuteAction::CreateSP(this, &SDisplayClusterLightCardEditorViewport::CanPasteLightCardsHere));
+
+		CommandList->MapAction(
+			Commands.ToggleAllLabels,
+			FExecuteAction::CreateSP(this, &SDisplayClusterLightCardEditorViewport::ToggleLabels),
+			FCanExecuteAction(),
+			FIsActionChecked::CreateSP(this, &SDisplayClusterLightCardEditorViewport::AreLabelsToggled));
 	}
 }
 
@@ -687,4 +721,19 @@ bool SDisplayClusterLightCardEditorViewport::CanPasteLightCardsHere() const
 	return false;
 }
 
+void SDisplayClusterLightCardEditorViewport::ToggleLabels()
+{
+	// TODO: Handle CCR
+	if (LightCardEditorPtr.IsValid())
+	{
+		return LightCardEditorPtr.Pin()->ToggleLightCardLabels();
+	}
+}
+
+bool SDisplayClusterLightCardEditorViewport::AreLabelsToggled() const
+{
+	// TODO: Handle CCR
+	return LightCardEditorPtr.IsValid() && LightCardEditorPtr.Pin()->ShouldShowLightCardLabels();
+}
+
 #undef LOCTEXT_NAMESPACE
diff --git a/Engine/Plugins/Runtime/nDisplay/Source/DisplayClusterLightCardEditor/Private/Viewport/DisplayClusterLightcardEditorViewport.h b/Engine/Plugins/Runtime/nDisplay/Source/DisplayClusterLightCardEditor/Private/Viewport/DisplayClusterLightcardEditorViewport.h
--- a/Engine/Plugins/Runtime/nDisplay/Source/DisplayClusterLightCardEditor/Private/Viewport/DisplayClusterLightcardEditorViewport.h
+++ b/Engine/Plugins/Runtime/nDisplay/Source/DisplayClusterLightCardEditor/Private/Viewport/DisplayClusterLightcardEditorViewport.h
@@ -98,6 +98,12 @@ private:
 	/** Determines if there are any light cards in the clipboard that can be pasted */
 	bool CanPasteLightCardsHere() const;
 
+	/** Toggles light card labels on or off */
+	void ToggleLabels();
+
+	/** If light card labels are currently toggled */
+	bool AreLabelsToggled() const;
+	
 private:
 	/** The preview scene to use by the viewport client to manage the preview 3D world */
 	TSharedPtr<FPreviewScene> PreviewScene;
diff --git a/Engine/Plugins/Runtime/nDisplay/Source/DisplayClusterLightCardEditor/Public/IDisplayClusterLightCardEditor.h b/Engine/Plugins/Runtime/nDisplay/Source/DisplayClusterLightCardEditor/Public/IDisplayClusterLightCardEditor.h
--- a/Engine/Plugins/Runtime/nDisplay/Source/DisplayClusterLightCardEditor/Public/IDisplayClusterLightCardEditor.h
+++ b/Engine/Plugins/Runtime/nDisplay/Source/DisplayClusterLightCardEditor/Public/IDisplayClusterLightCardEditor.h
@@ -37,4 +37,31 @@ public:
 	{
 		return FModuleManager::Get().IsModuleLoaded(ModuleName);
 	}
+
+	struct FLabelArgs
+	{
+		// TODO: Specify type (Light Card vs CCR)
+		
+		/**
+		 * [Required] The root actor to display labels for
+		 */
+		class ADisplayClusterRootActor* RootActor = nullptr;
+		
+		/**
+		 * The scale to apply to the label
+		 */
+		float Scale = 1.f;
+
+		/**
+		 * Should the label be visible
+		 */
+		bool bVisible = false;
+	};
+	
+	/**
+	 * Show the light card labels for the given root actor and save current settings.
+	 *
+	 * @param InArgs The arguments for displaying labels.
+	 */
+	virtual void ShowLabels(const FLabelArgs& InArgs) = 0;
 };
\ No newline at end of file
diff --git a/Engine/Plugins/Runtime/nDisplay/Source/DisplayClusterLightCardEditorShaders/Private/DisplayClusterMeshProjectionRenderer.cpp b/Engine/Plugins/Runtime/nDisplay/Source/DisplayClusterLightCardEditorShaders/Private/DisplayClusterMeshProjectionRenderer.cpp
--- a/Engine/Plugins/Runtime/nDisplay/Source/DisplayClusterLightCardEditorShaders/Private/DisplayClusterMeshProjectionRenderer.cpp
+++ b/Engine/Plugins/Runtime/nDisplay/Source/DisplayClusterLightCardEditorShaders/Private/DisplayClusterMeshProjectionRenderer.cpp
@@ -873,7 +873,10 @@ void FDisplayClusterMeshProjectionRenderer::RemoveActor(AActor* Actor)
 	for (const TWeakObjectPtr<UPrimitiveComponent>& PrimitiveComponent : ComponentsToRemove)
 	{
 #if WITH_EDITOR
-		PrimitiveComponent->SelectionOverrideDelegate.Unbind();
+		if (PrimitiveComponent.IsValid() && PrimitiveComponent->SelectionOverrideDelegate.IsBound())
+		{
+			PrimitiveComponent->SelectionOverrideDelegate.Unbind();
+		}
 #endif
 
 		PrimitiveComponents.Remove(PrimitiveComponent);
@@ -1518,6 +1521,11 @@ void FDisplayClusterMeshProjectionRenderer::RenderPrimitives_RenderThread(const
 		for (const FSceneProxyElement& PrimitiveProxyElement : PrimitiveSceneProxies)
 		{
 			FPrimitiveSceneProxy* PrimitiveProxy = PrimitiveProxyElement.PrimitiveSceneProxy;
+			if (PrimitiveProxy == nullptr || PrimitiveProxy->GetPrimitiveSceneInfo() == nullptr)
+			{
+				continue;
+			}
+
 			if (const FMeshBatch* MeshBatch = PrimitiveProxy->GetPrimitiveSceneInfo()->GetMeshBatch(PrimitiveProxy->GetPrimitiveSceneInfo()->StaticMeshes.Num() - 1))
 			{
 				MeshBatch->MaterialRenderProxy->UpdateUniformExpressionCacheIfNeeded(View->GetFeatureLevel());
@@ -1668,7 +1676,7 @@ void FDisplayClusterMeshProjectionRenderer::GetSceneProxies(TArray<FSceneProxyEl
 {
 	for (const TWeakObjectPtr<UPrimitiveComponent>& PrimitiveComponent : PrimitiveComponents)
 	{
-		if (PrimitiveComponent.IsValid() && PrimitiveComponent->SceneProxy)
+		if (PrimitiveComponent.IsValid() && PrimitiveComponent->SceneProxy && !PrimitiveComponent->bVisibleInSceneCaptureOnly)
 		{
 			if (PrimitiveFilter(PrimitiveComponent.Get()) && RenderSettings.PrimitiveFilter.ShouldRenderPrimitive(PrimitiveComponent.Get()))
 			{
diff --git a/Engine/Plugins/Runtime/nDisplay/Source/DisplayClusterMultiUser/Private/DisplayClusterMultiUserManager.cpp b/Engine/Plugins/Runtime/nDisplay/Source/DisplayClusterMultiUser/Private/DisplayClusterMultiUserManager.cpp
--- a/Engine/Plugins/Runtime/nDisplay/Source/DisplayClusterMultiUser/Private/DisplayClusterMultiUserManager.cpp
+++ b/Engine/Plugins/Runtime/nDisplay/Source/DisplayClusterMultiUser/Private/DisplayClusterMultiUserManager.cpp
@@ -55,8 +55,9 @@ void FDisplayClusterMultiUserManager::OnApplyRemoteTransaction(ETransactionNotif
 
 ETransactionFilterResult FDisplayClusterMultiUserManager::ShouldObjectBeTransacted(UObject* InObject, UPackage* InPackage)
 {
-	if (InObject && InObject->IsA<UDisplayClusterConfigurationData_Base>() &&
-		!InObject->IsTemplate() && !InObject->HasAnyFlags(RF_Transient) && InPackage != GetTransientPackage())
+	if (InObject && ((InObject->IsA<UDisplayClusterConfigurationData_Base>() &&
+		!InObject->IsTemplate() && !InObject->HasAnyFlags(RF_Transient) && InPackage != GetTransientPackage()) ||
+		(InObject->GetClass()->HasMetaData(TEXT("DisplayClusterMultiUserInclude")))))
 	{
 		UE_LOG(LogDisplayClusterMultiUser, Log, TEXT("FDisplayClusterMultiUser transaction for object: %s"), *InObject->GetName());
 		return ETransactionFilterResult::IncludeObject;
diff --git a/Engine/Plugins/Runtime/nDisplay/Source/DisplayClusterOperator/Private/SDisplayClusterOperatorPanel.cpp b/Engine/Plugins/Runtime/nDisplay/Source/DisplayClusterOperator/Private/SDisplayClusterOperatorPanel.cpp
--- a/Engine/Plugins/Runtime/nDisplay/Source/DisplayClusterOperator/Private/SDisplayClusterOperatorPanel.cpp
+++ b/Engine/Plugins/Runtime/nDisplay/Source/DisplayClusterOperator/Private/SDisplayClusterOperatorPanel.cpp
@@ -110,8 +110,13 @@ void SDisplayClusterOperatorPanel::Construct(const FArguments& InArgs, const TSh
 	];
 
 	// TODO: Move ownership of active root controller to a view model object, instead of letting the toolbar control which root actor is active
-	if (Toolbar.IsValid())
+	if (ToolbarContainer.IsValid())
 	{
+		// Create toolbar after tab has been restored so child windows can register their toolbar extensions
+		ToolbarContainer->SetContent(
+		SAssignNew(Toolbar, SDisplayClusterOperatorToolbar)
+		.CommandList(TSharedPtr<FUICommandList>()));
+
 		TWeakObjectPtr<ADisplayClusterRootActor> ActiveRootActor = Toolbar->GetActiveRootActor();
 		if (ActiveRootActor.IsValid())
 		{
@@ -126,8 +131,7 @@ TSharedRef<SDockTab> SDisplayClusterOperatorPanel::SpawnToolbarTab(const FSpawnT
 		.ShouldAutosize(true)
 		.TabRole(ETabRole::PanelTab)
 		[
-			SAssignNew(Toolbar, SDisplayClusterOperatorToolbar)
-				.CommandList(TSharedPtr<FUICommandList>())
+			SAssignNew(ToolbarContainer, SBox)
 		];
 }
 
diff --git a/Engine/Plugins/Runtime/nDisplay/Source/DisplayClusterOperator/Private/SDisplayClusterOperatorPanel.h b/Engine/Plugins/Runtime/nDisplay/Source/DisplayClusterOperator/Private/SDisplayClusterOperatorPanel.h
--- a/Engine/Plugins/Runtime/nDisplay/Source/DisplayClusterOperator/Private/SDisplayClusterOperatorPanel.h
+++ b/Engine/Plugins/Runtime/nDisplay/Source/DisplayClusterOperator/Private/SDisplayClusterOperatorPanel.h
@@ -9,6 +9,7 @@ class FTabManager;
 class FLayoutExtender;
 class FSpawnTabArgs;
 class FMenuBuilder;
+class SBox;
 class SDockTab;
 class SWindow;
 class SDisplayClusterOperatorToolbar;
@@ -64,6 +65,9 @@ private:
 	/** The layout extender used by the operator's layout */
 	TSharedPtr<FLayoutExtender> LayoutExtender;
 
+	/** The container for the toolbar */
+	TSharedPtr<SBox> ToolbarContainer;
+	
 	/** A reference to the operator panel's toolbar widget */
 	TSharedPtr<SDisplayClusterOperatorToolbar> Toolbar;
 
diff --git a/Engine/Plugins/Runtime/nDisplay/Source/DisplayClusterOperator/Private/SDisplayClusterOperatorToolbar.cpp b/Engine/Plugins/Runtime/nDisplay/Source/DisplayClusterOperator/Private/SDisplayClusterOperatorToolbar.cpp
--- a/Engine/Plugins/Runtime/nDisplay/Source/DisplayClusterOperator/Private/SDisplayClusterOperatorToolbar.cpp
+++ b/Engine/Plugins/Runtime/nDisplay/Source/DisplayClusterOperator/Private/SDisplayClusterOperatorToolbar.cpp
@@ -41,12 +41,8 @@ void SDisplayClusterOperatorToolbar::Construct(const FArguments& InArgs)
 	CommandList = InArgs._CommandList;
 
 	FillRootActorList();
-	
-	TSharedPtr<FExtender> ToolBarExtender = IDisplayClusterOperator::Get().GetOperatorToolBarExtensibilityManager()->GetAllExtenders();
-
-    FToolBarBuilder ToolBarBuilder(CommandList, FMultiBoxCustomization::None, ToolBarExtender);
 
-	RootActorComboBox = SNew(SComboBox<TSharedPtr<FString>>)
+	SAssignNew(RootActorComboBox, SComboBox<TSharedPtr<FString>>)
 		.OptionsSource(&RootActorList)
 		.OnSelectionChanged(this, &SDisplayClusterOperatorToolbar::OnRootActorChanged)
 		.OnComboBoxOpening(this, &SDisplayClusterOperatorToolbar::OnRootActorComboBoxOpening)
@@ -57,14 +53,27 @@ void SDisplayClusterOperatorToolbar::Construct(const FArguments& InArgs)
 			.Text(this, &SDisplayClusterOperatorToolbar::GetRootActorComboBoxText)
 		];
 
-	if (RootActorList.Num())
-	{
-		RootActorComboBox->SetSelectedItem(RootActorList[0]);
-	}
-
+	const TSharedPtr<FExtender> ToolBarExtender = IDisplayClusterOperator::Get().GetOperatorToolBarExtensibilityManager()->GetAllExtenders();
+	FToolBarBuilder ToolBarBuilder(CommandList, FMultiBoxCustomization::None, ToolBarExtender, true);
+	
+	const TSharedPtr<SHorizontalBox> RootActorSelectionBox = SNew(SHorizontalBox)
+	+SHorizontalBox::Slot()
+	.AutoWidth()
+	.VAlign(VAlign_Center)
+	.Padding(5.f, 0.f)
+	[
+		SNew(STextBlock)
+		.Text(LOCTEXT("RootActorPickerLabel", "nDisplay Configuration"))
+	]
+	+SHorizontalBox::Slot()
+	.AutoWidth()
+	[
+		RootActorComboBox.ToSharedRef()
+	];
+	
 	ToolBarBuilder.BeginSection("General");
 	{
-		ToolBarBuilder.AddToolBarWidget(RootActorComboBox.ToSharedRef(), LOCTEXT("RootActorPickerLabel", "nDisplay Actor"));
+		ToolBarBuilder.AddWidget(RootActorSelectionBox.ToSharedRef());
 	}
 	ToolBarBuilder.EndSection();
 
@@ -77,6 +86,11 @@ void SDisplayClusterOperatorToolbar::Construct(const FArguments& InArgs)
 			ToolBarBuilder.MakeWidget()
 		]
 	];
+	
+	if (RootActorList.Num())
+	{
+		RootActorComboBox->SetSelectedItem(RootActorList[0]);
+	}
 
 	if (GEngine != nullptr)
 	{
