nDisplay: Fix localization conflict.
Horde issue: 62d9869ab97cae1eae0fc652 - 196375

diff --git a/Engine/Plugins/Runtime/nDisplay/Source/DisplayClusterOperator/Private/SDisplayClusterOperatorToolbar.cpp b/Engine/Plugins/Runtime/nDisplay/Source/DisplayClusterOperator/Private/SDisplayClusterOperatorToolbar.cpp
--- a/Engine/Plugins/Runtime/nDisplay/Source/DisplayClusterOperator/Private/SDisplayClusterOperatorToolbar.cpp
+++ b/Engine/Plugins/Runtime/nDisplay/Source/DisplayClusterOperator/Private/SDisplayClusterOperatorToolbar.cpp
@@ -12,7 +12,7 @@
 #include "Widgets/Layout/SBorder.h"
 #include "Widgets/Input/STextComboBox.h"
 
-#define LOCTEXT_NAMESPACE "SDisplayClusterConfiguratorOperator"
+#define LOCTEXT_NAMESPACE "SDisplayClusterOperatorToolbar"
 
 SDisplayClusterOperatorToolbar::~SDisplayClusterOperatorToolbar()
 {
