nDisplay: fixed typo

diff --git a/Engine/Plugins/Runtime/nDisplay/Source/DisplayCluster/Private/Render/Viewport/LightCard/DisplayClusterViewportLightCardManager.cpp b/Engine/Plugins/Runtime/nDisplay/Source/DisplayCluster/Private/Render/Viewport/LightCard/DisplayClusterViewportLightCardManager.cpp
--- a/Engine/Plugins/Runtime/nDisplay/Source/DisplayCluster/Private/Render/Viewport/LightCard/DisplayClusterViewportLightCardManager.cpp
+++ b/Engine/Plugins/Runtime/nDisplay/Source/DisplayCluster/Private/Render/Viewport/LightCard/DisplayClusterViewportLightCardManager.cpp
@@ -129,7 +129,7 @@ void FDisplayClusterViewportLightCardManager::RenderFrame()
 				}
 			}
 
-			ENQUEUE_RENDER_COMMAND(DisplayClusterViewportLightCardManager_InitializeUVLightCardMap)(
+			ENQUEUE_RENDER_COMMAND(DisplayClusterViewportLightCardManager_RenderFrame)(
 				[InProxyData = ProxyData, bLoadedPrimitives, SceneInterface](FRHICommandListImmediate& RHICmdList)
 				{
 					InProxyData->RenderLightCardMap_RenderThread(RHICmdList, bLoadedPrimitives, SceneInterface);
