nDisplay: Optimize light card proxy full refresh when not using interactive changes.

diff --git a/Engine/Plugins/Runtime/nDisplay/Source/DisplayClusterLightCardEditor/Private/DisplayClusterLightCardEditorLog.h b/Engine/Plugins/Runtime/nDisplay/Source/DisplayClusterLightCardEditor/Private/DisplayClusterLightCardEditorLog.h
--- a/Engine/Plugins/Runtime/nDisplay/Source/DisplayClusterLightCardEditor/Private/DisplayClusterLightCardEditorLog.h
+++ b/Engine/Plugins/Runtime/nDisplay/Source/DisplayClusterLightCardEditor/Private/DisplayClusterLightCardEditorLog.h
@@ -3,6 +3,9 @@
 #pragma once
 
 #include "CoreMinimal.h"
+#include "Stats/Stats2.h"
 
 // Module-wide log categories
 DECLARE_LOG_CATEGORY_EXTERN(DisplayClusterLightCardEditorLog, Log, All);
+
+DECLARE_STATS_GROUP(TEXT("nDisplay Light Card Editor"), STATGROUP_NDisplayLightCardEditor, STATCAT_Advanced)
\ No newline at end of file
diff --git a/Engine/Plugins/Runtime/nDisplay/Source/DisplayClusterLightCardEditor/Private/Viewport/DisplayClusterLightCardEditorViewportClient.cpp b/Engine/Plugins/Runtime/nDisplay/Source/DisplayClusterLightCardEditor/Private/Viewport/DisplayClusterLightCardEditorViewportClient.cpp
--- a/Engine/Plugins/Runtime/nDisplay/Source/DisplayClusterLightCardEditor/Private/Viewport/DisplayClusterLightCardEditorViewportClient.cpp
+++ b/Engine/Plugins/Runtime/nDisplay/Source/DisplayClusterLightCardEditor/Private/Viewport/DisplayClusterLightCardEditorViewportClient.cpp
@@ -417,6 +417,14 @@ FDisplayClusterLightCardEditorViewportClient::~FDisplayClusterLightCardEditorVie
 	{
 		RootActorLevelInstance->UnsubscribeFromPostProcessRenderTarget(reinterpret_cast<uint8*>(this));
 	}
+
+	if (RefreshTimerHandle.IsValid())
+	{
+		if (const UWorld* PreviewWorld = PreviewScene->GetWorld())
+		{
+			PreviewWorld->GetTimerManager().ClearTimer(RefreshTimerHandle);
+		}
+	}
 }
 
 FLinearColor FDisplayClusterLightCardEditorViewportClient::GetBackgroundColor() const
@@ -1488,17 +1496,23 @@ bool FDisplayClusterLightCardEditorViewportClient::DropObjectsAtCoordinates(int3
 void FDisplayClusterLightCardEditorViewportClient::UpdatePreviewActor(ADisplayClusterRootActor* RootActor, bool bForce,
                                                                       EDisplayClusterLightCardEditorProxyType ProxyType)
 {
-	if (!bForce && RootActor == RootActorLevelInstance.Get())
+	DECLARE_SCOPE_CYCLE_COUNTER(TEXT("FDisplayClusterLightCardEditorViewportClient::UpdatePreviewActor"), STAT_UpdatePreviewActor, STATGROUP_NDisplayLightCardEditor);
+	
+	if ((!bForce && RootActor == RootActorLevelInstance.Get()) ||
+		(ProxyTypesRefreshing.Contains(ProxyType) || ProxyTypesRefreshing.Contains(EDisplayClusterLightCardEditorProxyType::All)))
 	{
 		return;
 	}
 
-	auto Finalize = [this]()
+	ProxyTypesRefreshing.Add(ProxyType);
+
+	auto Finalize = [this, ProxyType]()
 	{
 		Viewport->InvalidateHitProxy();
 		bShouldCheckHitProxy = true;
 		InvalidateNormalMap();
 
+		ProxyTypesRefreshing.Remove(ProxyType);
 		OnNextSceneRefreshDelegate.Broadcast();
 		OnNextSceneRefreshDelegate.Clear();
 	};
@@ -1513,16 +1527,17 @@ void FDisplayClusterLightCardEditorViewportClient::UpdatePreviewActor(ADisplayCl
 		UWorld* PreviewWorld = PreviewScene->GetWorld();
 		check(PreviewWorld);
 
-		TWeakObjectPtr<ADisplayClusterRootActor> RootActorPtr (RootActor);
+		const TWeakObjectPtr<ADisplayClusterRootActor> RootActorPtr (RootActor);
 		
 		// Schedule for the next tick so CDO changes get propagated first in the event of config editor skeleton
 		// regeneration & compiles. nDisplay's custom propagation may have issues if the archetype isn't correct.
-		PreviewWorld->GetTimerManager().SetTimerForNextTick([=]()
+		RefreshTimerHandle = PreviewWorld->GetTimerManager().SetTimerForNextTick([=]()
 		{			
 			DestroyProxies(ProxyType);
 
 			if (!RootActorPtr.IsValid())
 			{
+				Finalize();
 				return;
 			}
 			
diff --git a/Engine/Plugins/Runtime/nDisplay/Source/DisplayClusterLightCardEditor/Private/Viewport/DisplayClusterLightCardEditorViewportClient.h b/Engine/Plugins/Runtime/nDisplay/Source/DisplayClusterLightCardEditor/Private/Viewport/DisplayClusterLightCardEditorViewportClient.h
--- a/Engine/Plugins/Runtime/nDisplay/Source/DisplayClusterLightCardEditor/Private/Viewport/DisplayClusterLightCardEditorViewportClient.h
+++ b/Engine/Plugins/Runtime/nDisplay/Source/DisplayClusterLightCardEditor/Private/Viewport/DisplayClusterLightCardEditorViewportClient.h
@@ -349,6 +349,12 @@ private:
 	TWeakObjectPtr<ADisplayClusterRootActor> RootActorProxy;
 	TWeakObjectPtr<ADisplayClusterRootActor> RootActorLevelInstance;
 
+	/** Proxy types that are currently refreshing over 1-2 frames */
+	TSet<EDisplayClusterLightCardEditorProxyType> ProxyTypesRefreshing;
+
+	/** Handle to world timer when refreshing */
+	FTimerHandle RefreshTimerHandle;
+	
 	/** The radius of the bounding sphere that entirely encapsulates the root actor */
 	float RootActorBoundingRadius = 0.0f;
 
