[nDisplay][EpicStageApp] API for light card position control via WebSocket

diff --git a/Engine/Plugins/Runtime/nDisplay/Shaders/Private/MeshProjectionNormalSmoothing.usf b/Engine/Plugins/Runtime/nDisplay/Shaders/Private/MeshProjectionNormalSmoothing.usf
--- a/Engine/Plugins/Runtime/nDisplay/Shaders/Private/MeshProjectionNormalSmoothing.usf
+++ b/Engine/Plugins/Runtime/nDisplay/Shaders/Private/MeshProjectionNormalSmoothing.usf
@@ -13,6 +13,8 @@ RWTexture2D<float4> RWColor;
 RWTexture2D<float> RWDepth;
 RWTexture2D<uint> RWStencil;
 
+float4x4 NormalCorrectionMatrix;
+
 #define FILTER_KERNEL_SIZE 25
 
 DECLARE_SCALAR_ARRAY(float, SpatialKernel, FILTER_KERNEL_SIZE);
@@ -73,7 +75,7 @@ void CreateRWTexturesCS(uint2 DispatchThreadId : SV_DispatchThreadID)
 	float3 ViewPos = float3(sin(PolarCoords.x) * cos(PolarCoords.y), sin(PolarCoords.x) * sin(PolarCoords.y), cos(PolarCoords.x)) * Rho;
 	float3 WorldPos = mul(float4(ViewPos, 0), View.ViewToTranslatedWorld).xyz;
 	
-	float3 RadialVector = normalize(WorldPos);
+	float3 RadialVector = normalize(mul(WorldPos, (float3x3) NormalCorrectionMatrix));
 	
 	float4 Color = SceneColor.Load(uint3(PixelPos, 0));
 	float Depth = SceneDepth.Load(uint3(PixelPos, 0)).r;
diff --git a/Engine/Plugins/Runtime/nDisplay/Shaders/Private/MeshProjectionShaders.usf b/Engine/Plugins/Runtime/nDisplay/Shaders/Private/MeshProjectionShaders.usf
--- a/Engine/Plugins/Runtime/nDisplay/Shaders/Private/MeshProjectionShaders.usf
+++ b/Engine/Plugins/Runtime/nDisplay/Shaders/Private/MeshProjectionShaders.usf
@@ -139,6 +139,8 @@ void MainPS(FMeshProjectionVSToPS Input OPTIONAL_IsFrontFace, out float4 OutputC
 	OutputColor.a = lerp(OutputColor.a, 0, AlphaBlendWeight);
 }
 
+float4x4 NormalCorrectionMatrix;
+
 void NormalPS(FMeshProjectionVSToPS Input OPTIONAL_IsFrontFace, out float4 OutputColor : SV_Target0)
 {
 	ResolvedView = ResolveView();
@@ -151,8 +153,8 @@ void NormalPS(FMeshProjectionVSToPS Input OPTIONAL_IsFrontFace, out float4 Outpu
 	CalcMaterialParameters(MaterialParameters, PixelMaterialInputs, Input.Position, bIsFrontFace);
 
 	GetMaterialCoverageAndClipping(MaterialParameters, PixelMaterialInputs);
-
-	float3 Normal = TransformTangentNormalToWorld(MaterialParameters.TangentToWorld, float3(0, 0, 1));
+	
+	float3 Normal = mul(TransformTangentNormalToWorld(MaterialParameters.TangentToWorld, float3(0, 0, 1)), (float3x3) NormalCorrectionMatrix);
 	
 	OutputColor = float4(Normal * 0.5 + 0.5, 1.0f);
 }
\ No newline at end of file
diff --git a/Engine/Plugins/Runtime/nDisplay/Source/DisplayCluster/Public/Components/DisplayClusterPreviewComponent.h b/Engine/Plugins/Runtime/nDisplay/Source/DisplayCluster/Public/Components/DisplayClusterPreviewComponent.h
--- a/Engine/Plugins/Runtime/nDisplay/Source/DisplayCluster/Public/Components/DisplayClusterPreviewComponent.h
+++ b/Engine/Plugins/Runtime/nDisplay/Source/DisplayCluster/Public/Components/DisplayClusterPreviewComponent.h
@@ -76,6 +76,11 @@ public:
 		return PreviewMesh;
 	}
 
+	UTexture* GetOverrideTexture() const
+	{
+		return OverrideTexture;
+	}
+
 	/** Create and retrieve a render texture 2d from the render target. */
 	UTexture* GetViewportPreviewTexture2D();
 
diff --git a/Engine/Plugins/Runtime/nDisplay/Source/DisplayClusterLightCardEditor/Private/Viewport/DisplayClusterLightCardEditorViewportClient.cpp b/Engine/Plugins/Runtime/nDisplay/Source/DisplayClusterLightCardEditor/Private/Viewport/DisplayClusterLightCardEditorViewportClient.cpp
--- a/Engine/Plugins/Runtime/nDisplay/Source/DisplayClusterLightCardEditor/Private/Viewport/DisplayClusterLightCardEditorViewportClient.cpp
+++ b/Engine/Plugins/Runtime/nDisplay/Source/DisplayClusterLightCardEditor/Private/Viewport/DisplayClusterLightCardEditorViewportClient.cpp
@@ -11,6 +11,7 @@
 #include "Components/DisplayClusterCameraComponent.h"
 #include "Components/DisplayClusterScreenComponent.h"
 #include "DisplayClusterConfigurationTypes.h"
+#include "DisplayClusterLightCardEditorHelper.h"
 #include "DisplayClusterProjectionStrings.h"
 #include "DisplayClusterRootActor.h"
 #include "DisplayClusterLightCardActor.h"
@@ -62,295 +63,6 @@
 
 #define LOCTEXT_NAMESPACE "DisplayClusterLightCardEditorViewportClient"
 
-
-FDisplayClusterLightCardEditorViewportClient::FSphericalCoordinates::FSphericalCoordinates(const FVector& CartesianPosition)
-{
-	Radius = CartesianPosition.Size();
-
-	if (Radius > UE_SMALL_NUMBER)
-	{
-		Inclination = FMath::Acos(CartesianPosition.Z / Radius);
-	}
-	else
-	{
-		Inclination = 0;
-	}
-
-	Azimuth = FMath::Atan2(CartesianPosition.Y, CartesianPosition.X);
-}
-
-FDisplayClusterLightCardEditorViewportClient::FSphericalCoordinates::FSphericalCoordinates()
-	: Radius(0)
-	, Inclination(0)
-	, Azimuth(0)
-{ 
-
-}
-
-FVector FDisplayClusterLightCardEditorViewportClient::FSphericalCoordinates::AsCartesian() const
-{
-	const double SinAzimuth = FMath::Sin(Azimuth);
-	const double CosAzimuth = FMath::Cos(Azimuth);
-
-	const double SinInclination = FMath::Sin(Inclination);
-	const double CosInclination = FMath::Cos(Inclination);
-
-	return FVector(
-		Radius * CosAzimuth * SinInclination,
-		Radius * SinAzimuth * SinInclination,
-		Radius * CosInclination
-	);
-}
-
-FDisplayClusterLightCardEditorViewportClient::FSphericalCoordinates 
-FDisplayClusterLightCardEditorViewportClient::FSphericalCoordinates::operator+(FSphericalCoordinates const& Other) const
-{
-	FSphericalCoordinates Result;
-
-	Result.Radius = Radius + Other.Radius;
-	Result.Inclination = Inclination + Other.Inclination;
-	Result.Azimuth = Azimuth + Other.Azimuth;
-
-	return Result;
-}
-
-FDisplayClusterLightCardEditorViewportClient::FSphericalCoordinates 
-FDisplayClusterLightCardEditorViewportClient::FSphericalCoordinates::operator-(FSphericalCoordinates const& Other) const
-{
-	FSphericalCoordinates Result;
-
-	Result.Radius = Radius - Other.Radius;
-	Result.Inclination = Inclination - Other.Inclination;
-	Result.Azimuth = Azimuth - Other.Azimuth;
-
-	return Result;
-}
-
-void FDisplayClusterLightCardEditorViewportClient::FSphericalCoordinates::Conform()
-{
-	if (Radius < 0)
-	{
-		Radius = -Radius;
-		Inclination += PI;
-	}
-
-	if (Inclination < 0 || Inclination > PI)
-	{
-		// -2PI to 2PI
-		Inclination = FMath::Fmod(Inclination, 2 * PI);
-
-		// 0 to 2PI
-		if (Inclination < 0)
-		{
-			Inclination += 2 * PI;
-		}
-
-		// 0 to PI
-		if (Inclination > PI)
-		{
-			Inclination = 2 * PI - Inclination;
-			Azimuth += PI;
-		}
-	}
-
-	if (Azimuth < -PI || Azimuth > PI)
-	{
-		// -2PI to 2PI
-		Azimuth = FMath::Fmod(Azimuth, 2 * PI);
-
-		// -PI to PI
-		if (Azimuth > PI)
-		{
-			Azimuth -= 2 * PI;
-		}
-		else if (Azimuth < -PI)
-		{
-			Azimuth += 2 * PI;
-		}
-	}
-
-	checkSlow(Radius >= 0);
-	checkSlow(Inclination >= 0 && Inclination <= PI);
-	checkSlow(Azimuth >= -PI && Azimuth <= PI);
-}
-
-FDisplayClusterLightCardEditorViewportClient::FSphericalCoordinates 
-FDisplayClusterLightCardEditorViewportClient::FSphericalCoordinates::GetConformed() const
-{
-	FSphericalCoordinates Result = *this;
-	Result.Conform();
-	return Result;
-}
-
-bool FDisplayClusterLightCardEditorViewportClient::FSphericalCoordinates::IsPointingAtPole(double Margin) const
-{
-	FSphericalCoordinates CoordsConformed = GetConformed();
-
-	return FMath::IsNearlyZero(CoordsConformed.Inclination, Margin)
-		|| FMath::IsNearlyEqual(CoordsConformed.Inclination, PI, Margin);
-}
-
-//////////////////////////////////////////////////////////////////////////
-// FNormalMap
-
-const int32 FDisplayClusterLightCardEditorViewportClient::FNormalMap::NormalMapSize = 512;
-const float FDisplayClusterLightCardEditorViewportClient::FNormalMap::NormalMapFOV = 2.0f * FMath::RadiansToDegrees(FMath::Atan(0.55 * PI)); // Equation for FOV from desired angle from north pole;
-
-void FDisplayClusterLightCardEditorViewportClient::FNormalMap::Init(const FSceneViewInitOptions& InSceneViewInitOptions)
-{
-	SizeX = InSceneViewInitOptions.GetViewRect().Width();
-	SizeY = InSceneViewInitOptions.GetViewRect().Height();
-
-	ViewMatrices = FViewMatrices(InSceneViewInitOptions);
-
-	if (NormalMapTexture.IsValid())
-	{
-		NormalMapTexture->MarkAsGarbage();
-		NormalMapTexture = nullptr;
-	}
-
-	ENQUEUE_RENDER_COMMAND(InitRHIResourcesCommand)([this](FRHICommandListImmediate& RHICmdList)
-		{
-			const FRHITextureCreateDesc Desc =
-				FRHITextureCreateDesc::Create2D(TEXT("NormalMapTexture"))
-				.SetExtent(SizeX, SizeY)
-				.SetFormat(PF_FloatRGBA)
-				.SetClearValue(FClearValueBinding::Black)
-				.SetFlags(ETextureCreateFlags::RenderTargetable | ETextureCreateFlags::ShaderResource)
-				.SetInitialState(ERHIAccess::SRVMask);
-
-			RenderTargetTextureRHI = RHICreateTexture(Desc);
-		});
-}
-
-void FDisplayClusterLightCardEditorViewportClient::FNormalMap::Release()
-{
-	ENQUEUE_RENDER_COMMAND(ReleaseRHIResourcesCommand)([this](FRHICommandListImmediate& RHICmdList)
-		{
-			RenderTargetTextureRHI.SafeRelease();
-		});
-}
-
-bool FDisplayClusterLightCardEditorViewportClient::FNormalMap::GetNormalAndDistanceAtPosition(FVector Position, FVector& OutNormal, float& OutDistance) const
-{
-	auto GetPixel = [this](uint32 InX, uint32 InY)
-	{
-		uint32 ClampedX = FMath::Clamp(InX, (uint32)0, SizeX - 1);
-		uint32 ClampedY = FMath::Clamp(InY, (uint32)0, SizeY - 1);
-
-		return CachedNormalData[ClampedY * SizeX + ClampedX].GetFloats();
-	};
-
-	const FVector ViewPos = FVector(ViewMatrices.GetViewMatrix().TransformFVector4(FVector4(Position, 1)));
-	const FVector ProjectedViewPos = FDisplayClusterMeshProjectionRenderer::ProjectViewPosition(ViewPos, EDisplayClusterMeshProjectionType::Azimuthal);
-
-	const FVector4 ScreenPos = ViewMatrices.GetProjectionMatrix().TransformFVector4(FVector4(ProjectedViewPos, 1));
-
-	if (ScreenPos.W != 0.0)
-	{
-		const float InvW = (ScreenPos.W > 0.0f ? 1.0f : -1.0f) / ScreenPos.W;
-		const float Y = (GProjectionSignY > 0.0f) ? ScreenPos.Y : 1.0f - ScreenPos.Y;
-		const FVector2D PixelPos = FVector2D((0.5f + ScreenPos.X * 0.5f * InvW) * SizeX, (0.5f - Y * 0.5f * InvW) * SizeY);
-
-		// Perform a bilinear interpolation on the computed pixel position to ensure a continuous normal regardless of the resolution of the normal map
-		const uint32 PixelX = FMath::Floor(PixelPos.X - 0.5f);
-		const uint32 PixelY = FMath::Floor(PixelPos.Y - 0.5f);
-		const float PixelXFrac = FMath::Frac(PixelPos.X);
-		const float PixelYFrac = FMath::Frac(PixelPos.X);
-
-		FLinearColor NormalData;
-		NormalData = FMath::Lerp(
-			FMath::Lerp(GetPixel(PixelX, PixelY), GetPixel(PixelX + 1, PixelY), PixelXFrac),
-			FMath::Lerp(GetPixel(PixelX, PixelY + 1), GetPixel(PixelX + 1, PixelY + 1), PixelXFrac),
-			PixelYFrac);
-
-		const FVector NormalVector = 2.f * FVector(NormalData.R, NormalData.G, NormalData.B) - 1.f;
-		OutNormal = NormalVector.GetSafeNormal();
-
-		// Make sure the depth value is not 0, as that will cause a divide by zero when transformed, resulting in an NaN distance being returned
-		const float Depth = FMath::Max(0.001f, NormalData.A);
-
-		FVector4 DepthPos = ViewMatrices.GetInvProjectionMatrix().TransformFVector4(FVector4(ScreenPos.X * InvW, ScreenPos.Y * InvW, Depth, 1.0f));
-		DepthPos /= DepthPos.W;
-
-		const FVector UnprojectedDepthPos = FDisplayClusterMeshProjectionRenderer::UnprojectViewPosition(DepthPos, EDisplayClusterMeshProjectionType::Azimuthal);
-		OutDistance = UnprojectedDepthPos.Length();
-
-		return true;
-	}
-	else
-	{
-		OutNormal = FVector::ZeroVector;
-		OutDistance = 0.0;
-		return false;
-	}
-}
-
-void FDisplayClusterLightCardEditorViewportClient::FNormalMap::MorphProceduralMesh(UProceduralMeshComponent* InProceduralMeshComponent) const
-{
-	const FVector ViewOrigin = ViewMatrices.GetViewOrigin();
-	const FVector ViewDirection = ViewMatrices.GetViewMatrix().TransformVector(FVector::ZAxisVector);
-	const float MaxAngle = 1.0f / ViewMatrices.GetProjectionMatrix().M[0][0];
-
-	FProcMeshSection& Section = *InProceduralMeshComponent->GetProcMeshSection(0);
-
-
-	for (int32 Index = 0; Index < Section.ProcVertexBuffer.Num(); ++Index)
-	{
-		FProcMeshVertex& Vertex = Section.ProcVertexBuffer[Index];
-
-		const FVector VertexPosition = Vertex.Position;
-		const FVector VertexWorldPosition = VertexPosition + ViewOrigin;
-		const FVector VertexDirection = Vertex.Position.GetSafeNormal();
-		const float VertexAngle = FMath::Acos(VertexDirection | ViewDirection);
-
-		if (VertexAngle < MaxAngle)
-		{
-			FVector Normal;
-			float Depth;
-			GetNormalAndDistanceAtPosition(VertexWorldPosition, Normal, Depth);
-
-			const FVector NewPosition = VertexDirection * Depth;
-			Vertex.Position = NewPosition;
-
-			const FMatrix RadialBasis = FRotationMatrix::MakeFromX(VertexDirection);
-
-			const FVector WorldNormal = RadialBasis.TransformVector(Normal);
-			Vertex.Normal = WorldNormal;
-		}
-	}
-
-	InProceduralMeshComponent->SetProcMeshSection(0, Section);
-}
-
-UTexture2D* FDisplayClusterLightCardEditorViewportClient::FNormalMap::GenerateNormalMapTexture(const FString& TextureName)
-{
-	if (NormalMapTexture.IsValid())
-	{
-		NormalMapTexture->MarkAsGarbage();
-		NormalMapTexture = nullptr;
-	}
-
-	if (CachedNormalData.Num())
-	{
-		FCreateTexture2DParameters Params;
-		Params.bDeferCompression = true;
-
-		TArray<FColor> Bitmap;
-		Bitmap.AddZeroed(CachedNormalData.Num());
-
-		for (int32 Index = 0; Index < CachedNormalData.Num(); ++Index)
-		{
-			Bitmap[Index] = CachedNormalData[Index].GetFloats().ToFColor(false);
-		}
-
-		NormalMapTexture = FImageUtils::CreateTexture2D(SizeX, SizeY, Bitmap, GetTransientPackage(), TextureName, RF_Transient, Params);
-	}
-
-	return GetNormalMapTexture();
-}
-
-
 //////////////////////////////////////////////////////////////////////////
 // FDisplayClusterLightCardEditorViewportClient
 
@@ -375,6 +87,8 @@ FDisplayClusterLightCardEditorViewportClient::FDisplayClusterLightCardEditorView
 		FDisplayClusterMeshProjectionRenderer::FSimpleElementPass::CreateRaw(this, &FDisplayClusterLightCardEditorViewportClient::Draw)
 	);
 
+	ProjectionHelper = MakeShared<FDisplayClusterLightCardEditorHelper>(PreviewRendererId);
+
 	InputMode = EInputMode::Idle;
 	DragWidgetOffset = FVector::ZeroVector;
 	
@@ -513,20 +227,6 @@ void FDisplayClusterLightCardEditorViewportClient::Tick(float DeltaSeconds)
 
 void FDisplayClusterLightCardEditorViewportClient::Draw(FViewport* InViewport, FCanvas* Canvas)
 {
-	if (bNormalMapInvalid)
-	{
-		RenderNormalMap(NorthNormalMap, FVector::UpVector);
-		RenderNormalMap(SouthNormalMap, -FVector::UpVector);
-
-		if (NormalMapMeshComponent.IsValid())
-		{
-			NorthNormalMap.MorphProceduralMesh(NormalMapMeshComponent.Get());
-			SouthNormalMap.MorphProceduralMesh(NormalMapMeshComponent.Get());
-		}
-
-		bNormalMapInvalid = false;
-	}
-
 	FViewport* ViewportBackup = Viewport;
 	Viewport = InViewport ? InViewport : Viewport;
 
@@ -655,23 +355,16 @@ void FDisplayClusterLightCardEditorViewportClient::Draw(FViewport* InViewport, F
 
 	if (bDisplayNormalMapVisualization)
 	{
-		auto DrawNormalMap = [Canvas](FNormalMap& NormalMap, const FString& TextureName, FVector2D Position)
+		auto DrawNormalMap = [Canvas, this](bool bShowNorthMap, FVector2D Position)
 		{
-			UTexture2D* NormalMapTexture = NormalMap.GetNormalMapTexture();
-
-			if (!NormalMapTexture)
-			{
-				NormalMapTexture = NormalMap.GenerateNormalMapTexture(TextureName);
-			}
-
-			if (NormalMapTexture)
+			if (const UTexture* NormalMapTexture = ProjectionHelper->GetNormalMapTexture(bShowNorthMap))
 			{
 				Canvas->DrawTile(Position.X, Position.Y, 512, 512, 0, 0, 1, 1, FLinearColor::White, NormalMapTexture->GetResource());
 			}
 		};
 
-		DrawNormalMap(NorthNormalMap, TEXT("DisplayClusterLightCardEditor.NorthNormalMap"), FVector2D(0.0f, 0.0f));
-		DrawNormalMap(SouthNormalMap, TEXT("DisplayClusterLightCardEditor.SouthNormalMap"), FVector2D(0.0f, 512.0f));
+		DrawNormalMap(true, FVector2D(0.0f, 0.0f));
+		DrawNormalMap(false, FVector2D(0.0f, 512.0f));
 	}
 
 	// Remove temporary debug lines.
@@ -908,7 +601,7 @@ void FDisplayClusterLightCardEditorViewportClient::TrackingStarted(
 
 		FVector Origin;
 		FVector Direction;
-		PixelToWorld(*View, MousePos, Origin, Direction);
+		ProjectionHelper->PixelToWorld(*View, MousePos, Origin, Direction);
 
 		if (RenderViewportType != LVT_Perspective)
 		{
@@ -945,50 +638,6 @@ void FDisplayClusterLightCardEditorViewportClient::TrackingStopped()
 	FEditorViewportClient::TrackingStopped();
 }
 
-void FDisplayClusterLightCardEditorViewportClient::CalculateNormalAndPositionInDirection(
-	const FVector& InViewOrigin, 
-	const FVector& InDirection, 
-	FVector& OutWorldPosition, 
-	FVector& OutRelativeNormal, 
-	double InDesiredDistanceFromFlush) const
-{
-	if (ProjectionMode == EDisplayClusterMeshProjectionType::UV)
-	{
-		// In UV projection mode, all relevant geometry is projected onto the UV plane, so compute the world position and normal
-		// by performing a ray-plane intersection on the UV projection plane
-
-		const float UVProjectionPlaneDistance = ADisplayClusterLightCardActor::UVPlaneDefaultDistance;
-
-		const FPlane UVProjectionPlane(InViewOrigin + FVector::ForwardVector * UVProjectionPlaneDistance, -FVector::ForwardVector);
-		const FVector PlaneIntersection = FMath::RayPlaneIntersection(InViewOrigin, InDirection, UVProjectionPlane);
-
-		OutRelativeNormal = UVProjectionPlane.GetNormal();
-		OutWorldPosition = PlaneIntersection;
-	}
-	else
-	{
-		const FVector Position = InViewOrigin + InDirection; // We fabricate a position in the right direction
-
-		float Distance;
-
-		// We find the normal and distance from origin
-		if (Position.Z < InViewOrigin.Z)
-		{
-			SouthNormalMap.GetNormalAndDistanceAtPosition(Position, OutRelativeNormal, Distance);
-		}
-		else
-		{
-			NorthNormalMap.GetNormalAndDistanceAtPosition(Position, OutRelativeNormal, Distance);
-		}
-
-		Distance = CalculateFinalLightCardDistance(Distance, InDesiredDistanceFromFlush);
-
-		// Calculate world position
-		OutWorldPosition = InViewOrigin + Distance * InDirection;
-	}
-};
-
-
 void FDisplayClusterLightCardEditorViewportClient::CreateDrawnLightCard(const TArray<FIntPoint>& MousePositions)
 {
 #if WITH_OPENCV
@@ -1018,16 +667,7 @@ void FDisplayClusterLightCardEditorViewportClient::CreateDrawnLightCard(const TA
 
 	for (int32 PointIdx = 0; PointIdx < MousePositions.Num(); ++PointIdx)
 	{
-		PixelToWorld(*View, MousePositions[PointIdx], ScreenOrigin, ScreenDirection);
-
-		if (RenderViewportType != LVT_Perspective)
-		{
-			// For orthogonal projections, PixelToWorld does not return the view origin or a direction from the view origin. Use TraceScreenRay
-			// to find a useful direction away from the view origin to use
-
-			ScreenDirection = TraceScreenRay(ScreenOrigin, ScreenDirection, ViewOrigin);
-		}
-
+		ProjectionHelper->CalculateOriginAndDirectionFromPixelPosition(MousePositions[PointIdx], *View, FVector::ZeroVector, ScreenOrigin, ScreenDirection);
 		MouseWorldDirections[PointIdx] = ScreenDirection;
 	}
 
@@ -1051,7 +691,7 @@ void FDisplayClusterLightCardEditorViewportClient::CreateDrawnLightCard(const TA
 			FVector MouseWorldPosition;
 			FVector MouseRelativeNormal;
 
-			CalculateNormalAndPositionInDirection(ViewOrigin, MouseWorldDirection, MouseWorldPosition, MouseRelativeNormal);
+			ProjectionHelper->CalculateNormalAndPositionInDirection(ViewOrigin, MouseWorldDirection, MouseWorldPosition, MouseRelativeNormal);
 
 			MinX = FMath::Min(MinX, MouseWorldPosition.X);
 			MinY = FMath::Min(MinY, MouseWorldPosition.Y);
@@ -1089,7 +729,7 @@ void FDisplayClusterLightCardEditorViewportClient::CreateDrawnLightCard(const TA
 	{
 		FVector LightCardRelativeNormal;
 
-		CalculateNormalAndPositionInDirection(ViewOrigin, LightCardDirection, LightCardLocation, LightCardRelativeNormal);
+		ProjectionHelper->CalculateNormalAndPositionInDirection(ViewOrigin, LightCardDirection, LightCardLocation, LightCardRelativeNormal);
 
 		const FMatrix RotEffector = FRotationMatrix::MakeFromX(-LightCardRelativeNormal);
 		const FMatrix RotArm = ProjectionMode != EDisplayClusterMeshProjectionType::UV ? FRotationMatrix::MakeFromX(LightCardDirection) : FMatrix::Identity;
@@ -1251,8 +891,8 @@ void FDisplayClusterLightCardEditorViewportClient::CreateDrawnLightCard(const TA
 	}
 	else
 	{
-		const FSphericalCoordinates LightCardCoords(LightCardLocation - ViewOrigin);
-		MoveLightCardTo(*LightCard, LightCardCoords);
+		const FDisplayClusterLightCardEditorHelper::FSphericalCoordinates LightCardCoords(LightCardLocation - ViewOrigin);
+		MoveLightCardTo(*LightCard, LightCardLocation - ViewOrigin);
 	}
 
 #endif // WITH_OPENCV
@@ -1412,7 +1052,10 @@ bool FDisplayClusterLightCardEditorViewportClient::UpdateDropPreviewActors(int32
 
 	for (TWeakObjectPtr<ADisplayClusterLightCardActor>& LightCard : DropPreviewLightCards)
 	{
-		VerifyAndFixLightCardOrigin(LightCard.Get());
+		if (LightCard.IsValid())
+		{
+			ProjectionHelper->VerifyAndFixLightCardOrigin(*LightCard);
+		}
 	}
 	
 	MoveLightCardsToPixel(NewMousePos, DropPreviewLightCards);
@@ -1457,7 +1100,7 @@ bool FDisplayClusterLightCardEditorViewportClient::DropObjectsAtCoordinates(int3
 				bCreateDropPreview ? PreviewWorld->GetCurrentLevel() : nullptr, bCreateDropPreview);
 			check(LightCardActor);
 
-			VerifyAndFixLightCardOrigin(LightCardActor);
+			ProjectionHelper->VerifyAndFixLightCardOrigin(*LightCardActor);
 			
 			CreatedLightCards.Add(LightCardActor);
 			
@@ -1510,7 +1153,6 @@ void FDisplayClusterLightCardEditorViewportClient::UpdatePreviewActor(ADisplayCl
 	{
 		Viewport->InvalidateHitProxy();
 		bShouldCheckHitProxy = true;
-		InvalidateNormalMap();
 
 		ProxyTypesRefreshing.Remove(ProxyType);
 		OnNextSceneRefreshDelegate.Broadcast();
@@ -1566,6 +1208,9 @@ void FDisplayClusterLightCardEditorViewportClient::UpdatePreviewActor(ADisplayCl
 				RootActorProxy->UpdatePreviewComponents();
 				RootActorProxy->EnableEditorRender(false);
 
+				ProjectionHelper->SetLevelInstanceRootActor(*RootActorLevelInstance);
+				ProjectionHelper->SetEditorViewportClient(AsShared());
+
 				if (UDisplayClusterConfigurationData* ProxyConfig = RootActorProxy->GetConfigData())
 				{
 					// Disable lightcards so that it doesn't try to update the ones in the level instance world.
@@ -1574,40 +1219,6 @@ void FDisplayClusterLightCardEditorViewportClient::UpdatePreviewActor(ADisplayCl
 
 				FBox BoundingBox = RootActorProxy->GetComponentsBoundingBox();
 				RootActorBoundingRadius = FMath::Max(BoundingBox.Min.Length(), BoundingBox.Max.Length());
-
-				// Create the normal map mesh component, which is used to run traces on the generated normal map of the stage
-				{
-					NormalMapMeshComponent = NewObject<UProceduralMeshComponent>(GetTransientPackage(), TEXT("NormalMapMesh"));
-					NormalMapMeshComponent->SetCollisionEnabled(ECollisionEnabled::QueryOnly);
-					PreviewScene->AddComponent(NormalMapMeshComponent.Get(), FTransform(ProjectionOriginComponent.Get()->GetComponentLocation()));
-
-					UStaticMesh* IcoSphereMesh = LoadObject<UStaticMesh>(nullptr, TEXT("/nDisplay/Meshes/SM_IcoSphere.SM_IcoSphere"), nullptr, LOAD_None, nullptr);
-
-					if (ensure(IcoSphereMesh))
-					{
-						int32 NumSections = IcoSphereMesh->GetNumSections(0);
-						for (int32 SectionIndex = 0; SectionIndex < NumSections; SectionIndex++)
-						{
-							TArray<FVector> Vertices;
-							TArray<int32> Triangles;
-							TArray<FVector> Normals;
-							TArray<FVector2D> UVs;
-							TArray<FProcMeshTangent> Tangents;
-
-							UKismetProceduralMeshLibrary::GetSectionFromStaticMesh(IcoSphereMesh, 0, SectionIndex, Vertices, Triangles, Normals, UVs, Tangents);
-
-							TArray<FVector2D> EmptyUVs;
-							TArray<FLinearColor> EmptyColors;
-							NormalMapMeshComponent.Get()->CreateMeshSection_LinearColor(SectionIndex, Vertices, Triangles, Normals, UVs, EmptyUVs, EmptyUVs, EmptyUVs, EmptyColors, Tangents, true);
-
-							for (int32 Index = 0; Index < IcoSphereMesh->GetStaticMaterials().Num(); ++Index)
-							{
-								UMaterialInterface* MaterialInterface = IcoSphereMesh->GetStaticMaterials()[Index].MaterialInterface;
-								NormalMapMeshComponent.Get()->SetMaterial(Index, MaterialInterface);
-							}
-						}
-					}
-				}
 			}
 
 			// Filter out any primitives hidden in game except screen components
@@ -1755,12 +1366,6 @@ void FDisplayClusterLightCardEditorViewportClient::DestroyProxies(
 			RootActorLevelInstance->UnsubscribeFromPostProcessRenderTarget(reinterpret_cast<uint8*>(this));
 			RootActorLevelInstance.Reset();
 		}
-
-		if (NormalMapMeshComponent.IsValid())
-		{
-			PreviewScene->RemoveComponent(NormalMapMeshComponent.Get());
-			NormalMapMeshComponent.Reset();
-		}
 	}
 	
 	if (ProxyType == EDisplayClusterLightCardEditorProxyType::All ||
@@ -1798,6 +1403,9 @@ void FDisplayClusterLightCardEditorViewportClient::SetProjectionMode(EDisplayClu
 	ProjectionMode = InProjectionMode;
 	RenderViewportType = InViewportType;
 
+	ProjectionHelper->SetProjectionMode(ProjectionMode);
+	ProjectionHelper->SetIsOrthographic(InViewportType != LVT_Perspective);
+
 	if (ProjectionMode == EDisplayClusterMeshProjectionType::Linear || ProjectionMode == EDisplayClusterMeshProjectionType::UV)
 	{
 		// TODO: Do we want to cache the perspective rotation and restore it when the user switches back?
@@ -1869,36 +1477,16 @@ void FDisplayClusterLightCardEditorViewportClient::ResetCamera(bool bLocationOnl
 	ResetFOVs();
 }
 
-void FDisplayClusterLightCardEditorViewportClient::MoveLightCardTo(ADisplayClusterLightCardActor& LightCard, const FSphericalCoordinates& SphericalCoords) const
+void FDisplayClusterLightCardEditorViewportClient::MoveLightCardTo(ADisplayClusterLightCardActor& LightCard, const FDisplayClusterLightCardEditorHelper::FSphericalCoordinates& SphericalCoords) const
 {
 	const FVector Origin = GetViewLocation(); // Assumed the same as the LC origin in proxy space. Call VerifyAndFixLightCardOrigin to ensure this.
-	const FVector LightCardPosition = Origin + SphericalCoords.AsCartesian();
 
-	FVector DesiredNormal;
-	float DesiredDistance;
-
-	// If the light card is in the southern hemisphere of the view origin, use the southern normal map; otherwise, use the north normal map
-	if (LightCardPosition.Z < Origin.Z)
-	{
-		SouthNormalMap.GetNormalAndDistanceAtPosition(LightCardPosition, DesiredNormal, DesiredDistance);
-	}
-	else
-	{
-		NorthNormalMap.GetNormalAndDistanceAtPosition(LightCardPosition, DesiredNormal, DesiredDistance);
-	}
-
-	SetLightCardCoordinates(&LightCard, SphericalCoords);
-	LightCard.DistanceFromCenter = CalculateFinalLightCardDistance(DesiredDistance);
-
-	const FRotator Rotation = FRotationMatrix::MakeFromX(-DesiredNormal).Rotator();
-
-	LightCard.Pitch = Rotation.Pitch;
-	LightCard.Yaw = Rotation.Yaw;
+	ProjectionHelper->MoveLightCardsTo({ &LightCard }, SphericalCoords);
 }
 
 void FDisplayClusterLightCardEditorViewportClient::CenterLightCardInView(ADisplayClusterLightCardActor& LightCard)
 {
-	VerifyAndFixLightCardOrigin(&LightCard);
+	ProjectionHelper->VerifyAndFixLightCardOrigin(LightCard);
 
 	if (LightCard.bIsUVLightCard)
 	{
@@ -1906,7 +1494,7 @@ void FDisplayClusterLightCardEditorViewportClient::CenterLightCardInView(ADispla
 	}
 	else
 	{
-		MoveLightCardTo(LightCard, FSphericalCoordinates(GetViewRotation().RotateVector(FVector::ForwardVector)));
+		MoveLightCardTo(LightCard, FDisplayClusterLightCardEditorHelper::FSphericalCoordinates(GetViewRotation().RotateVector(FVector::ForwardVector)));
 	}
 
 	// If this is a proxy light, propagate to its counterpart in the level.
@@ -1941,20 +1529,45 @@ void FDisplayClusterLightCardEditorViewportClient::GetScenePrimitiveComponents(T
 
 void FDisplayClusterLightCardEditorViewportClient::GetSceneViewInitOptions(FSceneViewInitOptions& OutViewInitOptions)
 {
+	const FViewportCameraTransform& ViewTransform = GetViewTransform();
+	const FRotator& ViewRotation = ViewTransform.GetRotation();
+	const FMatrix RotationMatrix = CalcViewRotationMatrix(ViewRotation);
+
 	FSceneViewInitOptions ViewInitOptions;
+	ProjectionHelper->GetSceneViewInitOptions(
+		ViewInitOptions,
+		GetProjectionModeFOV(ProjectionMode),
+		Viewport->GetSizeXY(),
+		ViewTransform.GetLocation(),
+		ViewTransform.GetRotation(),
+		GetDefault<ULevelEditorViewportSettings>()->AspectRatioAxisConstraint,
+		GetNearClipPlane(),
+		&RotationMatrix,
+		GetDPIScale()
+	);
 
-	FViewportCameraTransform& ViewTransform = GetViewTransform();
+	ViewInitOptions.SceneViewStateInterface = ViewState.GetReference();
+	ViewInitOptions.ViewElementDrawer = this;
+
+	ViewInitOptions.BackgroundColor = GetBackgroundColor();
 
-	ViewInitOptions.ViewLocation = ProjectionMode == EDisplayClusterMeshProjectionType::UV ? FVector::ZeroVector : ViewTransform.GetLocation();
-	ViewInitOptions.ViewRotation = ProjectionMode == EDisplayClusterMeshProjectionType::UV ? FRotator::ZeroRotator : ViewTransform.GetRotation();
-	ViewInitOptions.ViewOrigin = ViewInitOptions.ViewLocation;
+	ViewInitOptions.EditorViewBitflag = (uint64)1 << ViewIndex, // send the bit for this view - each actor will check it's visibility bits against this
+
+	ViewInitOptions.OverrideFarClippingPlaneDistance = GetFarClipPlaneOverride();
+	ViewInitOptions.CursorPos = CurrentMousePos;
+
+	OutViewInitOptions = ViewInitOptions;
+}
+
+void FDisplayClusterLightCardEditorViewportClient::GetNormalMapSceneViewInitOptions(FIntPoint NormalMapSize, float NormalMapFOV, const FVector& ViewDirection, FSceneViewInitOptions& OutViewInitOptions)
+{
+	FViewportCameraTransform& ViewTransform = GetViewTransform();
 
-	FIntPoint ViewportSize = Viewport->GetSizeXY();
-	ViewportSize.X = FMath::Max(ViewportSize.X, 1);
-	ViewportSize.Y = FMath::Max(ViewportSize.Y, 1);
-	FIntPoint ViewportOffset(0, 0);
+	OutViewInitOptions.ViewLocation = ProjectionOriginComponent.IsValid() ? ProjectionOriginComponent->GetComponentLocation() : FVector::ZeroVector;
+	OutViewInitOptions.ViewRotation = ViewDirection.Rotation();
+	OutViewInitOptions.ViewOrigin = OutViewInitOptions.ViewLocation;
 
-	ViewInitOptions.SetViewRectangle(FIntRect(ViewportOffset, ViewportOffset + ViewportSize));
+	OutViewInitOptions.SetViewRectangle(FIntRect(0, 0, NormalMapSize.X, NormalMapSize.Y));
 
 	AWorldSettings* WorldSettings = nullptr;
 	if (GetScene() != nullptr && GetScene()->GetWorld() != nullptr)
@@ -1964,176 +1577,48 @@ void FDisplayClusterLightCardEditorViewportClient::GetSceneViewInitOptions(FScen
 
 	if (WorldSettings != nullptr)
 	{
-		ViewInitOptions.WorldToMetersScale = WorldSettings->WorldToMeters;
+		OutViewInitOptions.WorldToMetersScale = WorldSettings->WorldToMeters;
 	}
 
 	// Rotate view 90 degrees
-	ViewInitOptions.ViewRotationMatrix = CalcViewRotationMatrix(ViewInitOptions.ViewRotation) * FMatrix(
+	OutViewInitOptions.ViewRotationMatrix = FInverseRotationMatrix(OutViewInitOptions.ViewRotation) * FMatrix(
 		FPlane(0, 0, 1, 0),
 		FPlane(1, 0, 0, 0),
 		FPlane(0, 1, 0, 0),
 		FPlane(0, 0, 0, 1));
 
 	const float MinZ = GetNearClipPlane();
-	const float MaxZ = MinZ;
-	const float FieldOfView = FMath::DegreesToRadians(GetProjectionModeFOV(ProjectionMode));
+	const float MaxZ = FMath::Max(RootActorBoundingRadius, MinZ);
 
 	// Avoid zero ViewFOV's which cause divide by zero's in projection matrix
-	const float HalfFOV = FMath::Max(0.5f * FieldOfView, 0.001f);
-
-	float XAxisMultiplier;
-	float YAxisMultiplier;
+	const float MatrixFOV = FMath::Max(0.001f, NormalMapFOV) * (float)PI / 360.0f;
 
-	EAspectRatioAxisConstraint AspectRatioAxisConstraint = GetDefault<ULevelEditorViewportSettings>()->AspectRatioAxisConstraint;
+	const float XAxisMultiplier = 1.0f;
+	const float YAxisMultiplier = 1.0f;
 
-	if (((ViewportSize.X > ViewportSize.Y) && (AspectRatioAxisConstraint == AspectRatio_MajorAxisFOV)) || (AspectRatioAxisConstraint == AspectRatio_MaintainXFOV))
+	if ((bool)ERHIZBuffer::IsInverted)
 	{
-		//if the viewport is wider than it is tall
-		XAxisMultiplier = 1.0f;
-		YAxisMultiplier = ViewportSize.X / (float)ViewportSize.Y;
+		OutViewInitOptions.ProjectionMatrix = FReversedZPerspectiveMatrix(
+			MatrixFOV,
+			MatrixFOV,
+			XAxisMultiplier,
+			YAxisMultiplier,
+			MinZ,
+			MaxZ);
 	}
 	else
 	{
-		//if the viewport is taller than it is wide
-		XAxisMultiplier = ViewportSize.Y / (float)ViewportSize.X;
-		YAxisMultiplier = 1.0f;
+		OutViewInitOptions.ProjectionMatrix = FPerspectiveMatrix(
+			MatrixFOV,
+			MatrixFOV,
+			XAxisMultiplier,
+			YAxisMultiplier,
+			MinZ,
+			MaxZ);
 	}
 
-	if (RenderViewportType == LVT_Perspective)
-	{
-		if ((bool)ERHIZBuffer::IsInverted)
-		{
-			ViewInitOptions.ProjectionMatrix = FReversedZPerspectiveMatrix(
-				HalfFOV,
-				HalfFOV,
-				XAxisMultiplier,
-				YAxisMultiplier,
-				MinZ,
-				MaxZ
-			);
-		}
-		else
-		{
-			ViewInitOptions.ProjectionMatrix = FPerspectiveMatrix(
-				HalfFOV,
-				HalfFOV,
-				XAxisMultiplier,
-				YAxisMultiplier,
-				MinZ,
-				MaxZ
-			);
-		}
-	}
-	else
-	{
-		const float ZScale = 0.5f / HALF_WORLD_MAX;
-		const float ZOffset = HALF_WORLD_MAX;
-
-		const float FOVScale = FMath::Tan(HalfFOV) / GetDPIScale();
-
-		const float OrthoWidth = 0.5f * FOVScale * ViewportSize.X;
-		const float OrthoHeight = 0.5f * FOVScale * ViewportSize.Y;
-
-		if ((bool)ERHIZBuffer::IsInverted)
-		{
-			ViewInitOptions.ProjectionMatrix = FReversedZOrthoMatrix(
-				OrthoWidth,
-				OrthoHeight,
-				ZScale,
-				ZOffset
-			);
-		}
-		else
-		{
-			ViewInitOptions.ProjectionMatrix = FOrthoMatrix(
-				OrthoWidth,
-				OrthoHeight,
-				ZScale,
-				ZOffset
-			);
-		}
-	}
-
-	if (!ViewInitOptions.IsValidViewRectangle())
-	{
-		// Zero sized rects are invalid, so fake to 1x1 to avoid asserts later on
-		ViewInitOptions.SetViewRectangle(FIntRect(0, 0, 1, 1));
-	}
-
-	ViewInitOptions.SceneViewStateInterface = ViewState.GetReference();
-	ViewInitOptions.ViewElementDrawer = this;
-
-	ViewInitOptions.BackgroundColor = GetBackgroundColor();
-
-	ViewInitOptions.EditorViewBitflag = (uint64)1 << ViewIndex, // send the bit for this view - each actor will check it's visibility bits against this
-
-	ViewInitOptions.FOV = FieldOfView;
-	ViewInitOptions.OverrideFarClippingPlaneDistance = GetFarClipPlaneOverride();
-	ViewInitOptions.CursorPos = CurrentMousePos;
-
-	OutViewInitOptions = ViewInitOptions;
-}
-
-void FDisplayClusterLightCardEditorViewportClient::GetNormalMapSceneViewInitOptions(FIntPoint NormalMapSize, float NormalMapFOV, const FVector& ViewDirection, FSceneViewInitOptions& OutViewInitOptions)
-{
-	FViewportCameraTransform& ViewTransform = GetViewTransform();
-
-	OutViewInitOptions.ViewLocation = ProjectionOriginComponent.IsValid() ? ProjectionOriginComponent->GetComponentLocation() : FVector::ZeroVector;
-	OutViewInitOptions.ViewRotation = ViewDirection.Rotation();
-	OutViewInitOptions.ViewOrigin = OutViewInitOptions.ViewLocation;
-
-	OutViewInitOptions.SetViewRectangle(FIntRect(0, 0, NormalMapSize.X, NormalMapSize.Y));
-
-	AWorldSettings* WorldSettings = nullptr;
-	if (GetScene() != nullptr && GetScene()->GetWorld() != nullptr)
-	{
-		WorldSettings = GetScene()->GetWorld()->GetWorldSettings();
-	}
-
-	if (WorldSettings != nullptr)
-	{
-		OutViewInitOptions.WorldToMetersScale = WorldSettings->WorldToMeters;
-	}
-
-	// Rotate view 90 degrees
-	OutViewInitOptions.ViewRotationMatrix = FInverseRotationMatrix(OutViewInitOptions.ViewRotation) * FMatrix(
-		FPlane(0, 0, 1, 0),
-		FPlane(1, 0, 0, 0),
-		FPlane(0, 1, 0, 0),
-		FPlane(0, 0, 0, 1));
-
-	const float MinZ = GetNearClipPlane();
-	const float MaxZ = FMath::Max(RootActorBoundingRadius, MinZ);
-
-	// Avoid zero ViewFOV's which cause divide by zero's in projection matrix
-	const float MatrixFOV = FMath::Max(0.001f, NormalMapFOV) * (float)PI / 360.0f;
-
-	const float XAxisMultiplier = 1.0f;
-	const float YAxisMultiplier = 1.0f;
-
-	if ((bool)ERHIZBuffer::IsInverted)
-	{
-		OutViewInitOptions.ProjectionMatrix = FReversedZPerspectiveMatrix(
-			MatrixFOV,
-			MatrixFOV,
-			XAxisMultiplier,
-			YAxisMultiplier,
-			MinZ,
-			MaxZ);
-	}
-	else
-	{
-		OutViewInitOptions.ProjectionMatrix = FPerspectiveMatrix(
-			MatrixFOV,
-			MatrixFOV,
-			XAxisMultiplier,
-			YAxisMultiplier,
-			MinZ,
-			MaxZ);
-	}
-
-	OutViewInitOptions.SceneViewStateInterface = ViewState.GetReference();
-	OutViewInitOptions.ViewElementDrawer = this;
+	OutViewInitOptions.SceneViewStateInterface = ViewState.GetReference();
+	OutViewInitOptions.ViewElementDrawer = this;
 
 	OutViewInitOptions.EditorViewBitflag = (uint64)1 << ViewIndex, // send the bit for this view - each actor will check it's visibility bits against this
 
@@ -2296,58 +1781,6 @@ void FDisplayClusterLightCardEditorViewportClient::PropagateLightCardTransform(A
 	}
 }
 
-void FDisplayClusterLightCardEditorViewportClient::SetLightCardCoordinates(
-	ADisplayClusterLightCardActor* LightCard, const FSphericalCoordinates& SphericalCoords) const
-{
-	if (!LightCard)
-	{
-		return;
-	}
-
-	LightCard->DistanceFromCenter = SphericalCoords.Radius;
-	LightCard->Latitude = 90.f - FMath::RadiansToDegrees(SphericalCoords.Inclination);
-
-	// Keep the same longitude when pointing at the pole. This helps with continuity 
-	// and also mitigates sudden changes in apparent spin when moving around the poles
-	if (!SphericalCoords.IsPointingAtPole())
-	{
-		LightCard->Longitude = FRotator::ClampAxis(FMath::RadiansToDegrees(SphericalCoords.Azimuth) - 180);
-	}
-}
-
-void FDisplayClusterLightCardEditorViewportClient::VerifyAndFixLightCardOrigin(ADisplayClusterLightCardActor* LightCard) const
-{
-	// Center lightcard on the current view origin, let it keep its current world placement 
-	// (but not its spin/yaw/pitch since that will be happen later using the cache).
-
-	if (!LightCard)
-	{
-		return;
-	}
-
-	const ADisplayClusterRootActor* RootActor = LightCard->bIsProxy ? RootActorProxy.Get() : RootActorLevelInstance.Get();
-	const USceneComponent* OriginComponent = LightCard->bIsProxy ? ProjectionOriginComponent.Get() : FindProjectionOriginComponent(RootActor);
-
-	// Set location at the view origin
-	const FVector& NewLightCardActorLocation = OriginComponent ? OriginComponent->GetComponentLocation() : FVector::ZeroVector;
-
-	// Set rotation to match the root actor
-	const FRotator& NewLightCardActorRotation = RootActor ? RootActor->GetActorRotation() : FRotator::ZeroRotator;
-
-	const FVector LightCardEndEffectorLocation = LightCard->GetLightCardTransform().GetLocation();
-
-	LightCard->SetActorLocation(NewLightCardActorLocation);
-	LightCard->SetActorRotation(NewLightCardActorRotation);
-
-	// Update the light card spherical coordinates to match its current world coordinates
-
-	const FVector LightCardRelativeLocation = NewLightCardActorRotation.UnrotateVector(LightCardEndEffectorLocation - NewLightCardActorLocation);
-
-	const FSphericalCoordinates SphericalCoords(LightCardRelativeLocation);
-
-	SetLightCardCoordinates(LightCard, SphericalCoords);
-}
-
 void FDisplayClusterLightCardEditorViewportClient::MoveSelectedLightCards(FViewport* InViewport, EAxisList::Type CurrentAxis)
 {
 	FSceneViewFamilyContext ViewFamily(FSceneViewFamily::ConstructionValues(
@@ -2359,116 +1792,33 @@ void FDisplayClusterLightCardEditorViewportClient::MoveSelectedLightCards(FViewp
 
 	FSceneView* View = CalcSceneView(&ViewFamily);
 
-	if (LastSelectedLightCard.IsValid() && !LastSelectedLightCard->bIsUVLightCard)
-	{
-		const bool bUseDeltaRotation = (CurrentAxis == EAxisList::Type::XYZ) || (CurrentAxis == EAxisList::Type::Y);
-
-		const FRotator DeltaRotation = 
-			bUseDeltaRotation ?
-			GetLightCardRotationDelta(InViewport, LastSelectedLightCard.Get(), CurrentAxis) 
-			: FRotator::ZeroRotator;
-
-		const FSphericalCoordinates DeltaCoords = 
-			bUseDeltaRotation ?
-			FSphericalCoordinates()
-			: GetLightCardTranslationDelta(InViewport, LastSelectedLightCard.Get(), CurrentAxis);
-
-		for (const TWeakObjectPtr<ADisplayClusterLightCardActor>& LightCard : SelectedLightCards)
-		{
-			if (!LightCard.IsValid() || LightCard->bIsUVLightCard)
-			{
-				continue;
-			}
-
-			VerifyAndFixLightCardOrigin(LightCard.Get());
-
-			// Note: GetLightCardCoordinates maintains last known Azimuth when looking at the poles
-			const FSphericalCoordinates CurrentCoords = GetLightCardCoordinates(LightCard.Get()); 
-
-			// We will adjust the spin (to maintain the apparent spin) when using center of gizmo
-			// or dragging longitudinally (this seems to provide an intuitive behavior)
-			if (CurrentAxis == EAxisList::Type::XYZ || CurrentAxis == EAxisList::Type::Y) // Dragging center of gizmo
-			{
-				// We might need this to put back the LightCard exactly as it was
-				const ADisplayClusterLightCardActor::PositionalParams OriginalPositionalParams = LightCard->GetPositionalParams();
-
-				const FVector CurrentPos = CurrentCoords.AsCartesian();
-				const FVector NewPos = DeltaRotation.RotateVector(CurrentPos);
-
-				// Calculations are only valid if translation is not too small
-				const FSphericalCoordinates NewCoords(NewPos);
-
-				const FTransform Transform_A = LightCard->GetLightCardTransform(false /*bIgnoreSpinYawPitch*/);
-
-				MoveLightCardTo(*LightCard.Get(), NewCoords);
-				LightCard->UpdateLightCardTransform(); // We must call this for GetLightCardTransform to be valid
-
-				const FTransform Transform_B = LightCard->GetLightCardTransform(false /*bIgnoreSpinYawPitch*/);
-
-				// Calculate world delta translation of moving from A to B
-				const FVector WorldDelta = Transform_B.GetLocation() - Transform_A.GetLocation(); // X towards front of stage. Y towards right of stage. Z towards ceiling.
-
-				// Calculate LC "Y" unit vector at A and B. ("X" is LC normal)
-				const FVector Y_A = Transform_A.Rotator().RotateVector(FVector::YAxisVector);
-				const FVector Y_B = Transform_B.Rotator().RotateVector(FVector::YAxisVector);
-
-				// Calculate card normal vector
-				const FVector CardNormal_A = Transform_A.Rotator().RotateVector(FVector::XAxisVector); // When card is on ceiling, expect around (0,0,-1).
-				const FVector CardNormal_B = Transform_B.Rotator().RotateVector(FVector::XAxisVector);
-
-				// Calculate projection of movement onto surface tangent plane at A and B
-				const FVector UnitWorldDeltaPlane_A = FVector::VectorPlaneProject(WorldDelta, CardNormal_A).GetSafeNormal();
-				const FVector UnitWorldDeltaPlane_B = FVector::VectorPlaneProject(WorldDelta, CardNormal_B).GetSafeNormal();
-
-				if (!FMath::IsNearlyZero(UnitWorldDeltaPlane_A.Length()) && !FMath::IsNearlyZero(UnitWorldDeltaPlane_B.Length()))
-				{
-					// Calculate relative spin angle at A, which is the angle between Y_A and UnitWorldDeltaPlane_A
-					const double SpinDotProduct_A = FVector::DotProduct(Y_A, UnitWorldDeltaPlane_A);
-					const FVector SpinCrossProduct_A = FVector::CrossProduct(Y_A, UnitWorldDeltaPlane_A);
-					const int32 SpinSign_A = FVector::DotProduct(CardNormal_A, SpinCrossProduct_A) > 0 ? -1 : 1;
-					const double RelativeSpinAngle_A = FMath::Acos(SpinDotProduct_A) * SpinSign_A; // radians
-
-					// Now we need to find the spin that keeps the same RelativeSpinAngle in B as it was in A
-					const double SpinDotProduct_B = FVector::DotProduct(Y_B, UnitWorldDeltaPlane_B);
-					const FVector SpinCrossProduct_B = FVector::CrossProduct(Y_B, UnitWorldDeltaPlane_B);
-					const int32 SpinSign_B = FVector::DotProduct(CardNormal_B, SpinCrossProduct_B) > 0 ? -1 : 1;
-					const double RelativeSpinAngle_B = FMath::Acos(SpinDotProduct_B) * SpinSign_B; // radians
-
-					const double DeltaSpin = RelativeSpinAngle_B - RelativeSpinAngle_A;
-
-					// Apply delta spin to lightcard
-					LightCard->Spin += FMath::RadiansToDegrees(DeltaSpin);
-				}
-				else
-				{
-					// Leave it where it was to avoid apparent spins even though motion would have been insignificant.
-					LightCard->SetPositionalParams(OriginalPositionalParams);
-				}
-			}
-			else // Dragging latitudinally
-			{
-				MoveLightCardTo(*LightCard.Get(), CurrentCoords + DeltaCoords);
-			}
-
-			PropagateLightCardTransform(LightCard.Get());
-		}
+	FIntPoint MousePos;
+	InViewport->GetMousePos(MousePos);
 
-		FSceneViewInitOptions SceneVewInitOptions;
-		GetSceneViewInitOptions(SceneVewInitOptions);
-		FViewMatrices ViewMatrices(SceneVewInitOptions);
+	// Move the light cards
+	ProjectionHelper->DragLightCards(SelectedLightCards, MousePos, *View, DragWidgetOffset, CurrentAxis, LastSelectedLightCard.Get());
 
-		FVector LightCardProjectedLocation = ProjectWorldPosition(CachedEditorWidgetWorldTransform.GetTranslation(), ViewMatrices);
-		FVector2D ScreenPercentage;
+	// Update the level instances
+	for (const TWeakObjectPtr<ADisplayClusterLightCardActor>& LightCard : SelectedLightCards)
+	{
+		PropagateLightCardTransform(LightCard.Get());
+	}
 
-		if (IsLocationCloseToEdge(LightCardProjectedLocation, InViewport, View, &ScreenPercentage))
-		{
-			DesiredLookAtSpeed = FMath::Max(ScreenPercentage.X, ScreenPercentage.Y) * MaxDesiredLookAtSpeed;
-			DesiredLookAtLocation = LightCardProjectedLocation;
-		}
-		else
-		{
-			DesiredLookAtLocation.Reset();
-		}
+	FSceneViewInitOptions SceneViewInitOptions;
+	GetSceneViewInitOptions(SceneViewInitOptions);
+	FViewMatrices ViewMatrices(SceneViewInitOptions);
+	
+	FVector LightCardProjectedLocation = ProjectWorldPosition(CachedEditorWidgetWorldTransform.GetTranslation(), ViewMatrices);
+	FVector2D ScreenPercentage;
+	
+	if (IsLocationCloseToEdge(LightCardProjectedLocation, InViewport, View, &ScreenPercentage))
+	{
+		DesiredLookAtSpeed = FMath::Max(ScreenPercentage.X, ScreenPercentage.Y) * MaxDesiredLookAtSpeed;
+		DesiredLookAtLocation = LightCardProjectedLocation;
+	}
+	else
+	{
+		DesiredLookAtLocation.Reset();
 	}
 }
 
@@ -2483,203 +1833,20 @@ void FDisplayClusterLightCardEditorViewportClient::MoveLightCardsToPixel(const F
 
 	FSceneView* View = CalcSceneView(&ViewFamily);
 
-	FVector Origin;
-	FVector Direction;
-	PixelToWorld(*View, PixelPos, Origin, Direction);
-
-	if (RenderViewportType != LVT_Perspective)
-	{
-		// For orthogonal projections, PixelToWorld does not return the view origin or a direction from the view origin. Use TraceScreenRay
-		// to find a useful direction away from the view origin to use
-
-		const FVector ViewOrigin = View->ViewLocation;
-
-		Direction = TraceScreenRay(Origin, Direction, ViewOrigin);
-		Origin = ViewOrigin;
-	}
-
-	if (ProjectionMode == EDisplayClusterMeshProjectionType::UV)
-	{
-		// Find the average position of all selected light cards. This group average is what is moved to the specified pixel
-		FVector2D AverageUVCoords = FVector2D::ZeroVector;
-		int32 NumLightCards = 0;
-
-		for (const TWeakObjectPtr<ADisplayClusterLightCardActor>& LightCard : InLightCards)
-		{
-			if (LightCard.IsValid() && LightCard->bIsUVLightCard)
-			{
-				AverageUVCoords += LightCard->UVCoordinates;
-				++NumLightCards;
-			}
-		}
-
-		AverageUVCoords /= NumLightCards;
-
-		// Compute the desired coordinates by projecting the specified screen ray onto the UV projection plane
-		const float UVProjectionPlaneSize = ADisplayClusterLightCardActor::UVPlaneDefaultSize;
-		const float UVProjectionPlaneDistance = ADisplayClusterLightCardActor::UVPlaneDefaultDistance;
-
-		const FPlane UVProjectionPlane(FVector::ForwardVector * UVProjectionPlaneDistance, -FVector::ForwardVector);
-		const FVector PlaneIntersection = FMath::RayPlaneIntersection(FVector::ZeroVector, Direction, UVProjectionPlane);
-		const FVector2D DesiredUVCoords = FVector2D(PlaneIntersection.Y / UVProjectionPlaneSize + 0.5f, 0.5f - PlaneIntersection.Z / UVProjectionPlaneSize);
-
-		const FVector2D DeltaUVCoords = DesiredUVCoords - AverageUVCoords;
+	ProjectionHelper->MoveLightCardsToPixel(SelectedLightCards, PixelPos, *View);
 
-		for (const TWeakObjectPtr<ADisplayClusterLightCardActor>& LightCard : InLightCards)
-		{
-			if (LightCard.IsValid() && LightCard->bIsUVLightCard)
-			{
-				LightCard->UVCoordinates += DeltaUVCoords;
-				PropagateLightCardTransform(LightCard.Get());
-			}
-		}
-	}
-	else
+	// Update each light card with the delta coordinates; the flush constraint is applied by MoveLightCardTo, ensuring the light card is always flush to screens
+	for (const TWeakObjectPtr<ADisplayClusterLightCardActor>& LightCard : InLightCards)
 	{
-		// Find the average position of all selected light cards. This group average is what is moved to the specified pixel
-		FSphericalCoordinates AverageCoords = FSphericalCoordinates();
-		int32 NumLightCards = 0;
-
-		for (const TWeakObjectPtr<ADisplayClusterLightCardActor>& LightCard : InLightCards)
+		if (LightCard.IsValid() &&
+			((LightCard->bIsUVLightCard && ProjectionMode == EDisplayClusterMeshProjectionType::UV) ||
+			(!LightCard->bIsUVLightCard && ProjectionMode != EDisplayClusterMeshProjectionType::UV)))
 		{
-			if (LightCard.IsValid() && !LightCard->bIsUVLightCard)
-			{
-				const FSphericalCoordinates LightCardCoords = GetLightCardCoordinates(LightCard.Get());
-
-				AverageCoords = AverageCoords + LightCardCoords;
-				++NumLightCards;
-			}
-		}
-
-		AverageCoords.Radius /= NumLightCards;
-		AverageCoords.Azimuth /= NumLightCards;
-		AverageCoords.Inclination /= NumLightCards;
-		AverageCoords.Conform();
-
-		// Compute desired coordinates (radius doesn't matter here since we will use the flush constraint on the light cards after moving them)
-		const FSphericalCoordinates DesiredCoords(Direction * 100.0f);
-		const FSphericalCoordinates DeltaCoords = DesiredCoords - AverageCoords;
-
-		// Update each light card with the delta coordinates; the flush constraint is applied by MoveLightCardTo, ensuring the light card is always flush to screens
-		for (const TWeakObjectPtr<ADisplayClusterLightCardActor>& LightCard : InLightCards)
-		{
-			if (LightCard.IsValid() && !LightCard->bIsUVLightCard)
-			{
-				const FSphericalCoordinates LightCardCoords = GetLightCardCoordinates(LightCard.Get());
-				const FSphericalCoordinates NewCoords = LightCardCoords + DeltaCoords;
-
-				MoveLightCardTo(*LightCard.Get(), NewCoords);
-				PropagateLightCardTransform(LightCard.Get());
-			}
+			PropagateLightCardTransform(LightCard.Get());
 		}
 	}
 }
 
-FRotator FDisplayClusterLightCardEditorViewportClient::GetLightCardRotationDelta(
-	FViewport* InViewport,
-	ADisplayClusterLightCardActor* LightCard,
-	EAxisList::Type CurrentAxis)
-{
-	check(LightCard);
-
-	const FSphericalCoordinates DeltaCoords = GetLightCardTranslationDelta(InViewport, LightCard, CurrentAxis);
-	const FSphericalCoordinates LightCardCoords = GetLightCardCoordinates(LightCard);
-	const FSphericalCoordinates NewCoords = LightCardCoords + DeltaCoords;
-
-	const FVector LightCardPos = LightCardCoords.AsCartesian();
-	const FVector NewPos = NewCoords.AsCartesian();
-
-	const FVector PosCrossProduct = FVector::CrossProduct(LightCardPos, NewPos);
-
-	if (FMath::IsNearlyZero(PosCrossProduct.Length()))
-	{
-		return FQuat(FVector::ForwardVector, 0).Rotator();
-	}
-
-	const FVector AxisOfRotation = PosCrossProduct.GetSafeNormal();
-	const double Angle = FMath::Acos(FVector::DotProduct(LightCardPos.GetSafeNormal(), NewPos.GetSafeNormal()));
-
-	return FQuat(AxisOfRotation, Angle).Rotator();
-}
-
-FDisplayClusterLightCardEditorViewportClient::FSphericalCoordinates FDisplayClusterLightCardEditorViewportClient::GetLightCardTranslationDelta(
-	FViewport* InViewport, 
-	ADisplayClusterLightCardActor* LightCard,
-	EAxisList::Type CurrentAxis)
-{
-	FIntPoint MousePos;
-	InViewport->GetMousePos(MousePos);
-
-	FSceneViewFamilyContext ViewFamily(FSceneViewFamily::ConstructionValues(
-		InViewport,
-		GetScene(),
-		EngineShowFlags)
-		.SetRealtimeUpdate(IsRealtime()));
-
-	FSceneView* View = CalcSceneView(&ViewFamily);
-
-	FVector Origin;
-	FVector Direction;
-	PixelToWorld(*View, MousePos, Origin, Direction);
-
-	if (RenderViewportType != LVT_Perspective)
-	{
-		// For orthogonal projections, PixelToWorld does not return the view origin or a direction from the view origin. Use TraceScreenRay
-		// to find a useful direction away from the view origin to use
-
-		const FVector ViewOrigin = View->ViewLocation;
-
-		Direction = TraceScreenRay(Origin - DragWidgetOffset, Direction, ViewOrigin);
-		Origin = ViewOrigin;
-	}
-	else
-	{
-		Direction = (Direction - DragWidgetOffset).GetSafeNormal();
-	}
-
-	const FVector LocalDirection = LightCard->GetActorRotation().RotateVector(Direction);
-	const FVector LightCardLocation = LightCard->GetLightCardTransform().GetTranslation() - Origin;
-
-	FVector Normal;
-	float Distance;
-
-	// If the light card is in the southern hemisphere of the view origin, use the southern normal map; otherwise, use the north normal map
-	if (LightCardLocation.Z < 0.0f)
-	{
-		SouthNormalMap.GetNormalAndDistanceAtPosition(LightCard->GetLightCardTransform().GetTranslation(), Normal, Distance);
-	}
-	else
-	{
-		NorthNormalMap.GetNormalAndDistanceAtPosition(LightCard->GetLightCardTransform().GetTranslation(), Normal, Distance);
-	}
-
-	const FSphericalCoordinates LightCardCoords = GetLightCardCoordinates(LightCard);
-	const FSphericalCoordinates RequestedCoords(LocalDirection * Distance);
-
-	FSphericalCoordinates DeltaCoords = RequestedCoords - LightCardCoords;
-
-	if (CurrentAxis == EAxisList::Type::X)
-	{
-		DeltaCoords.Inclination = 0;
-	}
-	else if (CurrentAxis == EAxisList::Type::Y)
-	{
-		// Convert the inclination to Cartesian coordinates, project it to the x-z plane, and convert back to spherical coordinates. This ensures that the motion in the inclination
-		// plane always lines up with the mouse's projected location along that plane
-		const double FixedInclination = FMath::Abs(FMath::Atan2(
-			FMath::Cos(DeltaCoords.Azimuth) * FMath::Sin(RequestedCoords.Inclination), 
-			FMath::Cos(RequestedCoords.Inclination))
-		);
-
-		// When translating along the inclination axis, the azimuth delta can only be intervals of pi
-		const double FixedAzimuth = FMath::RoundToInt(DeltaCoords.Azimuth / PI) * PI;
-
-		DeltaCoords.Azimuth = FixedAzimuth;
-		DeltaCoords.Inclination = FixedInclination - LightCardCoords.Inclination;
-	}
-
-	return DeltaCoords;
-}
 
 void FDisplayClusterLightCardEditorViewportClient::MoveSelectedUVLightCards(FViewport* InViewport, EAxisList::Type CurrentAxis)
 {
@@ -2714,7 +1881,7 @@ FVector2D FDisplayClusterLightCardEditorViewportClient::GetUVLightCardTranslatio
 
 	FVector Origin;
 	FVector Direction;
-	PixelToWorld(*View, MousePos, Origin, Direction);
+	ProjectionHelper->PixelToWorld(*View, MousePos, Origin, Direction);
 
 	const float UVProjectionPlaneSize = ADisplayClusterLightCardActor::UVPlaneDefaultSize;
 	const float UVProjectionPlaneDistance = ADisplayClusterLightCardActor::UVPlaneDefaultDistance;
@@ -2867,7 +2034,7 @@ double FDisplayClusterLightCardEditorViewportClient::GetLightCardSpinDelta(FView
 	FVector2D ScreenXAxis = FVector2D::ZeroVector;
 	FVector2D ScreenYAxis = FVector2D::ZeroVector;
 
-	WorldToPixel(*View, WorldWidgetOrigin, ScreenOrigin);
+	ProjectionHelper->WorldToPixel(*View, WorldWidgetOrigin, ScreenOrigin);
 	WorldToScreenDirection(*View, WorldWidgetOrigin, WorldWidgetXAxis, ScreenXAxis);
 	WorldToScreenDirection(*View, WorldWidgetOrigin, WorldWidgetYAxis, ScreenYAxis);
 
@@ -2880,121 +2047,6 @@ double FDisplayClusterLightCardEditorViewportClient::GetLightCardSpinDelta(FView
 	return FMath::RadiansToDegrees(Theta - LastTheta);
 }
 
-FDisplayClusterLightCardEditorViewportClient::FSphericalCoordinates FDisplayClusterLightCardEditorViewportClient::GetLightCardCoordinates(ADisplayClusterLightCardActor* LightCard) const
-{
-	const FVector LightCardLocation = LightCard->GetLightCardTransform().GetTranslation() - LightCard->GetActorLocation();
-
-	FSphericalCoordinates LightCardCoords(LightCardLocation);
-
-	// If the light card points at any of the poles, the spherical coordinates will have an "undefined" azimuth value. 
-	// For continuity when dragging a light card positioned there, 
-	// we can manually set the azimuthal value to match the light card's configured longitude
-
-	if (LightCardCoords.IsPointingAtPole())
-	{
-		LightCardCoords.Azimuth = FMath::DegreesToRadians(LightCard->Longitude + 180);
-	}
-
-	return LightCardCoords;
-}
-
-bool FDisplayClusterLightCardEditorViewportClient::TraceStage(const FVector& RayStart, const FVector& RayEnd, FVector& OutHitLocation) const
-{
-	UWorld* PreviewWorld = PreviewScene->GetWorld();
-	check(PreviewWorld);
-
-	FCollisionQueryParams TraceParams(SCENE_QUERY_STAT(DisplayClusterStageTrace), true);
-
-	for (const FLightCardProxy& ProxyRef : LightCardProxies)
-	{
-		if (ProxyRef.Proxy.IsValid())
-		{
-			TraceParams.AddIgnoredActor(ProxyRef.Proxy.Get());
-		}
-	}
-
-	TraceParams.AddIgnoredComponent(NormalMapMeshComponent.Get());
-
-	FHitResult HitResult;
-	if (PreviewWorld->LineTraceSingleByObjectType(HitResult, RayStart, RayEnd, FCollisionObjectQueryParams(FCollisionObjectQueryParams::InitType::AllObjects), TraceParams))
-	{
-		// If the actor we hit was the stage root actor, return the hit location.
-		if (RootActorProxy.Get() == HitResult.GetActor())
-		{
-			OutHitLocation = HitResult.Location;
-			return true;
-		}
-	}
-
-	OutHitLocation = FVector::ZeroVector;
-	return false;
-}
-
-FVector FDisplayClusterLightCardEditorViewportClient::TraceScreenRay(const FVector& OrthogonalOrigin, const FVector& OrthogonalDirection, const FVector& ViewOrigin)
-{
-	const FVector RayStart = OrthogonalOrigin;
-	const FVector RayEnd = OrthogonalOrigin + OrthogonalDirection * WORLD_MAX;
-
-	FVector Direction = FVector::ZeroVector;
-
-	if (ProjectionMode == EDisplayClusterMeshProjectionType::UV)
-	{
-		// In the UV projection mode, all stage geometry has been projected onto a UV projection plane, so perform a ray trace against that plane to find the 
-		// screen ray direction
-		const float UVProjectionPlaneDistance = ADisplayClusterLightCardActor::UVPlaneDefaultDistance;
-
-		const FPlane UVProjectionPlane(ViewOrigin + FVector::ForwardVector * UVProjectionPlaneDistance, -FVector::ForwardVector);
-		const FVector PlaneIntersection = FMath::RayPlaneIntersection(OrthogonalOrigin, OrthogonalDirection, UVProjectionPlane);
-
-		Direction = (PlaneIntersection - ViewOrigin).GetSafeNormal();
-	}
-	else
-	{
-		// First, trace against the stage actor to see if the screen ray hits it; if so, simply return the direction from the view origin to this hit point
-		FVector HitLocation = FVector::ZeroVector;
-		if (TraceStage(RayStart, RayEnd, HitLocation))
-		{
-			Direction = (HitLocation - ViewOrigin).GetSafeNormal();
-		}
-		else
-		{
-			// If we didn't hit any stage geometry, try to trace against the normal map mesh. Procedural meshes does not appear to handle inward pointing normals correctly,
-			// so we need to reverse the trace start and end locations to get a useful hit result
-
-			UWorld* PreviewWorld = PreviewScene->GetWorld();
-			check(PreviewWorld);
-
-			FCollisionQueryParams TraceParams(SCENE_QUERY_STAT(DisplayClusterStageTrace), true);
-
-			for (const FLightCardProxy& ProxyRef : LightCardProxies)
-			{
-				if (ProxyRef.Proxy.IsValid())
-				{
-					TraceParams.AddIgnoredActor(ProxyRef.Proxy.Get());
-				}
-			}
-
-			TraceParams.AddIgnoredActor(RootActorProxy.Get());
-
-			FHitResult HitResult;
-			if (PreviewWorld->LineTraceSingleByObjectType(HitResult, RayEnd, RayStart, FCollisionObjectQueryParams(FCollisionObjectQueryParams::InitType::AllObjects), TraceParams))
-			{
-				HitLocation = HitResult.Location;
-				Direction = (HitLocation - ViewOrigin).GetSafeNormal();
-			}
-			else
-			{
-				// If the screen ray does not hit the stage or the normal map mesh, then simply use the closest point on the ray to the view origin
-				const FVector ClosestPoint = OrthogonalOrigin + ((ViewOrigin - OrthogonalOrigin) | OrthogonalDirection) * OrthogonalDirection;
-
-				Direction = (ClosestPoint - ViewOrigin).GetSafeNormal();
-			}
-		}
-	}
-
-	return Direction;
-}
-
 ADisplayClusterLightCardActor* FDisplayClusterLightCardEditorViewportClient::TraceScreenForLightCard(const FSceneView& View, int32 HitX, int32 HitY)
 {
 	UWorld* PreviewWorld = PreviewScene->GetWorld();
@@ -3002,15 +2054,13 @@ ADisplayClusterLightCardActor* FDisplayClusterLightCardEditorViewportClient::Tra
 
 	FVector Origin;
 	FVector Direction;
-	PixelToWorld(View, FIntPoint(HitX, HitY), Origin, Direction);
+	ProjectionHelper->PixelToWorld(View, FIntPoint(HitX, HitY), Origin, Direction);
 
 	const FVector CursorRayStart = Origin;
 	const FVector CursorRayEnd = CursorRayStart + Direction * (RenderViewportType == LVT_Perspective ? HALF_WORLD_MAX : WORLD_MAX);
 
 	FCollisionQueryParams TraceParams(SCENE_QUERY_STAT(LightCardTrace), true);
 
-	TraceParams.AddIgnoredComponent(NormalMapMeshComponent.Get());
-
 	bool bHitLightCard = false;
 	FHitResult ScreenHitResult;
 	if (PreviewWorld->LineTraceSingleByObjectType(ScreenHitResult, CursorRayStart, CursorRayEnd, FCollisionObjectQueryParams(FCollisionObjectQueryParams::InitType::AllObjects), TraceParams))
@@ -3072,49 +2122,13 @@ FVector FDisplayClusterLightCardEditorViewportClient::ProjectWorldPosition(const
 	return Transform.ProjectPosition(UnprojectedWorldPosition);
 }
 
-void FDisplayClusterLightCardEditorViewportClient::PixelToWorld(const FSceneView& View, const FIntPoint& PixelPos, FVector& OutOrigin, FVector& OutDirection)
-{
-	const FMatrix& InvProjMatrix = View.ViewMatrices.GetInvProjectionMatrix();
-	const FMatrix& InvViewMatrix = View.ViewMatrices.GetInvViewMatrix();
-
-	const FVector4 ScreenPos = View.PixelToScreen(PixelPos.X, PixelPos.Y, 0);
-	const FVector4 HomogeneousPos = RenderViewportType == LVT_Perspective
-		? FVector4(ScreenPos.X * GNearClippingPlane, ScreenPos.Y * GNearClippingPlane, 0.0f, GNearClippingPlane)
-		: FVector4(ScreenPos.X, ScreenPos.Y, 1.0f, 1.0f);
-
-	const FVector ViewPos = FVector(InvProjMatrix.TransformFVector4(HomogeneousPos));
-	const FVector UnprojectedViewPos = FDisplayClusterMeshProjectionRenderer::UnprojectViewPosition(ViewPos, ProjectionMode);
-
-	if (RenderViewportType == LVT_Perspective)
-	{
-		OutOrigin = View.ViewMatrices.GetViewOrigin();
-		OutDirection = InvViewMatrix.TransformVector(UnprojectedViewPos).GetSafeNormal();
-	}
-	else
-	{
-		OutOrigin = InvViewMatrix.TransformFVector4(UnprojectedViewPos);
-		OutDirection = InvViewMatrix.TransformVector(FVector(0, 0, 1)).GetSafeNormal();
-	}
-}
-
-bool FDisplayClusterLightCardEditorViewportClient::WorldToPixel(const FSceneView& View, const FVector& WorldPos, FVector2D& OutPixelPos) const
-{
-	const FMatrix& ViewMatrix = View.ViewMatrices.GetViewMatrix();
-	const FMatrix& ProjMatrix = View.ViewMatrices.GetProjectionMatrix();
-
-	const FVector ViewPos = ViewMatrix.TransformPosition(WorldPos);
-	const FVector ProjectedViewPos = FDisplayClusterMeshProjectionRenderer::ProjectViewPosition(ViewPos, ProjectionMode);
-	const FVector4 ScreenPos = ProjMatrix.TransformFVector4(FVector4(ProjectedViewPos, 1));
-
-	return View.ScreenToPixel(ScreenPos, OutPixelPos);
-}
-
 bool FDisplayClusterLightCardEditorViewportClient::WorldToScreenDirection(const FSceneView& View, const FVector& WorldPos, const FVector& WorldDirection, FVector2D& OutScreenDir)
 {
 	FVector2D ScreenVectorStart = FVector2D::ZeroVector;
 	FVector2D ScreenVectorEnd = FVector2D::ZeroVector;
 
-	if (WorldToPixel(View, WorldPos, ScreenVectorStart) && WorldToPixel(View, WorldPos + WorldDirection, ScreenVectorEnd))
+	if (ProjectionHelper->WorldToPixel(View, WorldPos, ScreenVectorStart)
+		&& ProjectionHelper->WorldToPixel(View, WorldPos + WorldDirection, ScreenVectorEnd))
 	{
 		OutScreenDir = (ScreenVectorEnd - ScreenVectorStart).GetSafeNormal();
 		return true;
@@ -3128,7 +2142,8 @@ bool FDisplayClusterLightCardEditorViewportClient::WorldToScreenDirection(const
 		const FVector Offset = ViewDirection * (FVector::DotProduct(ViewLocation - WorldPos, ViewDirection) + 100.0f);
 		const FVector AdjustedWorldPos = WorldPos + Offset;
 
-		if (WorldToPixel(View, AdjustedWorldPos, ScreenVectorStart) && WorldToPixel(View, AdjustedWorldPos + WorldDirection, ScreenVectorEnd))
+		if (ProjectionHelper->WorldToPixel(View, AdjustedWorldPos, ScreenVectorStart)
+			&& ProjectionHelper->WorldToPixel(View, AdjustedWorldPos + WorldDirection, ScreenVectorEnd))
 		{
 			OutScreenDir = -(ScreenVectorEnd - ScreenVectorStart).GetSafeNormal();
 			return true;
@@ -3189,45 +2204,6 @@ bool FDisplayClusterLightCardEditorViewportClient::CalcEditorWidgetTransform(FTr
 	return true;
 }
 
-void FDisplayClusterLightCardEditorViewportClient::RenderNormalMap(FNormalMap& NormalMap, const FVector& NormalMapDirection)
-{
-	// Only render primitive components from the stage actor for the normal map
-	FDisplayClusterMeshProjectionPrimitiveFilter PrimitiveFilter;
-	PrimitiveFilter.ShouldRenderPrimitiveDelegate = FDisplayClusterMeshProjectionPrimitiveFilter::FPrimitiveFilter::CreateLambda([this](const UPrimitiveComponent* PrimitiveComponent)
-	{
-		return PrimitiveComponent->GetOwner() == RootActorProxy;
-	});
-
-	FDisplayClusterMeshProjectionRenderSettings RenderSettings;
-	RenderSettings.RenderType = EDisplayClusterMeshProjectionOutput::Normals;
-	RenderSettings.EngineShowFlags = EngineShowFlags;
-	RenderSettings.ProjectionType = EDisplayClusterMeshProjectionType::Azimuthal;
-	RenderSettings.PrimitiveFilter = PrimitiveFilter;
-	
-	FSceneViewInitOptions ViewInitOptions;
-	GetNormalMapSceneViewInitOptions(FIntPoint(FNormalMap::NormalMapSize), FNormalMap::NormalMapFOV, NormalMapDirection, RenderSettings.ViewInitOptions);
-
-	NormalMap.Init(RenderSettings.ViewInitOptions);
-
-	FCanvas Canvas(&NormalMap, nullptr, GetWorld(), GetScene()->GetFeatureLevel(), FCanvas::CDM_DeferDrawing, 1.0f);
-	{
-		Canvas.Clear(FLinearColor::Black);
-
-		IDisplayClusterScenePreview::Get().Render(PreviewRendererId, RenderSettings, Canvas);
-	}
-	Canvas.Flush_GameThread();
-
-	NormalMap.ReadFloat16Pixels(NormalMap.GetCachedNormalData());
-	NormalMap.Release();
-
-	FlushRenderingCommands();
-}
-
-void FDisplayClusterLightCardEditorViewportClient::InvalidateNormalMap()
-{
-	bNormalMapInvalid = true;
-}
-
 bool FDisplayClusterLightCardEditorViewportClient::IsLocationCloseToEdge(const FVector& InPosition, const FViewport* InViewport,
                                                                          const FSceneView* InView, FVector2D* OutPercentageToEdge)
 {
diff --git a/Engine/Plugins/Runtime/nDisplay/Source/DisplayClusterLightCardEditor/Private/Viewport/DisplayClusterLightCardEditorViewportClient.h b/Engine/Plugins/Runtime/nDisplay/Source/DisplayClusterLightCardEditor/Private/Viewport/DisplayClusterLightCardEditorViewportClient.h
--- a/Engine/Plugins/Runtime/nDisplay/Source/DisplayClusterLightCardEditor/Private/Viewport/DisplayClusterLightCardEditorViewportClient.h
+++ b/Engine/Plugins/Runtime/nDisplay/Source/DisplayClusterLightCardEditor/Private/Viewport/DisplayClusterLightCardEditorViewportClient.h
@@ -11,6 +11,7 @@
 #include "DisplayClusterMeshProjectionRenderer.h"
 #include "DisplayClusterLightCardActor.h"
 #include "DisplayClusterLightCardEditorWidget.h"
+#include "DisplayClusterLightCardEditorHelper.h"
 
 class ADisplayClusterRootActor;
 class SDisplayClusterLightCardEditor;
@@ -41,89 +42,6 @@ public:
 private:
 	DECLARE_MULTICAST_DELEGATE(FOnNextSceneRefresh);
 
-	struct FSphericalCoordinates
-	{
-	public:
-
-		/** Constructors */
-		FSphericalCoordinates(const FVector& CartesianPosition);
-		FSphericalCoordinates();
-
-		/** Return equivalent cartesian coordinates */
-		FVector AsCartesian() const;
-
-		/** Addition operator */
-		FSphericalCoordinates operator+(FSphericalCoordinates const& Other) const;
-
-		/** Subtraction operator */
-		FSphericalCoordinates operator-(FSphericalCoordinates const& Other) const;
-
-		/** Conform parameters to their normal ranges */
-		void Conform();
-		
-		/** Returns a conformed version of this struct without changing the current one */
-		FSphericalCoordinates GetConformed() const;
-
-		/** Returns true if the inclination is pointing at north or south poles, within the given margin (in radians) */
-		bool IsPointingAtPole(double Margin = 1e-6) const;
-
-		double Radius = 0;      // unitless   0+      (when conforming)
-		double Inclination = 0; // radians    0 to PI (when conforming)
-		double Azimuth = 0;     // radians  -PI to PI (when conforming)
-	};
-
-	/** Custom render target that stores the normal data for the stage */
-	class FNormalMap : public FRenderTarget
-	{
-	public:
-		/** The size of the normal map */
-		static const int32 NormalMapSize;
-
-		/** The field of view to render the normal map with. Using the azimuthal projection,
-		 * this is set so that the entire 360 degree scene is rendered */
-		static const float NormalMapFOV;
-
-		/** Initializes the normal map render target using the specified scene view options */
-		void Init(const FSceneViewInitOptions& InSceneViewInitOptions);
-
-		/** Releases the normal map render target's resources */
-		void Release();
-
-		/** Gets the size of the render target */
-		virtual FIntPoint GetSizeXY() const override { return FIntPoint(SizeX, SizeY); }
-
-		/** Gets a reference to the normal map data array, which stores the normal vector in the RGB components (color = 0.5 * Normal + 0.5) and the depth in the A component */
-		TArray<FFloat16Color>& GetCachedNormalData() { return CachedNormalData; }
-
-		/** Gets the normal vector and distance at the specified world location. The normal and distance are bilinearly interpolated from the nearest pixels in the normal map */
-		bool GetNormalAndDistanceAtPosition(FVector Position, FVector& OutNormal, float& OutDistance) const;
-
-		/** Morphs the vertices of the specified prodedural mesh to match the normal map */
-		void MorphProceduralMesh(UProceduralMeshComponent* InProceduralMeshComponent) const;
-
-		/** Generates a texture object that can be used to visualize the normal map */
-		UTexture2D* GenerateNormalMapTexture(const FString& TextureName);
-
-		/** Gets the normal map visualization texture, or null if it hasn't been generated */
-		UTexture2D* GetNormalMapTexture() const { return NormalMapTexture.IsValid() ? NormalMapTexture.Get() : nullptr; }
-
-	private:
-		/** The view matrices used when the normal map was last rendered */
-		FViewMatrices ViewMatrices;
-
-		/** The cached normal map data from the last normal map render */
-		TArray<FFloat16Color> CachedNormalData;
-
-		/** A texture that contains the normal map, for visualization purposes */
-		TWeakObjectPtr<UTexture2D> NormalMapTexture;
-
-		/** The width of the normal map. */
-		uint32 SizeX = 0;
-
-		/** The height of the normal map. */
-		uint32 SizeY = 0;
-	};
-
 public:
 	FDisplayClusterLightCardEditorViewportClient(FPreviewScene& InPreviewScene,
 		const TWeakPtr<SDisplayClusterLightCardEditorViewport>& InEditorViewportWidget);
@@ -201,7 +119,7 @@ public:
 	 * @param SphericalCoords specifies desired location of light card in spherical coordinates with respect to view origin.
 	 * 
 	*/
-	void MoveLightCardTo(ADisplayClusterLightCardActor& LightCard, const FSphericalCoordinates& SphericalCoords) const;
+	void MoveLightCardTo(ADisplayClusterLightCardActor& LightCard, const FDisplayClusterLightCardEditorHelper::FSphericalCoordinates& SphericalCoords) const;
 
 	/** Places the given light card in the middle of the current viewport */
 	void CenterLightCardInView(ADisplayClusterLightCardActor& LightCard);
@@ -257,15 +175,6 @@ private:
 
 	/** Moves all given light cards to the specified pixel position */
 	void MoveLightCardsToPixel(const FIntPoint& PixelPos, const TArray<TWeakObjectPtr<ADisplayClusterLightCardActor>>& InLightCards);
-	
-	/** Ensures that the light card root component is at the same location as the projection/view origin */
-	void VerifyAndFixLightCardOrigin(ADisplayClusterLightCardActor* LightCard) const;
-
-	/** Determines the appropriate delta rotation needed to move the specified light card to the mouse's location */
-	FRotator GetLightCardRotationDelta(FViewport* InViewport, ADisplayClusterLightCardActor* LightCard, EAxisList::Type CurrentAxis);
-
-	/** Determines the appropriate delta in spherical coordinates needed to move the specified light card to the mouse's location */
-	FSphericalCoordinates GetLightCardTranslationDelta(FViewport* InViewport, ADisplayClusterLightCardActor* LightCard, EAxisList::Type CurrentAxis);
 
 	/** Moves the currently selected UV light cards */
 	void MoveSelectedUVLightCards(FViewport* InViewport, EAxisList::Type CurrentAxis);
@@ -285,41 +194,17 @@ private:
 	/** Determines the appropriate spin delta needed to rotate the light card */
 	double GetLightCardSpinDelta(FViewport* InViewport, ADisplayClusterLightCardActor* LightCard);
 
-	/** Gets the spherical coordinates of the specified light card */
-	FSphericalCoordinates GetLightCardCoordinates(ADisplayClusterLightCardActor* LightCard) const;
-
-	/** Sets the light card position to the given spherical coordinates */
-	void SetLightCardCoordinates(ADisplayClusterLightCardActor* LightCard, const FSphericalCoordinates& SphericalCoords) const;
-
-	/** Performs a ray trace against the stage's geometry, and returns the hit point */
-	bool TraceStage(const FVector& RayStart, const FVector& RayEnd, FVector& OutHitLocation) const;
-
-	/** Traces the world geometry to find the best direction vector from the view origin to a valid point in space using a screen ray */
-	FVector TraceScreenRay(const FVector& RayOrigin, const FVector& RayDirection, const FVector& ViewOrigin);
-
 	/** Traces to find the light card corresponding to a click on a stage screen */
 	ADisplayClusterLightCardActor* TraceScreenForLightCard(const FSceneView& View, int32 HitX, int32 HitY);
 
 	/** Projects the specified world position to the viewport's current projection space */
 	FVector ProjectWorldPosition(const FVector& UnprojectedWorldPosition, const FViewMatrices& ViewMatrices) const;
 
-	/** Converts a pixel coordinate into a point and direction vector in world space */
-	void PixelToWorld(const FSceneView& View, const FIntPoint& PixelPos, FVector& OutOrigin, FVector& OutDirection);
-
-	/** Converts a world coordinate into a point in screen space, and returns true if the world position is on the screen */
-	bool WorldToPixel(const FSceneView& View, const FVector& WorldPos, FVector2D& OutPixelPos) const;
-
 	/** Converts a direction vector from world space to screen screen space, and returns true of the direction vector is on the screen */
 	bool WorldToScreenDirection(const FSceneView& View, const FVector& WorldPos, const FVector& WorldDirection, FVector2D& OutScreenDir);
 
 	/** Calculates the world transform to render the editor widget with to align it with the selected light card */
 	bool CalcEditorWidgetTransform(FTransform& WidgetTransform);
-	
-	/** Renders the viewport's normal map and stores the texture data to be used later */
-	void RenderNormalMap(FNormalMap& NormalMap, const FVector& NormalMapDirection);
-
-	/** Invalidates the viewport's normal map, forcing it to be rerendered on the next draw call */
-	void InvalidateNormalMap();
 
 	/** Checks if the location is approaching the edge of the view space */
 	bool IsLocationCloseToEdge(const FVector& InPosition, const FViewport* InViewport = nullptr, const FSceneView* InView = nullptr, FVector2D* OutPercentageToEdge = nullptr);
@@ -333,9 +218,6 @@ private:
 	/** Calculates the final distance from the origin of a light card, given its flush distance and a desired offset */
 	double CalculateFinalLightCardDistance(double FlushDistance, double DesiredOffsetFromFlush = 0.) const;
 
-	/** Calculates the relative normal vector and world position in the specified direction from the given view origin */
-	void CalculateNormalAndPositionInDirection(const FVector& InViewOrigin, const FVector& InDirection, FVector& OutWorldLocation, FVector& OutRelativeNormal, double InDesiredDistanceFromFlush = 0.) const;
-
 	/** Callback passed into the mesh projection renderer to filter which primitives are drawn */
 	bool ShouldRenderPrimitive(const UPrimitiveComponent* PrimitiveComponent);
 
@@ -382,6 +264,9 @@ private:
 	
 	/** The index of the scene preview renderer returned from IDisplayClusterScenePreview */
 	int32 PreviewRendererId = -1;
+
+	/** Helper used to convert between screen and world coordinates for lightcard positions. */
+	TSharedPtr<FDisplayClusterLightCardEditorHelper> ProjectionHelper;
 	
 	/** The LC editor widget used to manipulate light cards */
 	TSharedPtr<FDisplayClusterLightCardEditorWidget> EditorWidget;
@@ -422,18 +307,6 @@ private:
 	/** The increment to change the FOV by when using the scroll wheel */
 	float FOVScrollIncrement = 5.0f;
 
-	/** The render target used to render a map of the screens' normals for the northern hemisphere of the view */
-	FNormalMap NorthNormalMap;
-
-	/** The render target used to render a map of the screens' normals for the southern hemisphere of the view */
-	FNormalMap SouthNormalMap;
-
-	/** A morphed ico-sphere mesh component that approximates the normal and depth map */
-	TWeakObjectPtr<UProceduralMeshComponent> NormalMapMeshComponent;
-
-	/** Indicates if the cached normal map is invalid and needs to be redrawn */
-	bool bNormalMapInvalid = false;
-
 	/** Indicates if the normal map should be displayed to the screen */
 	bool bDisplayNormalMapVisualization = false;
 
diff --git a/Engine/Plugins/Runtime/nDisplay/Source/DisplayClusterLightCardEditorShaders/Private/DisplayClusterMeshProjectionRenderer.cpp b/Engine/Plugins/Runtime/nDisplay/Source/DisplayClusterLightCardEditorShaders/Private/DisplayClusterMeshProjectionRenderer.cpp
--- a/Engine/Plugins/Runtime/nDisplay/Source/DisplayClusterLightCardEditorShaders/Private/DisplayClusterMeshProjectionRenderer.cpp
+++ b/Engine/Plugins/Runtime/nDisplay/Source/DisplayClusterLightCardEditorShaders/Private/DisplayClusterMeshProjectionRenderer.cpp
@@ -32,6 +32,7 @@ class FMeshProjectionShaderElementData : public FMeshMaterialShaderElementData
 {
 public:
 	FDisplayClusterMeshProjectionTypeSettings ProjectionTypeSettings;
+	FMatrix44f NormalCorrectionMatrix;
 };
 
 /** Base class for all mesh pass processors used by the projection renderer.
@@ -50,6 +51,7 @@ public:
 		const FDisplayClusterMeshProjectionRenderSettings& InRenderSettings)
 		: FMeshPassProcessor(InScene, GMaxRHIFeatureLevel, InView, InDrawListContext)
 		, ProjectionTypeSettings(InRenderSettings.ProjectionTypeSettings)
+		, NormalCorrectionMatrix(InRenderSettings.NormalCorrectionMatrix)
 		, StencilValue(0)
 	{
 		DrawRenderState.SetDepthStencilState(TStaticDepthStencilState<>::GetRHI());
@@ -139,6 +141,7 @@ protected:
 		ShaderElementDataType ShaderElementData;
 		ShaderElementData.InitializeMeshMaterialData(ViewIfDynamicMeshCommand, PrimitiveSceneProxy, MeshBatch, StaticMeshId, false);
 		ShaderElementData.ProjectionTypeSettings = ProjectionTypeSettings;
+		ShaderElementData.NormalCorrectionMatrix = NormalCorrectionMatrix;
 
 		return ShaderElementData;
 	}
@@ -155,6 +158,7 @@ protected:
 protected:
 	FMeshPassProcessorRenderState DrawRenderState;
 	FDisplayClusterMeshProjectionTypeSettings ProjectionTypeSettings;
+	FMatrix44f NormalCorrectionMatrix;
 	uint32 StencilValue;
 };
 
@@ -227,6 +231,7 @@ public:
 		: FMeshMaterialShader(Initializer)
 	{
 		PassUniformBuffer.Bind(Initializer.ParameterMap, FSceneTextureUniformParameters::StaticStructMetadata.GetShaderVariableName());
+		NormalCorrectionMatrix.Bind(Initializer.ParameterMap, TEXT("NormalCorrectionMatrix"), SPF_Optional);
 	}
 
 	static bool ShouldCompilePermutation(const FMeshMaterialShaderPermutationParameters& Parameters)
@@ -235,6 +240,24 @@ public:
 			&& IsFeatureLevelSupported(Parameters.Platform, ERHIFeatureLevel::SM5)
 			&& Parameters.VertexFactoryType == FindVertexFactoryType(TEXT("FLocalVertexFactory"));
 	}
+
+	void GetShaderBindings(
+		const FScene* Scene,
+		ERHIFeatureLevel::Type FeatureLevel,
+		const FPrimitiveSceneProxy* PrimitiveSceneProxy,
+		const FMaterialRenderProxy& MaterialRenderProxy,
+		const FMaterial& Material,
+		const FMeshPassProcessorRenderState& DrawRenderState,
+		const FMeshProjectionShaderElementData& ShaderElementData,
+		FMeshDrawSingleShaderBindings& ShaderBindings) const
+	{
+		FMeshMaterialShader::GetShaderBindings(Scene, FeatureLevel, PrimitiveSceneProxy, MaterialRenderProxy, Material, DrawRenderState, ShaderElementData, ShaderBindings);
+
+		ShaderBindings.Add(NormalCorrectionMatrix, ShaderElementData.NormalCorrectionMatrix);
+	}
+
+private:
+	LAYOUT_FIELD(FShaderParameter, NormalCorrectionMatrix);
 };
 
 using FMeshProjectionColorPS = FMeshProjectionPS<EDisplayClusterMeshProjectionOutput::Color>;
@@ -480,6 +503,7 @@ public:
 		SHADER_PARAMETER_RDG_TEXTURE_UAV(RWTexture2D<float4>, RWColor)
 		SHADER_PARAMETER_RDG_TEXTURE_UAV(RWTexture2D<float>, RWDepth)
 		SHADER_PARAMETER_RDG_TEXTURE_UAV(RWTexture2D<uint>, RWStencil)
+		SHADER_PARAMETER(FMatrix44f, NormalCorrectionMatrix)
 	END_SHADER_PARAMETER_STRUCT()
 
 	static bool ShouldCompilePermutation(const FGlobalShaderPermutationParameters& Parameters)
@@ -939,6 +963,7 @@ void FDisplayClusterMeshProjectionRenderer::Render(FCanvas* Canvas, FSceneInterf
 
 			FRDGTextureRef OutputTexture = GraphBuilder.RegisterExternalTexture(CreateRenderTarget(RenderTarget->GetRenderTargetTexture(), TEXT("ViewRenderTarget")));
 			FRenderTargetBinding OutputRenderTargetBinding(OutputTexture, ERenderTargetLoadAction::ELoad);
+			FMeshProjectionElementCollector ElementCollector(HitProxyConsumer);
 
 			switch (RenderSettings.RenderType)
 			{
@@ -951,6 +976,12 @@ void FDisplayClusterMeshProjectionRenderer::Render(FCanvas* Canvas, FSceneInterf
 				{
 					RenderColorOutput(GraphBuilder, View, RenderSettings, OutputRenderTargetBinding);
 				}
+
+				if (RenderSimpleElementsDelegate.IsBound())
+				{
+					RenderSimpleElementsDelegate.Execute(View, &ElementCollector);
+					AddSimpleElementPass(GraphBuilder, View, OutputRenderTargetBinding, ElementCollector);
+				}
 				break;
 
 			case EDisplayClusterMeshProjectionOutput::Normals:
@@ -958,13 +989,6 @@ void FDisplayClusterMeshProjectionRenderer::Render(FCanvas* Canvas, FSceneInterf
 				break;
 			}
 
-			FMeshProjectionElementCollector ElementCollector(HitProxyConsumer);
-			if (RenderSimpleElementsDelegate.IsBound())
-			{
-				RenderSimpleElementsDelegate.Execute(View, &ElementCollector);
-				AddSimpleElementPass(GraphBuilder, View, OutputRenderTargetBinding, ElementCollector);
-			}
-
 			GraphBuilder.Execute();
 		});
 }
@@ -1103,7 +1127,7 @@ void FDisplayClusterMeshProjectionRenderer::RenderNormalsOutput(FRDGBuilder& Gra
 	FDepthStencilBinding NormalsDepthStencilBinding(NormalsDepthTexture, ERenderTargetLoadAction::EClear, ERenderTargetLoadAction::EClear, FExclusiveDepthStencil::DepthWrite_StencilWrite);
 
 	AddNormalsRenderPass(GraphBuilder, View, RenderSettings, NormalsRenderTargetBinding, NormalsDepthStencilBinding);
-	AddNormalsFilterPass(GraphBuilder, View, OutputRenderTargetBinding, NormalsTexture, NormalsDepthTexture);
+	AddNormalsFilterPass(GraphBuilder, View, OutputRenderTargetBinding, NormalsTexture, NormalsDepthTexture, RenderSettings.NormalCorrectionMatrix);
 }
 
 // Helper macros that generate appropriate switch cases for each projection type the mesh projection renderer supports, allowing new projection types to be easily added
@@ -1300,7 +1324,8 @@ void FDisplayClusterMeshProjectionRenderer::AddNormalsFilterPass(FRDGBuilder& Gr
 	const FViewInfo* View,
 	FRenderTargetBinding& OutputRenderTargetBinding,
 	FRDGTexture* SceneColor,
-	FRDGTexture* SceneDepth)
+	FRDGTexture* SceneDepth,
+	const FMatrix44f& NormalCorrectionMatrix)
 {
 	// The normal map filter pass is made up of several separate passes:
 	// * First, the scene color, depth, and stencil textures must be copied into UAV read-write textures for the compute shaders to operate on. This pass also
@@ -1334,6 +1359,7 @@ void FDisplayClusterMeshProjectionRenderer::AddNormalsFilterPass(FRDGBuilder& Gr
 		PassParameters->RWColor = RWColorUAV;
 		PassParameters->RWDepth = RWDepthUAV;
 		PassParameters->RWStencil = RWStencilUAV;
+		PassParameters->NormalCorrectionMatrix = NormalCorrectionMatrix;
 
 		TShaderMapRef<FMeshProjectionNormalsCreateRWTexturesCS> ComputeShader(View->ShaderMap);
 
diff --git a/Engine/Plugins/Runtime/nDisplay/Source/DisplayClusterLightCardEditorShaders/Public/DisplayClusterMeshProjectionRenderer.h b/Engine/Plugins/Runtime/nDisplay/Source/DisplayClusterLightCardEditorShaders/Public/DisplayClusterMeshProjectionRenderer.h
--- a/Engine/Plugins/Runtime/nDisplay/Source/DisplayClusterLightCardEditorShaders/Public/DisplayClusterMeshProjectionRenderer.h
+++ b/Engine/Plugins/Runtime/nDisplay/Source/DisplayClusterLightCardEditorShaders/Public/DisplayClusterMeshProjectionRenderer.h
@@ -92,6 +92,9 @@ struct DISPLAYCLUSTERLIGHTCARDEDITORSHADERS_API FDisplayClusterMeshProjectionRen
 	/** The output type to use for the renderer. */
 	EDisplayClusterMeshProjectionOutput RenderType = EDisplayClusterMeshProjectionOutput::Color;
 
+	/** A matrix used to rotate the normals in order to account for the root actor's rotation. */
+	FMatrix44f NormalCorrectionMatrix = FMatrix44f::Identity;
+
 	/**
 	 * Optional filter to prevent specific primitives from being rendered.
 	 * This only applies when RenderType is ERenderType::Normals.
@@ -204,7 +207,8 @@ private:
 		const FViewInfo* View,
 		FRenderTargetBinding& OutputRenderTargetBinding,
 		FRDGTexture* SceneColor,
-		FRDGTexture* SceneDepth);
+		FRDGTexture* SceneDepth,
+		const FMatrix44f& NormalCorrectionMatrix);
 
 #if WITH_EDITOR
 	/** Adds a pass to perform the depth render for any selected primitives for the mesh projection */
diff --git a/Engine/Plugins/Runtime/nDisplay/Source/DisplayClusterScenePreview/DisplayClusterScenePreview.Build.cs b/Engine/Plugins/Runtime/nDisplay/Source/DisplayClusterScenePreview/DisplayClusterScenePreview.Build.cs
--- a/Engine/Plugins/Runtime/nDisplay/Source/DisplayClusterScenePreview/DisplayClusterScenePreview.Build.cs
+++ b/Engine/Plugins/Runtime/nDisplay/Source/DisplayClusterScenePreview/DisplayClusterScenePreview.Build.cs
@@ -1,7 +1,6 @@
 // Copyright Epic Games, Inc. All Rights Reserved.
 
 using UnrealBuildTool;
-using System.IO;
 
 public class DisplayClusterScenePreview : ModuleRules
 {
@@ -10,6 +9,7 @@ public class DisplayClusterScenePreview : ModuleRules
 		PublicDependencyModuleNames.AddRange(
 			new string[] {
 				"DisplayClusterLightCardEditorShaders",
+				"RHI"
 			});
 
 		PrivateDependencyModuleNames.AddRange(
@@ -20,7 +20,13 @@ public class DisplayClusterScenePreview : ModuleRules
 				"Core",
 				"CoreUObject",
 				"Engine",
+				"ProceduralMeshComponent",
 				"RenderCore"
 			});
+
+		if (Target.bBuildEditor == true)
+		{
+			PrivateDependencyModuleNames.Add("UnrealEd");
+		}
 	}
 }
diff --git a/Engine/Plugins/Runtime/nDisplay/Source/DisplayClusterScenePreview/Private/DisplayClusterLightCardEditorHelper.cpp b/Engine/Plugins/Runtime/nDisplay/Source/DisplayClusterScenePreview/Private/DisplayClusterLightCardEditorHelper.cpp
new file mode 100644
--- /dev/null
+++ b/Engine/Plugins/Runtime/nDisplay/Source/DisplayClusterScenePreview/Private/DisplayClusterLightCardEditorHelper.cpp
@@ -0,0 +1,1538 @@
+// Copyright Epic Games, Inc. All Rights Reserved.
+
+#include "DisplayClusterLightCardEditorHelper.h"
+
+#include "CanvasTypes.h"
+#include "Containers/ArrayView.h"
+#include "DisplayClusterRootActor.h"
+#include "DisplayClusterLightCardActor.h"
+#include "Engine/Texture2D.h"
+#include "IDisplayClusterScenePreview.h"
+#include "ImageUtils.h"
+#include "KismetProceduralMeshLibrary.h"
+#include "PreviewScene.h"
+#include "ProceduralMeshComponent.h"
+#include "UObject/UObjectGlobals.h"
+
+#if WITH_EDITOR
+#include "EditorViewportClient.h"
+#endif
+
+//////////////////////////////////////////////////////////////////////////
+// FSphericalCoordinates
+
+FDisplayClusterLightCardEditorHelper::FSphericalCoordinates::FSphericalCoordinates(const FVector& CartesianPosition)
+{
+	Radius = CartesianPosition.Size();
+
+	if (Radius > UE_SMALL_NUMBER)
+	{
+		Inclination = FMath::Acos(CartesianPosition.Z / Radius);
+	}
+	else
+	{
+		Inclination = 0;
+	}
+
+	Azimuth = FMath::Atan2(CartesianPosition.Y, CartesianPosition.X);
+}
+
+FDisplayClusterLightCardEditorHelper::FSphericalCoordinates::FSphericalCoordinates()
+	: Radius(0)
+	, Inclination(0)
+	, Azimuth(0)
+{
+
+}
+
+FVector FDisplayClusterLightCardEditorHelper::FSphericalCoordinates::AsCartesian() const
+{
+	const double SinAzimuth = FMath::Sin(Azimuth);
+	const double CosAzimuth = FMath::Cos(Azimuth);
+
+	const double SinInclination = FMath::Sin(Inclination);
+	const double CosInclination = FMath::Cos(Inclination);
+
+	return FVector(
+		Radius * CosAzimuth * SinInclination,
+		Radius * SinAzimuth * SinInclination,
+		Radius * CosInclination
+	);
+}
+
+FDisplayClusterLightCardEditorHelper::FSphericalCoordinates
+FDisplayClusterLightCardEditorHelper::FSphericalCoordinates::operator+(FSphericalCoordinates const& Other) const
+{
+	FSphericalCoordinates Result;
+
+	Result.Radius = Radius + Other.Radius;
+	Result.Inclination = Inclination + Other.Inclination;
+	Result.Azimuth = Azimuth + Other.Azimuth;
+
+	return Result;
+}
+
+FDisplayClusterLightCardEditorHelper::FSphericalCoordinates
+FDisplayClusterLightCardEditorHelper::FSphericalCoordinates::operator-(FSphericalCoordinates const& Other) const
+{
+	FSphericalCoordinates Result;
+
+	Result.Radius = Radius - Other.Radius;
+	Result.Inclination = Inclination - Other.Inclination;
+	Result.Azimuth = Azimuth - Other.Azimuth;
+
+	return Result;
+}
+
+void FDisplayClusterLightCardEditorHelper::FSphericalCoordinates::Conform()
+{
+	if (Radius < 0)
+	{
+		Radius = -Radius;
+		Inclination += PI;
+	}
+
+	if (Inclination < 0 || Inclination > PI)
+	{
+		// -2PI to 2PI
+		Inclination = FMath::Fmod(Inclination, 2 * PI);
+
+		// 0 to 2PI
+		if (Inclination < 0)
+		{
+			Inclination += 2 * PI;
+		}
+
+		// 0 to PI
+		if (Inclination > PI)
+		{
+			Inclination = 2 * PI - Inclination;
+			Azimuth += PI;
+		}
+	}
+
+	if (Azimuth < -PI || Azimuth > PI)
+	{
+		// -2PI to 2PI
+		Azimuth = FMath::Fmod(Azimuth, 2 * PI);
+
+		// -PI to PI
+		if (Azimuth > PI)
+		{
+			Azimuth -= 2 * PI;
+		}
+		else if (Azimuth < -PI)
+		{
+			Azimuth += 2 * PI;
+		}
+	}
+
+	checkSlow(Radius >= 0);
+	checkSlow(Inclination >= 0 && Inclination <= PI);
+	checkSlow(Azimuth >= -PI && Azimuth <= PI);
+}
+
+FDisplayClusterLightCardEditorHelper::FSphericalCoordinates
+FDisplayClusterLightCardEditorHelper::FSphericalCoordinates::GetConformed() const
+{
+	FSphericalCoordinates Result = *this;
+	Result.Conform();
+	return Result;
+}
+
+bool FDisplayClusterLightCardEditorHelper::FSphericalCoordinates::IsPointingAtPole(double Margin) const
+{
+	FSphericalCoordinates CoordsConformed = GetConformed();
+
+	return FMath::IsNearlyZero(CoordsConformed.Inclination, Margin)
+		|| FMath::IsNearlyEqual(CoordsConformed.Inclination, PI, Margin);
+}
+
+//////////////////////////////////////////////////////////////////////////
+// FNormalMap
+
+const int32 FDisplayClusterLightCardEditorHelper::FNormalMap::NormalMapSize = 512;
+const float FDisplayClusterLightCardEditorHelper::FNormalMap::NormalMapFOV = 2.0f * FMath::RadiansToDegrees(FMath::Atan(0.55 * PI)); // Equation for FOV from desired angle from north pole;
+
+void FDisplayClusterLightCardEditorHelper::FNormalMap::Init(const FSceneViewInitOptions& InSceneViewInitOptions)
+{
+	SizeX = InSceneViewInitOptions.GetViewRect().Width();
+	SizeY = InSceneViewInitOptions.GetViewRect().Height();
+
+	ViewMatrices = FViewMatrices(InSceneViewInitOptions);
+
+	if (NormalMapTexture.IsValid())
+	{
+		NormalMapTexture->MarkAsGarbage();
+		NormalMapTexture = nullptr;
+	}
+
+	ENQUEUE_RENDER_COMMAND(InitRHIResourcesCommand)([this](FRHICommandListImmediate& RHICmdList)
+		{
+			const FRHITextureCreateDesc Desc =
+				FRHITextureCreateDesc::Create2D(TEXT("NormalMapTexture"))
+				.SetExtent(SizeX, SizeY)
+				.SetFormat(PF_FloatRGBA)
+				.SetClearValue(FClearValueBinding::Black)
+				.SetFlags(ETextureCreateFlags::RenderTargetable | ETextureCreateFlags::ShaderResource)
+				.SetInitialState(ERHIAccess::SRVMask);
+
+			RenderTargetTextureRHI = RHICreateTexture(Desc);
+		});
+}
+
+void FDisplayClusterLightCardEditorHelper::FNormalMap::Release()
+{
+	ENQUEUE_RENDER_COMMAND(ReleaseRHIResourcesCommand)([this](FRHICommandListImmediate& RHICmdList)
+		{
+			RenderTargetTextureRHI.SafeRelease();
+		});
+}
+
+bool FDisplayClusterLightCardEditorHelper::FNormalMap::GetNormalAndDistanceAtPosition(FVector Position, FVector& OutNormal, float& OutDistance) const
+{
+	auto GetPixel = [this](uint32 InX, uint32 InY)
+	{
+		uint32 ClampedX = FMath::Clamp(InX, (uint32)0, SizeX - 1);
+		uint32 ClampedY = FMath::Clamp(InY, (uint32)0, SizeY - 1);
+
+		return CachedNormalData[ClampedY * SizeX + ClampedX].GetFloats();
+	};
+
+	const FVector ViewPos = FVector(ViewMatrices.GetViewMatrix().TransformFVector4(FVector4(Position, 1)));
+	const FVector ProjectedViewPos = FDisplayClusterMeshProjectionRenderer::ProjectViewPosition(ViewPos, EDisplayClusterMeshProjectionType::Azimuthal);
+
+	const FVector4 ScreenPos = ViewMatrices.GetProjectionMatrix().TransformFVector4(FVector4(ProjectedViewPos, 1));
+
+	if (ScreenPos.W != 0.0)
+	{
+		const float InvW = (ScreenPos.W > 0.0f ? 1.0f : -1.0f) / ScreenPos.W;
+		const float Y = (GProjectionSignY > 0.0f) ? ScreenPos.Y : 1.0f - ScreenPos.Y;
+		const FVector2D PixelPos = FVector2D((0.5f + ScreenPos.X * 0.5f * InvW) * SizeX, (0.5f - Y * 0.5f * InvW) * SizeY);
+
+		// Perform a bilinear interpolation on the computed pixel position to ensure a continuous normal regardless of the resolution of the normal map
+		const uint32 PixelX = FMath::Floor(PixelPos.X - 0.5f);
+		const uint32 PixelY = FMath::Floor(PixelPos.Y - 0.5f);
+		const float PixelXFrac = FMath::Frac(PixelPos.X);
+		const float PixelYFrac = FMath::Frac(PixelPos.X);
+
+		FLinearColor NormalData;
+		NormalData = FMath::Lerp(
+			FMath::Lerp(GetPixel(PixelX, PixelY), GetPixel(PixelX + 1, PixelY), PixelXFrac),
+			FMath::Lerp(GetPixel(PixelX, PixelY + 1), GetPixel(PixelX + 1, PixelY + 1), PixelXFrac),
+			PixelYFrac);
+
+		const FVector NormalVector = 2.f * FVector(NormalData.R, NormalData.G, NormalData.B) - 1.f;
+		OutNormal = NormalVector.GetSafeNormal();
+
+		// Make sure the depth value is not 0, as that will cause a divide by zero when transformed, resulting in an NaN distance being returned
+		const float Depth = FMath::Max(0.001f, NormalData.A);
+
+		FVector4 DepthPos = ViewMatrices.GetInvProjectionMatrix().TransformFVector4(FVector4(ScreenPos.X * InvW, ScreenPos.Y * InvW, Depth, 1.0f));
+		DepthPos /= DepthPos.W;
+
+		const FVector UnprojectedDepthPos = FDisplayClusterMeshProjectionRenderer::UnprojectViewPosition(DepthPos, EDisplayClusterMeshProjectionType::Azimuthal);
+		OutDistance = UnprojectedDepthPos.Length();
+
+		return true;
+	}
+	else
+	{
+		OutNormal = FVector::ZeroVector;
+		OutDistance = 0.0;
+		return false;
+	}
+}
+
+void FDisplayClusterLightCardEditorHelper::FNormalMap::MorphProceduralMesh(UProceduralMeshComponent* InProceduralMeshComponent) const
+{
+	const FVector ViewOrigin = ViewMatrices.GetViewOrigin();
+	const FVector ViewDirection = ViewMatrices.GetViewMatrix().TransformVector(FVector::ZAxisVector);
+	const float MaxAngle = 1.0f / ViewMatrices.GetProjectionMatrix().M[0][0];
+
+	FProcMeshSection& Section = *InProceduralMeshComponent->GetProcMeshSection(0);
+
+	for (int32 Index = 0; Index < Section.ProcVertexBuffer.Num(); ++Index)
+	{
+		FProcMeshVertex& Vertex = Section.ProcVertexBuffer[Index];
+
+		const FVector VertexPosition = Vertex.Position;
+		const FVector VertexWorldPosition = VertexPosition + ViewOrigin;
+		const FVector VertexDirection = Vertex.Position.GetSafeNormal();
+		const float VertexAngle = FMath::Acos(VertexDirection | ViewDirection);
+
+		if (VertexAngle < MaxAngle)
+		{
+			FVector Normal;
+			float Depth;
+			GetNormalAndDistanceAtPosition(VertexWorldPosition, Normal, Depth);
+
+			const FVector NewPosition = VertexDirection * Depth;
+			Vertex.Position = NewPosition;
+
+			const FMatrix RadialBasis = FRotationMatrix::MakeFromX(VertexDirection);
+
+			const FVector WorldNormal = RadialBasis.TransformVector(Normal);
+			Vertex.Normal = WorldNormal;
+		}
+	}
+
+	InProceduralMeshComponent->SetProcMeshSection(0, Section);
+}
+
+UTexture2D* FDisplayClusterLightCardEditorHelper::FNormalMap::GenerateNormalMapTexture(const FString& TextureName)
+{
+	if (NormalMapTexture.IsValid())
+	{
+		NormalMapTexture->MarkAsGarbage();
+		NormalMapTexture = nullptr;
+	}
+
+	if (CachedNormalData.Num())
+	{
+		FCreateTexture2DParameters Params;
+		Params.bDeferCompression = true;
+
+		TArray<FColor> Bitmap;
+		Bitmap.AddZeroed(CachedNormalData.Num());
+
+		for (int32 Index = 0; Index < CachedNormalData.Num(); ++Index)
+		{
+			Bitmap[Index] = CachedNormalData[Index].GetFloats().ToFColor(false);
+		}
+
+		NormalMapTexture = FImageUtils::CreateTexture2D(SizeX, SizeY, Bitmap, GetTransientPackage(), TextureName, RF_Transient, Params);
+	}
+
+	return GetNormalMapTexture();
+}
+
+
+//////////////////////////////////////////////////////////////////////////
+// FDisplayClusterLightCardEditorHelper
+
+FDisplayClusterLightCardEditorHelper::FDisplayClusterLightCardEditorHelper()
+	: FDisplayClusterLightCardEditorHelper(IDisplayClusterScenePreview::Get().CreateRenderer())
+{
+	bCreatedRenderer = true;
+}
+
+FDisplayClusterLightCardEditorHelper::FDisplayClusterLightCardEditorHelper(int32 RendererId)
+	: RendererId(RendererId), bCreatedRenderer(false)
+{
+}
+
+FDisplayClusterLightCardEditorHelper::~FDisplayClusterLightCardEditorHelper()
+{
+	if (bCreatedRenderer)
+	{
+		IDisplayClusterScenePreview::Get().DestroyRenderer(RendererId);
+	}
+
+	FWorldDelegates::OnWorldCleanup.RemoveAll(this);
+
+#if WITH_EDITORONLY_DATA
+	FCoreUObjectDelegates::OnObjectPropertyChanged.RemoveAll(this);
+
+	if (ADisplayClusterRootActor* RootActor = CachedRootActor.Get())
+	{
+		if (UBlueprint* Blueprint = UBlueprint::GetBlueprintFromClass(RootActor->GetClass()))
+		{
+			Blueprint->OnCompiled().RemoveAll(this);
+		}
+	}
+#endif
+}
+
+#if WITH_EDITOR
+void FDisplayClusterLightCardEditorHelper::SetEditorViewportClient(TWeakPtr<FEditorViewportClient> InViewportClient)
+{
+	ViewportClient = InViewportClient;
+}
+#endif
+
+void FDisplayClusterLightCardEditorHelper::SetProjectionMode(EDisplayClusterMeshProjectionType Value)
+{
+	ProjectionMode = Value;
+}
+
+EDisplayClusterMeshProjectionType FDisplayClusterLightCardEditorHelper::GetProjectionMode() const
+{
+	return ProjectionMode;
+}
+
+void FDisplayClusterLightCardEditorHelper::SetIsOrthographic(bool bValue)
+{
+	bIsOrthographic = bValue;
+}
+
+bool FDisplayClusterLightCardEditorHelper::GetIsOrthographic() const
+{
+	return bIsOrthographic;
+}
+
+void FDisplayClusterLightCardEditorHelper::SetRootActor(ADisplayClusterRootActor& NewRootActor)
+{
+	if (!ensureMsgf(bCreatedRenderer, TEXT("SetRootActor can't be called on an FDisplayClusterLightCardEditorHelper that was created with an existing preview renderer")))
+	{
+		return;
+	}
+
+	IDisplayClusterScenePreview::Get().SetRendererRootActor(RendererId, &NewRootActor, false);
+}
+
+void FDisplayClusterLightCardEditorHelper::SetLevelInstanceRootActor(ADisplayClusterRootActor& NewRootActor)
+{
+	LevelInstanceRootActor = &NewRootActor;
+}
+
+const UTexture2D* FDisplayClusterLightCardEditorHelper::GetNormalMapTexture(bool bShowNorthMap)
+{
+	FNormalMap& NormalMap = bShowNorthMap ? NorthNormalMap : SouthNormalMap;
+	UTexture2D* NormalMapTexture = NormalMap.GetNormalMapTexture();
+
+	if (!NormalMapTexture)
+	{
+		const FString TextureName = bShowNorthMap ? TEXT("DisplayClusterLightCardEditorHelper.NorthNormalMap") : TEXT("DisplayClusterLightCardEditorHelper.SouthNormalMap");
+		NormalMapTexture = NormalMap.GenerateNormalMapTexture(TextureName);
+	}
+
+	return NormalMapTexture;
+}
+
+void FDisplayClusterLightCardEditorHelper::MoveLightCardsToPixel(
+	const TArray<TWeakObjectPtr<ADisplayClusterLightCardActor>>& LightCards, const FIntPoint& PixelPos, const FSceneView& SceneView)
+{
+	FVector Origin;
+	FVector Direction;
+	CalculateOriginAndDirectionFromPixelPosition(PixelPos, SceneView, FVector::ZeroVector, Origin, Direction);
+
+	if (ProjectionMode == EDisplayClusterMeshProjectionType::UV)
+	{
+		// Find the average position of all selected light cards. This group average is what is moved to the specified pixel
+		FVector2D AverageUVCoords = FVector2D::ZeroVector;
+		int32 NumLightCards = 0;
+
+		for (const TWeakObjectPtr<ADisplayClusterLightCardActor>& LightCard : LightCards)
+		{
+			if (LightCard.IsValid() && LightCard->bIsUVLightCard)
+			{
+				AverageUVCoords += LightCard->UVCoordinates;
+				++NumLightCards;
+			}
+		}
+
+		AverageUVCoords /= NumLightCards;
+
+		// Compute the desired coordinates by projecting the specified screen ray onto the UV projection plane
+		const float UVProjectionPlaneSize = ADisplayClusterLightCardActor::UVPlaneDefaultSize;
+		const float UVProjectionPlaneDistance = ADisplayClusterLightCardActor::UVPlaneDefaultDistance;
+
+		const FPlane UVProjectionPlane(FVector::ForwardVector * UVProjectionPlaneDistance, -FVector::ForwardVector);
+		const FVector PlaneIntersection = FMath::RayPlaneIntersection(FVector::ZeroVector, Direction, UVProjectionPlane);
+		const FVector2D DesiredUVCoords = FVector2D(PlaneIntersection.Y / UVProjectionPlaneSize + 0.5f, 0.5f - PlaneIntersection.Z / UVProjectionPlaneSize);
+
+		const FVector2D DeltaUVCoords = DesiredUVCoords - AverageUVCoords;
+
+		for (const TWeakObjectPtr<ADisplayClusterLightCardActor>& LightCard : LightCards)
+		{
+			if (LightCard.IsValid() && LightCard->bIsUVLightCard)
+			{
+				LightCard->UVCoordinates += DeltaUVCoords;
+			}
+		}
+	}
+	else
+	{
+		// Find the average position of all selected light cards. This group average is what is moved to the specified pixel
+		FSphericalCoordinates AverageCoords = FSphericalCoordinates();
+		int32 NumLightCards = 0;
+
+		for (const TWeakObjectPtr<ADisplayClusterLightCardActor> LightCard : LightCards)
+		{
+			if (LightCard.IsValid())
+			{
+				const FSphericalCoordinates LightCardCoords = GetLightCardCoordinates(*LightCard);
+
+				AverageCoords = AverageCoords + LightCardCoords;
+				++NumLightCards;
+			}
+		}
+
+		AverageCoords.Radius /= NumLightCards;
+		AverageCoords.Azimuth /= NumLightCards;
+		AverageCoords.Inclination /= NumLightCards;
+		AverageCoords.Conform();
+
+		// Compute desired coordinates (radius doesn't matter here since we will use the flush constraint on the light cards after moving them)
+		const FSphericalCoordinates DesiredCoords(Direction * 100.0f);
+		const FSphericalCoordinates DeltaCoords = DesiredCoords - AverageCoords;
+
+		// Update each light card with the delta coordinates; the flush constraint is applied by MoveLightCardTo, ensuring the light card is always flush to screens
+		for (const TWeakObjectPtr<ADisplayClusterLightCardActor> LightCard : LightCards)
+		{
+			if (LightCard.IsValid() && !LightCard->bIsUVLightCard)
+			{
+				const FSphericalCoordinates LightCardCoords = GetLightCardCoordinates(*LightCard);
+				const FSphericalCoordinates NewCoords = LightCardCoords + DeltaCoords;
+
+				MoveLightCardsTo({ LightCard }, NewCoords);
+			}
+		}
+	}
+}
+
+void FDisplayClusterLightCardEditorHelper::MoveLightCardsTo(
+	const TArray<TWeakObjectPtr<ADisplayClusterLightCardActor>>& LightCards, const FSphericalCoordinates& SphericalCoords)
+{
+	ADisplayClusterRootActor* RootActor = UpdateRootActor();
+	if (!RootActor)
+	{
+		return;
+	}
+
+	if (!UpdateNormalMaps())
+	{
+		return;
+	}
+
+	for (const TWeakObjectPtr<ADisplayClusterLightCardActor> LightCard : LightCards)
+	{
+		if (!LightCard.IsValid())
+		{
+			continue;
+		}
+
+		InternalMoveLightCardTo(*LightCard, SphericalCoords, true);
+	}
+}
+
+void FDisplayClusterLightCardEditorHelper::DragLightCards(const TArray<TWeakObjectPtr<ADisplayClusterLightCardActor>>& LightCards, const FIntPoint& PixelPos, const FSceneView& SceneView,
+	const FVector& DragWidgetOffset, EAxisList::Type DragAxis, ADisplayClusterLightCardActor* PrimaryLightCard)
+{
+	if (LightCards.IsEmpty())
+	{
+		return;
+	}
+
+	ADisplayClusterRootActor* RootActor = UpdateRootActor();
+	if (!RootActor)
+	{
+		return;
+	}
+
+	if (!UpdateNormalMaps())
+	{
+		return;
+	}
+
+	// Maps are already ready, so we can move immediately
+	InternalDragLightCards(LightCards, PixelPos, SceneView, DragWidgetOffset, DragAxis, PrimaryLightCard);
+}
+
+void FDisplayClusterLightCardEditorHelper::VerifyAndFixLightCardOrigin(ADisplayClusterLightCardActor& LightCard)
+{
+	// Center lightcard on the current view origin, let it keep its current world placement 
+	// (but not its spin/yaw/pitch since that will be happen later using the cache).
+
+	ADisplayClusterRootActor* RootActor = UpdateRootActor();
+	const ADisplayClusterRootActor* LightCardRootActor = (LightCard.bIsProxy || !LevelInstanceRootActor.IsValid()) ? RootActor : LevelInstanceRootActor.Get();
+
+	const USceneComponent* OriginComponent = nullptr;
+	if (LightCardRootActor)
+	{
+		OriginComponent = LightCard.bIsProxy ? ProjectionOriginComponent.Get() : LightCardRootActor->GetCommonViewPoint();
+	}
+
+	if (!OriginComponent)
+	{
+		return;
+	}
+
+	// Set location at the view origin
+	const FVector& NewLightCardActorLocation = OriginComponent ? OriginComponent->GetComponentLocation() : FVector::ZeroVector;
+
+	// Set rotation to match the root actor
+	const FRotator& NewLightCardActorRotation = LightCardRootActor ? LightCardRootActor->GetActorRotation() : FRotator::ZeroRotator;
+
+	const FVector LightCardEndEffectorLocation = LightCard.GetLightCardTransform().GetLocation();
+
+	LightCard.SetActorLocation(NewLightCardActorLocation);
+	LightCard.SetActorRotation(NewLightCardActorRotation);
+
+	// Update the light card spherical coordinates to match its current world coordinates
+
+	const FVector LightCardRelativeLocation = NewLightCardActorRotation.UnrotateVector(LightCardEndEffectorLocation - NewLightCardActorLocation);
+
+	const FSphericalCoordinates SphericalCoords(LightCardRelativeLocation);
+
+	SetLightCardCoordinates(LightCard, SphericalCoords);
+
+#if WITH_EDITOR
+	PostEditChangePropertiesForMovedLightCard(LightCard);
+#endif
+}
+
+bool FDisplayClusterLightCardEditorHelper::CalculateNormalAndPositionInDirection(
+	const FVector& InViewOrigin,
+	const FVector& InDirection,
+	FVector& OutWorldPosition,
+	FVector& OutRelativeNormal,
+	double InDesiredDistanceFromFlush)
+{
+	if (!UpdateNormalMaps())
+	{
+		return false;
+	}
+
+	if (ProjectionMode == EDisplayClusterMeshProjectionType::UV)
+	{
+		// In UV projection mode, all relevant geometry is projected onto the UV plane, so compute the world position and normal
+		// by performing a ray-plane intersection on the UV projection plane
+
+		const float UVProjectionPlaneDistance = ADisplayClusterLightCardActor::UVPlaneDefaultDistance;
+
+		const FPlane UVProjectionPlane(InViewOrigin + FVector::ForwardVector * UVProjectionPlaneDistance, -FVector::ForwardVector);
+		const FVector PlaneIntersection = FMath::RayPlaneIntersection(InViewOrigin, InDirection, UVProjectionPlane);
+
+		OutRelativeNormal = UVProjectionPlane.GetNormal();
+		OutWorldPosition = PlaneIntersection;
+	}
+	else
+	{
+		const FVector Position = InViewOrigin + InDirection; // We fabricate a position in the right direction
+
+		float Distance;
+
+		// We find the normal and distance from origin
+		if (Position.Z < InViewOrigin.Z)
+		{
+			SouthNormalMap.GetNormalAndDistanceAtPosition(Position, OutRelativeNormal, Distance);
+		}
+		else
+		{
+			NorthNormalMap.GetNormalAndDistanceAtPosition(Position, OutRelativeNormal, Distance);
+		}
+
+		Distance = CalculateFinalLightCardDistance(Distance, InDesiredDistanceFromFlush);
+
+		// Calculate world position
+		OutWorldPosition = InViewOrigin + Distance * InDirection;
+	}
+
+	return true;
+}
+
+bool FDisplayClusterLightCardEditorHelper::CalculateOriginAndDirectionFromPixelPosition(const FIntPoint& PixelPos, const FSceneView& SceneView, const FVector& OriginOffset,
+	FVector& OutOrigin, FVector& OutDirection)
+{
+	PixelToWorld(SceneView, PixelPos, OutOrigin, OutDirection);
+
+	if (bIsOrthographic)
+	{
+		if (!UpdateNormalMaps())
+		{
+			return false;
+		}
+
+		// For orthogonal projections, PixelToWorld does not return the view origin or a direction from the view origin. Use TraceScreenRay
+		// to find a useful direction away from the view origin to use
+
+		const FVector ViewOrigin = SceneView.ViewLocation;
+
+		OutDirection = TraceScreenRay(OutOrigin + OriginOffset, OutDirection, ViewOrigin);
+		OutOrigin = ViewOrigin;
+	}
+
+	return true;
+}
+
+void FDisplayClusterLightCardEditorHelper::PixelToWorld(const FSceneView& View, const FIntPoint& PixelPos, FVector& OutOrigin, FVector& OutDirection) const
+{
+	const FMatrix& InvProjMatrix = View.ViewMatrices.GetInvProjectionMatrix();
+	FMatrix InvViewMatrix = View.ViewMatrices.GetInvViewMatrix();
+
+	if (USceneComponent* OriginComponent = ProjectionOriginComponent.Get())
+	{
+		// Cancel out the root actor's orientation
+		InvViewMatrix *= OriginComponent->GetComponentTransform().ToInverseMatrixWithScale();
+	}
+
+	const FVector4 ScreenPos = View.PixelToScreen(PixelPos.X, PixelPos.Y, 0);
+	const FVector4 HomogeneousPos = bIsOrthographic
+		? FVector4(ScreenPos.X, ScreenPos.Y, 1.0f, 1.0f)
+		: FVector4(ScreenPos.X * GNearClippingPlane, ScreenPos.Y * GNearClippingPlane, 0.0f, GNearClippingPlane);
+
+	const FVector ViewPos = FVector(InvProjMatrix.TransformFVector4(HomogeneousPos));
+	const FVector UnprojectedViewPos = FDisplayClusterMeshProjectionRenderer::UnprojectViewPosition(ViewPos, ProjectionMode);
+
+	if (bIsOrthographic)
+	{
+		OutOrigin = InvViewMatrix.TransformFVector4(UnprojectedViewPos);
+		OutDirection = InvViewMatrix.TransformVector(FVector(0, 0, 1)).GetSafeNormal();
+	}
+	else
+	{
+		OutOrigin = View.ViewMatrices.GetViewOrigin();
+		OutDirection = InvViewMatrix.TransformVector(UnprojectedViewPos).GetSafeNormal();
+	}
+}
+
+bool FDisplayClusterLightCardEditorHelper::WorldToPixel(const FSceneView& View, const FVector& WorldPos, FVector2D& OutPixelPos) const
+{
+	const FMatrix& ViewMatrix = View.ViewMatrices.GetViewMatrix();
+	const FMatrix& ProjMatrix = View.ViewMatrices.GetProjectionMatrix();
+
+	const FVector ViewPos = ViewMatrix.TransformPosition(WorldPos);
+	const FVector ProjectedViewPos = FDisplayClusterMeshProjectionRenderer::ProjectViewPosition(ViewPos, ProjectionMode);
+	const FVector4 ScreenPos = ProjMatrix.TransformFVector4(FVector4(ProjectedViewPos, 1));
+
+	return View.ScreenToPixel(ScreenPos, OutPixelPos);
+}
+
+void FDisplayClusterLightCardEditorHelper::GetSceneViewInitOptions(
+	FSceneViewInitOptions& OutViewInitOptions,
+	float InFOV,
+	const FIntPoint& InViewportSize,
+	const FVector& InLocation,
+	const FRotator& InRotation,
+	const EAspectRatioAxisConstraint InAspectRatioAxisConstraint,
+	float InNearClipPlane,
+	const FMatrix* InRotationMatrix,
+	float InDPIScale)
+{
+	OutViewInitOptions.ViewLocation = (ProjectionMode == EDisplayClusterMeshProjectionType::UV) ? FVector::ZeroVector : InLocation;
+	OutViewInitOptions.ViewRotation = (ProjectionMode == EDisplayClusterMeshProjectionType::UV) ? FRotator::ZeroRotator : InRotation;
+	OutViewInitOptions.ViewOrigin = OutViewInitOptions.ViewLocation;
+
+	FIntPoint ViewportSize = InViewportSize;
+	ViewportSize.X = FMath::Max(ViewportSize.X, 1);
+	ViewportSize.Y = FMath::Max(ViewportSize.Y, 1);
+	FIntPoint ViewportOffset(0, 0);
+
+	OutViewInitOptions.SetViewRectangle(FIntRect(ViewportOffset, ViewportOffset + ViewportSize));
+
+	if (const ADisplayClusterRootActor* RootActor = UpdateRootActor())
+	{
+		const AWorldSettings* WorldSettings = nullptr;
+		if (RootActor->GetWorld() != nullptr)
+		{
+			WorldSettings = RootActor->GetWorld()->GetWorldSettings();
+		}
+
+		if (WorldSettings != nullptr)
+		{
+			OutViewInitOptions.WorldToMetersScale = WorldSettings->WorldToMeters;
+		}
+	}
+
+	// Rotate view 90 degrees
+	const FMatrix Rotate90(
+		FPlane(0, 0, 1, 0),
+		FPlane(1, 0, 0, 0),
+		FPlane(0, 1, 0, 0),
+		FPlane(0, 0, 0, 1)
+	);
+	OutViewInitOptions.ViewRotationMatrix = (InRotationMatrix ? *InRotationMatrix : FInverseRotationMatrix(OutViewInitOptions.ViewRotation)) * Rotate90;
+
+	// Avoid zero ViewFOV's which cause divide by zero's in projection matrix
+	const float RadianFOV = FMath::DegreesToRadians(InFOV);
+	const float HalfFOV = FMath::Max(0.001f, RadianFOV * 0.5f);
+
+	// Determine FOV multipliers to match render target's aspect ratio
+	float XAxisMultiplier;
+	float YAxisMultiplier;
+
+	if (((ViewportSize.X > ViewportSize.Y) && (InAspectRatioAxisConstraint == AspectRatio_MajorAxisFOV)) || (InAspectRatioAxisConstraint == AspectRatio_MaintainXFOV))
+	{
+		//if the viewport is wider than it is tall
+		XAxisMultiplier = 1.0f;
+		YAxisMultiplier = ViewportSize.X / (float)ViewportSize.Y;
+	}
+	else
+	{
+		//if the viewport is taller than it is wide
+		XAxisMultiplier = ViewportSize.Y / (float)ViewportSize.X;
+		YAxisMultiplier = 1.0f;
+	}
+
+	if (bIsOrthographic)
+	{
+		const float ZScale = 0.5f / HALF_WORLD_MAX;
+		const float ZOffset = HALF_WORLD_MAX;
+
+		const float FOVScale = FMath::Tan(HalfFOV) / InDPIScale;
+
+		const float OrthoWidth = 0.5f * FOVScale * ViewportSize.X;
+		const float OrthoHeight = 0.5f * FOVScale * ViewportSize.Y;
+
+		if ((bool)ERHIZBuffer::IsInverted)
+		{
+			OutViewInitOptions.ProjectionMatrix = FReversedZOrthoMatrix(
+				OrthoWidth,
+				OrthoHeight,
+				ZScale,
+				ZOffset
+			);
+		}
+		else
+		{
+			OutViewInitOptions.ProjectionMatrix = FOrthoMatrix(
+				OrthoWidth,
+				OrthoHeight,
+				ZScale,
+				ZOffset
+			);
+		}
+	}
+	else
+	{
+		if ((bool)ERHIZBuffer::IsInverted)
+		{
+			OutViewInitOptions.ProjectionMatrix = FReversedZPerspectiveMatrix(
+				HalfFOV,
+				HalfFOV,
+				XAxisMultiplier,
+				YAxisMultiplier,
+				InNearClipPlane,
+				InNearClipPlane
+			);
+		}
+		else
+		{
+			OutViewInitOptions.ProjectionMatrix = FPerspectiveMatrix(
+				HalfFOV,
+				HalfFOV,
+				XAxisMultiplier,
+				YAxisMultiplier,
+				InNearClipPlane,
+				InNearClipPlane
+			);
+		}
+	}
+
+	if (!OutViewInitOptions.IsValidViewRectangle())
+	{
+		// Zero sized rects are invalid, so fake to 1x1 to avoid asserts later on
+		OutViewInitOptions.SetViewRectangle(FIntRect(0, 0, 1, 1));
+	}
+
+	OutViewInitOptions.BackgroundColor = FLinearColor::Black;
+	OutViewInitOptions.FOV = InFOV;
+}
+
+FDisplayClusterLightCardEditorHelper::FSphericalCoordinates FDisplayClusterLightCardEditorHelper::GetLightCardCoordinates(const ADisplayClusterLightCardActor& LightCard)
+{
+	const FVector LightCardLocation = LightCard.GetLightCardTransform().GetTranslation() - LightCard.GetActorLocation();
+
+	FSphericalCoordinates LightCardCoords(LightCardLocation);
+
+	// If the light card points at any of the poles, the spherical coordinates will have an "undefined" azimuth value. 
+	// For continuity when dragging a light card positioned there, 
+	// we can manually set the azimuthal value to match the light card's configured longitude
+
+	if (LightCardCoords.IsPointingAtPole())
+	{
+		LightCardCoords.Azimuth = FMath::DegreesToRadians(LightCard.Longitude + 180);
+	}
+
+	return LightCardCoords;
+}
+
+ADisplayClusterRootActor* FDisplayClusterLightCardEditorHelper::UpdateRootActor()
+{
+	ADisplayClusterRootActor* NewRootActor = IDisplayClusterScenePreview::Get().GetRendererRootActor(RendererId);
+	if (NewRootActor == CachedRootActor)
+	{
+		return NewRootActor;
+	}
+
+#if WITH_EDITOR
+	// Clean up old subscribed events
+	if (ADisplayClusterRootActor* RootActor = CachedRootActor.Get())
+	{
+		if (UBlueprint* Blueprint = UBlueprint::GetBlueprintFromClass(RootActor->GetClass()))
+		{
+			Blueprint->OnCompiled().RemoveAll(this);
+		}
+	}
+
+	FCoreUObjectDelegates::OnObjectPropertyChanged.RemoveAll(this);
+#endif
+
+	CachedRootActor = NewRootActor;
+	InvalidateNormalMap();
+
+	if (NewRootActor)
+	{
+		const FBox BoundingBox = NewRootActor->GetComponentsBoundingBox();
+		RootActorBoundingRadius = FMath::Max(BoundingBox.Min.Length(), BoundingBox.Max.Length());
+
+#if WITH_EDITOR
+		FCoreUObjectDelegates::OnObjectPropertyChanged.AddRaw(this, &FDisplayClusterLightCardEditorHelper::OnActorPropertyChanged);
+
+		if (UBlueprint* Blueprint = UBlueprint::GetBlueprintFromClass(NewRootActor->GetClass()))
+		{
+			Blueprint->OnCompiled().AddRaw(this, &FDisplayClusterLightCardEditorHelper::OnRootActorBlueprintCompiled);
+		}
+#endif
+	}
+	else
+	{
+		RootActorBoundingRadius = 0.f;
+	}
+
+	return NewRootActor;
+}
+
+USceneComponent* FDisplayClusterLightCardEditorHelper::UpdateProjectionOriginComponent()
+{
+	if (ADisplayClusterRootActor* RootActor = UpdateRootActor())
+	{
+		ProjectionOriginComponent = RootActor->GetCommonViewPoint();
+	}
+	else
+	{
+		ProjectionOriginComponent = nullptr;
+	}
+
+	return ProjectionOriginComponent.Get();
+}
+
+FVector FDisplayClusterLightCardEditorHelper::GetProjectionOrigin() const
+{
+	if (ProjectionOriginComponent.IsValid())
+	{
+		return ProjectionOriginComponent->GetComponentTransform().GetLocation();
+	}
+
+	if (CachedRootActor.IsValid())
+	{
+		return CachedRootActor->GetTransform().GetLocation();
+	}
+
+	return FVector::Zero();
+}
+
+FRotator FDisplayClusterLightCardEditorHelper::GetLightCardRotationDelta(const FIntPoint& PixelPos, const FSceneView& View, ADisplayClusterLightCardActor& LightCard,
+	EAxisList::Type DragAxis, const FVector& DragWidgetOffset)
+{
+	const FSphericalCoordinates DeltaCoords = GetLightCardTranslationDelta(PixelPos, View, LightCard, DragAxis, DragWidgetOffset);
+	const FSphericalCoordinates LightCardCoords = GetLightCardCoordinates(LightCard);
+	const FSphericalCoordinates NewCoords = LightCardCoords + DeltaCoords;
+
+	const FVector LightCardPos = LightCardCoords.AsCartesian();
+	const FVector NewPos = NewCoords.AsCartesian();
+
+	const FVector PosCrossProduct = FVector::CrossProduct(LightCardPos, NewPos);
+
+	if (FMath::IsNearlyZero(PosCrossProduct.Length()))
+	{
+		return FQuat(FVector::ForwardVector, 0).Rotator();
+	}
+
+	const FVector AxisOfRotation = PosCrossProduct.GetSafeNormal();
+	const double Angle = FMath::Acos(FVector::DotProduct(LightCardPos.GetSafeNormal(), NewPos.GetSafeNormal()));
+
+	return FQuat(AxisOfRotation, Angle).Rotator();
+}
+
+FDisplayClusterLightCardEditorHelper::FSphericalCoordinates FDisplayClusterLightCardEditorHelper::GetLightCardTranslationDelta(
+	const FIntPoint& PixelPos,
+	const FSceneView& View,
+	ADisplayClusterLightCardActor& LightCard,
+	EAxisList::Type DragAxis,
+	const FVector& DragWidgetOffset)
+{
+	FVector Origin;
+	FVector Direction;
+	CalculateOriginAndDirectionFromPixelPosition(PixelPos, View, -DragWidgetOffset, Origin, Direction);
+
+	if (!bIsOrthographic)
+	{
+		Direction = (Direction - DragWidgetOffset).GetSafeNormal();
+	}
+
+	const FVector LocalDirection = LightCard.GetActorRotation().RotateVector(Direction);
+	const FVector LightCardLocation = LightCard.GetLightCardTransform().GetTranslation() - Origin;
+
+	FVector Normal;
+	float Distance;
+
+	// If the light card is in the southern hemisphere of the view origin, use the southern normal map; otherwise, use the north normal map
+	if (LightCardLocation.Z < 0.0f)
+	{
+		SouthNormalMap.GetNormalAndDistanceAtPosition(LightCard.GetLightCardTransform().GetTranslation(), Normal, Distance);
+	}
+	else
+	{
+		NorthNormalMap.GetNormalAndDistanceAtPosition(LightCard.GetLightCardTransform().GetTranslation(), Normal, Distance);
+	}
+
+	const FSphericalCoordinates LightCardCoords = GetLightCardCoordinates(LightCard);
+	const FSphericalCoordinates RequestedCoords(LocalDirection * Distance);
+
+	FSphericalCoordinates DeltaCoords = RequestedCoords - LightCardCoords;
+
+	if (DragAxis == EAxisList::Type::X)
+	{
+		DeltaCoords.Inclination = 0;
+	}
+	else if (DragAxis == EAxisList::Type::Y)
+	{
+		// Convert the inclination to Cartesian coordinates, project it to the x-z plane, and convert back to spherical coordinates. This ensures that the motion in the inclination
+		// plane always lines up with the mouse's projected location along that plane
+		const double FixedInclination = FMath::Abs(FMath::Atan2(
+			FMath::Cos(DeltaCoords.Azimuth) * FMath::Sin(RequestedCoords.Inclination),
+			FMath::Cos(RequestedCoords.Inclination))
+		);
+
+		// When translating along the inclination axis, the azimuth delta can only be intervals of pi
+		const double FixedAzimuth = FMath::RoundToInt(DeltaCoords.Azimuth / PI) * PI;
+
+		DeltaCoords.Azimuth = FixedAzimuth;
+		DeltaCoords.Inclination = FixedInclination - LightCardCoords.Inclination;
+	}
+
+	return DeltaCoords;
+}
+
+void FDisplayClusterLightCardEditorHelper::InternalMoveLightCardTo(
+	ADisplayClusterLightCardActor& LightCard, const FSphericalCoordinates& Position, bool bIsFinalChange) const
+{
+	if (bNormalMapInvalid || !CachedRootActor.IsValid())
+	{
+		ensure(false);
+		return;
+	}
+
+	const FVector Origin = GetProjectionOrigin();
+	const FVector LightCardPosition = Origin + Position.AsCartesian();
+
+	FVector DesiredNormal;
+	float DesiredDistance;
+
+	// If the light card is in the southern hemisphere of the view origin, use the southern normal map; otherwise, use the north normal map
+	if (LightCardPosition.Z < Origin.Z)
+	{
+		SouthNormalMap.GetNormalAndDistanceAtPosition(LightCardPosition, DesiredNormal, DesiredDistance);
+	}
+	else
+	{
+		NorthNormalMap.GetNormalAndDistanceAtPosition(LightCardPosition, DesiredNormal, DesiredDistance);
+	}
+
+	// Remove actor rotation from the position before setting the coordinate. Note that we don't need to do this for the normal map,
+	// which already takes the root actor rotation into account as part of its view matrix.
+	const FQuat RootRotation = CachedRootActor->GetTransform().GetRotation().Inverse();
+	const FVector InverseRotatedPosition = RootRotation.RotateVector(Position.AsCartesian());
+	SetLightCardCoordinates(LightCard, FSphericalCoordinates(InverseRotatedPosition));
+
+	LightCard.DistanceFromCenter = CalculateFinalLightCardDistance(DesiredDistance);
+
+	const FRotator Rotation = FRotationMatrix::MakeFromX(-DesiredNormal).Rotator();
+
+	LightCard.Pitch = Rotation.Pitch;
+	LightCard.Yaw = Rotation.Yaw;
+
+#if WITH_EDITOR
+	if (bIsFinalChange && !LightCard.bIsProxy)
+	{
+		PostEditChangePropertiesForMovedLightCard(LightCard);
+	}
+#endif
+}
+
+void FDisplayClusterLightCardEditorHelper::InternalDragLightCards(const TArray<TWeakObjectPtr<ADisplayClusterLightCardActor>>& LightCards, const FIntPoint& PixelPos, const FSceneView& View,
+	const FVector& DragWidgetOffset, EAxisList::Type DragAxis, ADisplayClusterLightCardActor* PrimaryLightCard)
+{
+	if (LightCards.IsEmpty() || !ensure(!bNormalMapInvalid) || !UpdateRootActor())
+	{
+		return;
+	}
+
+	if (!PrimaryLightCard)
+	{
+		PrimaryLightCard = LightCards.Last().Get();
+	}
+
+	if (PrimaryLightCard && !PrimaryLightCard->bIsUVLightCard)
+	{
+		const bool bUseDeltaRotation = (DragAxis == EAxisList::Type::XYZ) || (DragAxis == EAxisList::Type::Y);
+
+		const FRotator DeltaRotation =
+			bUseDeltaRotation ?
+			GetLightCardRotationDelta(PixelPos, View, *PrimaryLightCard, DragAxis, DragWidgetOffset)
+			: FRotator::ZeroRotator;
+
+		const FSphericalCoordinates DeltaCoords =
+			bUseDeltaRotation ?
+			FSphericalCoordinates()
+			: GetLightCardTranslationDelta(PixelPos, View, *PrimaryLightCard, DragAxis, DragWidgetOffset);
+
+		for (const TWeakObjectPtr<ADisplayClusterLightCardActor> LightCard : LightCards)
+		{
+			if (!LightCard.IsValid() || LightCard->bIsUVLightCard)
+			{
+				continue;
+			}
+
+			VerifyAndFixLightCardOrigin(*LightCard);
+
+			// Note: GetLightCardCoordinates maintains last known Azimuth when looking at the poles
+			const FSphericalCoordinates CurrentCoords = GetLightCardCoordinates(*LightCard.Get());
+
+			// We will adjust the spin (to maintain the apparent spin) when using center of gizmo
+			// or dragging longitudinally (this seems to provide an intuitive behavior)
+			if (DragAxis == EAxisList::Type::XYZ || DragAxis == EAxisList::Type::Y) // Dragging center of gizmo
+			{
+				// We might need this to put back the LightCard exactly as it was
+				const ADisplayClusterLightCardActor::PositionalParams OriginalPositionalParams = LightCard->GetPositionalParams();
+
+				const FVector CurrentPos = CurrentCoords.AsCartesian();
+				const FVector NewPos = DeltaRotation.RotateVector(CurrentPos);
+
+				// Calculations are only valid if translation is not too small
+				const FSphericalCoordinates NewCoords(NewPos);
+
+				const FTransform Transform_A = LightCard->GetLightCardTransform(false /*bIgnoreSpinYawPitch*/);
+
+				// Don't fire off property change events yet since we're not done modifying the lightcard
+				InternalMoveLightCardTo(*LightCard, NewCoords, false);
+				LightCard->UpdateLightCardTransform(); // We must call this for GetLightCardTransform to be valid
+
+				const FTransform Transform_B = LightCard->GetLightCardTransform(false /*bIgnoreSpinYawPitch*/);
+
+				// Calculate world delta translation of moving from A to B
+				FVector WorldDelta = Transform_B.GetLocation() - Transform_A.GetLocation(); // X towards front of stage. Y towards right of stage. Z towards ceiling.
+
+				// Calculate LC "Y" unit vector at A and B. ("X" is LC normal)
+				const FVector Y_A = Transform_A.Rotator().RotateVector(FVector::YAxisVector);
+				const FVector Y_B = Transform_B.Rotator().RotateVector(FVector::YAxisVector);
+
+				// Calculate card normal vector
+				const FVector CardNormal_A = Transform_A.Rotator().RotateVector(FVector::XAxisVector); // When card is on ceiling, expect around (0,0,-1).
+				const FVector CardNormal_B = Transform_B.Rotator().RotateVector(FVector::XAxisVector);
+
+				// Calculate projection of movement onto surface tangent plane at A and B
+				const FVector UnitWorldDeltaPlane_A = FVector::VectorPlaneProject(WorldDelta, CardNormal_A).GetSafeNormal();
+				const FVector UnitWorldDeltaPlane_B = FVector::VectorPlaneProject(WorldDelta, CardNormal_B).GetSafeNormal();
+
+				if (!FMath::IsNearlyZero(UnitWorldDeltaPlane_A.Length()) && !FMath::IsNearlyZero(UnitWorldDeltaPlane_B.Length()))
+				{
+					// Calculate relative spin angle at A, which is the angle between Y_A and UnitWorldDeltaPlane_A
+					const double SpinDotProduct_A = FVector::DotProduct(Y_A, UnitWorldDeltaPlane_A);
+					const FVector SpinCrossProduct_A = FVector::CrossProduct(Y_A, UnitWorldDeltaPlane_A);
+					const int32 SpinSign_A = FVector::DotProduct(CardNormal_A, SpinCrossProduct_A) > 0 ? -1 : 1;
+					const double RelativeSpinAngle_A = FMath::Acos(SpinDotProduct_A) * SpinSign_A; // radians
+
+					// Now we need to find the spin that keeps the same RelativeSpinAngle in B as it was in A
+					const double SpinDotProduct_B = FVector::DotProduct(Y_B, UnitWorldDeltaPlane_B);
+					const FVector SpinCrossProduct_B = FVector::CrossProduct(Y_B, UnitWorldDeltaPlane_B);
+					const int32 SpinSign_B = FVector::DotProduct(CardNormal_B, SpinCrossProduct_B) > 0 ? -1 : 1;
+					const double RelativeSpinAngle_B = FMath::Acos(SpinDotProduct_B) * SpinSign_B; // radians
+
+					const double DeltaSpin = RelativeSpinAngle_B - RelativeSpinAngle_A;
+
+					// Apply delta spin to lightcard
+					LightCard->Spin += FMath::RadiansToDegrees(DeltaSpin);
+				}
+				else
+				{
+					// Leave it where it was to avoid apparent spins even though motion would have been insignificant.
+					LightCard->SetPositionalParams(OriginalPositionalParams);
+				}
+
+#if WITH_EDITOR
+				PostEditChangePropertiesForMovedLightCard(*LightCard);
+#endif
+			}
+			else // Dragging latitudinally
+			{
+				InternalMoveLightCardTo(*LightCard, CurrentCoords + DeltaCoords, true);
+			}
+		}
+	}
+}
+
+void FDisplayClusterLightCardEditorHelper::GetNormalMapSceneViewInitOptions(const FVector& ViewDirection, FSceneViewInitOptions& OutViewInitOptions)
+{
+	OutViewInitOptions.ViewLocation = ProjectionOriginComponent.IsValid() ? ProjectionOriginComponent->GetComponentLocation() : FVector::ZeroVector;
+	OutViewInitOptions.ViewRotation = ViewDirection.Rotation();
+	OutViewInitOptions.ViewOrigin = OutViewInitOptions.ViewLocation;
+
+	if (ADisplayClusterRootActor* RootActor = UpdateRootActor())
+	{
+		OutViewInitOptions.ViewRotation = FRotator(RootActor->GetActorRotation().Quaternion() * OutViewInitOptions.ViewRotation.Quaternion());
+	}
+
+	OutViewInitOptions.SetViewRectangle(FIntRect(0, 0, FNormalMap::NormalMapSize, FNormalMap::NormalMapSize));
+
+	AWorldSettings* WorldSettings = nullptr;
+	ADisplayClusterRootActor* RootActor = UpdateRootActor();
+
+	if (RootActor && RootActor->GetWorld())
+	{
+		WorldSettings = RootActor->GetWorld()->GetWorldSettings();
+	}
+
+	if (WorldSettings != nullptr)
+	{
+		OutViewInitOptions.WorldToMetersScale = WorldSettings->WorldToMeters;
+	}
+
+	// Rotate view 90 degrees
+	OutViewInitOptions.ViewRotationMatrix = FInverseRotationMatrix(OutViewInitOptions.ViewRotation) * FMatrix(
+		FPlane(0, 0, 1, 0),
+		FPlane(1, 0, 0, 0),
+		FPlane(0, 1, 0, 0),
+		FPlane(0, 0, 0, 1));
+
+	const float MinZ = GNearClippingPlane;
+	const float MaxZ = FMath::Max(RootActorBoundingRadius, MinZ);
+
+	// Avoid zero ViewFOV's which cause divide by zero's in projection matrix
+	const float MatrixFOV = FMath::Max(0.001f, FNormalMap::NormalMapFOV) * (float)PI / 360.0f;
+
+	const float XAxisMultiplier = 1.0f;
+	const float YAxisMultiplier = 1.0f;
+
+	if ((bool)ERHIZBuffer::IsInverted)
+	{
+		OutViewInitOptions.ProjectionMatrix = FReversedZPerspectiveMatrix(
+			MatrixFOV,
+			MatrixFOV,
+			XAxisMultiplier,
+			YAxisMultiplier,
+			MinZ,
+			MaxZ);
+	}
+	else
+	{
+		OutViewInitOptions.ProjectionMatrix = FPerspectiveMatrix(
+			MatrixFOV,
+			MatrixFOV,
+			XAxisMultiplier,
+			YAxisMultiplier,
+			MinZ,
+			MaxZ);
+	}
+
+	OutViewInitOptions.FOV = FNormalMap::NormalMapFOV;
+
+#if WITH_EDITOR
+	if (ViewportClient.IsValid())
+	{
+		FEditorViewportClient& ViewportClientRef = *ViewportClient.Pin().Get();
+
+		OutViewInitOptions.SceneViewStateInterface = ViewportClientRef.ViewState.GetReference();
+		OutViewInitOptions.ViewElementDrawer = &ViewportClientRef;
+		OutViewInitOptions.OverrideFarClippingPlaneDistance = ViewportClientRef.GetFarClipPlaneOverride();
+		OutViewInitOptions.EditorViewBitflag = (uint64)1 << ViewportClientRef.ViewIndex; // send the bit for this view - each actor will check it's visibility bits against this
+	}
+#endif
+}
+
+void FDisplayClusterLightCardEditorHelper::SetLightCardCoordinates(ADisplayClusterLightCardActor& LightCard, const FSphericalCoordinates& SphericalCoords) const
+{
+	LightCard.DistanceFromCenter = SphericalCoords.Radius;
+	LightCard.Latitude = 90.f - FMath::RadiansToDegrees(SphericalCoords.Inclination);
+
+	// Keep the same longitude when pointing at the pole. This helps with continuity 
+	// and also mitigates sudden changes in apparent spin when moving around the poles
+	if (!SphericalCoords.IsPointingAtPole())
+	{
+		LightCard.Longitude = FRotator::ClampAxis(FMath::RadiansToDegrees(SphericalCoords.Azimuth) - 180);
+	}
+}
+
+bool FDisplayClusterLightCardEditorHelper::TraceStage(const FVector& RayStart, const FVector& RayEnd, FVector& OutHitLocation)
+{
+	ADisplayClusterRootActor* RootActor = UpdateRootActor();
+	if (RootActor)
+	{
+		FCollisionQueryParams TraceParams(SCENE_QUERY_STAT(DisplayClusterStageTrace), true);
+
+		FHitResult HitResult;
+		if (RootActor->ActorLineTraceSingle(HitResult, RayStart, RayEnd, ECollisionChannel::ECC_PhysicsBody, TraceParams))
+		{
+			OutHitLocation = HitResult.Location;
+			return true;
+		}
+	}
+
+	OutHitLocation = FVector::ZeroVector;
+	return false;
+}
+
+FVector FDisplayClusterLightCardEditorHelper::TraceScreenRay(const FVector& OrthogonalOrigin, const FVector& OrthogonalDirection, const FVector& ViewOrigin)
+{
+	const FVector RayStart = OrthogonalOrigin;
+	const FVector RayEnd = OrthogonalOrigin + OrthogonalDirection * WORLD_MAX;
+
+	FVector Direction = FVector::ZeroVector;
+
+	if (ProjectionMode == EDisplayClusterMeshProjectionType::UV)
+	{
+		// In the UV projection mode, all stage geometry has been projected onto a UV projection plane, so perform a ray trace against that plane to find the 
+		// screen ray direction
+		const float UVProjectionPlaneDistance = ADisplayClusterLightCardActor::UVPlaneDefaultDistance;
+
+		const FPlane UVProjectionPlane(ViewOrigin + FVector::ForwardVector * UVProjectionPlaneDistance, -FVector::ForwardVector);
+		const FVector PlaneIntersection = FMath::RayPlaneIntersection(OrthogonalOrigin, OrthogonalDirection, UVProjectionPlane);
+
+		Direction = (PlaneIntersection - ViewOrigin).GetSafeNormal();
+	}
+	else
+	{
+		// First, trace against the stage actor to see if the screen ray hits it; if so, simply return the direction from the view origin to this hit point
+		FVector HitLocation = FVector::ZeroVector;
+		if (TraceStage(RayStart, RayEnd, HitLocation))
+		{
+			Direction = (HitLocation - ViewOrigin).GetSafeNormal();
+		}
+		else
+		{
+			// If we didn't hit any stage geometry, try to trace against the normal map mesh. Procedural meshes does not appear to handle inward pointing normals correctly,
+			// so we need to reverse the trace start and end locations to get a useful hit result
+
+			FCollisionQueryParams TraceParams(SCENE_QUERY_STAT(DisplayClusterStageTrace), true);
+
+			FHitResult HitResult;
+			if (ensure(!bNormalMapInvalid) && NormalMapMeshComponent->LineTraceComponent(HitResult, RayEnd, RayStart, TraceParams))
+			{
+				HitLocation = HitResult.Location;
+				Direction = (HitLocation - ViewOrigin).GetSafeNormal();
+			}
+			else
+			{
+				// If the screen ray does not hit the stage or the normal map mesh, then simply use the closest point on the ray to the view origin
+				const FVector ClosestPoint = OrthogonalOrigin + ((ViewOrigin - OrthogonalOrigin) | OrthogonalDirection) * OrthogonalDirection;
+
+				Direction = (ClosestPoint - ViewOrigin).GetSafeNormal();
+			}
+		}
+	}
+
+	return Direction;
+}
+
+double FDisplayClusterLightCardEditorHelper::CalculateFinalLightCardDistance(double FlushDistance, double DesiredOffsetFromFlush) const
+{
+	double Distance = FMath::Min(FlushDistance, RootActorBoundingRadius) + DesiredOffsetFromFlush;
+
+	return FMath::Max(Distance, 0);
+}
+
+void FDisplayClusterLightCardEditorHelper::InvalidateNormalMap()
+{
+	bNormalMapInvalid = true;
+}
+
+bool FDisplayClusterLightCardEditorHelper::UpdateNormalMaps()
+{
+	if (!bNormalMapInvalid)
+	{
+		return true;
+	}
+
+	ADisplayClusterRootActor* RootActor = UpdateRootActor();
+	if (!RootActor)
+	{
+		return false;
+	}
+
+	// Update this so we render from the latest projection point
+	UpdateProjectionOriginComponent();
+
+	if (!RenderNormalMap(NorthNormalMap, true))
+	{
+		return false;
+	}
+
+	if (!RenderNormalMap(SouthNormalMap, false))
+	{
+		return false;
+	}
+
+	bNormalMapInvalid = false;
+	UpdateNormalMapMesh();
+
+	return !bNormalMapInvalid;
+}
+
+void FDisplayClusterLightCardEditorHelper::UpdateNormalMapMesh()
+{
+	if (!NormalMeshScene.IsValid())
+	{
+		FWorldDelegates::OnWorldCleanup.AddRaw(this, &FDisplayClusterLightCardEditorHelper::OnWorldCleanup);
+		NormalMeshScene = MakeShared<FPreviewScene>(FPreviewScene::ConstructionValues());
+	}
+	else if (NormalMapMeshComponent.IsValid())
+	{
+		NormalMeshScene->RemoveComponent(NormalMapMeshComponent.Get());
+	}
+
+	UpdateProjectionOriginComponent();
+
+	NormalMapMeshComponent = NewObject<UProceduralMeshComponent>(GetTransientPackage(), TEXT("NormalMapMesh"));
+	NormalMapMeshComponent->SetCollisionEnabled(ECollisionEnabled::QueryOnly);
+	NormalMeshScene->AddComponent(NormalMapMeshComponent.Get(), FTransform(GetProjectionOrigin()));
+
+	UStaticMesh* IcoSphereMesh = LoadObject<UStaticMesh>(nullptr, TEXT("/nDisplay/Meshes/SM_IcoSphere.SM_IcoSphere"), nullptr, LOAD_None, nullptr);
+
+	if (ensure(IcoSphereMesh))
+	{
+		IcoSphereMesh->bAllowCPUAccess = true;
+
+		int32 NumSections = IcoSphereMesh->GetNumSections(0);
+		for (int32 SectionIndex = 0; SectionIndex < NumSections; SectionIndex++)
+		{
+			TArray<FVector> Vertices;
+			TArray<int32> Triangles;
+			TArray<FVector> Normals;
+			TArray<FVector2D> UVs;
+			TArray<FProcMeshTangent> Tangents;
+
+			UKismetProceduralMeshLibrary::GetSectionFromStaticMesh(IcoSphereMesh, 0, SectionIndex, Vertices, Triangles, Normals, UVs, Tangents);
+
+			TArray<FVector2D> EmptyUVs;
+			TArray<FLinearColor> EmptyColors;
+			NormalMapMeshComponent.Get()->CreateMeshSection_LinearColor(SectionIndex, Vertices, Triangles, Normals, UVs, EmptyUVs, EmptyUVs, EmptyUVs, EmptyColors, Tangents, true);
+
+			for (int32 Index = 0; Index < IcoSphereMesh->GetStaticMaterials().Num(); ++Index)
+			{
+				UMaterialInterface* MaterialInterface = IcoSphereMesh->GetStaticMaterials()[Index].MaterialInterface;
+				NormalMapMeshComponent.Get()->SetMaterial(Index, MaterialInterface);
+			}
+		}
+	}
+}
+
+bool FDisplayClusterLightCardEditorHelper::RenderNormalMap(FNormalMap& NormalMap, bool bIsNorthMap)
+{
+	ADisplayClusterRootActor* RootActor = UpdateRootActor();
+	if (!RootActor)
+	{
+		return false;
+	}
+
+	UWorld* World = RootActor->GetWorld();
+	if (!World)
+	{
+		return false;
+	}
+
+	// Rotate the normal map's direction to account for the actor's world rotation
+	FRotator InverseRootRotation = RootActor->GetTransform().Rotator().GetInverse();
+	FVector NormalMapDirection = InverseRootRotation.RotateVector(bIsNorthMap ? FVector::UpVector : FVector::DownVector);
+
+	// Only render primitive components from the stage actor for the normal map
+	FDisplayClusterMeshProjectionPrimitiveFilter PrimitiveFilter;
+	PrimitiveFilter.ShouldRenderPrimitiveDelegate = FDisplayClusterMeshProjectionPrimitiveFilter::FPrimitiveFilter::CreateWeakLambda(
+		RootActor, [RootActor](const UPrimitiveComponent* PrimitiveComponent)
+		{
+			return PrimitiveComponent->GetOwner() == RootActor;
+		});
+
+	FDisplayClusterMeshProjectionRenderSettings RenderSettings;
+	RenderSettings.RenderType = EDisplayClusterMeshProjectionOutput::Normals;
+	RenderSettings.ProjectionType = EDisplayClusterMeshProjectionType::Azimuthal;
+	RenderSettings.PrimitiveFilter = PrimitiveFilter;
+	RenderSettings.NormalCorrectionMatrix = FMatrix44f(FRotationMatrix::Make(InverseRootRotation));
+
+#if WITH_EDITOR
+	if (ViewportClient.IsValid())
+	{
+		const FEditorViewportClient& ViewportClientRef = *ViewportClient.Pin().Get();
+		RenderSettings.EngineShowFlags = ViewportClientRef.EngineShowFlags;
+	}
+#endif
+
+	FSceneViewInitOptions ViewInitOptions;
+	GetNormalMapSceneViewInitOptions(NormalMapDirection, RenderSettings.ViewInitOptions);
+
+	NormalMap.Init(RenderSettings.ViewInitOptions);
+
+	NormalMap.Canvas = MakeShared<FCanvas>(&NormalMap, nullptr, World, World->Scene->GetFeatureLevel(), FCanvas::CDM_DeferDrawing, 1.0f);
+	NormalMap.Canvas->Clear(FLinearColor::Black);
+
+	IDisplayClusterScenePreview::Get().Render(RendererId, RenderSettings, *NormalMap.Canvas);
+	NormalMap.Canvas->Flush_GameThread();
+
+	NormalMap.ReadFloat16Pixels(NormalMap.GetCachedNormalData());
+	NormalMap.Release();
+
+	FlushRenderingCommands();
+
+	return true;
+}
+
+void FDisplayClusterLightCardEditorHelper::OnWorldCleanup(UWorld* World, bool bSessionEnded, bool bCleanupResources)
+{
+	if (!NormalMeshScene.IsValid() || !CachedRootActor.IsValid() || World != CachedRootActor->GetWorld())
+	{
+		return;
+	}
+
+	if (NormalMapMeshComponent.IsValid())
+	{
+		NormalMeshScene->RemoveComponent(NormalMapMeshComponent.Get());
+	}
+
+	NormalMeshScene.Reset();
+}
+
+#if WITH_EDITORONLY_DATA
+void FDisplayClusterLightCardEditorHelper::OnActorPropertyChanged(UObject* ObjectBeingModified, FPropertyChangedEvent& PropertyChangedEvent)
+{
+	if (!CachedRootActor.IsValid())
+	{
+		FCoreUObjectDelegates::OnObjectPropertyChanged.RemoveAll(this);
+		return;
+	}
+
+	if (ObjectBeingModified == CachedRootActor.Get())
+	{
+		InvalidateNormalMap();
+	}
+}
+
+void FDisplayClusterLightCardEditorHelper::OnRootActorBlueprintCompiled(UBlueprint* Blueprint)
+{
+	InvalidateNormalMap();
+}
+#endif
+
+#if WITH_EDITOR
+void FDisplayClusterLightCardEditorHelper::PostEditChangePropertiesForMovedLightCard(ADisplayClusterLightCardActor& LightCard) const
+{
+	if (LightCard.bIsProxy)
+	{
+		return;
+	}
+
+	LightCard.Modify();
+
+	TArray<const FProperty*> ChangedProperties;
+	auto ModifyLightCardProperty = [&](const FName& PropertyName)
+	{
+		FProperty* Property = FindFProperty<FProperty>(LightCard.GetClass(), PropertyName);
+		FPropertyChangedEvent PropertyEvent(Property);
+		LightCard.PostEditChangeProperty(PropertyEvent);
+
+		ChangedProperties.Add(Property);
+	};
+
+	ModifyLightCardProperty(GET_MEMBER_NAME_CHECKED(ADisplayClusterLightCardActor, Longitude));
+	ModifyLightCardProperty(GET_MEMBER_NAME_CHECKED(ADisplayClusterLightCardActor, Latitude));
+	ModifyLightCardProperty(GET_MEMBER_NAME_CHECKED(ADisplayClusterLightCardActor, DistanceFromCenter));
+	ModifyLightCardProperty(GET_MEMBER_NAME_CHECKED(ADisplayClusterLightCardActor, Spin));
+	ModifyLightCardProperty(GET_MEMBER_NAME_CHECKED(ADisplayClusterLightCardActor, Pitch));
+	ModifyLightCardProperty(GET_MEMBER_NAME_CHECKED(ADisplayClusterLightCardActor, Yaw));
+	ModifyLightCardProperty(GET_MEMBER_NAME_CHECKED(ADisplayClusterLightCardActor, Scale));
+	ModifyLightCardProperty(GET_MEMBER_NAME_CHECKED(ADisplayClusterLightCardActor, UVCoordinates));
+
+	SnapshotTransactionBuffer(&LightCard, MakeArrayView(ChangedProperties));
+}
+#endif
diff --git a/Engine/Plugins/Runtime/nDisplay/Source/DisplayClusterScenePreview/Private/DisplayClusterScenePreviewModule.cpp b/Engine/Plugins/Runtime/nDisplay/Source/DisplayClusterScenePreview/Private/DisplayClusterScenePreviewModule.cpp
--- a/Engine/Plugins/Runtime/nDisplay/Source/DisplayClusterScenePreview/Private/DisplayClusterScenePreviewModule.cpp
+++ b/Engine/Plugins/Runtime/nDisplay/Source/DisplayClusterScenePreview/Private/DisplayClusterScenePreviewModule.cpp
@@ -5,6 +5,7 @@
 #include "Blueprints/DisplayClusterBlueprintLib.h"
 #include "CanvasTypes.h"
 #include "Components/DisplayClusterScreenComponent.h"
+#include "Components/DisplayClusterPreviewComponent.h"
 #include "DisplayClusterLightCardActor.h"
 #include "DisplayClusterRootActor.h"
 #include "Engine/TextureRenderTarget2D.h"
@@ -25,6 +26,13 @@ void FDisplayClusterScenePreviewModule::ShutdownModule()
 {
 	FTSTicker::GetCoreTicker().RemoveTicker(RenderTickerHandle);
 	RenderTickerHandle.Reset();
+
+	TArray<int32> RendererIds;
+	RendererConfigs.GenerateKeyArray(RendererIds);
+	for (const int32 RendererId : RendererIds)
+	{
+		DestroyRenderer(RendererId);
+	}
 }
 
 int32 FDisplayClusterScenePreviewModule::CreateRenderer()
@@ -88,11 +96,37 @@ ADisplayClusterRootActor* FDisplayClusterScenePreviewModule::GetRendererRootActo
 	return nullptr;
 }
 
+bool FDisplayClusterScenePreviewModule::GetActorsInRendererScene(int32 RendererId, bool bIncludeRoot, TArray<AActor*>& OutActors)
+{
+	if (FRendererConfig* Config = RendererConfigs.Find(RendererId))
+	{
+		if (bIncludeRoot)
+		{
+			OutActors.Add(GetRendererRootActor(RendererId));
+		}
+
+		for (const TWeakObjectPtr<AActor>& Actor : Config->AddedActors)
+		{
+			if (!Actor.IsValid())
+			{
+				continue;
+			}
+
+			OutActors.Add(Actor.Get());
+		}
+
+		return true;
+	}
+
+	return false;
+}
+
 bool FDisplayClusterScenePreviewModule::AddActorToRenderer(int32 RendererId, AActor* Actor)
 {
 	if (FRendererConfig* Config = RendererConfigs.Find(RendererId))
 	{
 		Config->Renderer->AddActor(Actor);
+		Config->AddedActors.Add(Actor);
 		return true;
 	}
 
@@ -126,6 +160,7 @@ bool FDisplayClusterScenePreviewModule::ClearRendererScene(int32 RendererId)
 	if (FRendererConfig* Config = RendererConfigs.Find(RendererId))
 	{
 		Config->Renderer->ClearScene();
+		Config->AddedActors.Empty();
 		return true;
 	}
 
@@ -154,6 +189,17 @@ bool FDisplayClusterScenePreviewModule::SetRendererRenderSimpleElementsDelegate(
 	return false;
 }
 
+bool FDisplayClusterScenePreviewModule::SetRendererUsePostProcessTexture(int32 RendererId, bool bUsePostProcessTexture)
+{
+	if (FRendererConfig* Config = RendererConfigs.Find(RendererId))
+	{
+		Config->bUsePostProcessTexture = bUsePostProcessTexture;
+		return true;
+	}
+
+	return false;
+}
+
 bool FDisplayClusterScenePreviewModule::Render(int32 RendererId, FDisplayClusterMeshProjectionRenderSettings& RenderSettings, FCanvas& Canvas)
 {
 	if (FRendererConfig* Config = RendererConfigs.Find(RendererId))
@@ -166,22 +212,23 @@ bool FDisplayClusterScenePreviewModule::Render(int32 RendererId, FDisplayCluster
 
 bool FDisplayClusterScenePreviewModule::RenderQueued(int32 RendererId, FDisplayClusterMeshProjectionRenderSettings& RenderSettings, const FIntPoint& Size, FRenderResultDelegate ResultDelegate)
 {
-	if (FRendererConfig* Config = RendererConfigs.Find(RendererId))
-	{
-		RenderQueue.Enqueue(FPreviewRenderJob(RendererId, RenderSettings, Size, ResultDelegate));
+	return InternalRenderQueued(RendererId, RenderSettings, nullptr, Size, ResultDelegate);
+}
 
-		if (!RenderTickerHandle.IsValid())
-		{
-			RenderTickerHandle = FTSTicker::GetCoreTicker().AddTicker(
-				FTickerDelegate::CreateRaw(this, &FDisplayClusterScenePreviewModule::OnTick),
-				CVarDisplayClusterScenePreviewRenderTickDelay.GetValueOnGameThread()
-			);
-		}
+bool FDisplayClusterScenePreviewModule::RenderQueued(int32 RendererId, FDisplayClusterMeshProjectionRenderSettings& RenderSettings, const TWeakPtr<FCanvas> Canvas, FRenderResultDelegate ResultDelegate)
+{
+	if (!Canvas.IsValid())
+	{
+		return false;
+	}
 
-		return true;
+	FRenderTarget* RenderTarget = Canvas.Pin()->GetRenderTarget();
+	if (!RenderTarget)
+	{
+		return false;
 	}
 
-	return false;
+	return InternalRenderQueued(RendererId, RenderSettings, Canvas, RenderTarget->GetSizeXY(), ResultDelegate);
 }
 
 ADisplayClusterRootActor* FDisplayClusterScenePreviewModule::InternalGetRendererRootActor(FRendererConfig& RendererConfig)
@@ -199,11 +246,31 @@ ADisplayClusterRootActor* FDisplayClusterScenePreviewModule::InternalGetRenderer
 	return RendererConfig.RootActor.Get();
 }
 
+bool FDisplayClusterScenePreviewModule::InternalRenderQueued(int32 RendererId, FDisplayClusterMeshProjectionRenderSettings& RenderSettings, TWeakPtr<FCanvas> Canvas,
+	const FIntPoint& Size, FRenderResultDelegate ResultDelegate)
+{
+	if (FRendererConfig* Config = RendererConfigs.Find(RendererId))
+	{
+		RenderQueue.Enqueue(FPreviewRenderJob(RendererId, RenderSettings, Size, Canvas, ResultDelegate));
+
+		if (!RenderTickerHandle.IsValid())
+		{
+			RenderTickerHandle = FTSTicker::GetCoreTicker().AddTicker(
+				FTickerDelegate::CreateRaw(this, &FDisplayClusterScenePreviewModule::OnTick),
+				CVarDisplayClusterScenePreviewRenderTickDelay.GetValueOnGameThread()
+			);
+		}
+
+		return true;
+	}
+
+	return false;
+}
+
 void FDisplayClusterScenePreviewModule::InternalSetRendererRootActor(FRendererConfig& RendererConfig, ADisplayClusterRootActor* Actor, bool bAutoUpdateLightcards)
 {
 	// Determine these values before we update the config's RootActor/bAutoUpdateLightcards
 	const bool bRootChanged = RendererConfig.RootActor != Actor;
-	const bool bStartAutoUpdating = bAutoUpdateLightcards && (!RendererConfig.bAutoUpdateLightcards || bRootChanged);
 
 	if (bRootChanged)
 	{
@@ -218,11 +285,7 @@ void FDisplayClusterScenePreviewModule::InternalSetRendererRootActor(FRendererCo
 		AutoPopulateScene(RendererConfig);
 	}
 
-	if (bStartAutoUpdating)
-	{
-		RegisterRootActorEvents(Actor, true);
-	}
-
+	RegisterRootActorEvents(Actor, true);
 	RegisterOrUnregisterGlobalActorEvents();
 }
 
@@ -239,7 +302,44 @@ bool FDisplayClusterScenePreviewModule::InternalRenderImmediate(FRendererConfig&
 		AutoPopulateScene(RendererConfig);
 	}
 
+#if WITH_EDITOR
+	TMap<UDisplayClusterPreviewComponent*, UTexture*> PreviewComponentOverrideTextures;
+	const bool bUsePostProcessTexture = RendererConfig.bUsePostProcessTexture && (RenderSettings.RenderType == EDisplayClusterMeshProjectionOutput::Color);
+
+	// Update the preview components with the post-processed texture
+	if (bUsePostProcessTexture)
+	{
+		UDisplayClusterConfigurationData* Config = RendererConfig.RootActor->GetConfigData();
+
+		for (const TPair<FString, UDisplayClusterConfigurationClusterNode*>& NodePair : Config->Cluster->Nodes)
+		{
+			const UDisplayClusterConfigurationClusterNode* Node = NodePair.Value;
+			for (const TPair<FString, UDisplayClusterConfigurationViewport*>& ViewportPair : Node->Viewports)
+			{
+				UDisplayClusterPreviewComponent* PreviewComp = RendererConfig.RootActor->GetPreviewComponent(NodePair.Key, ViewportPair.Key);
+				if (PreviewComp)
+				{
+					PreviewComponentOverrideTextures.Add(PreviewComp, PreviewComp->GetOverrideTexture());
+					PreviewComp->SetOverrideTexture(PreviewComp->GetRenderTargetTexturePostProcess());
+				}
+			}
+		}
+	}
+#endif
+
 	RendererConfig.Renderer->Render(&Canvas, World->Scene, RenderSettings);
+
+#if WITH_EDITOR
+	// Remove the override texture so the actor renders normally
+	if (bUsePostProcessTexture)
+	{
+		for (const TPair<UDisplayClusterPreviewComponent*, UTexture*>& PreviewPair : PreviewComponentOverrideTextures)
+		{
+			PreviewPair.Key->SetOverrideTexture(PreviewPair.Value);
+		}
+	}
+#endif
+
 	return true;
 }
 
@@ -266,6 +366,7 @@ void FDisplayClusterScenePreviewModule::RegisterOrUnregisterGlobalActorEvents()
 		if (GEngine != nullptr)
 		{
 			GEngine->OnLevelActorDeleted().AddRaw(this, &FDisplayClusterScenePreviewModule::OnLevelActorDeleted);
+			GEngine->OnLevelActorAdded().AddRaw(this, &FDisplayClusterScenePreviewModule::OnLevelActorAdded);
 		}
 	}
 	else if (!bShouldBeRegistered && bIsRegisteredForActorEvents)
@@ -277,12 +378,13 @@ void FDisplayClusterScenePreviewModule::RegisterOrUnregisterGlobalActorEvents()
 		if (GEngine != nullptr)
 		{
 			GEngine->OnLevelActorDeleted().RemoveAll(this);
+			GEngine->OnLevelActorAdded().RemoveAll(this);
 		}
 	}
 #endif
 }
 
-void FDisplayClusterScenePreviewModule::RegisterRootActorEvents(AActor* Actor, bool bShouldRegister)
+void FDisplayClusterScenePreviewModule::RegisterRootActorEvents(ADisplayClusterRootActor* Actor, bool bShouldRegister)
 {
 #if WITH_EDITOR
 	if (!Actor)
@@ -290,16 +392,23 @@ void FDisplayClusterScenePreviewModule::RegisterRootActorEvents(AActor* Actor, b
 		return;
 	}
 
+	// Register/unregister actors to use post-processing for preview renders so their lighting looks correct
+	Actor->UnsubscribeFromPostProcessRenderTarget(reinterpret_cast<uint8*>(this));
+
+	if (bShouldRegister)
+	{
+		Actor->SubscribeToPostProcessRenderTarget(reinterpret_cast<uint8*>(this));
+	}
+
+	// Register/unregister for Blueprint events
 	if (UBlueprint* Blueprint = UBlueprint::GetBlueprintFromClass(Actor->GetClass()))
 	{
+		Blueprint->OnCompiled().RemoveAll(this);
+
 		if (bShouldRegister)
 		{
 			Blueprint->OnCompiled().AddRaw(this, &FDisplayClusterScenePreviewModule::OnBlueprintCompiled);
 		}
-		else
-		{
-			Blueprint->OnCompiled().RemoveAll(this);
-		}
 	}
 #endif
 }
@@ -324,6 +433,7 @@ void FDisplayClusterScenePreviewModule::AutoPopulateScene(FRendererConfig& Rende
 			for (ADisplayClusterLightCardActor* LightCard : LightCards)
 			{
 				RendererConfig.Renderer->AddActor(LightCard);
+				RendererConfig.AddedActors.Add(LightCard);
 				RendererConfig.AutoLightcards.Add(LightCard);
 			}
 		}
@@ -334,45 +444,73 @@ void FDisplayClusterScenePreviewModule::AutoPopulateScene(FRendererConfig& Rende
 
 bool FDisplayClusterScenePreviewModule::OnTick(float DeltaTime)
 {
-	FPreviewRenderJob Job;
-	if (!RenderQueue.Dequeue(Job))
+	// This loop should break when we either run out of jobs or complete a single job
+	while (!RenderQueue.IsEmpty())
 	{
-		RenderTickerHandle.Reset();
-		return false;
-	}
+		FPreviewRenderJob Job;
+		if (!RenderQueue.Dequeue(Job))
+		{
+			break;
+		}
 
-	ensure(Job.ResultDelegate.IsBound());
+		ensure(Job.ResultDelegate.IsBound());
 
-	if (FRendererConfig* Config = RendererConfigs.Find(Job.RendererId))
-	{
-		if (UWorld* World = Config->RootActor.IsValid() ? Config->RootActor->GetWorld() : nullptr)
+		if (FRendererConfig* Config = RendererConfigs.Find(Job.RendererId))
 		{
-			UTextureRenderTarget2D* RenderTarget = Config->RenderTarget.Get();
-
-			if (!RenderTarget)
+			if (Job.bWasCanvasProvided)
 			{
-				// Create a new render target
-				RenderTarget = NewObject<UTextureRenderTarget2D>();
-				RenderTarget->InitCustomFormat(Job.Size.X, Job.Size.Y, PF_B8G8R8A8, false);
-
-				Config->RenderTarget = TStrongObjectPtr<UTextureRenderTarget2D>(RenderTarget);
+				// We were provided a canvas for this render job, so use it if possible
+				if (!Job.Canvas.IsValid())
+				{
+					continue;
+				}
+
+				TSharedPtr<FCanvas, ESPMode::ThreadSafe> Canvas = Job.Canvas.Pin();
+				FRenderTarget* RenderTarget = Canvas->GetRenderTarget();
+				if (!RenderTarget)
+				{
+					continue;
+				}
+
+				InternalRenderImmediate(*Config, Job.Settings, *Canvas);
+				Job.ResultDelegate.Execute(*RenderTarget);
+				break;
 			}
-			else if (RenderTarget->SizeX != Job.Size.X || RenderTarget->SizeY != Job.Size.Y)
+
+			if (UWorld* World = Config->RootActor.IsValid() ? Config->RootActor->GetWorld() : nullptr)
 			{
-				// Resize to match the new size
-				RenderTarget->ResizeTarget(Job.Size.X, Job.Size.Y);
-				
-				// Flush commands so target is immediately ready to render at the new size
-				FlushRenderingCommands();
+				// We need to provide the render target for this job
+				UTextureRenderTarget2D* RenderTarget = Config->RenderTarget.Get();
+
+				if (!RenderTarget)
+				{
+					// Create a new render target (which will be reused for this config in the future)
+					RenderTarget = NewObject<UTextureRenderTarget2D>();
+					RenderTarget->InitCustomFormat(Job.Size.X, Job.Size.Y, PF_B8G8R8A8, true);
+
+					Config->RenderTarget = TStrongObjectPtr<UTextureRenderTarget2D>(RenderTarget);
+				}
+				else if (RenderTarget->SizeX != Job.Size.X || RenderTarget->SizeY != Job.Size.Y)
+				{
+					// Resize to match the new size
+					RenderTarget->ResizeTarget(Job.Size.X, Job.Size.Y);
+
+					// Flush commands so target is immediately ready to render at the new size
+					FlushRenderingCommands();
+				}
+
+				FTextureRenderTargetResource* RenderTargetResource = RenderTarget->GameThread_GetRenderTargetResource();
+				FCanvas Canvas(RenderTargetResource, nullptr, FGameTime::GetTimeSinceAppStart(), World->Scene->GetFeatureLevel());
+
+				InternalRenderImmediate(*Config, Job.Settings, Canvas);
+				Job.ResultDelegate.Execute(*RenderTargetResource);
+				break;
 			}
-
-			FTextureRenderTargetResource* RenderTargetResource = RenderTarget->GameThread_GetRenderTargetResource();
-			FCanvas Canvas(RenderTargetResource, nullptr, FGameTime::GetTimeSinceAppStart(), World->Scene->GetFeatureLevel());
-
-			InternalRenderImmediate(*Config, Job.Settings, Canvas);
-
-			Job.ResultDelegate.Execute(*RenderTargetResource);
+			
+			// No canvas and no world, so try the next render
 		}
+
+		// Config no longer exists, so try the next render
 	}
 
 	if (RenderQueue.IsEmpty())
@@ -408,6 +546,20 @@ void FDisplayClusterScenePreviewModule::OnLevelActorDeleted(AActor* Actor)
 	}
 }
 
+void FDisplayClusterScenePreviewModule::OnLevelActorAdded(AActor* Actor)
+{
+	if (Cast<ADisplayClusterLightCardActor>(Actor) == nullptr)
+	{
+		return;
+	}
+
+	// The actor won't be added to a root actor yet, so we can't check who it belongs to. Easier to just mark all configs as dirty.
+	for (TPair<int32, FRendererConfig>& ConfigPair : RendererConfigs)
+	{
+		ConfigPair.Value.bIsSceneDirty = true;
+	}
+}
+
 void FDisplayClusterScenePreviewModule::OnBlueprintCompiled(UBlueprint* Blueprint)
 {
 #if WITH_EDITOR
@@ -417,6 +569,7 @@ void FDisplayClusterScenePreviewModule::OnBlueprintCompiled(UBlueprint* Blueprin
 		if (Config.RootActor.IsValid() && Blueprint == UBlueprint::GetBlueprintFromClass(Config.RootActor->GetClass()))
 		{
 			Config.bIsSceneDirty = true;
+			RegisterRootActorEvents(Config.RootActor.Get(), true);
 		}
 	}
 #endif
diff --git a/Engine/Plugins/Runtime/nDisplay/Source/DisplayClusterScenePreview/Private/DisplayClusterScenePreviewModule.h b/Engine/Plugins/Runtime/nDisplay/Source/DisplayClusterScenePreview/Private/DisplayClusterScenePreviewModule.h
--- a/Engine/Plugins/Runtime/nDisplay/Source/DisplayClusterScenePreview/Private/DisplayClusterScenePreviewModule.h
+++ b/Engine/Plugins/Runtime/nDisplay/Source/DisplayClusterScenePreview/Private/DisplayClusterScenePreviewModule.h
@@ -29,14 +29,17 @@ public:
 	virtual bool SetRendererRootActorPath(int32 RendererId, const FString& ActorPath, bool bAutoUpdateLightcards = false) override;
 	virtual bool SetRendererRootActor(int32 RendererId, ADisplayClusterRootActor* Actor, bool bAutoUpdateLightcards = false) override;
 	virtual ADisplayClusterRootActor* GetRendererRootActor(int32 RendererId) override;
+	virtual bool GetActorsInRendererScene(int32 RendererId, bool bIncludeRoot, TArray<AActor*>& OutActors) override;
 	virtual bool AddActorToRenderer(int32 RendererId, AActor* Actor) override;
 	virtual bool AddActorToRenderer(int32 RendererId, AActor* Actor, const TFunctionRef<bool(const UPrimitiveComponent*)>& PrimitiveFilter) override;
 	virtual bool RemoveActorFromRenderer(int32 RendererId, AActor* Actor) override;
 	virtual bool ClearRendererScene(int32 RendererId) override;
 	virtual bool SetRendererActorSelectedDelegate(int32 RendererId, FDisplayClusterMeshProjectionRenderer::FSelection ActorSelectedDelegate) override;
 	virtual bool SetRendererRenderSimpleElementsDelegate(int32 RendererId, FDisplayClusterMeshProjectionRenderer::FSimpleElementPass RenderSimpleElementsDelegate) override;
+	virtual bool SetRendererUsePostProcessTexture(int32 RendererId, bool bUsePostProcessTexture) override;
 	virtual bool Render(int32 RendererId, FDisplayClusterMeshProjectionRenderSettings& RenderSettings, FCanvas& Canvas) override;
 	virtual bool RenderQueued(int32 RendererId, FDisplayClusterMeshProjectionRenderSettings& RenderSettings, const FIntPoint& Size, FRenderResultDelegate ResultDelegate) override;
+	virtual bool RenderQueued(int32 RendererId, FDisplayClusterMeshProjectionRenderSettings& RenderSettings, const TWeakPtr<FCanvas> Canvas, FRenderResultDelegate ResultDelegate) override;
 	//~ End IDisplayClusterScenePreview Interface
 
 private:
@@ -52,6 +55,9 @@ private:
 		/** The path of the root actor that this renderer is previewing. If this is not empty and the root actor becomes invalid, we will attempt to find it again using this path. */
 		FString RootActorPath;
 
+		/** All actors that have been added to the renderer (except for the root actor). */
+		TArray<TWeakObjectPtr<AActor>> AddedActors;
+
 		/** Lightcards that have been automatically added to the scene. */
 		TArray<TWeakObjectPtr<ADisplayClusterLightCardActor>> AutoLightcards;
 
@@ -61,6 +67,9 @@ private:
 		/** If true, the scene needs to be updated before the next render. This is only relevant if bAutoUpdateLightcards is true. */
 		bool bIsSceneDirty = true;
 
+		/** If true, apply post-processing to the nDisplay preview texture and override the preview component with it before rendering. */
+		bool bUsePostProcessTexture = false;
+
 		/** The render target to use for queued renders. */
 		TStrongObjectPtr<UTextureRenderTarget2D> RenderTarget = nullptr;
 	};
@@ -71,8 +80,8 @@ private:
 		FPreviewRenderJob() {}
 
 		FPreviewRenderJob(int32 RendererId, const FDisplayClusterMeshProjectionRenderSettings& Settings, const FIntPoint& Size,
-			FRenderResultDelegate ResultDelegate)
-			: RendererId(RendererId), Settings(Settings), Size(Size), ResultDelegate(ResultDelegate)
+			TWeakPtr<FCanvas> Canvas, FRenderResultDelegate ResultDelegate)
+			: RendererId(RendererId), Settings(Settings), Size(Size), Canvas(Canvas), bWasCanvasProvided(Canvas.IsValid()), ResultDelegate(ResultDelegate)
 		{
 		}
 
@@ -85,6 +94,12 @@ private:
 		/** The size of the image to render. */
 		FIntPoint Size;
 
+		/** The canvas to render to, if provided. */
+		TWeakPtr<FCanvas> Canvas;
+
+		/** Whether a canvas was provided for this job. */
+		bool bWasCanvasProvided;
+
 		/** The delegate to call when the render is completed. */
 		FRenderResultDelegate ResultDelegate;
 	};
@@ -95,6 +110,10 @@ private:
 	/** Set the root actor for a config, update its scene, and register events accordingly. */
 	void InternalSetRendererRootActor(FRendererConfig& RendererConfig, ADisplayClusterRootActor* Actor, bool bAutoUpdateLightcards);
 
+	/** Queue a preview to be rendered. */
+	bool InternalRenderQueued(int32 RendererId, FDisplayClusterMeshProjectionRenderSettings& RenderSettings, TWeakPtr<FCanvas> Canvas,
+		const FIntPoint& Size, FRenderResultDelegate ResultDelegate);
+
 	/** Immediately render with the given renderer config and settings to the given canvas. */
 	bool InternalRenderImmediate(FRendererConfig& RendererConfig, FDisplayClusterMeshProjectionRenderSettings& RenderSettings, FCanvas& Canvas);
 
@@ -102,7 +121,7 @@ private:
 	void RegisterOrUnregisterGlobalActorEvents();
 
 	/** Register/unregister to events affecting a cluster root actor. */
-	void RegisterRootActorEvents(AActor* Actor, bool bShouldRegister);
+	void RegisterRootActorEvents(ADisplayClusterRootActor* Actor, bool bShouldRegister);
 
 	/** Clear and re-populate a renderer's scene with the root actor and lightcards if applicable. */
 	void AutoPopulateScene(FRendererConfig& RendererConfig);
@@ -116,6 +135,9 @@ private:
 	/** Called when the user deletes an actor from the level. */
 	void OnLevelActorDeleted(AActor* Actor);
 
+	/** Called when the user adds an actor to the level. */
+	void OnLevelActorAdded(AActor* Actor);
+
 	/** Called when a blueprint for an actor we care about is compiled. */
 	void OnBlueprintCompiled(UBlueprint* Blueprint);
 
diff --git a/Engine/Plugins/Runtime/nDisplay/Source/DisplayClusterScenePreview/Public/DisplayClusterLightCardEditorHelper.h b/Engine/Plugins/Runtime/nDisplay/Source/DisplayClusterScenePreview/Public/DisplayClusterLightCardEditorHelper.h
new file mode 100644
--- /dev/null
+++ b/Engine/Plugins/Runtime/nDisplay/Source/DisplayClusterScenePreview/Public/DisplayClusterLightCardEditorHelper.h
@@ -0,0 +1,371 @@
+// Copyright Epic Games, Inc. All Rights Reserved.
+
+#pragma once
+
+#include "CoreMinimal.h"
+#include "Containers/Union.h"
+#include "DisplayClusterMeshProjectionRenderer.h"
+#include "SceneView.h"
+#include "UObject/StrongObjectPtr.h"
+#include "UnrealClient.h"
+
+class ADisplayClusterLightCardActor;
+class ADisplayClusterRootActor;
+class FPreviewScene;
+class FSceneView;
+class UProceduralMeshComponent;
+class UTexture2D;
+struct FSceneViewInitOptions;
+
+#if WITH_EDITOR
+class FEditorViewportClient;
+#endif
+
+/**
+ * Helper class for moving lightcards in an nDisplay cluster in various projection modes.
+ * Contains functions to perform coordinate conversion, scene setup, and lightcard movement, and manages the
+ * normal maps needed to convert from projected to world coordinates.
+*/
+struct FDisplayClusterLightCardEditorHelper
+{
+public:
+	struct FSphericalCoordinates
+	{
+	public:
+
+		/** Constructors */
+		DISPLAYCLUSTERSCENEPREVIEW_API FSphericalCoordinates(const FVector& CartesianPosition);
+		DISPLAYCLUSTERSCENEPREVIEW_API FSphericalCoordinates();
+
+		/** Return equivalent cartesian coordinates */
+		DISPLAYCLUSTERSCENEPREVIEW_API FVector AsCartesian() const;
+
+		/** Addition operator */
+		DISPLAYCLUSTERSCENEPREVIEW_API FSphericalCoordinates operator+(FSphericalCoordinates const& Other) const;
+
+		/** Subtraction operator */
+		DISPLAYCLUSTERSCENEPREVIEW_API FSphericalCoordinates operator-(FSphericalCoordinates const& Other) const;
+
+		/** Conform parameters to their normal ranges */
+		DISPLAYCLUSTERSCENEPREVIEW_API void Conform();
+
+		/** Returns a conformed version of this struct without changing the current one */
+		DISPLAYCLUSTERSCENEPREVIEW_API FSphericalCoordinates GetConformed() const;
+
+		/** Returns true if the inclination is pointing at north or south poles, within the given margin (in radians) */
+		DISPLAYCLUSTERSCENEPREVIEW_API bool IsPointingAtPole(double Margin = 1e-6) const;
+
+		double Radius = 0;      // unitless   0+      (when conforming)
+		double Inclination = 0; // radians    0 to PI (when conforming)
+		double Azimuth = 0;     // radians  -PI to PI (when conforming)
+	};
+
+private:
+	/** Custom render target that stores the normal data for the stage */
+	class FNormalMap : public FRenderTarget
+	{
+	public:
+		/** The size of the normal map */
+		static const int32 NormalMapSize;
+
+		/**
+		 * The field of view to render the normal map with. Using the azimuthal projection,
+		 * this is set so that the entire 360 degree scene is rendered
+		 */
+		static const float NormalMapFOV;
+
+		/** Initializes the normal map render target using the specified scene view options */
+		void Init(const FSceneViewInitOptions& InSceneViewInitOptions);
+
+		/** Releases the normal map render target's resources */
+		void Release();
+
+		/** Gets the size of the render target */
+		virtual FIntPoint GetSizeXY() const override { return FIntPoint(SizeX, SizeY); }
+
+		/** Gets a reference to the normal map data array, which stores the normal vector in the RGB components (color = 0.5 * Normal + 0.5) and the depth in the A component */
+		TArray<FFloat16Color>& GetCachedNormalData() { return CachedNormalData; }
+
+		/** Gets the normal vector and distance at the specified world location. The normal and distance are bilinearly interpolated from the nearest pixels in the normal map */
+		bool GetNormalAndDistanceAtPosition(FVector Position, FVector& OutNormal, float& OutDistance) const;
+
+		/** Morphs the vertices of the specified prodedural mesh to match the normal map */
+		void MorphProceduralMesh(UProceduralMeshComponent* InProceduralMeshComponent) const;
+
+		/** Generates a texture object that can be used to visualize the normal map */
+		UTexture2D* GenerateNormalMapTexture(const FString& TextureName);
+
+		/** Gets the normal map visualization texture, or null if it hasn't been generated */
+		UTexture2D* GetNormalMapTexture() const { return NormalMapTexture.IsValid() ? NormalMapTexture.Get() : nullptr; }
+
+	public:
+		/** True if this is currently waiting for a render to complete. */
+		bool bIsPendingRender = true;
+
+		/** The canvas currently being used to render this map. */
+		TSharedPtr<FCanvas> Canvas;
+
+	private:
+		/** The view matrices used when the normal map was last rendered */
+		FViewMatrices ViewMatrices;
+
+		/** The cached normal map data from the last normal map render */
+		TArray<FFloat16Color> CachedNormalData;
+
+		/** A texture that contains the normal map, for visualization purposes */
+		TWeakObjectPtr<UTexture2D> NormalMapTexture;
+
+		/** The width of the normal map. */
+		uint32 SizeX = 0;
+
+		/** The height of the normal map. */
+		uint32 SizeY = 0;
+	};
+
+public:
+	/** Create a projection helper that automatically creates its own preview renderer for normal map generation. */
+	DISPLAYCLUSTERSCENEPREVIEW_API FDisplayClusterLightCardEditorHelper();
+
+	/** Create a projection helper that uses an existing preview renderer for normal map generation. */
+	DISPLAYCLUSTERSCENEPREVIEW_API FDisplayClusterLightCardEditorHelper(int32 RendererId);
+
+	/** Clean up the projection renderer and any data used by the helper. */
+	DISPLAYCLUSTERSCENEPREVIEW_API ~FDisplayClusterLightCardEditorHelper();
+
+#if WITH_EDITOR
+	/** Set the viewport client for which this is handling rendering. If set, its data will be included in the SceneViewInitOptions generated by this helper. */
+	DISPLAYCLUSTERSCENEPREVIEW_API void SetEditorViewportClient(TWeakPtr<FEditorViewportClient> InViewportClient);
+#endif
+
+	/** Set the projection mode to use. */
+	DISPLAYCLUSTERSCENEPREVIEW_API void SetProjectionMode(EDisplayClusterMeshProjectionType Value);
+
+	/** Get the projection mode to use. */
+	DISPLAYCLUSTERSCENEPREVIEW_API EDisplayClusterMeshProjectionType GetProjectionMode() const;
+
+	/** Set whether to use an orthographic projection. */
+	DISPLAYCLUSTERSCENEPREVIEW_API void SetIsOrthographic(bool bValue);
+
+	/** Get whether this is using an orthographic projection. */
+	DISPLAYCLUSTERSCENEPREVIEW_API bool GetIsOrthographic() const;
+
+	/**
+	 * Set the root actor of the DisplayCluster being controlled. This will also update the renderer's root actor.
+	 * Don't call this if the helper was created with an existing preview renderer, as it will automatically the root actor from that renderer.
+	 */
+	DISPLAYCLUSTERSCENEPREVIEW_API void SetRootActor(ADisplayClusterRootActor& NewRootActor);
+
+	/**
+	 * If the root actor is a proxy, use this to set the corresponding level instance that is being proxied.
+	 */
+	DISPLAYCLUSTERSCENEPREVIEW_API void SetLevelInstanceRootActor(ADisplayClusterRootActor& NewRootActor);
+
+	/**
+	 * Get (or generate) a visualization of one of the helper's normal maps.
+	 *
+	 * @param bShowNorthMap If true, return the texture for the north normal map. Otherwise, return the texture for the south normal map.
+	 */
+	DISPLAYCLUSTERSCENEPREVIEW_API const UTexture2D* GetNormalMapTexture(bool bShowNorthMap);
+
+	/**
+	 * Moves the given light cards to the specified pixel position within the provided scene view.
+	 *
+	 * @param LightCards The light cards that we are moving
+	 * @param PixelPos The pixel location to move the light cards to
+	 * @param SceneView The scene view used to convert from pixel position to 3D position
+	*/
+	DISPLAYCLUSTERSCENEPREVIEW_API void MoveLightCardsToPixel(const TArray<TWeakObjectPtr<ADisplayClusterLightCardActor>>& LightCards, const FIntPoint& PixelPos, const FSceneView& SceneView);
+
+	/**
+	 * Moves specified cards to desired coordinates. Actual radius will be based on flush constraint and LightCard's RadialOffset.
+	 *
+	 * @param LightCards The light cards that we are moving
+	 * @param SphericalCoords The desired location of light card in spherical coordinates with respect to view origin.
+	*/
+	DISPLAYCLUSTERSCENEPREVIEW_API void MoveLightCardsTo(const TArray<TWeakObjectPtr<ADisplayClusterLightCardActor>>& LightCards, const FSphericalCoordinates& SphericalCoords);
+
+	/**
+	 * Moves specified cards to a coordinate in viewport space as if dragged by a translate widget.
+	 * 
+	 * @param LightCards The light cards that we are moving
+	 * @param PixelPos The screen pixel position of the widget
+	 * @param SceneView The scene view used to convert from pixel position to 3D position
+	 * @param DragWidgetOffset The offset between the actual cursor position and the position of the widget when the drag action started
+	 * @param DragAxis The axis along which the widget is being dragged
+	 * @param PrimaryLightCard The light card used to calculate the translation/rotation delta. If not provided, the last entry in LightCards will be used.
+	 */
+	DISPLAYCLUSTERSCENEPREVIEW_API void DragLightCards(const TArray<TWeakObjectPtr<ADisplayClusterLightCardActor>>& LightCards, const FIntPoint& PixelPos, const FSceneView& SceneView,
+		const FVector& DragWidgetOffset, EAxisList::Type DragAxis, ADisplayClusterLightCardActor* PrimaryLightCard = nullptr);
+
+	/** Ensures that the light card root component is at the same location as the projection/view origin */
+	DISPLAYCLUSTERSCENEPREVIEW_API void VerifyAndFixLightCardOrigin(ADisplayClusterLightCardActor& LightCard);
+
+	/**
+	 * Calculates the relative normal vector and world position in the specified direction from the given view origin.
+	 * 
+	 * @param InViewOrigin The origin point of the view.
+	 * @param InDirection The direction in which check the normal.
+	 * @param OutWorldPosition The world-space coordinates of the calculated position.
+	 * @param OutRelativeNormal The normal vector at the calculated position.
+	 * @param InDesiredDistanceFromFlush The desired flush distance from the initial calculated position.
+	 * @return true if the position was found, or false if the normal maps need to be updated first.
+	 */
+	DISPLAYCLUSTERSCENEPREVIEW_API bool CalculateNormalAndPositionInDirection(const FVector& InViewOrigin, const FVector& InDirection, FVector& OutWorldPosition,
+		FVector& OutRelativeNormal, double InDesiredDistanceFromFlush = 0.);
+
+	/**
+	 * Calculates the desired direction from the view origin given a pixel position within the view.
+	 * 
+	 * @param InPixelPos The desired position within the view in pixels.
+	 * @param InSceneView The view that was displayed when the pixel position was chosen.
+	 * @param InOriginOffset An offset to apply to the origin before calling TraceScreenRay (if this is in orthographic mode).
+	 * @param OutOrigin The origin point of the view.
+	 * @param OutDirection The direction relative to the origin.
+	 * @return true if the direction was found
+	 */
+	DISPLAYCLUSTERSCENEPREVIEW_API bool CalculateOriginAndDirectionFromPixelPosition(const FIntPoint& InPixelPos, const FSceneView& InSceneView, const FVector& InOriginOffset,
+		FVector& OutOrigin, FVector& OutDirection);
+
+	/** Converts a pixel coordinate into a point and direction vector in world space */
+	DISPLAYCLUSTERSCENEPREVIEW_API void PixelToWorld(const FSceneView& View, const FIntPoint& PixelPos, FVector& OutOrigin, FVector& OutDirection) const;
+
+	/** Converts a world coordinate into a point in screen space, and returns true if the world position is on the screen */
+	DISPLAYCLUSTERSCENEPREVIEW_API bool WorldToPixel(const FSceneView& View, const FVector& WorldPos, FVector2D& OutPixelPos) const;
+
+	/**
+	 * Sets up FSceneViewInitOptions' basic settings in a consistent way for preview renders.
+	 * If RotationMatrix is not provided, one will be automatically generated based on the view rotation.
+	 */
+	DISPLAYCLUSTERSCENEPREVIEW_API void GetSceneViewInitOptions(
+		FSceneViewInitOptions& OutViewInitOptions,
+		float InFOV,
+		const FIntPoint& InViewportSize,
+		const FVector& InLocation,
+		const FRotator& InRotation,
+		const EAspectRatioAxisConstraint InAspectRatioAxisConstraint = EAspectRatioAxisConstraint::AspectRatio_MaintainYFOV,
+		float InNearClipPlane = GNearClippingPlane,
+		const FMatrix* InRotationMatrix = nullptr,
+		float InDPIScale = 1.f);
+
+	/** Gets the spherical coordinates of the specified light card. */
+	DISPLAYCLUSTERSCENEPREVIEW_API FSphericalCoordinates GetLightCardCoordinates(const ADisplayClusterLightCardActor& LightCard);
+
+private:
+	/**
+	 * Update the cached pointer to the root actor from the associated preview renderer and, if the actor has changed, update related data.
+	 * Returns the root actor.
+	 */
+	ADisplayClusterRootActor* UpdateRootActor();
+
+	/**
+	 * Update the cached pointer to the component marking the origin point of the projection and return it.
+	 * Note that this has a cost, so avoid calling it frequently.
+	 */
+	USceneComponent* UpdateProjectionOriginComponent();
+
+	/** Get the origin point of the projection from the ProjectionOriginComponent (or root actor as a fallback). */
+	FVector GetProjectionOrigin() const;
+
+	/** Determines the appropriate delta rotation needed to move the specified light card to the given position within the view. */
+	FRotator GetLightCardRotationDelta(const FIntPoint& PixelPos, const FSceneView& View, ADisplayClusterLightCardActor& LightCard, EAxisList::Type CurrentAxis,
+		const FVector& DragWidgetOffset);
+
+	/** Determines the appropriate delta in spherical coordinates needed to move the specified light card to the given position within the view. */
+	FSphericalCoordinates GetLightCardTranslationDelta(const FIntPoint& PixelPos, const FSceneView& View, ADisplayClusterLightCardActor& LightCard,
+		EAxisList::Type DragAxis, const FVector& DragWidgetOffset);
+
+	/** Moves specified card to desired coordinates immediately using the current normal maps. Requires valid normal maps. */
+	void InternalMoveLightCardTo(ADisplayClusterLightCardActor& LightCard, const FSphericalCoordinates& Position, bool bIsFinalChange) const;
+
+	/** Moves specified cards to a coordinate in viewport space as if dragged by a widget. Requires valid normal maps. */
+	void InternalDragLightCards(const TArray<TWeakObjectPtr<ADisplayClusterLightCardActor>>& LightCards, const FIntPoint& PixelPos, const FSceneView& View,
+		const FVector& DragWidgetOffset, EAxisList::Type DragAxis, ADisplayClusterLightCardActor* PrimaryLightCard);
+
+	/** Gets the scene view init options to use when rendering the normal map cache */
+	void GetNormalMapSceneViewInitOptions(const FVector& ViewDirection, FSceneViewInitOptions& OutViewInitOptions);
+
+	/** Sets the light card position to the given spherical coordinates */
+	void SetLightCardCoordinates(ADisplayClusterLightCardActor& LightCard, const FSphericalCoordinates& SphericalCoords) const;
+
+	/** Performs a ray trace against the stage's geometry, and returns the hit point. */
+	bool TraceStage(const FVector& RayStart, const FVector& RayEnd, FVector& OutHitLocation);
+
+	/** Traces the world geometry to find the best direction vector from the view origin to a valid point in space using a screen ray. Requires valid normal maps. */
+	FVector TraceScreenRay(const FVector& RayOrigin, const FVector& RayDirection, const FVector& ViewOrigin);
+
+	/** Calculates the final distance from the origin of a light card, given its flush distance and a desired offset */
+	double CalculateFinalLightCardDistance(double FlushDistance, double DesiredOffsetFromFlush = 0.) const;
+
+	/** Invalidates the viewport's normal map, forcing it to be re-rendered before using it next */
+	void InvalidateNormalMap();
+
+	/** Update the normal maps and mesh if necessary. Returns true if there are now valid normal maps to read from. */
+	bool UpdateNormalMaps();
+
+	/** Update the normal map mesh. Only call this once both maps are valid. */
+	void UpdateNormalMapMesh();
+
+	/** Renders the viewport's normal map and stores the texture data to be used later */
+	bool RenderNormalMap(FNormalMap& NormalMap, bool bIsNorthMap);
+
+	/** Called when a world beings cleanup. We destroy our preview scene if the world containing the root actor is cleaned up. */
+	void OnWorldCleanup(UWorld* World, bool bSessionEnded, bool bCleanupResources);
+
+#if WITH_EDITORONLY_DATA
+	/** Called when any actor properties changed, so we can invalidate normals if it is the root actor. */
+	void OnActorPropertyChanged(UObject* ObjectBeingModified, FPropertyChangedEvent& PropertyChangedEvent);
+
+	/** Called when the root actor's blueprint is compiled so we can invalidate normals. */
+	void OnRootActorBlueprintCompiled(UBlueprint* Blueprint);
+#endif
+
+#if WITH_EDITOR
+	/** Send property change events for any properties that may have changed due to a lightcard being moved. */
+	void PostEditChangePropertiesForMovedLightCard(ADisplayClusterLightCardActor& LightCard) const;
+#endif
+
+private:
+#if WITH_EDITOR
+	/** The viewport client for which this is handling projections. This can be null. */
+	TWeakPtr<FEditorViewportClient> ViewportClient;
+#endif
+
+	/** The preview scene containing NormalMapMeshComponent. */
+	TSharedPtr<FPreviewScene> NormalMeshScene;
+
+	/** A morphed ico-sphere mesh component that approximates the normal and depth map. */
+	TWeakObjectPtr<UProceduralMeshComponent> NormalMapMeshComponent;
+
+	/** The cached root actor of the DisplayCluster being controlled. You should generally use UpdateRootActor() instead of accessing this directly. */
+	TWeakObjectPtr<ADisplayClusterRootActor> CachedRootActor;
+
+	/** The level instance root actor of the DisplayCluster being controlled, if RootActor is a proxy. */
+	TWeakObjectPtr<ADisplayClusterRootActor> LevelInstanceRootActor;
+
+	/** The component of the root actor that is acting as the projection origin. Can be either the root component (stage origin) or a view origin component */
+	TWeakObjectPtr<USceneComponent> ProjectionOriginComponent;
+
+	/** The render target used to render a map of the screens' normals for the northern hemisphere of the view */
+	FNormalMap NorthNormalMap;
+
+	/** The render target used to render a map of the screens' normals for the southern hemisphere of the view */
+	FNormalMap SouthNormalMap;
+
+	/** The projection mode of the view this is helping. */
+	EDisplayClusterMeshProjectionType ProjectionMode;
+
+	/** The ID of the preview renderer used for normal map generation, as returned by IDisplayClusterScenePreview. */
+	int32 RendererId;
+
+	/** The radius of the bounding sphere that entirely encapsulates the root actor */
+	float RootActorBoundingRadius = 0.0f;
+
+	/** Whether this created its own preview renderer (and therefore is responsible for destroying it). */
+	bool bCreatedRenderer = false;
+
+	/** If true, this is using an orthographic projection mode. */
+	bool bIsOrthographic = false;
+
+	/** Indicates if the cached normal map is invalid and needs to be redrawn */
+	bool bNormalMapInvalid = false;
+};
diff --git a/Engine/Plugins/Runtime/nDisplay/Source/DisplayClusterScenePreview/Public/IDisplayClusterScenePreview.h b/Engine/Plugins/Runtime/nDisplay/Source/DisplayClusterScenePreview/Public/IDisplayClusterScenePreview.h
--- a/Engine/Plugins/Runtime/nDisplay/Source/DisplayClusterScenePreview/Public/IDisplayClusterScenePreview.h
+++ b/Engine/Plugins/Runtime/nDisplay/Source/DisplayClusterScenePreview/Public/IDisplayClusterScenePreview.h
@@ -85,6 +85,16 @@ public:
 	 */
 	virtual ADisplayClusterRootActor* GetRendererRootActor(int32 RendererId) = 0;
 
+	/**
+	 * Get a list of all actors that have been added to a renderer's scene.
+	 * This includes any actors that the renderer automatically added to itself.
+	 *
+	 * @param RendererId The ID of the renderer as returned from CreateRenderer.
+	 * @param bIncludeRoot If true, include the renderer's root actor.
+	 * @param OutActors The array to fill with actors.
+	 */
+	virtual bool GetActorsInRendererScene(int32 RendererId, bool bIncludeRoot, TArray<AActor*>& OutActors) = 0;
+
 	/**
 	 * Add an actor to the preview scene for a renderer.
 	 * 
@@ -135,6 +145,14 @@ public:
 	 */
 	virtual bool SetRendererRenderSimpleElementsDelegate(int32 RendererId, FDisplayClusterMeshProjectionRenderer::FSimpleElementPass RenderSimpleElementsDelegate) = 0;
 
+	/**
+	 * Set whether the renderer should use post-processed nDisplay preview textures.
+	 *
+	 * @param RendererId The ID of the renderer as returned from CreateRenderer.
+	 * @param bUsePostProcessTexture Whether to use post-processed nDisplay preview texture for future renders.
+	 */
+	virtual bool SetRendererUsePostProcessTexture(int32 RendererId, bool bUsePostProcessTexture) = 0;
+
 	/**
 	 * Immediately render a preview.
 	 * 
@@ -145,7 +163,7 @@ public:
 	virtual bool Render(int32 RendererId, FDisplayClusterMeshProjectionRenderSettings& RenderSettings, FCanvas& Canvas) = 0;
 
 	/**
-	 * Queue a preview to be rendered in the future.
+	 * Queue a preview to be rendered in the future, automatically creating a canvas and render target.
 	 *
 	 * @param RendererId The ID of the renderer as returned from CreateRenderer.
 	 * @param RenderSettings Settings controlling how the scene will be rendered.
@@ -153,4 +171,14 @@ public:
 	 * @param ResultDelegate The delegate to call when the render is complete. It will be passed a FRenderTarget containing the rendered preview.
 	 */
 	virtual bool RenderQueued(int32 RendererId, FDisplayClusterMeshProjectionRenderSettings& RenderSettings, const FIntPoint& Size, FRenderResultDelegate ResultDelegate) = 0;
+
+	/**
+	 * Queue a preview to be rendered in the future using an existing canvas. The canvas must have a valid render target assigned.
+	 *
+	 * @param RendererId The ID of the renderer as returned from CreateRenderer.
+	 * @param RenderSettings Settings controlling how the scene will be rendered.
+	 * @param Canvas The canvas to draw to. If this is invalid when the render is ready to start, the render will be skipped.
+	 * @param ResultDelegate The delegate to call when the render is complete. It will be passed a FRenderTarget containing the rendered preview.
+	 */
+	virtual bool RenderQueued(int32 RendererId, FDisplayClusterMeshProjectionRenderSettings& RenderSettings, TWeakPtr<FCanvas> Canvas, FRenderResultDelegate ResultDelegate) = 0;
 };
\ No newline at end of file
diff --git a/Engine/Plugins/VirtualProduction/EpicStageApp/Source/EpicStageApp/EpicStageApp.build.cs b/Engine/Plugins/VirtualProduction/EpicStageApp/Source/EpicStageApp/EpicStageApp.build.cs
--- a/Engine/Plugins/VirtualProduction/EpicStageApp/Source/EpicStageApp/EpicStageApp.build.cs
+++ b/Engine/Plugins/VirtualProduction/EpicStageApp/Source/EpicStageApp/EpicStageApp.build.cs
@@ -14,7 +14,7 @@ public class EpicStageApp : ModuleRules
 			}
 		);
 
-        PrivateDependencyModuleNames.AddRange(
+		PrivateDependencyModuleNames.AddRange(
 			new string[] {
 				"DeveloperSettings",
 				"DisplayCluster",
@@ -25,9 +25,19 @@ public class EpicStageApp : ModuleRules
 				"Networking",
 				"RemoteControl",
 				"RemoteControlCommon",
+				"RHI",
 				"Sockets",
 				"WebRemoteControl"
 			}
-        );
+		);
+
+		if (Target.Type == TargetType.Editor)
+		{
+			PrivateDependencyModuleNames.AddRange(
+				new string[] {
+					"UnrealEd",
+				}
+			);
+		}
 	}
 }
diff --git a/Engine/Plugins/VirtualProduction/EpicStageApp/Source/EpicStageApp/Private/StageAppRequest.h b/Engine/Plugins/VirtualProduction/EpicStageApp/Source/EpicStageApp/Private/StageAppRequest.h
--- a/Engine/Plugins/VirtualProduction/EpicStageApp/Source/EpicStageApp/Private/StageAppRequest.h
+++ b/Engine/Plugins/VirtualProduction/EpicStageApp/Source/EpicStageApp/Private/StageAppRequest.h
@@ -20,7 +20,57 @@ UENUM()
 enum class ERCWebSocketNDisplayPreviewRenderProjectionType : uint8
 {
 	Perspective,
-	Azimuthal
+	Azimuthal,
+	Orthographic
+};
+
+/** Preview renderer settings exposed to WebSocket clients. */
+USTRUCT()
+struct FRCWebSocketNDisplayPreviewRendererSettings
+{
+	GENERATED_BODY()
+
+	/**
+	 * The type of render to perform.
+	 */
+	UPROPERTY()
+	ERCWebSocketNDisplayPreviewRenderType RenderType = ERCWebSocketNDisplayPreviewRenderType::Color;
+
+	/**
+	 * The type of projection to use.
+	 */
+	UPROPERTY()
+	ERCWebSocketNDisplayPreviewRenderProjectionType ProjectionType = ERCWebSocketNDisplayPreviewRenderProjectionType::Azimuthal;
+
+	/**
+	 * The resolution of the image to render.
+	 */
+	UPROPERTY()
+	FIntPoint Resolution = FIntPoint(1024, 1024);
+
+	/**
+	 * The preview camera's field of view (both horizontal and vertical) in degrees.
+	 */
+	UPROPERTY()
+	float FOV = 130.0f;
+
+	/**
+	 * The preview camera's Euler rotation relative to the camera's actual rotation in the scene.
+	 */
+	UPROPERTY()
+	FRotator Rotation = FRotator(90, 0, 0);
+
+	/**
+	 * The quality of the JPEG to send back to the requesting client, in range 50-100.
+	 */
+	UPROPERTY()
+	int32 JpegQuality = 50;
+
+	/**
+	 * If true, include a list of projected positions within the preview render for each rendered actor when responding to a render request.
+	 */
+	UPROPERTY()
+	bool IncludeActorPositions = false;
 };
 
 /**
@@ -46,6 +96,12 @@ struct FRCWebSocketNDisplayPreviewRendererCreateBody : public FRCRequest
 	 */
 	UPROPERTY()
 	FString RootActorPath = "";
+
+	/**
+	 * Initial settings for the renderer.
+	 */
+	UPROPERTY()
+	FRCWebSocketNDisplayPreviewRendererSettings Settings;
 };
 
 /**
@@ -104,6 +160,27 @@ struct FRCWebSocketNDisplayPreviewRendererDestroyBody : public FRCRequest
 	int32 RendererId = -1;
 };
 
+/**
+ * Change a preview renderer's settings for future renders.
+ */
+USTRUCT()
+struct FRCWebSocketNDisplayPreviewRendererConfigureBody : public FRCRequest
+{
+	GENERATED_BODY()
+		
+	/**
+	 * The ID of the renderer returned when it was created.
+	 */
+	UPROPERTY()
+	int32 RendererId = -1;
+
+	/**
+	 * Settings to use for future renders.
+	 */
+	UPROPERTY()
+	FRCWebSocketNDisplayPreviewRendererSettings Settings;
+};
+
 /**
  * Holds a request made via websocket to render a preview of a display cluster.
  */
@@ -127,40 +204,125 @@ struct FRCWebSocketNDisplayPreviewRenderBody : public FRCRequest
 	 */
 	UPROPERTY()
 	int32 RendererId = -1;
+};
+
+/**
+ * Holds a request made via websocket to start dragging one or more light cards.
+ */
+USTRUCT()
+struct FRCWebSocketNDisplayPreviewLightCardDragBeginBody : public FRCRequest
+{
+	GENERATED_BODY()
+
+	FRCWebSocketNDisplayPreviewLightCardDragBeginBody()
+	{
+		AddStructParameter(ParametersFieldLabel());
+	}
 
 	/**
-	 * The type of render to perform.
+	 * Get the label for the property value struct.
+	 */
+	static FString ParametersFieldLabel() { return TEXT("Parameters"); }
+
+	/**
+	 * The ID of the preview renderer returned when it was created.
 	 */
 	UPROPERTY()
-	ERCWebSocketNDisplayPreviewRenderType RenderType = ERCWebSocketNDisplayPreviewRenderType::Color;
+	int32 RendererId = -1;
 
 	/**
-	 * The type of projection to use.
+	 * Paths of the light cards that will be included in this drag action.
 	 */
 	UPROPERTY()
-	ERCWebSocketNDisplayPreviewRenderProjectionType ProjectionType = ERCWebSocketNDisplayPreviewRenderProjectionType::Azimuthal;
+	TArray<FString> LightCards;
 
 	/**
-	 * The resolution of the image to render.
+	 * The light card to use as the origin point for the drag.
 	 */
 	UPROPERTY()
-	FIntPoint Resolution = FIntPoint(1024, 1024);
+	FString PrimaryLightCard;
 
 	/**
-	 * The preview camera's field of view (both horizontal and vertical) in degrees.
+	 * The sequence number of this change. The highest sequence number received from this client will be
+	 * sent back to the client in future position updates.
 	 */
 	UPROPERTY()
-	float FOV = 130.0f;
+	int64 SequenceNumber = -1;
+};
+
+
+/**
+ * Holds a request made via websocket to move the light cards that are currently being dragged for a preview renderer.
+ */
+USTRUCT()
+struct FRCWebSocketNDisplayPreviewLightCardDragMoveBody : public FRCRequest
+{
+	GENERATED_BODY()
+
+	FRCWebSocketNDisplayPreviewLightCardDragMoveBody()
+	{
+		AddStructParameter(ParametersFieldLabel());
+	}
 
 	/**
-	 * The preview camera's Euler rotation relative to the camera's actual rotation in the scene.
+	 * Get the label for the property value struct.
+	 */
+	static FString ParametersFieldLabel() { return TEXT("Parameters"); }
+
+	/**
+	 * The ID of the preview renderer returned when it was created.
 	 */
 	UPROPERTY()
-	FRotator Rotation = FRotator(90, 0, 0);
+	int32 RendererId = -1;
 
 	/**
-	 * The quality of the JPEG to send back to the requesting client, in range 50-100.
+	 * The position of the drag cursor in coordinates normalized to the size of the preview image.
 	 */
 	UPROPERTY()
-	int32 JpegQuality = 50;
+	FVector2D DragPosition;
+
+	/**
+	 * The sequence number of this change. The highest sequence number received from this client will be
+	 * sent back to the client in future position updates.
+	 */
+	UPROPERTY()
+	int64 SequenceNumber = -1;
+};
+
+/**
+ * Holds a request made via websocket to stop dragging light cards for a preview renderer.
+ */
+USTRUCT()
+struct FRCWebSocketNDisplayPreviewLightCardDragEndBody : public FRCRequest
+{
+	GENERATED_BODY()
+
+	FRCWebSocketNDisplayPreviewLightCardDragEndBody()
+	{
+		AddStructParameter(ParametersFieldLabel());
+	}
+
+	/**
+	 * Get the label for the property value struct.
+	 */
+	static FString ParametersFieldLabel() { return TEXT("Parameters"); }
+
+	/**
+	 * The ID of the preview renderer returned when it was created.
+	 */
+	UPROPERTY()
+	int32 RendererId = -1;
+
+	/**
+	 * The final position of the drag cursor in coordinates normalized to the size of the preview image.
+	 */
+	UPROPERTY()
+	FVector2D DragPosition;
+
+	/**
+	 * The sequence number of this change. The highest sequence number received from this client will be
+	 * sent back to the client in future position updates.
+	 */
+	UPROPERTY()
+	int64 SequenceNumber = -1;
 };
diff --git a/Engine/Plugins/VirtualProduction/EpicStageApp/Source/EpicStageApp/Private/StageAppResponse.h b/Engine/Plugins/VirtualProduction/EpicStageApp/Source/EpicStageApp/Private/StageAppResponse.h
--- a/Engine/Plugins/VirtualProduction/EpicStageApp/Source/EpicStageApp/Private/StageAppResponse.h
+++ b/Engine/Plugins/VirtualProduction/EpicStageApp/Source/EpicStageApp/Private/StageAppResponse.h
@@ -31,6 +31,19 @@ struct FRCPreviewRendererCreatedEvent
 	int32 RendererId = -1;
 };
 
+USTRUCT()
+struct FRCPreviewRenderCompletedEventActorPosition
+{
+	GENERATED_BODY()
+
+	/** The actor's path. */
+	UPROPERTY()
+	FString Path;
+
+	/** The actor's position on the preview texture in coordinates normalized to the size of the preview render's texture. */
+	UPROPERTY()
+	FVector2D Position;
+};
 
 /**
  * Event sent to a client that created a renderer.
@@ -62,4 +75,43 @@ struct FRCPreviewRenderCompletedEvent
 	 */
 	UPROPERTY()
 	FString ImageBase64 = "";
+
+	/**
+	 * A list of projected positions within the preview render for each rendered actor.
+	 */
+	UPROPERTY()
+	TArray<FRCPreviewRenderCompletedEventActorPosition> ActorPositions;
+
+	/**
+	 * The last sequence number that the client sent during a light card drag operation on this renderer.
+	 */
+	UPROPERTY()
+	int64 SequenceNumber = -1;
+};
+
+/**
+ * Event sent to a client when a drag operation is cancelled by the engine.
+ */
+USTRUCT()
+struct FRCLightCardDragCancelled
+{
+	GENERATED_BODY()
+
+	FRCLightCardDragCancelled()
+	: Type(TEXT("LightCardDragCancelled"))
+	{
+	}
+
+	/**
+	 * Type of the event.
+	 */
+	UPROPERTY()
+	FString Type;
+
+	/**
+	 * The ID of the renderer associated with the drag operation.
+	 */
+	UPROPERTY()
+	int32 RendererId = -1;
 };
+
diff --git a/Engine/Plugins/VirtualProduction/EpicStageApp/Source/EpicStageApp/Private/StageAppRouteHandler.cpp b/Engine/Plugins/VirtualProduction/EpicStageApp/Source/EpicStageApp/Private/StageAppRouteHandler.cpp
--- a/Engine/Plugins/VirtualProduction/EpicStageApp/Source/EpicStageApp/Private/StageAppRouteHandler.cpp
+++ b/Engine/Plugins/VirtualProduction/EpicStageApp/Source/EpicStageApp/Private/StageAppRouteHandler.cpp
@@ -5,6 +5,7 @@
 #include "CanvasTypes.h"
 #include "Components/DisplayClusterCameraComponent.h"
 #include "DisplayClusterLightCardActor.h"
+#include "DisplayClusterLightCardEditorHelper.h"
 #include "DisplayClusterRootActor.h"
 #include "Engine/TextureRenderTarget2D.h"
 #include "IDisplayClusterScenePreview.h"
@@ -12,10 +13,166 @@
 #include "Misc/Base64.h"
 #include "Modules/ModuleManager.h"
 #include "RemoteControlWebsocketRoute.h"
-#include "StageAppRequest.h"
 #include "StageAppResponse.h"
 #include "WebRemoteControlUtils.h"
 
+#if WITH_EDITOR
+#include "Editor.h"
+#endif
+
+#define DEBUG_STAGEAPP_NORMALS false
+
+#define LOCTEXT_NAMESPACE "StageAppRouteHandler"
+
+static TAutoConsoleVariable<float> CVarStageAppDragTimeoutCheckInterval(
+	TEXT("StageApp.DragTimeoutCheckInterval"),
+	2.f,
+	TEXT("The interval at which to check whether a drag operation is still active, and if not, time it out.")
+);
+
+namespace StageAppRouteHandlerUtils
+{
+	/**
+	 * Convert from a projection type provided over WebSocket to the values needed internally.
+	 */
+	void GetProjectionSettingsFromWebSocketType(
+		ERCWebSocketNDisplayPreviewRenderProjectionType InProjectionType,
+		EDisplayClusterMeshProjectionType& OutProjectionType,
+		bool& bOutIsOrthographic)
+	{
+		switch (InProjectionType)
+		{
+		case ERCWebSocketNDisplayPreviewRenderProjectionType::Azimuthal:
+			OutProjectionType = EDisplayClusterMeshProjectionType::Azimuthal;
+			bOutIsOrthographic = false;
+			break;
+
+		case ERCWebSocketNDisplayPreviewRenderProjectionType::Perspective:
+			OutProjectionType = EDisplayClusterMeshProjectionType::Linear;
+			bOutIsOrthographic = false;
+			break;
+
+		case ERCWebSocketNDisplayPreviewRenderProjectionType::Orthographic:
+			OutProjectionType = EDisplayClusterMeshProjectionType::Linear;
+			bOutIsOrthographic = true;
+			break;
+
+		default:
+			checkf(false, TEXT("Unknown projection type %d"), InProjectionType);
+		}
+	}
+
+	/**
+	 * Convert from a preview render type provided over WebSocket to the internal value used for the preview renderer.
+	 */
+	EDisplayClusterMeshProjectionOutput GetInternalPreviewRenderType(ERCWebSocketNDisplayPreviewRenderType InRenderType)
+	{
+		switch (InRenderType)
+		{
+		case ERCWebSocketNDisplayPreviewRenderType::Color:
+			return EDisplayClusterMeshProjectionOutput::Color;
+
+		case ERCWebSocketNDisplayPreviewRenderType::Normals:
+			return EDisplayClusterMeshProjectionOutput::Normals;
+
+		default:
+			checkf(false, TEXT("Unknown projection type %d"), InRenderType);
+			return EDisplayClusterMeshProjectionOutput::Color;
+		}
+	}
+}
+
+
+//////////////////////////////////////////////////////////////////////////
+// FPerRendererData
+
+FStageAppRouteHandler::FPerRendererData::FPerRendererData(int32 RendererId)
+	: RendererId(RendererId)
+{
+}
+
+FDisplayClusterLightCardEditorHelper& FStageAppRouteHandler::FPerRendererData::GetLightCardHelper()
+{
+	if (!LightCardHelper.IsValid())
+	{
+		LightCardHelper = MakeShared<FDisplayClusterLightCardEditorHelper>(RendererId);
+		UpdateLightCardHelperSettings();
+	}
+
+	return *LightCardHelper;
+}
+
+const FRCWebSocketNDisplayPreviewRendererSettings& FStageAppRouteHandler::FPerRendererData::GetPreviewSettings() const
+{
+	return PreviewSettings;
+}
+
+void FStageAppRouteHandler::FPerRendererData::SetPreviewSettings(const FRCWebSocketNDisplayPreviewRendererSettings& NewSettings)
+{
+	PreviewSettings = NewSettings;
+	UpdateLightCardHelperSettings();
+}
+
+bool FStageAppRouteHandler::FPerRendererData::GetSceneViewInitOptions(FSceneViewInitOptions& ViewInitOptions)
+{
+	ADisplayClusterRootActor* RootActor = GetRootActor();
+	if (!RootActor)
+	{
+		return false;
+	}
+
+	USceneComponent* ViewOriginComponent = RootActor->GetCommonViewPoint();
+	if (!ViewOriginComponent)
+	{
+		return false;
+	}
+
+	GetLightCardHelper().GetSceneViewInitOptions(
+		ViewInitOptions,
+		PreviewSettings.FOV,
+		PreviewSettings.Resolution,
+		ViewOriginComponent->GetComponentLocation(),
+		FRotator(ViewOriginComponent->GetComponentRotation().Quaternion() * PreviewSettings.Rotation.Quaternion())
+	);
+
+	return true;
+}
+
+void FStageAppRouteHandler::FPerRendererData::UpdateDragSequenceNumber(int64 ReceivedSequenceNumber)
+{
+	if (ReceivedSequenceNumber > SequenceNumber)
+	{
+		SequenceNumber = ReceivedSequenceNumber;
+	}
+}
+
+int64 FStageAppRouteHandler::FPerRendererData::GetDragSequenceNumber() const
+{
+	return SequenceNumber;
+}
+
+ADisplayClusterRootActor* FStageAppRouteHandler::FPerRendererData::GetRootActor() const
+{
+	return IDisplayClusterScenePreview::Get().GetRendererRootActor(RendererId);
+}
+
+void FStageAppRouteHandler::FPerRendererData::UpdateLightCardHelperSettings()
+{
+	if (FDisplayClusterLightCardEditorHelper* LightCardHelperPtr = LightCardHelper.Get())
+	{
+		EDisplayClusterMeshProjectionType ProjectionType;
+		bool bIsOrthographic;
+
+		StageAppRouteHandlerUtils::GetProjectionSettingsFromWebSocketType(PreviewSettings.ProjectionType, ProjectionType, bIsOrthographic);
+
+		LightCardHelperPtr->SetProjectionMode(ProjectionType);
+		LightCardHelperPtr->SetIsOrthographic(bIsOrthographic);
+	}
+}
+
+//////////////////////////////////////////////////////////////////////////
+// FStageAppRouteHandler
+
 void FStageAppRouteHandler::RegisterRoutes(IWebRemoteControlModule& WebRemoteControl)
 {
 	RemoteControlModule = &WebRemoteControl;
@@ -46,6 +203,24 @@ void FStageAppRouteHandler::RegisterRoutes(IWebRemoteControlModule& WebRemoteCon
 		TEXT("ndisplay.preview.render"),
 		FWebSocketMessageDelegate::CreateRaw(this, &FStageAppRouteHandler::HandleWebSocketNDisplayPreviewRender)
 	));
+
+	RegisterRoute(MakeUnique<FRemoteControlWebsocketRoute>(
+		TEXT("Start dragging lightcards relative to a projected preview"),
+		TEXT("ndisplay.preview.lightcard.drag.begin"),
+		FWebSocketMessageDelegate::CreateRaw(this, &FStageAppRouteHandler::HandleWebSocketNDisplayPreviewLightCardDragBegin)
+	));
+
+	RegisterRoute(MakeUnique<FRemoteControlWebsocketRoute>(
+		TEXT("Move the light cards that are currently being dragged"),
+		TEXT("ndisplay.preview.lightcard.drag.move"),
+		FWebSocketMessageDelegate::CreateRaw(this, &FStageAppRouteHandler::HandleWebSocketNDisplayPreviewLightCardDragMove)
+	));
+
+	RegisterRoute(MakeUnique<FRemoteControlWebsocketRoute>(
+		TEXT("Finish dragging the light cards that are currently being dragged"),
+		TEXT("ndisplay.preview.lightcard.drag.end"),
+		FWebSocketMessageDelegate::CreateRaw(this, &FStageAppRouteHandler::HandleWebSocketNDisplayPreviewLightCardDragEnd)
+	));
 }
 
 void FStageAppRouteHandler::UnregisterRoutes(IWebRemoteControlModule& WebRemoteControl)
@@ -57,6 +232,12 @@ void FStageAppRouteHandler::UnregisterRoutes(IWebRemoteControlModule& WebRemoteC
 			WebRemoteControl.UnregisterWebsocketRoute(*Route);
 		}
 	}
+
+	if (DragTimeoutTickerHandle.IsValid())
+	{
+		FTSTicker::GetCoreTicker().RemoveTicker(DragTimeoutTickerHandle);
+		DragTimeoutTickerHandle.Reset();
+	}
 }
 
 void FStageAppRouteHandler::RegisterRoute(TUniquePtr<FRemoteControlWebsocketRoute> Route)
@@ -75,14 +256,20 @@ void FStageAppRouteHandler::HandleWebSocketNDisplayPreviewRendererCreate(const F
 		return;
 	}
 
-	const int32 RendererId = IDisplayClusterScenePreview::Get().CreateRenderer();
-	NDisplayPreviewRendererIdsByClientId.FindOrAdd(WebSocketMessage.ClientId).Add(RendererId);
+	IDisplayClusterScenePreview& PreviewModule = IDisplayClusterScenePreview::Get();
+
+	const int32 RendererId = PreviewModule.CreateRenderer();
+	PerRendererDataMapsByClientId.FindOrAdd(WebSocketMessage.ClientId).Emplace(RendererId, RendererId);
+
+	PreviewModule.SetRendererUsePostProcessTexture(RendererId, true);
 
 	if (!Body.RootActorPath.IsEmpty())
 	{
-		IDisplayClusterScenePreview::Get().SetRendererRootActorPath(RendererId, Body.RootActorPath, true);
+		PreviewModule.SetRendererRootActorPath(RendererId, Body.RootActorPath, true);
 	}
 
+	ChangePreviewRendererSettings(WebSocketMessage.ClientId, RendererId, Body.Settings);
+
 	FRCPreviewRendererCreatedEvent Event;
 	Event.RendererId = RendererId;
 
@@ -103,6 +290,17 @@ void FStageAppRouteHandler::HandleWebSocketNDisplayPreviewRendererSetRoot(const
 	IDisplayClusterScenePreview::Get().SetRendererRootActorPath(Body.RendererId, Body.RootActorPath, true);
 }
 
+void FStageAppRouteHandler::HandleWebSocketNDisplayPreviewRendererConfigure(const FRemoteControlWebSocketMessage& WebSocketMessage)
+{
+	FRCWebSocketNDisplayPreviewRendererConfigureBody Body;
+	if (!WebRemoteControlUtils::DeserializeMessage(WebSocketMessage.RequestPayload, Body))
+	{
+		return;
+	}
+
+	ChangePreviewRendererSettings(WebSocketMessage.ClientId, Body.RendererId, Body.Settings);
+}
+
 void FStageAppRouteHandler::HandleWebSocketNDisplayPreviewRendererDestroy(const FRemoteControlWebSocketMessage& WebSocketMessage)
 {
 	FRCWebSocketNDisplayPreviewRendererDestroyBody Body;
@@ -111,10 +309,10 @@ void FStageAppRouteHandler::HandleWebSocketNDisplayPreviewRendererDestroy(const
 		return;
 	}
 
-	if (TArray<int32>* ClientRendererIds = NDisplayPreviewRendererIdsByClientId.Find(WebSocketMessage.ClientId))
+	if (TClientIdToPerRendererDataMap* PerRendererDataMap = PerRendererDataMapsByClientId.Find(WebSocketMessage.ClientId))
 	{
 		// Check that this client created the renderer
-		if (ClientRendererIds->Remove(Body.RendererId) > 0)
+		if (PerRendererDataMap->Remove(Body.RendererId) > 0)
 		{
 			IDisplayClusterScenePreview::Get().DestroyRenderer(Body.RendererId);
 		}
@@ -134,148 +332,352 @@ void FStageAppRouteHandler::HandleWebSocketNDisplayPreviewRender(const FRemoteCo
 		return;
 	}
 
-	IDisplayClusterScenePreview& PreviewModule = IDisplayClusterScenePreview::Get();
-	ADisplayClusterRootActor* RootActor = PreviewModule.GetRendererRootActor(Body.RendererId);
-	if (RootActor == nullptr)
+	FPerRendererData* PerRendererData = GetClientPerRendererData(WebSocketMessage.ClientId, Body.RendererId);
+	if (!PerRendererData)
 	{
 		return;
 	}
 
-	UWorld* World = RootActor->GetWorld();
-	if (!World)
+	// Set up the render settings
+	FDisplayClusterMeshProjectionRenderSettings RenderSettings;
+	if (!PerRendererData->GetSceneViewInitOptions(RenderSettings.ViewInitOptions))
 	{
 		return;
 	}
 
-	USceneComponent* ViewOriginComponent = RootActor->GetCommonViewPoint();
-	if (!ViewOriginComponent)
+	const FRCWebSocketNDisplayPreviewRendererSettings& PreviewSettings = PerRendererData->GetPreviewSettings();
+	RenderSettings.RenderType = StageAppRouteHandlerUtils::GetInternalPreviewRenderType(PreviewSettings.RenderType);
+	RenderSettings.EngineShowFlags.SetSelectionOutline(false);
+
+	EDisplayClusterMeshProjectionType ProjectionType;
+	bool bIsOrthographic;
+
+	StageAppRouteHandlerUtils::GetProjectionSettingsFromWebSocketType(PreviewSettings.ProjectionType, ProjectionType, bIsOrthographic);
+	EngineShowFlagOrthographicOverride(!bIsOrthographic, RenderSettings.EngineShowFlags);
+
+	FSceneViewInitOptions ViewInitOptions = RenderSettings.ViewInitOptions;
+	const FIntPoint Resolution = PreviewSettings.Resolution;
+	const FGuid ClientId = WebSocketMessage.ClientId;
+	const int32 RendererId = Body.RendererId;
+	const int32 JpegQuality = PreviewSettings.JpegQuality;
+	const bool bIncludeActorPositions = PreviewSettings.IncludeActorPositions;
+
+	// Render the image
+	IDisplayClusterScenePreview::Get().RenderQueued(Body.RendererId, RenderSettings, Resolution, FRenderResultDelegate::CreateLambda(
+		[this, ViewInitOptions, ClientId, RendererId, Resolution, JpegQuality, bIncludeActorPositions](FRenderTarget& RenderTarget)
+		{
+			FPerRendererData* PerRendererData = GetClientPerRendererData(ClientId, RendererId);
+
+#if DEBUG_STAGEAPP_NORMALS
+			// Draw the normal texture in the top-left corner if available
+			if (PerRendererData)
+			{
+				FDisplayClusterLightCardEditorHelper& Helper = PerRendererData->GetLightCardHelper();
+				const FGameTime Time = FGameTime::GetTimeSinceAppStart();
+				const ERHIFeatureLevel::Type FeatureLevel = PerRendererData->GetRootActor()->GetWorld()->Scene->GetFeatureLevel();
+				const FIntPoint HalfResolution = Resolution / 2;
+
+				auto DrawNormalMap = [&](const FIntPoint& Position, bool bShowNorthMap) {
+					if (const UTexture2D* NorthMapTexture = Helper.GetNormalMapTexture(bShowNorthMap))
+					{
+						FCanvas Canvas(&RenderTarget, nullptr, Time, FeatureLevel);
+						Canvas.DrawTile(Position.X, Position.Y, HalfResolution.X, HalfResolution.Y, 0, 0, 1, 1, FColor::White, NorthMapTexture->GetResource());
+						Canvas.Flush_GameThread();
+					}
+				};
+
+				DrawNormalMap(FIntPoint(0, 0), true);
+				DrawNormalMap(FIntPoint(HalfResolution.X, 0), false);
+			}
+#endif
+
+			// Read image data
+			TArray<FColor> PixelData;
+			RenderTarget.ReadPixels(PixelData);
+			FImageView ImageView(PixelData.GetData(), Resolution.X, Resolution.Y, 1, ERawImageFormat::BGRA8, EGammaSpace::Linear);
+
+			// Convert to JPEG
+			TArray64<uint8> JpegData;
+			ImageWrapperModule->CompressImage(JpegData, EImageFormat::JPEG, ImageView, JpegQuality);
+
+			// Encode to Base64
+			FRCPreviewRenderCompletedEvent Event;
+			Event.RendererId = RendererId;
+			Event.ImageBase64 = FBase64::Encode(JpegData.GetData(), JpegData.Num());
+
+			if (PerRendererData)
+			{
+				Event.SequenceNumber = PerRendererData->GetDragSequenceNumber();
+
+				// Add the list of actors, if requested
+				if (bIncludeActorPositions)
+				{
+					TArray<AActor*> Actors;
+					if (IDisplayClusterScenePreview::Get().GetActorsInRendererScene(RendererId, false, Actors))
+					{
+						const FSceneView View(ViewInitOptions);
+
+						for (const AActor* Actor : Actors)
+						{
+							if (!Actor)
+							{
+								continue;
+							}
+
+							FVector WorldPosition;
+							if (const ADisplayClusterLightCardActor* LightCard = Cast<ADisplayClusterLightCardActor>(Actor))
+							{
+								if (LightCard->bIsUVLightCard)
+								{
+									// TODO: add support for UV light cards and projection mode
+									continue;
+								}
+
+								WorldPosition = LightCard->GetLightCardTransform().GetTranslation();
+							}
+							else
+							{
+								WorldPosition = Actor->GetActorLocation();
+							}
+
+							FRCPreviewRenderCompletedEventActorPosition ActorInfo;
+							ActorInfo.Path = Actor->GetPathName();
+
+							PerRendererData->GetLightCardHelper().WorldToPixel(
+								View,
+								WorldPosition,
+								ActorInfo.Position
+							);
+
+							// Normalize to the size of the texture
+							ActorInfo.Position /= Resolution;
+
+							Event.ActorPositions.Add(ActorInfo);
+						}
+					}
+				}
+			}
+
+			// Send data back to the client
+			TArray<uint8> Payload;
+			WebRemoteControlUtils::SerializeMessage(Event, Payload);
+			RemoteControlModule->SendWebsocketMessage(ClientId, Payload);
+		}
+	));
+}
+
+void FStageAppRouteHandler::HandleWebSocketNDisplayPreviewLightCardDragBegin(const FRemoteControlWebSocketMessage& WebSocketMessage)
+{
+	FRCWebSocketNDisplayPreviewLightCardDragBeginBody Body;
+	if (!WebRemoteControlUtils::DeserializeMessage(WebSocketMessage.RequestPayload, Body))
 	{
 		return;
 	}
 
-	// Set up the render
-	FDisplayClusterMeshProjectionRenderSettings RenderSettings;
+	FPerRendererData* PerRendererData = GetClientPerRendererData(WebSocketMessage.ClientId, Body.RendererId);
+	if (!PerRendererData)
+	{
+		return;
+	}
 
-	switch (Body.RenderType)
+	if (!PerRendererData->DraggedLightCards.IsEmpty())
 	{
-	case ERCWebSocketNDisplayPreviewRenderType::Color:
-		RenderSettings.RenderType = EDisplayClusterMeshProjectionOutput::Color;
-		break;
+		// A drag is already in progress
+		return;
+	}
+
+#if WITH_EDITOR
+	GEditor->BeginTransaction(LOCTEXT("DragLightcardsTransaction", "Drag Light Cards with Stage App"));
+#endif
+
+	PerRendererData->UpdateDragSequenceNumber(Body.SequenceNumber);
+	PerRendererData->bHasDragMovedRecently = true;
 
-	case ERCWebSocketNDisplayPreviewRenderType::Normals:
-		RenderSettings.RenderType = EDisplayClusterMeshProjectionOutput::Normals;
-		break;
+	for (const FString& LightCardPath : Body.LightCards)
+	{
+		if (ADisplayClusterLightCardActor* Actor = FindObject<ADisplayClusterLightCardActor>(nullptr, *LightCardPath))
+		{
+			PerRendererData->DraggedLightCards.Emplace(Actor);
+		}
 	}
 
-	switch (Body.ProjectionType)
+	if (!Body.PrimaryLightCard.IsEmpty())
 	{
-	case ERCWebSocketNDisplayPreviewRenderProjectionType::Azimuthal:
-		RenderSettings.ProjectionType = EDisplayClusterMeshProjectionType::Azimuthal;
-		break;
+		PerRendererData->PrimaryLightCard = FindObject<ADisplayClusterLightCardActor>(nullptr, *Body.PrimaryLightCard);
+	}
 
-	case ERCWebSocketNDisplayPreviewRenderProjectionType::Perspective:
-		RenderSettings.ProjectionType = EDisplayClusterMeshProjectionType::Linear;
-		break;
+	if (!DragTimeoutTickerHandle.IsValid())
+	{
+		DragTimeoutTickerHandle = FTSTicker::GetCoreTicker().AddTicker(
+			FTickerDelegate::CreateRaw(this, &FStageAppRouteHandler::TimeOutDrags),
+			CVarStageAppDragTimeoutCheckInterval.GetValueOnGameThread()
+		);
 	}
+}
 
-	RenderSettings.EngineShowFlags.SetSelectionOutline(false);	// Determine FOV multipliers to match render target's aspect ratio
-	float XAxisMultiplier;
-	float YAxisMultiplier;
+void FStageAppRouteHandler::HandleWebSocketNDisplayPreviewLightCardDragMove(const FRemoteControlWebSocketMessage& WebSocketMessage)
+{
+	FRCWebSocketNDisplayPreviewLightCardDragMoveBody Body;
+	if (!WebRemoteControlUtils::DeserializeMessage(WebSocketMessage.RequestPayload, Body))
+	{
+		return;
+	}
 
-	const FIntPoint TargetSize = Body.Resolution;
-	if (TargetSize.X > TargetSize.Y)
+	FPerRendererData* PerRendererData = GetClientPerRendererData(WebSocketMessage.ClientId, Body.RendererId);
+	if (!PerRendererData)
 	{
-		XAxisMultiplier = 1.0f;
-		YAxisMultiplier = TargetSize.X / (float)TargetSize.Y;
+		return;
 	}
-	else
+
+	PerRendererData->UpdateDragSequenceNumber(Body.SequenceNumber);
+	PerRendererData->bHasDragMovedRecently = true;
+
+	DragLightCards(*PerRendererData, Body.DragPosition);
+}
+
+void FStageAppRouteHandler::HandleWebSocketNDisplayPreviewLightCardDragEnd(const FRemoteControlWebSocketMessage& WebSocketMessage)
+{
+	FRCWebSocketNDisplayPreviewLightCardDragEndBody Body;
+	if (!WebRemoteControlUtils::DeserializeMessage(WebSocketMessage.RequestPayload, Body))
 	{
-		XAxisMultiplier = TargetSize.Y / (float)TargetSize.X;
-		YAxisMultiplier = 1.0f;
+		return;
 	}
 
-	const float FOV = Body.FOV * (float)PI / 360.0f;
+	FPerRendererData* PerRendererData = GetClientPerRendererData(WebSocketMessage.ClientId, Body.RendererId);
+	if (!PerRendererData)
+	{
+		return;
+	}
 
-	FSceneViewInitOptions& ViewInitOptions = RenderSettings.ViewInitOptions;
-	ViewInitOptions.BackgroundColor = FLinearColor::Black;
-	ViewInitOptions.SetViewRectangle(FIntRect(0, 0, TargetSize.X, TargetSize.Y));
-	ViewInitOptions.ViewOrigin = ViewOriginComponent->GetComponentLocation();
-	ViewInitOptions.ViewLocation = ViewInitOptions.ViewOrigin;
-	ViewInitOptions.ViewRotation = ViewOriginComponent->GetComponentRotation() + Body.Rotation;
+	PerRendererData->UpdateDragSequenceNumber(Body.SequenceNumber);
+	DragLightCards(*PerRendererData, Body.DragPosition);
+	EndLightCardDrag(*PerRendererData, WebSocketMessage.ClientId, Body.RendererId, true);
+}
 
-	// Rotate view 90 degrees to match FDisplayClusterLightCardEditorViewportClient
-	ViewInitOptions.ViewRotationMatrix = FInverseRotationMatrix(ViewInitOptions.ViewRotation) * FMatrix(
-		FPlane(0, 0, 1, 0),
-		FPlane(1, 0, 0, 0),
-		FPlane(0, 1, 0, 0),
-		FPlane(0, 0, 0, 1)
-	);
+void FStageAppRouteHandler::HandleClientDisconnected(FGuid ClientId)
+{
+	// Destroy the client's renderers
+	if (const TClientIdToPerRendererDataMap* PerRendererDataMap = PerRendererDataMapsByClientId.Find(ClientId))
+	{
+		IDisplayClusterScenePreview& PreviewModule = IDisplayClusterScenePreview::Get();
+		for (const TClientIdToPerRendererDataMap::ElementType& PerRendererDataPair : *PerRendererDataMap)
+		{
+			PreviewModule.DestroyRenderer(PerRendererDataPair.Key);
+		}
+
+		PerRendererDataMapsByClientId.Remove(ClientId);
+	}
+}
 
-	if ((bool)ERHIZBuffer::IsInverted)
+void FStageAppRouteHandler::ChangePreviewRendererSettings(const FGuid& ClientId, int32 RendererId, const FRCWebSocketNDisplayPreviewRendererSettings& NewSettings)
+{
+	FPerRendererData* PerRendererData = GetClientPerRendererData(ClientId, RendererId);
+	if (!PerRendererData)
 	{
-		ViewInitOptions.ProjectionMatrix = FReversedZPerspectiveMatrix(
-			FOV,
-			FOV,
-			XAxisMultiplier,
-			YAxisMultiplier,
-			GNearClippingPlane,
-			GNearClippingPlane
-		);
+		return;
 	}
-	else
+
+	PerRendererData->SetPreviewSettings(NewSettings);
+}
+
+FStageAppRouteHandler::FPerRendererData* FStageAppRouteHandler::GetClientPerRendererData(const FGuid& ClientId, int32 RendererId)
+{
+	TClientIdToPerRendererDataMap* PerRendererDataMap = PerRendererDataMapsByClientId.Find(ClientId);
+	if (!PerRendererDataMap)
 	{
-		ViewInitOptions.ProjectionMatrix = FPerspectiveMatrix(
-			FOV,
-			FOV,
-			XAxisMultiplier,
-			YAxisMultiplier,
-			GNearClippingPlane,
-			GNearClippingPlane
-		);
+		return nullptr;
 	}
 
-	const FIntPoint Resolution = Body.Resolution;
-	const FGuid ClientId = WebSocketMessage.ClientId;
-	const int32 RendererId = Body.RendererId;
-	const int32 JpegQuality = Body.JpegQuality;
+	return PerRendererDataMap->Find(RendererId);
+}
 
-	// Render the image
-	PreviewModule.RenderQueued(Body.RendererId, RenderSettings, Resolution, FRenderResultDelegate::CreateLambda(
-		[this, ClientId, RendererId, Resolution, JpegQuality](FRenderTarget& RenderTarget)
-		{
-			// Save the image
-			TArray<FColor> PixelData;
-			RenderTarget.ReadPixels(PixelData);
-			FImageView ImageView(PixelData.GetData(), Resolution.X, Resolution.Y, ERawImageFormat::BGRA8);
+void FStageAppRouteHandler::DragLightCards(FPerRendererData& PerRendererData, FVector2D DragPosition)
+{
+	// Set up the scene view relative to which we're dragging
+	FSceneViewInitOptions ViewInitOptions;
+	if (!PerRendererData.GetSceneViewInitOptions(ViewInitOptions))
+	{
+		return;
+	}
 
-			// Convert to JPEG
-			TArray64<uint8> JpegData;
-			ImageWrapperModule->CompressImage(JpegData, EImageFormat::JPEG, ImageView, JpegQuality);
+	PerRendererData.LastDragPosition = DragPosition;
 
-			// Encode to Base64
-			FRCPreviewRenderCompletedEvent Event;
-			Event.RendererId = RendererId;
-			Event.ImageBase64 = FBase64::Encode(JpegData.GetData(), JpegData.Num());
+	const FSceneView View(ViewInitOptions);
 
-			// Send it back to the client
-			TArray<uint8> Payload;
-			WebRemoteControlUtils::SerializeMessage(Event, Payload);
-			RemoteControlModule->SendWebsocketMessage(ClientId, Payload);
-		}
-	));
+	// Determine the screen pixel position to which we're dragging
+	const FRCWebSocketNDisplayPreviewRendererSettings& PreviewSettings = PerRendererData.GetPreviewSettings();
+	const FIntPoint& Resolution = PreviewSettings.Resolution;
+
+	const FIntPoint PixelPos(
+		FMath::RoundToInt(DragPosition.X * Resolution.X),
+		FMath::RoundToInt(DragPosition.Y * Resolution.Y)
+	);
+	const FVector DragWidgetOffset = FVector::ZeroVector;
+	
+	FDisplayClusterLightCardEditorHelper& LightCardHelper = PerRendererData.GetLightCardHelper();
+	LightCardHelper.DragLightCards(
+		PerRendererData.DraggedLightCards,
+		PixelPos,
+		View,
+		DragWidgetOffset,
+		EAxisList::Type::XYZ,
+		PerRendererData.PrimaryLightCard.Get()
+	);
 }
 
-void FStageAppRouteHandler::HandleClientDisconnected(FGuid ClientId)
+bool FStageAppRouteHandler::TimeOutDrags(float DeltaTime)
 {
-	/** Destroy the client's renderers. */
-	if (const TArray<int32>* RendererIds = NDisplayPreviewRendererIdsByClientId.Find(ClientId))
+	bool bAnyDragsInProgress = false;
+	for (TPair<FGuid, TClientIdToPerRendererDataMap>& ClientPair : PerRendererDataMapsByClientId)
 	{
-		IDisplayClusterScenePreview& PreviewModule = IDisplayClusterScenePreview::Get();
-		for (const int32 RendererId : *RendererIds)
+		for (TClientIdToPerRendererDataMap::ElementType& PerRendererDataPair : ClientPair.Value)
 		{
-			PreviewModule.DestroyRenderer(RendererId);
+			FPerRendererData& PerRendererData = PerRendererDataPair.Value;
+			if (PerRendererData.DraggedLightCards.IsEmpty())
+			{
+				continue;
+			}
+			
+			if (!PerRendererData.bHasDragMovedRecently)
+			{
+				EndLightCardDrag(PerRendererData, ClientPair.Key, PerRendererDataPair.Key, false);
+				continue;
+			}
+
+			PerRendererData.bHasDragMovedRecently = false;
+			bAnyDragsInProgress = true;
 		}
-		NDisplayPreviewRendererIdsByClientId.Remove(ClientId);
 	}
 
+	if (!bAnyDragsInProgress)
+	{
+		DragTimeoutTickerHandle.Reset();
+		return false;
+	}
+
+	return true;
 }
+
+void FStageAppRouteHandler::EndLightCardDrag(FPerRendererData& PerRendererData, const FGuid& ClientId, int32 RendererId, bool bEndedByClient)
+{
+	PerRendererData.DraggedLightCards.Empty();
+	PerRendererData.bHasDragMovedRecently = false;
+
+#if WITH_EDITOR
+	GEditor->EndTransaction();
+#endif
+
+	if (!bEndedByClient)
+	{
+		// Notify the client that its drag in progress was cancelled
+		FRCLightCardDragCancelled Event;
+		Event.RendererId = RendererId;
+
+		TArray<uint8> Payload;
+		WebRemoteControlUtils::SerializeMessage(Event, Payload);
+
+		RemoteControlModule->SendWebsocketMessage(ClientId, Payload);
+	}
+}
+
+#undef LOCTEXT_NAMESPACE
diff --git a/Engine/Plugins/VirtualProduction/EpicStageApp/Source/EpicStageApp/Private/StageAppRouteHandler.h b/Engine/Plugins/VirtualProduction/EpicStageApp/Source/EpicStageApp/Private/StageAppRouteHandler.h
--- a/Engine/Plugins/VirtualProduction/EpicStageApp/Source/EpicStageApp/Private/StageAppRouteHandler.h
+++ b/Engine/Plugins/VirtualProduction/EpicStageApp/Source/EpicStageApp/Private/StageAppRouteHandler.h
@@ -2,13 +2,81 @@
 
 #pragma once
 
+#include "Containers/Ticker.h"
 #include "CoreMinimal.h"
+#include "DisplayClusterLightCardEditorHelper.h"
 #include "UObject/ObjectMacros.h"
 #include "IWebRemoteControlModule.h"
 #include "RemoteControlWebsocketRoute.h"
+#include "StageAppRequest.h"
 
+class ADisplayClusterRootActor;
+class ADisplayClusterLightCardActor;
+
+/** Handler for Epic Stage App websocket requests. */
 class FStageAppRouteHandler
 {
+private:
+	/** Data used for editing lightcard positions within the context of a preview. */
+	struct FPerRendererData
+	{
+		FPerRendererData(int RendererId);
+
+		/** The list of light cards currently being dragged by this client. Empty if a drag is not in progress. */
+		TArray<TWeakObjectPtr<ADisplayClusterLightCardActor>> DraggedLightCards;
+
+		/** The light card to use as the origin point for the drag operation. */
+		TWeakObjectPtr<ADisplayClusterLightCardActor> PrimaryLightCard;
+
+		/** The last position sent by a client during a drag operation. */
+		FVector2D LastDragPosition;
+
+		/** If true, the dragged lightcards have moved since the last timeout check. */
+		bool bHasDragMovedRecently;
+
+	public:
+		/** Get or create a FDisplayClusterLightCardEditorHelper for this client preview renderer. */
+		FDisplayClusterLightCardEditorHelper& GetLightCardHelper();
+
+		/** Get the preview settings to use for renders. */
+		const FRCWebSocketNDisplayPreviewRendererSettings& GetPreviewSettings() const;
+
+		/** Change the preview settings to use for renders. */
+		void SetPreviewSettings(const FRCWebSocketNDisplayPreviewRendererSettings& NewSettings);
+
+		/** Get the scene view init options to be used for preview renders based on the stored settings. */
+		bool GetSceneViewInitOptions(FSceneViewInitOptions& ViewInitOptions);
+
+		/** Update the sequence number if the provided one is higher than the previous. */
+		void UpdateDragSequenceNumber(int64 ReceivedSequenceNumber);
+
+		/** Get the highest sequence number received from the client during a drag operation. */
+		int64 GetDragSequenceNumber() const;
+
+		/** Get the root actor used for previews. */
+		ADisplayClusterRootActor* GetRootActor() const;
+
+	private:
+		/** Update LightCardHelper's settings based on the ones stored in PreviewSettings. */
+		void UpdateLightCardHelperSettings();
+
+	private:
+		/** Helper used to move light cards and project their coordinates. */
+		TSharedPtr<FDisplayClusterLightCardEditorHelper> LightCardHelper;
+
+		/** Settings used for generating previews. */
+		FRCWebSocketNDisplayPreviewRendererSettings PreviewSettings;
+
+		/** The ID of the renderer with which this is associated. */
+		int RendererId;
+
+		/** Highest sequence number received from the client during a drag operation. */
+		int64 SequenceNumber = 0;
+	};
+
+	/** Map from renderer IDs created by a client (keyed by ID) to per-renderer data used by the handler. */
+	using TClientIdToPerRendererDataMap = TMap<int32, FPerRendererData>;
+
 public:
 	/** Registers the WebSocket routes with the provided WebRemoteControl module. */
 	void RegisterRoutes(IWebRemoteControlModule& WebRemoteControl);
@@ -26,15 +94,42 @@ private:
 	/** Handles setting the root actor of a renderer. */
 	void HandleWebSocketNDisplayPreviewRendererSetRoot(const FRemoteControlWebSocketMessage& WebSocketMessage);
 
+	/** Handles changing renderer settings. */
+	void HandleWebSocketNDisplayPreviewRendererConfigure(const FRemoteControlWebSocketMessage& WebSocketMessage);
+
 	/** Handles destroying a preview renderer. */
 	void HandleWebSocketNDisplayPreviewRendererDestroy(const FRemoteControlWebSocketMessage& WebSocketMessage);
 
 	/** Handles rendering a preview. */
 	void HandleWebSocketNDisplayPreviewRender(const FRemoteControlWebSocketMessage& WebSocketMessage);
 
+	/** Handles beginning a light card drag operation. */
+	void HandleWebSocketNDisplayPreviewLightCardDragBegin(const FRemoteControlWebSocketMessage& WebSocketMessage);
+
+	/** Handles moving light cards with a drag operation. */
+	void HandleWebSocketNDisplayPreviewLightCardDragMove(const FRemoteControlWebSocketMessage& WebSocketMessage);
+
+	/** Handles ending a light card drag operation. */
+	void HandleWebSocketNDisplayPreviewLightCardDragEnd(const FRemoteControlWebSocketMessage& WebSocketMessage);
+
 	/** Called when a client disconnects from the WebSocket server. */
 	void HandleClientDisconnected(FGuid ClientId);
 
+	/** Change the settings for a client's preview renderer. */
+	void ChangePreviewRendererSettings(const FGuid& ClientId, int32 RendererId, const FRCWebSocketNDisplayPreviewRendererSettings& NewSettings);
+
+	/** Get the nDisplay preview rendering data for the given client and renderer ID, or null if it's an invalid client/renderer ID. */
+	FPerRendererData* GetClientPerRendererData(const FGuid& ClientId, int32 RendererId);
+
+	/** Drag light cards to the specified position. */
+	void DragLightCards(FPerRendererData& PreviewData, FVector2D DragPosition);
+
+	/** Check if any light card drag operations have timed out, and if so, end them. */
+	bool TimeOutDrags(float DeltaTime);
+
+	/** End a client's light card drag operation. */
+	void EndLightCardDrag(FPerRendererData& PreviewData, const FGuid& ClientId, int32 RendererId, bool bEndedByClient);
+
 private:
 	/** The module for the WebSocket server for which this is handling routes. */
 	IWebRemoteControlModule* RemoteControlModule = nullptr;
@@ -42,9 +137,12 @@ private:
 	/** Routes this has registered with the WebSocket server. */
 	TArray<TUniquePtr<FRemoteControlWebsocketRoute>> Routes;
 
-	/** nDisplay preview renderers belonging to each client. */
-	TMap<FGuid, TArray<int32>> NDisplayPreviewRendererIdsByClientId;
+	/** Map from client GUID to preview renderer data maps for that client. */
+	TMap<FGuid, TClientIdToPerRendererDataMap> PerRendererDataMapsByClientId;
 
 	/** Image wrapper module used to compress preview images. */
 	class IImageWrapperModule* ImageWrapperModule = nullptr;
+
+	/** Handle for timing out drags. */
+	FTSTicker::FDelegateHandle DragTimeoutTickerHandle;
 };
