Fix non-unity
#rnx

diff --git a/Engine/Plugins/Runtime/nDisplay/Source/DisplayClusterScenePreview/Private/DisplayClusterLightCardEditorHelper.cpp b/Engine/Plugins/Runtime/nDisplay/Source/DisplayClusterScenePreview/Private/DisplayClusterLightCardEditorHelper.cpp
--- a/Engine/Plugins/Runtime/nDisplay/Source/DisplayClusterScenePreview/Private/DisplayClusterLightCardEditorHelper.cpp
+++ b/Engine/Plugins/Runtime/nDisplay/Source/DisplayClusterScenePreview/Private/DisplayClusterLightCardEditorHelper.cpp
@@ -6,6 +6,7 @@
 #include "Containers/ArrayView.h"
 #include "DisplayClusterRootActor.h"
 #include "DisplayClusterLightCardActor.h"
+#include "Engine/StaticMesh.h"
 #include "Engine/Texture2D.h"
 #include "IDisplayClusterScenePreview.h"
 #include "ImageUtils.h"
