nDisplay: Add Place Actors menu.
- Supports adding new or existing light cards.
- Supports viewing templates and favorites.
- Supports selecting recently placed items.

Add light card editor preference settings.
- Recent items persist here.
- The viewport projection settings are now saved.

Fix template favorites not saving.

diff --git a/Engine/Plugins/Runtime/nDisplay/Source/DisplayClusterLightCardEditor/Private/LightCardTemplates/DisplayClusterLightCardTemplate.h b/Engine/Plugins/Runtime/nDisplay/Source/DisplayClusterLightCardEditor/Private/LightCardTemplates/DisplayClusterLightCardTemplate.h
--- a/Engine/Plugins/Runtime/nDisplay/Source/DisplayClusterLightCardEditor/Private/LightCardTemplates/DisplayClusterLightCardTemplate.h
+++ b/Engine/Plugins/Runtime/nDisplay/Source/DisplayClusterLightCardEditor/Private/LightCardTemplates/DisplayClusterLightCardTemplate.h
@@ -9,7 +9,7 @@
 /**
  * A template asset to store appearance settings from Light Card actors.
  */
-UCLASS(NotBlueprintType, NotBlueprintable, NotPlaceable, config=EditorPerProjectUserSettings)
+UCLASS(NotBlueprintType, NotBlueprintable, NotPlaceable, PerObjectConfig, config=EditorPerProjectUserSettings)
 class DISPLAYCLUSTERLIGHTCARDEDITOR_API UDisplayClusterLightCardTemplate : public UObject
 {
 	GENERATED_BODY()
diff --git a/Engine/Plugins/Runtime/nDisplay/Source/DisplayClusterLightCardEditor/Private/LightCardTemplates/DisplayClusterLightCardTemplateHelpers.cpp b/Engine/Plugins/Runtime/nDisplay/Source/DisplayClusterLightCardEditor/Private/LightCardTemplates/DisplayClusterLightCardTemplateHelpers.cpp
new file mode 100644
--- /dev/null
+++ b/Engine/Plugins/Runtime/nDisplay/Source/DisplayClusterLightCardEditor/Private/LightCardTemplates/DisplayClusterLightCardTemplateHelpers.cpp
@@ -0,0 +1,29 @@
+// Copyright Epic Games, Inc. All Rights Reserved.
+
+#include "DisplayClusterLightCardTemplateHelpers.h"
+
+#include "DisplayClusterLightCardTemplate.h"
+#include "AssetRegistry/AssetRegistryModule.h"
+#include "AssetRegistry/IAssetRegistry.h"
+
+TArray<UDisplayClusterLightCardTemplate*> UE::DisplayClusterLightCardTemplateHelpers::GetLightCardTemplates(bool bFavoritesOnly)
+{
+	TArray<UDisplayClusterLightCardTemplate*> OutLightCardTemplates;
+	
+	const IAssetRegistry& AssetRegistry = FModuleManager::GetModuleChecked<FAssetRegistryModule>(AssetRegistryConstants::ModuleName).Get();
+	
+	TArray<FAssetData> TemplateAssets;
+	const FTopLevelAssetPath ClassPath = UDisplayClusterLightCardTemplate::StaticClass()->GetClassPathName();
+	AssetRegistry.GetAssetsByClass(ClassPath, TemplateAssets, true);
+
+	for (const FAssetData& AssetData : TemplateAssets)
+	{
+		UDisplayClusterLightCardTemplate* AssetObject = CastChecked<UDisplayClusterLightCardTemplate>(AssetData.GetAsset());
+		if (!bFavoritesOnly || AssetObject->bIsFavorite)
+		{
+			OutLightCardTemplates.Add(AssetObject);
+		}
+	}
+
+	return OutLightCardTemplates;
+}
diff --git a/Engine/Plugins/Runtime/nDisplay/Source/DisplayClusterLightCardEditor/Private/LightCardTemplates/DisplayClusterLightCardTemplateHelpers.h b/Engine/Plugins/Runtime/nDisplay/Source/DisplayClusterLightCardEditor/Private/LightCardTemplates/DisplayClusterLightCardTemplateHelpers.h
new file mode 100644
--- /dev/null
+++ b/Engine/Plugins/Runtime/nDisplay/Source/DisplayClusterLightCardEditor/Private/LightCardTemplates/DisplayClusterLightCardTemplateHelpers.h
@@ -0,0 +1,18 @@
+// Copyright Epic Games, Inc. All Rights Reserved.
+
+#pragma once
+
+#include "CoreMinimal.h"
+
+class UDisplayClusterLightCardTemplate;
+
+namespace UE::DisplayClusterLightCardTemplateHelpers
+{
+	/**
+	 * Load and return all light card templates
+	 * 
+	 * @param bFavoritesOnly only return light card templates that are marked as a user favorite
+	 * @return the light card template objects, loaded from their asset
+	 */
+	TArray<UDisplayClusterLightCardTemplate*> GetLightCardTemplates(bool bFavoritesOnly = false);
+}
\ No newline at end of file
diff --git a/Engine/Plugins/Runtime/nDisplay/Source/DisplayClusterLightCardEditor/Private/LightCardTemplates/SDisplayClusterLightCardTemplateList.cpp b/Engine/Plugins/Runtime/nDisplay/Source/DisplayClusterLightCardEditor/Private/LightCardTemplates/SDisplayClusterLightCardTemplateList.cpp
--- a/Engine/Plugins/Runtime/nDisplay/Source/DisplayClusterLightCardEditor/Private/LightCardTemplates/SDisplayClusterLightCardTemplateList.cpp
+++ b/Engine/Plugins/Runtime/nDisplay/Source/DisplayClusterLightCardEditor/Private/LightCardTemplates/SDisplayClusterLightCardTemplateList.cpp
@@ -4,6 +4,7 @@
 
 #include "DisplayClusterLightCardTemplate.h"
 #include "DisplayClusterLightCardTemplateDragDropOp.h"
+#include "DisplayClusterLightCardTemplateHelpers.h"
 #include "SDisplayClusterLightCardEditor.h"
 
 #include "AssetRegistry/AssetRegistryModule.h"
@@ -29,7 +30,9 @@ class SLightCardTemplateTreeItemRow : public SMultiColumnTableRow<TSharedPtr<SDi
 
 public:
 
-	void Construct(const FArguments& InArgs, const TSharedRef<STableViewBase>& InOwnerTable, const TSharedPtr<SDisplayClusterLightCardTemplateList>& InTemplateList, const TSharedPtr<SDisplayClusterLightCardTemplateList::FLightCardTemplateTreeItem> InLightCardTreeItem)
+	void Construct(const FArguments& InArgs, const TSharedRef<STableViewBase>& InOwnerTable,
+		const TSharedPtr<SDisplayClusterLightCardTemplateList>& InTemplateList,
+		const TSharedPtr<SDisplayClusterLightCardTemplateList::FLightCardTemplateTreeItem> InLightCardTreeItem)
 	{
 		LightCardTreeItem = InLightCardTreeItem;
 		OwningTemplateListPtr = InTemplateList;
@@ -123,7 +126,8 @@ private:
 	
 	FReply HandleDragDetected(const FGeometry& MyGeometry, const FPointerEvent& MouseEvent)
 	{
-		if (LightCardTreeItem.IsValid() && MouseEvent.IsMouseButtonDown(EKeys::LeftMouseButton))
+		if (OwningTemplateListPtr.IsValid() && OwningTemplateListPtr.Pin()->IsDragDropEnabled() &&
+			LightCardTreeItem.IsValid() && MouseEvent.IsMouseButtonDown(EKeys::LeftMouseButton))
 		{
 			const TSharedPtr<FDragDropOperation> DragDropOp = FDisplayClusterLightCardTemplateDragDropOp::New(LightCardTreeItem.Pin()->LightCardTemplate);
 
@@ -148,8 +152,8 @@ private:
 		if (LightCardTreeItem.IsValid() && LightCardTreeItem.Pin()->LightCardTemplate.IsValid())
 		{
 			LightCardTreeItem.Pin()->LightCardTemplate->bIsFavorite = !LightCardTreeItem.Pin()->LightCardTemplate->bIsFavorite;
-			LightCardTreeItem.Pin()->LightCardTemplate->SaveConfig();
 			LightCardTreeItem.Pin()->LightCardTemplate->PostEditChange();
+			LightCardTreeItem.Pin()->LightCardTemplate->SaveConfig();
 			
 			if (OwningTemplateListPtr.IsValid())
 			{
@@ -182,6 +186,8 @@ void SDisplayClusterLightCardTemplateList::Construct(const FArguments& InArgs, T
 {
 	LightCardEditorPtr = InLightCardEditor;
 
+	bSpawnOnSelection = InArgs._SpawnOnSelection;
+
 	IAssetRegistry& AssetRegistry = FModuleManager::LoadModuleChecked<FAssetRegistryModule>(AssetRegistryConstants::ModuleName).Get();
 	AssetRegistry.OnAssetAdded().AddRaw(this, &SDisplayClusterLightCardTemplateList::OnAssetAddedOrRemoved);
 	AssetRegistry.OnAssetRemoved().AddRaw(this, &SDisplayClusterLightCardTemplateList::OnAssetAddedOrRemoved);
@@ -191,80 +197,82 @@ void SDisplayClusterLightCardTemplateList::Construct(const FArguments& InArgs, T
 
 	ChildSlot
 	[
-		SNew(SVerticalBox)
-		+SVerticalBox::Slot()
-		.Padding(0.f, 8.f, 0.f, 5.f)
-		.AutoHeight()
+		SNew(SBox)
+		.MinDesiredWidth(260.f)
+		.Padding(FMargin(4.f, 0.f))
 		[
-			SNew(SHorizontalBox)
-			+SHorizontalBox::Slot()
-			.Padding(6.f, 0.f)
-			.AutoWidth()
-			.VAlign(VAlign_Center)
-			.HAlign(HAlign_Center)
+			SNew(SVerticalBox)
+			+SVerticalBox::Slot()
+			.Padding(0.f, 8.f, 0.f, 5.f)
+			.AutoHeight()
 			[
-				SNew(SImage)
-				.Image(this, &SDisplayClusterLightCardTemplateList::GetFavoriteFilterIcon)
-				.ToolTipText(LOCTEXT("FavoriteFilterToolTip", "Filter by favorites"))
-				.OnMouseButtonDown_Lambda([this](const FGeometry& MyGeometry, const FPointerEvent& MouseEvent)
-				{
-					bFilterFavorites = !bFilterFavorites;
-					ApplyFilter();
-					return FReply::Handled();
-				})
+				SNew(SHorizontalBox)
+				+SHorizontalBox::Slot()
+				.Padding(6.f, 0.f)
+				.AutoWidth()
+				.VAlign(VAlign_Center)
+				.HAlign(HAlign_Center)
+				[
+					SNew(SImage)
+					.Image(this, &SDisplayClusterLightCardTemplateList::GetFavoriteFilterIcon)
+					.ToolTipText(LOCTEXT("FavoriteFilterToolTip", "Filter by favorites"))
+					.OnMouseButtonDown_Lambda([this](const FGeometry& MyGeometry, const FPointerEvent& MouseEvent)
+					{
+						bFilterFavorites = !bFilterFavorites;
+						ApplyFilter();
+						return FReply::Handled();
+					})
+				]
+				+SHorizontalBox::Slot()
+				.FillWidth(1.f)
+				[
+					SNew(SSearchBox)
+					.SelectAllTextWhenFocused(true)
+					.OnTextChanged(this, &SDisplayClusterLightCardTemplateList::OnFilterTextChanged)
+					.HintText(LOCTEXT("SearchBoxHint", "Search Light Card Templates..."))
+				]
 			]
-			+SHorizontalBox::Slot()
-			.FillWidth(1.f)
+			+SVerticalBox::Slot()
 			[
-				SNew(SSearchBox)
-				.SelectAllTextWhenFocused(true)
-				.OnTextChanged(this, &SDisplayClusterLightCardTemplateList::OnFilterTextChanged)
-				.HintText(LOCTEXT("SearchBoxHint", "Search Light Card Templates..."))
+				SAssignNew(LightCardTemplateTreeView, STreeView<TSharedPtr<FLightCardTemplateTreeItem>>)
+				.TreeItemsSource(&FilteredLightCardTemplateTree)
+				.ItemHeight(28)
+				.SelectionMode(ESelectionMode::Single)
+				.OnGenerateRow(this, &SDisplayClusterLightCardTemplateList::GenerateTreeItemRow)
+				.OnGetChildren(this, &SDisplayClusterLightCardTemplateList::GetChildrenForTreeItem)
+				.OnSelectionChanged(this, &SDisplayClusterLightCardTemplateList::OnSelectionChanged)
+				.HeaderRow
+				(
+					SNew(SHeaderRow)
+					.Visibility(InArgs._HideHeader ? EVisibility::Collapsed : EVisibility::Visible)
+					+SHeaderRow::Column(DisplayClusterLightCardTemplateListColumnNames::LightCardTemplateName)
+					.DefaultLabel(LOCTEXT("LightCardTemplateName", "Light Card Templates"))
+					.FillWidth(0.8f)
+				)
 			]
 		]
-		+SVerticalBox::Slot()
-		[
-			SAssignNew(LightCardTemplateTreeView, STreeView<TSharedPtr<FLightCardTemplateTreeItem>>)
-			.TreeItemsSource(&FilteredLightCardTemplateTree)
-			.ItemHeight(28)
-			.SelectionMode(ESelectionMode::Single)
-			.OnGenerateRow(this, &SDisplayClusterLightCardTemplateList::GenerateTreeItemRow)
-			.OnGetChildren(this, &SDisplayClusterLightCardTemplateList::GetChildrenForTreeItem)
-			.HeaderRow
-			(
-				SNew(SHeaderRow)
-
-				+ SHeaderRow::Column(DisplayClusterLightCardTemplateListColumnNames::LightCardTemplateName)
-				.DefaultLabel(LOCTEXT("LightCardTemplateName", "Light Card Templates"))
-				.FillWidth(0.8f)
-			)
-		]
 	];
 }
 
 void SDisplayClusterLightCardTemplateList::RefreshTemplateList()
 {
 	LightCardTemplateTree.Reset();
-	
-	const IAssetRegistry& AssetRegistry = FModuleManager::GetModuleChecked<FAssetRegistryModule>(AssetRegistryConstants::ModuleName).Get();
-	
-	TArray<FAssetData> TemplateAssets;
-	const FTopLevelAssetPath ClassPath = UDisplayClusterLightCardTemplate::StaticClass()->GetClassPathName();
-	AssetRegistry.GetAssetsByClass(ClassPath, TemplateAssets, true);
 
-	for (const FAssetData& AssetData : TemplateAssets)
+	TArray<UDisplayClusterLightCardTemplate*> LightCardTemplates =
+		UE::DisplayClusterLightCardTemplateHelpers::GetLightCardTemplates();
+	
+	for (UDisplayClusterLightCardTemplate* Template : LightCardTemplates)
 	{
 		TSharedPtr<FLightCardTemplateTreeItem> TemplateTreeItem = MakeShared<FLightCardTemplateTreeItem>();
-		UDisplayClusterLightCardTemplate* AssetObject = CastChecked<UDisplayClusterLightCardTemplate>(AssetData.GetAsset());
 		
-		TemplateTreeItem->TemplateName = AssetData.AssetName;
-		TemplateTreeItem->LightCardTemplate = AssetObject;
+		TemplateTreeItem->TemplateName = Template->GetFName();
+		TemplateTreeItem->LightCardTemplate = Template;
 
-		if (AssetObject->LightCardActor && AssetObject->LightCardActor->Texture.Get())
+		if (Template->LightCardActor && Template->LightCardActor->Texture.Get())
 		{
 			TemplateTreeItem->SlateBrush = MakeShared<FSlateBrush>();
 		
-			TemplateTreeItem->SlateBrush->SetResourceObject(AssetObject->LightCardActor->Texture.Get());
+			TemplateTreeItem->SlateBrush->SetResourceObject(Template->LightCardActor->Texture.Get());
 			TemplateTreeItem->SlateBrush->ImageSize = FVector2D(64.f, 64.f);
 		}
 		
@@ -301,6 +309,11 @@ void SDisplayClusterLightCardTemplateList::ApplyFilter()
 	ApplyFilter(FilterText, bFilterFavorites);
 }
 
+bool SDisplayClusterLightCardTemplateList::IsDragDropEnabled() const
+{
+	return !bSpawnOnSelection;
+}
+
 TSharedRef<ITableRow> SDisplayClusterLightCardTemplateList::GenerateTreeItemRow(
 	TSharedPtr<FLightCardTemplateTreeItem> Item, const TSharedRef<STableViewBase>& OwnerTable)
 {
@@ -318,6 +331,16 @@ const FSlateBrush* SDisplayClusterLightCardTemplateList::GetFavoriteFilterIcon()
 	return FSlateIcon(FAppStyle::Get().GetStyleSetName(), FavoriteIconName).GetIcon();
 }
 
+void SDisplayClusterLightCardTemplateList::OnSelectionChanged(TSharedPtr<FLightCardTemplateTreeItem> InItem, ESelectInfo::Type SelectInfo)
+{
+	if (bSpawnOnSelection && (SelectInfo == ESelectInfo::OnKeyPress || SelectInfo == ESelectInfo::OnMouseClick) && InItem.IsValid() && LightCardEditorPtr.IsValid())
+	{
+		LightCardEditorPtr.Pin()->SpawnLightCardFromTemplate(InItem->LightCardTemplate.Get());
+		// Close the menu containing this widget
+		FSlateApplication::Get().DismissAllMenus();
+	}
+}
+
 void SDisplayClusterLightCardTemplateList::OnFilterTextChanged(const FText& SearchText)
 {
 	ApplyFilter(SearchText.ToString(), bFilterFavorites);
diff --git a/Engine/Plugins/Runtime/nDisplay/Source/DisplayClusterLightCardEditor/Private/LightCardTemplates/SDisplayClusterLightCardTemplateList.h b/Engine/Plugins/Runtime/nDisplay/Source/DisplayClusterLightCardEditor/Private/LightCardTemplates/SDisplayClusterLightCardTemplateList.h
--- a/Engine/Plugins/Runtime/nDisplay/Source/DisplayClusterLightCardEditor/Private/LightCardTemplates/SDisplayClusterLightCardTemplateList.h
+++ b/Engine/Plugins/Runtime/nDisplay/Source/DisplayClusterLightCardEditor/Private/LightCardTemplates/SDisplayClusterLightCardTemplateList.h
@@ -32,6 +32,8 @@ public:
 	
 	SLATE_BEGIN_ARGS(SDisplayClusterLightCardTemplateList)
 	{}
+	SLATE_ARGUMENT(bool, HideHeader)
+	SLATE_ARGUMENT(bool, SpawnOnSelection)
 	SLATE_END_ARGS()
 
 	virtual ~SDisplayClusterLightCardTemplateList() override;
@@ -46,11 +48,15 @@ public:
 	
 	/** Reapply the current filter */
 	void ApplyFilter();
+
+	/** If drag & drop should be used */
+	bool IsDragDropEnabled() const;
 	
 private:
 	TSharedRef<ITableRow> GenerateTreeItemRow(TSharedPtr<FLightCardTemplateTreeItem> Item, const TSharedRef<STableViewBase>& OwnerTable);
 	void GetChildrenForTreeItem(TSharedPtr<FLightCardTemplateTreeItem> InItem, TArray<TSharedPtr<FLightCardTemplateTreeItem>>& OutChildren);
 	const FSlateBrush* GetFavoriteFilterIcon() const;
+	void OnSelectionChanged(TSharedPtr<FLightCardTemplateTreeItem> InItem, ESelectInfo::Type SelectInfo);
 	
 	void OnFilterTextChanged(const FText& SearchText);
 	
@@ -75,4 +81,7 @@ private:
 
 	/** If the favorite filter should be applied */
 	bool bFilterFavorites = false;
+ 
+	/** If selecting a template should spawn it */
+	bool bSpawnOnSelection = false;
 };
\ No newline at end of file
diff --git a/Engine/Plugins/Runtime/nDisplay/Source/DisplayClusterLightCardEditor/Private/SDisplayClusterLightCardEditor.cpp b/Engine/Plugins/Runtime/nDisplay/Source/DisplayClusterLightCardEditor/Private/SDisplayClusterLightCardEditor.cpp
--- a/Engine/Plugins/Runtime/nDisplay/Source/DisplayClusterLightCardEditor/Private/SDisplayClusterLightCardEditor.cpp
+++ b/Engine/Plugins/Runtime/nDisplay/Source/DisplayClusterLightCardEditor/Private/SDisplayClusterLightCardEditor.cpp
@@ -8,6 +8,7 @@
 #include "SDisplayClusterLightCardList.h"
 #include "LightCardTemplates/SDisplayClusterLightCardTemplateList.h"
 #include "LightCardTemplates/DisplayClusterLightCardTemplate.h"
+#include "LightCardTemplates/DisplayClusterLightCardTemplateHelpers.h"
 #include "Settings/DisplayClusterLightCardEditorSettings.h"
 
 #include "Viewport/DisplayClusterLightcardEditorViewport.h"
@@ -33,6 +34,7 @@
 #include "ScopedTransaction.h"
 #include "Misc/TransactionObjectEvent.h"
 #include "Selection.h"
+#include "Styling/SlateIconFinder.h"
 #include "Toolkits/AssetEditorToolkit.h"
 #include "Widgets/Docking/SDockTab.h"
 #include "Widgets/Input/SNumericEntryBox.h"
@@ -121,7 +123,6 @@ void SDisplayClusterLightCardEditor::Construct(const FArguments& InArgs, const T
 	TabManager = FGlobalTabmanager::Get()->NewTabManager(MajorTabOwner);
 
 	const FName LightCardTab = TEXT("LightCards");
-	const FName LightCardTemplateTab = TEXT("LightCardTemplates");
 	const FName ViewportTab = TEXT("Viewport"); 
 	
 	const TSharedRef<FTabManager::FLayout> Layout = FTabManager::NewLayout("LightCardEditor")
@@ -141,13 +142,6 @@ void SDisplayClusterLightCardEditor::Construct(const FArguments& InArgs, const T
 				->SetHideTabWell(true)
 				->SetForegroundTab(LightCardTab)
 			)
-			->Split
-			(
-			FTabManager::NewStack()
-				->AddTab(LightCardTemplateTab, ETabState::OpenedTab)
-				->SetHideTabWell(true)
-				->SetForegroundTab(LightCardTab)
-			)
 		)
 		->Split
 		(
@@ -170,18 +164,6 @@ void SDisplayClusterLightCardEditor::Construct(const FArguments& InArgs, const T
 			CreateLightCardListWidget()
 		];
 	}));
-
-	TabManager->RegisterTabSpawner(LightCardTemplateTab, FOnSpawnTab::CreateLambda([this](const FSpawnTabArgs&) -> TSharedRef<SDockTab>
-	{
-		return SNew(SDockTab)
-		.TabRole(ETabRole::PanelTab)
-		.CanEverClose(false)
-		.OnCanCloseTab(SDockTab::FCanCloseTab::CreateLambda([](){ return false; }))
-		.ContentPadding(0)
-		[
-			CreateLightCardTemplateWidget()
-		];
-	}));
 	
 	TabManager->RegisterTabSpawner(ViewportTab, FOnSpawnTab::CreateLambda([this](const FSpawnTabArgs&) -> TSharedRef<SDockTab>
 	{
@@ -278,7 +260,9 @@ ADisplayClusterLightCardActor* SDisplayClusterLightCardEditor::SpawnLightCardFro
 		return nullptr;
 	}
 
-	check(InTemplate);
+	FScopedTransaction Transaction(LOCTEXT("SpawnLightCardFromTemplateTransactionMessage", "Spawn Light Card from Template"));
+	
+	check(InTemplate && InTemplate->LightCardActor);
 
 	ULevel* Level = InLevel ? InLevel : ActiveRootActor->GetWorld()->GetCurrentLevel();
 	
@@ -300,6 +284,12 @@ ADisplayClusterLightCardActor* SDisplayClusterLightCardEditor::SpawnLightCardFro
 	
 		const TArray<ADisplayClusterLightCardActor*> LightCards { NewLightCard } ;
 		AddLightCardsToActor(LightCards);
+
+		FDisplayClusterLightCardEditorRecentItem RecentlyPlacedItem;
+		RecentlyPlacedItem.ObjectPath = InTemplate;
+		RecentlyPlacedItem.ItemType = FDisplayClusterLightCardEditorRecentItem::Type_LightCardTemplate;
+		
+		AddRecentlyPlacedItem(MoveTemp(RecentlyPlacedItem));
 	}
 	
 	return NewLightCard;
@@ -317,6 +307,12 @@ void SDisplayClusterLightCardEditor::AddNewLightCard()
 	if (NewLightCard)
 	{
 		CenterLightCardInView(*NewLightCard);
+
+		FDisplayClusterLightCardEditorRecentItem RecentlyPlacedItem;
+		RecentlyPlacedItem.ObjectPath = NewLightCard->GetClass();
+		RecentlyPlacedItem.ItemType = FDisplayClusterLightCardEditorRecentItem::Type_LightCard;
+		
+		AddRecentlyPlacedItem(MoveTemp(RecentlyPlacedItem));
 	}
 }
 
@@ -787,6 +783,176 @@ TSharedRef<SWidget> SDisplayClusterLightCardEditor::CreateViewportWidget()
 	return SAssignNew(ViewportView, SDisplayClusterLightCardEditorViewport, SharedThis(this), CommandList);
 }
 
+TSharedRef<SWidget> SDisplayClusterLightCardEditor::GeneratePlaceActorsMenu()
+{
+	const bool bInShouldCloseWindowAfterMenuSelection = true;
+
+	FMenuBuilder MenuBuilder(bInShouldCloseWindowAfterMenuSelection, CommandList);
+	MenuBuilder.BeginSection("PlaceActors", LOCTEXT("PlaceActorsMenuHeader", "Place Actors"));
+	{
+		FSlateIcon LightCardIcon = FSlateIconFinder::FindIconForClass(ADisplayClusterLightCardActor::StaticClass());
+		
+		MenuBuilder.AddMenuEntry(FDisplayClusterLightCardEditorCommands::Get().AddNewLightCard,
+			NAME_None, TAttribute<FText>(), TAttribute<FText>(), LightCardIcon);
+		MenuBuilder.AddMenuEntry(FDisplayClusterLightCardEditorCommands::Get().AddExistingLightCard,
+			NAME_None, TAttribute<FText>(), TAttribute<FText>(), LightCardIcon);
+		
+		MenuBuilder.AddSubMenu(LOCTEXT("AllTemplatesLabel", "All Templates"),
+		LOCTEXT("AllTemplatesTooltip", "Select a template"),
+		FNewMenuDelegate::CreateSP(this, &SDisplayClusterLightCardEditor::GenerateTemplateSubMenu),
+		false,
+		FSlateIcon(),
+		true);
+	}
+	MenuBuilder.EndSection();
+	MenuBuilder.BeginSection("Favorites", LOCTEXT("FavoritesMenuHeader", "Favorites"));
+	{
+		TArray<UDisplayClusterLightCardTemplate*> LightCardTemplates =
+		UE::DisplayClusterLightCardTemplateHelpers::GetLightCardTemplates(/*bFavoritesOnly*/ true);
+
+		TemplateBrushes.Empty();
+		
+		for (UDisplayClusterLightCardTemplate* Template : LightCardTemplates)
+		{
+			// Create a brush if this template is using a custom texture.
+			if (Template->LightCardActor != nullptr && Template->LightCardActor->Texture.Get() != nullptr)
+			{
+				TSharedPtr<FSlateBrush> SlateBrush = MakeShared<FSlateBrush>();
+
+				SlateBrush->SetResourceObject(Template->LightCardActor->Texture.Get());
+				SlateBrush->ImageSize = FVector2D(16.f, 16.f);
+				TemplateBrushes.Add(Template, SlateBrush);
+			}
+
+			TWeakObjectPtr<UDisplayClusterLightCardTemplate> TemplateWeakPtr = MakeWeakObjectPtr(Template);
+
+			const TSharedPtr<SWidget> TemplateWidget = SNew(SHorizontalBox)
+			+SHorizontalBox::Slot()
+			.AutoWidth()
+			.VAlign(VAlign_Center)
+			[
+				SNew(SBox)
+				.HeightOverride(16.f)
+				.WidthOverride(16.f)
+				[
+					SNew(SImage)
+					.ColorAndOpacity(FSlateColor::UseForeground())
+					.Image(this, &SDisplayClusterLightCardEditor::GetLightCardTemplateIcon, TemplateWeakPtr)
+				]
+			]
+			+SHorizontalBox::Slot()
+			.AutoWidth()
+			.VAlign(VAlign_Center)
+			.Padding(4.f, 0.f)
+			[
+				SNew(STextBlock)
+				.Text(FText::FromString(Template->GetName()))
+			];
+
+			FMenuEntryParams EntryParams;
+			EntryParams.EntryWidget = TemplateWidget;
+			EntryParams.LabelOverride = FText::FromString(Template->GetName());
+			EntryParams.ToolTipOverride = LOCTEXT("FavoriteEntryTooltip", "Spawn this favorite");
+			EntryParams.DirectActions.ExecuteAction = FExecuteAction::CreateLambda([this, TemplateWeakPtr]()
+			{
+				if (TemplateWeakPtr.IsValid())
+				{
+					SpawnLightCardFromTemplate(TemplateWeakPtr.Get());
+				}
+			});
+			
+			MenuBuilder.AddMenuEntry(EntryParams);
+		}
+	}
+	MenuBuilder.EndSection();
+
+	CleanupRecentlyPlacedItems();
+	
+	MenuBuilder.BeginSection("RecentlyPlaced", LOCTEXT("RecentlyPlacedMenuHeader", "Recently Placed"));
+	{
+		const UDisplayClusterLightCardEditorSettings* Settings = GetDefault<UDisplayClusterLightCardEditorSettings>();
+		for (const FDisplayClusterLightCardEditorRecentItem& RecentlyPlacedItem : Settings->RecentlyPlacedItems)
+		{
+			TWeakObjectPtr<UObject> ObjectWeakPtr = MakeWeakObjectPtr(RecentlyPlacedItem.ObjectPath.LoadSynchronous());
+			check(ObjectWeakPtr.IsValid());
+			
+			const TSharedPtr<SWidget> RecentItemWidget = SNew(SHorizontalBox)
+			+SHorizontalBox::Slot()
+			.AutoWidth()
+			.VAlign(VAlign_Center)
+			[
+				SNew(SBox)
+				.HeightOverride(16.f)
+				.WidthOverride(16.f)
+				[
+					SNew(SImage)
+					.ColorAndOpacity(FSlateColor::UseForeground())
+					.Image(RecentlyPlacedItem.GetSlateBrush())
+				]
+			]
+			+SHorizontalBox::Slot()
+			.AutoWidth()
+			.VAlign(VAlign_Center)
+			.Padding(4.f, 0.f)
+			[
+				SNew(STextBlock)
+				.Text(RecentlyPlacedItem.GetItemDisplayName())
+			];
+
+			FMenuEntryParams EntryParams;
+			EntryParams.EntryWidget = RecentItemWidget;
+			EntryParams.LabelOverride = FText::FromString(ObjectWeakPtr->GetName());
+			EntryParams.ToolTipOverride = LOCTEXT("RecentItemEntryTooltip", "Spawn this recently placed item");
+			EntryParams.DirectActions.ExecuteAction = FExecuteAction::CreateLambda([this, RecentlyPlacedItem, ObjectWeakPtr]()
+			{
+				if (ObjectWeakPtr.IsValid())
+				{
+					if (RecentlyPlacedItem.ItemType == FDisplayClusterLightCardEditorRecentItem::Type_LightCard)
+					{
+						AddNewLightCard();
+					}
+					else if (RecentlyPlacedItem.ItemType == FDisplayClusterLightCardEditorRecentItem::Type_LightCardTemplate)
+					{
+						if (const UDisplayClusterLightCardTemplate* Template = Cast<UDisplayClusterLightCardTemplate>(ObjectWeakPtr.Get()))
+						{
+							SpawnLightCardFromTemplate(Template);
+						}
+					}
+				}
+			});
+			
+			MenuBuilder.AddMenuEntry(EntryParams);
+		}
+	}
+	MenuBuilder.EndSection();
+
+	return MenuBuilder.MakeWidget();
+}
+
+const FSlateBrush* SDisplayClusterLightCardEditor::GetLightCardTemplateIcon(const TWeakObjectPtr<UDisplayClusterLightCardTemplate> InTemplate) const
+{
+	if (!InTemplate.IsValid())
+	{
+		return nullptr;
+	}
+		
+	if (const TSharedPtr<FSlateBrush>* Brush = TemplateBrushes.Find(InTemplate))
+	{
+		return (*Brush).Get();
+	}
+	
+	return FSlateIconFinder::FindIconBrushForClass(InTemplate->GetClass());
+}
+
+void SDisplayClusterLightCardEditor::GenerateTemplateSubMenu(FMenuBuilder& InMenuBuilder)
+{
+	InMenuBuilder.BeginSection("Templates", LOCTEXT("TemplatesMenuHeader", "Templates"));
+	{
+		InMenuBuilder.AddWidget(CreateLightCardTemplateWidget(), FText::GetEmpty(), true);
+	}
+	InMenuBuilder.EndSection();
+}
+
 TSharedRef<SWidget> SDisplayClusterLightCardEditor::GenerateLabelsMenu()
 {
 	const bool bInShouldCloseWindowAfterMenuSelection = true;
@@ -833,7 +999,9 @@ TSharedRef<SWidget> SDisplayClusterLightCardEditor::GenerateLabelsMenu()
 
 TSharedRef<SWidget> SDisplayClusterLightCardEditor::CreateLightCardTemplateWidget()
 {
-	return SAssignNew(LightCardTemplateList, SDisplayClusterLightCardTemplateList, SharedThis(this));
+	return SNew(SDisplayClusterLightCardTemplateList, SharedThis(this))
+	.HideHeader(true)
+	.SpawnOnSelection(true);
 }
 
 void SDisplayClusterLightCardEditor::BindCommands()
@@ -905,14 +1073,27 @@ void SDisplayClusterLightCardEditor::ExtendToolbar(FToolBarBuilder& ToolbarBuild
 	// TODO: Any toolbar buttons needed for the lightcards editor can be added to the operator panel's toolbar using this toolbar extender
 
 	ToolbarBuilder.AddSeparator();
+	{
+		const FUIAction DefaultAction;
+		ToolbarBuilder.AddComboButton(
+			DefaultAction,
+			FOnGetContent::CreateSP(this, &SDisplayClusterLightCardEditor::GeneratePlaceActorsMenu),
+			TAttribute<FText>(),
+			LOCTEXT("PlaceActors_ToolTip", "Place actors in the scene"),
+			FSlateIcon(FAppStyle::GetAppStyleSetName(), "Icons.PlaceActors"));
+	}
+	
+	ToolbarBuilder.AddSeparator();
+	{
+		const FUIAction DefaultAction;
+		ToolbarBuilder.AddComboButton(
+			DefaultAction,
+			FOnGetContent::CreateSP(this, &SDisplayClusterLightCardEditor::GenerateLabelsMenu),
+			TAttribute<FText>(),
+			LOCTEXT("Labels_ToolTip", "Configure options for labels"),
+			FSlateIcon(FDisplayClusterLightCardEditorStyle::Get().GetStyleSetName(), "DisplayClusterLightCardEditor.Labels"));
+	}
 	
-	const FUIAction DefaultAction;
-	ToolbarBuilder.AddComboButton(
-		DefaultAction,
-		FOnGetContent::CreateSP(this, &SDisplayClusterLightCardEditor::GenerateLabelsMenu),
-		TAttribute<FText>(),
-		LOCTEXT("Labels_ToolTip", "Configure options for labels"),
-		FSlateIcon(FDisplayClusterLightCardEditorStyle::Get().GetStyleSetName(), "DisplayClusterLightCardEditor.Labels"));
 }
 
 void SDisplayClusterLightCardEditor::RefreshPreviewActors(EDisplayClusterLightCardEditorProxyType ProxyType)
@@ -1088,4 +1269,37 @@ void SDisplayClusterLightCardEditor::OnLightCardListChanged()
 	RefreshPreviewActors(EDisplayClusterLightCardEditorProxyType::LightCards);
 }
 
+void SDisplayClusterLightCardEditor::AddRecentlyPlacedItem(const FDisplayClusterLightCardEditorRecentItem& InItem)
+{
+	UDisplayClusterLightCardEditorSettings* Settings = GetMutableDefault<UDisplayClusterLightCardEditorSettings>();
+	
+	Settings->RecentlyPlacedItems.RemoveAll([&](const FDisplayClusterLightCardEditorRecentItem& Compare)
+	{
+		return Compare == InItem;
+	});
+	Settings->RecentlyPlacedItems.Insert(InItem, 0);
+
+	const int32 MaxRecentlyPlaced = 5;
+	
+	if (Settings->RecentlyPlacedItems.Num() > MaxRecentlyPlaced)
+	{
+		Settings->RecentlyPlacedItems.RemoveAt(MaxRecentlyPlaced - 1, Settings->RecentlyPlacedItems.Num() - MaxRecentlyPlaced);
+	}
+
+	Settings->PostEditChange();
+	Settings->SaveConfig();
+}
+
+void SDisplayClusterLightCardEditor::CleanupRecentlyPlacedItems()
+{
+	UDisplayClusterLightCardEditorSettings* Settings = GetMutableDefault<UDisplayClusterLightCardEditorSettings>();
+	Settings->RecentlyPlacedItems.RemoveAll([&](const FDisplayClusterLightCardEditorRecentItem& Compare)
+	{
+		return !Compare.ObjectPath.LoadSynchronous();
+	});
+
+	Settings->PostEditChange();
+	Settings->SaveConfig();
+}
+
 #undef LOCTEXT_NAMESPACE
diff --git a/Engine/Plugins/Runtime/nDisplay/Source/DisplayClusterLightCardEditor/Private/SDisplayClusterLightCardEditor.h b/Engine/Plugins/Runtime/nDisplay/Source/DisplayClusterLightCardEditor/Private/SDisplayClusterLightCardEditor.h
--- a/Engine/Plugins/Runtime/nDisplay/Source/DisplayClusterLightCardEditor/Private/SDisplayClusterLightCardEditor.h
+++ b/Engine/Plugins/Runtime/nDisplay/Source/DisplayClusterLightCardEditor/Private/SDisplayClusterLightCardEditor.h
@@ -10,9 +10,10 @@
 
 #include "EditorUndoClient.h"
 
-class FTabManager;
 class FLayoutExtender;
+class FMenuBuilder;
 class FSpawnTabArgs;
+class FTabManager;
 class FToolBarBuilder;
 class FUICommandList;
 class SDockTab;
@@ -23,6 +24,8 @@ class ADisplayClusterRootActor;
 class ADisplayClusterLightCardActor;
 class UDisplayClusterLightCardTemplate;
 
+struct FDisplayClusterLightCardEditorRecentItem;
+
 /** A panel that can be spawned in a tab that contains all the UI elements that make up the 2D light cards editor */
 class SDisplayClusterLightCardEditor : public SCompoundWidget, public FEditorUndoClient
 {
@@ -74,7 +77,7 @@ public:
 	ADisplayClusterLightCardActor* SpawnLightCard();
 
 	/** Spawns a new light card from a light card template */
-	ADisplayClusterLightCardActor* SpawnLightCardFromTemplate(const UDisplayClusterLightCardTemplate* InTemplate, ULevel* InLevel, bool bIsPreview);
+	ADisplayClusterLightCardActor* SpawnLightCardFromTemplate(const UDisplayClusterLightCardTemplate* InTemplate, ULevel* InLevel = nullptr, bool bIsPreview = false);
 
 	/** Adds a new light card to the root actor and centers it in the viewport */
 	void AddNewLightCard();
@@ -155,6 +158,15 @@ private:
 	/** Create the 3d viewport widget */
 	TSharedRef<SWidget> CreateViewportWidget();
 
+	/** Generate the place actors drop down menu */
+	TSharedRef<SWidget> GeneratePlaceActorsMenu();
+
+	/** Return the correct template icon to use */
+	const FSlateBrush* GetLightCardTemplateIcon(const TWeakObjectPtr<UDisplayClusterLightCardTemplate> InTemplate) const;
+
+	/** Generate the All Templates sub menu */
+	void GenerateTemplateSubMenu(FMenuBuilder& InMenuBuilder);
+	
 	/** Generate the labels dropdown menu */
 	TSharedRef<SWidget> GenerateLabelsMenu();
 
@@ -201,6 +213,12 @@ private:
 
 	/** Raised when the light card list has added or removed a card */
 	void OnLightCardListChanged();
+
+	/** Add an item to recently placed list. Reorders matching items to the top of the list and trims the list */
+	void AddRecentlyPlacedItem(const FDisplayClusterLightCardEditorRecentItem& InItem);
+
+	/** Make sure no invalid entries are present */
+	void CleanupRecentlyPlacedItems();
 	
 private:
 	TSharedPtr<FTabManager> TabManager;
@@ -208,8 +226,8 @@ private:
 	/** The light card list widget */
 	TSharedPtr<SDisplayClusterLightCardList> LightCardList;
 
-	/** The light card template list */
-	TSharedPtr<SDisplayClusterLightCardTemplateList> LightCardTemplateList;
+	/** Templates to their icon brushes */
+	TMap<TWeakObjectPtr<UDisplayClusterLightCardTemplate>, TSharedPtr<FSlateBrush>> TemplateBrushes;
 	
 	/** The 3d viewport */
 	TSharedPtr<SDisplayClusterLightCardEditorViewport> ViewportView;
diff --git a/Engine/Plugins/Runtime/nDisplay/Source/DisplayClusterLightCardEditor/Private/SDisplayClusterLightCardList.cpp b/Engine/Plugins/Runtime/nDisplay/Source/DisplayClusterLightCardEditor/Private/SDisplayClusterLightCardList.cpp
--- a/Engine/Plugins/Runtime/nDisplay/Source/DisplayClusterLightCardEditor/Private/SDisplayClusterLightCardList.cpp
+++ b/Engine/Plugins/Runtime/nDisplay/Source/DisplayClusterLightCardEditor/Private/SDisplayClusterLightCardList.cpp
@@ -151,30 +151,6 @@ void SDisplayClusterLightCardList::Construct(const FArguments& InArgs, TSharedPt
 	[
 		SNew(SVerticalBox)
 		+SVerticalBox::Slot()
-		.AutoHeight()
-		.Padding(0.f, 5.f, 0.f, 0.f)
-		[
-			SNew(SBorder)
-			.Padding(0)
-			.BorderImage(FAppStyle::Get().GetBrush("DetailsView.CategoryTop"))
-			.BorderBackgroundColor(FLinearColor(0.6f, 0.6f, 0.6f, 1.0f))
-			[
-				SNew(SHorizontalBox)
-				+SHorizontalBox::Slot()
-				.Padding(3.0f, 3.0f)
-				.AutoWidth()
-				.VAlign(VAlign_Center)
-				[
-						
-					SNew(SPositiveActionButton)
-					.Text(LOCTEXT("AddNewItem", "Add"))
-					.Icon(FAppStyle::Get().GetBrush("Icons.Plus"))
-					.ToolTipText(LOCTEXT("AddNewLightCardTooltip", "Adds a new Light Card to the display cluster"))
-					.OnGetMenuContent(this, &SDisplayClusterLightCardList::CreateAddNewMenuContent)
-				]
-			]
-		]
-		+SVerticalBox::Slot()
 		[
 			SAssignNew(LightCardTreeView, STreeView<TSharedPtr<FLightCardTreeItem>>)
 			.TreeItemsSource(&LightCardTree)
@@ -399,18 +375,6 @@ FReply SDisplayClusterLightCardList::OnKeyDown(const FGeometry& MyGeometry, cons
 	return FReply::Unhandled();
 }
 
-TSharedRef<SWidget> SDisplayClusterLightCardList::CreateAddNewMenuContent()
-{
-	const FDisplayClusterLightCardEditorCommands& Commands = FDisplayClusterLightCardEditorCommands::Get();
-	
-	const bool bCloseAfterSelection = true;
-	FMenuBuilder MenuBuilder(bCloseAfterSelection, CommandList, Extenders);
-	MenuBuilder.AddMenuEntry(Commands.AddNewLightCard);
-	MenuBuilder.AddMenuEntry(Commands.AddExistingLightCard);
-	
-	return MenuBuilder.MakeWidget();
-}
-
 TSharedPtr<SWidget> SDisplayClusterLightCardList::CreateContextMenu()
 {
 	const bool bCloseAfterSelection = true;
diff --git a/Engine/Plugins/Runtime/nDisplay/Source/DisplayClusterLightCardEditor/Private/SDisplayClusterLightCardList.h b/Engine/Plugins/Runtime/nDisplay/Source/DisplayClusterLightCardEditor/Private/SDisplayClusterLightCardList.h
--- a/Engine/Plugins/Runtime/nDisplay/Source/DisplayClusterLightCardEditor/Private/SDisplayClusterLightCardList.h
+++ b/Engine/Plugins/Runtime/nDisplay/Source/DisplayClusterLightCardEditor/Private/SDisplayClusterLightCardList.h
@@ -69,7 +69,6 @@ private:
 
 	virtual FReply OnKeyDown(const FGeometry& MyGeometry, const FKeyEvent& InKeyEvent) override;
 	
-	TSharedRef<SWidget> CreateAddNewMenuContent();
 	TSharedPtr<SWidget> CreateContextMenu();
 	
 private:
diff --git a/Engine/Plugins/Runtime/nDisplay/Source/DisplayClusterLightCardEditor/Private/Settings/DisplayClusterLightCardEditorSettings.cpp b/Engine/Plugins/Runtime/nDisplay/Source/DisplayClusterLightCardEditor/Private/Settings/DisplayClusterLightCardEditorSettings.cpp
--- a/Engine/Plugins/Runtime/nDisplay/Source/DisplayClusterLightCardEditor/Private/Settings/DisplayClusterLightCardEditorSettings.cpp
+++ b/Engine/Plugins/Runtime/nDisplay/Source/DisplayClusterLightCardEditor/Private/Settings/DisplayClusterLightCardEditorSettings.cpp
@@ -2,9 +2,82 @@
 
 #include "DisplayClusterLightCardEditorSettings.h"
 
+#include "LightCardTemplates/DisplayClusterLightCardTemplate.h"
+
+#include "DisplayClusterMeshProjectionRenderer.h"
+
+#include "DisplayClusterLightCardActor.h"
+
+#include "Editor/UnrealEdTypes.h"
+#include "Engine/Texture.h"
+#include "Styling/SlateBrush.h"
+#include "Styling/SlateIconFinder.h"
+
+#define LOCTEXT_NAMESPACE "DisplayClusterLightCardEditorSettings"
+
 UDisplayClusterLightCardEditorProjectSettings::UDisplayClusterLightCardEditorProjectSettings()
 {
 	LightCardTemplateDefaultPath.Path = TEXT("/Game/VP/LightCards");
 	LightCardLabelScale = 1.f;
 	bDisplayLightCardLabels = false;
 }
+
+const FName FDisplayClusterLightCardEditorRecentItem::Type_LightCard = "LightCard";
+const FName FDisplayClusterLightCardEditorRecentItem::Type_LightCardTemplate = "LightCardTemplate";
+
+FText FDisplayClusterLightCardEditorRecentItem::GetItemDisplayName() const
+{
+	if (ItemType == Type_LightCard)
+	{
+		return LOCTEXT("LightCardRecentItemName", "Light Card");
+	}
+	if (const UObject* Object = ObjectPath.LoadSynchronous())
+	{
+		return FText::FromString(Object->GetName());
+	}
+
+	return FText::GetEmpty();
+}
+
+const FSlateBrush* FDisplayClusterLightCardEditorRecentItem::GetSlateBrush() const
+{
+	if (ItemType == Type_LightCard)
+	{
+		return FSlateIconFinder::FindIconBrushForClass(ADisplayClusterLightCardActor::StaticClass());
+	}
+	
+	if (ItemType == Type_LightCardTemplate && !SlateBrush.IsValid())
+	{
+		if (const UDisplayClusterLightCardTemplate* Template = Cast<UDisplayClusterLightCardTemplate>(ObjectPath.LoadSynchronous()))
+		{
+			UTexture* Texture = Template->LightCardActor != nullptr ? Template->LightCardActor->Texture.Get() : nullptr;
+			if (Texture == nullptr)
+			{
+				return FSlateIconFinder::FindIconBrushForClass(Template->GetClass());;
+			}
+			
+			if (!SlateBrush.IsValid())
+			{
+				SlateBrush = MakeShared<FSlateBrush>();
+			}
+		
+			SlateBrush->SetResourceObject(Texture);
+			SlateBrush->ImageSize = FVector2D(16.f, 16.f);
+		}
+	}
+
+	if (SlateBrush.IsValid())
+	{
+		return &*SlateBrush;
+	}
+
+	return nullptr;
+}
+
+UDisplayClusterLightCardEditorSettings::UDisplayClusterLightCardEditorSettings()
+{
+	ProjectionMode = EDisplayClusterMeshProjectionType::Azimuthal;
+	RenderViewportType = ELevelViewportType::LVT_Perspective;
+}
+
+#undef LOCTEXT_NAMESPACE
\ No newline at end of file
diff --git a/Engine/Plugins/Runtime/nDisplay/Source/DisplayClusterLightCardEditor/Private/Settings/DisplayClusterLightCardEditorSettings.h b/Engine/Plugins/Runtime/nDisplay/Source/DisplayClusterLightCardEditor/Private/Settings/DisplayClusterLightCardEditorSettings.h
--- a/Engine/Plugins/Runtime/nDisplay/Source/DisplayClusterLightCardEditor/Private/Settings/DisplayClusterLightCardEditorSettings.h
+++ b/Engine/Plugins/Runtime/nDisplay/Source/DisplayClusterLightCardEditor/Private/Settings/DisplayClusterLightCardEditorSettings.h
@@ -3,9 +3,14 @@
 #pragma once
 
 #include "CoreMinimal.h"
+
 #include "DisplayClusterLightCardEditorSettings.generated.h"
 
+struct FSlateBrush;
 
+/**
+ * Default settings shared across users of the project
+ */
 UCLASS(config = Editor, defaultconfig, meta = (DisplayClusterMultiUserInclude))
 class UDisplayClusterLightCardEditorProjectSettings : public UObject
 {
@@ -26,3 +31,55 @@ public:
 	UPROPERTY(config, EditAnywhere, Category = LightCardLabels)
 	float LightCardLabelScale;
 };
+
+USTRUCT()
+struct FDisplayClusterLightCardEditorRecentItem
+{
+	GENERATED_BODY()
+	
+	static const FName Type_LightCard;
+	static const FName Type_LightCardTemplate;
+
+	/** Path of object to instantiate */
+	UPROPERTY()
+	TSoftObjectPtr<UObject> ObjectPath;
+	
+	/** Type of item placed */
+	UPROPERTY()
+	FName ItemType;
+
+	FText GetItemDisplayName() const;
+	const FSlateBrush* GetSlateBrush() const;
+
+	bool operator==(const FDisplayClusterLightCardEditorRecentItem& Rhs) const
+	{
+		return ObjectPath == Rhs.ObjectPath && ItemType == Rhs.ItemType;
+	}
+
+private:
+	mutable TSharedPtr<FSlateBrush> SlateBrush;
+};
+
+/**
+ * Editor preferences unique to this user
+ */
+UCLASS(config = EditorPerProjectUserSettings)
+class UDisplayClusterLightCardEditorSettings : public UObject
+{
+	GENERATED_BODY()
+
+public:
+	UDisplayClusterLightCardEditorSettings();
+
+	/** Items recently placed by the user */
+	UPROPERTY(config)
+	TArray<FDisplayClusterLightCardEditorRecentItem> RecentlyPlacedItems;
+
+	/** Last used projection mode user has set */
+	UPROPERTY(config)
+	uint8 ProjectionMode;
+
+	/** Last used viewport type user has set */
+	UPROPERTY(config)
+	uint8 RenderViewportType;
+};
diff --git a/Engine/Plugins/Runtime/nDisplay/Source/DisplayClusterLightCardEditor/Private/Viewport/DisplayClusterLightCardEditorViewportClient.cpp b/Engine/Plugins/Runtime/nDisplay/Source/DisplayClusterLightCardEditor/Private/Viewport/DisplayClusterLightCardEditorViewportClient.cpp
--- a/Engine/Plugins/Runtime/nDisplay/Source/DisplayClusterLightCardEditor/Private/Viewport/DisplayClusterLightCardEditorViewportClient.cpp
+++ b/Engine/Plugins/Runtime/nDisplay/Source/DisplayClusterLightCardEditor/Private/Viewport/DisplayClusterLightCardEditorViewportClient.cpp
@@ -18,6 +18,7 @@
 #include "DisplayClusterLightCardEditorLog.h"
 #include "IDisplayClusterScenePreview.h"
 #include "SDisplayClusterLightCardEditor.h"
+#include "Settings/DisplayClusterLightCardEditorSettings.h"
 
 #include "AudioDevice.h"
 #include "CameraController.h"
@@ -74,7 +75,7 @@ FDisplayClusterLightCardEditorViewportClient::FDisplayClusterLightCardEditorView
 	check(InEditorViewportWidget.IsValid());
 	
 	LightCardEditorPtr = InEditorViewportWidget.Pin()->GetLightCardEditor();
-	check(LightCardEditorPtr.IsValid())
+	check(LightCardEditorPtr.IsValid());
 
 	IDisplayClusterScenePreview& PreviewModule = IDisplayClusterScenePreview::Get();
 	PreviewRendererId = PreviewModule.CreateRenderer();
@@ -119,7 +120,10 @@ FDisplayClusterLightCardEditorViewportClient::FDisplayClusterLightCardEditorView
 	
 	UpdatePreviewActor(LightCardEditorPtr.Pin()->GetActiveRootActor().Get());
 
-	SetProjectionMode(EDisplayClusterMeshProjectionType::Azimuthal, ELevelViewportType::LVT_Perspective);
+	const UDisplayClusterLightCardEditorSettings* Settings = GetDefault<UDisplayClusterLightCardEditorSettings>();
+	
+	SetProjectionMode(static_cast<EDisplayClusterMeshProjectionType>(Settings->ProjectionMode),
+		static_cast<ELevelViewportType>(Settings->RenderViewportType));
 }
 
 FDisplayClusterLightCardEditorViewportClient::~FDisplayClusterLightCardEditorViewportClient()
@@ -1403,6 +1407,13 @@ void FDisplayClusterLightCardEditorViewportClient::SetProjectionMode(EDisplayClu
 	ProjectionMode = InProjectionMode;
 	RenderViewportType = InViewportType;
 
+	UDisplayClusterLightCardEditorSettings* Settings = GetMutableDefault<UDisplayClusterLightCardEditorSettings>();
+	
+	Settings->ProjectionMode = ProjectionMode;
+	Settings->RenderViewportType = InViewportType;
+	Settings->PostEditChange();
+	Settings->SaveConfig();
+
 	ProjectionHelper->SetProjectionMode(ProjectionMode);
 	ProjectionHelper->SetIsOrthographic(InViewportType != LVT_Perspective);
 
