Disable UV Light cards until UE-159748 is resolved

diff --git a/Engine/Plugins/Runtime/nDisplay/Source/DisplayCluster/Private/Render/Viewport/DisplayClusterViewportManager.cpp b/Engine/Plugins/Runtime/nDisplay/Source/DisplayCluster/Private/Render/Viewport/DisplayClusterViewportManager.cpp
--- a/Engine/Plugins/Runtime/nDisplay/Source/DisplayCluster/Private/Render/Viewport/DisplayClusterViewportManager.cpp
+++ b/Engine/Plugins/Runtime/nDisplay/Source/DisplayCluster/Private/Render/Viewport/DisplayClusterViewportManager.cpp
@@ -26,7 +26,7 @@
 
 #include "SceneViewExtension.h"
 
-#include "DisplayClusterRootActor.h" 
+#include "DisplayClusterRootActor.h"
 
 #include "Misc/DisplayClusterLog.h"
 #include "LegacyScreenPercentageDriver.h"
@@ -64,7 +64,8 @@ FDisplayClusterViewportManager::FDisplayClusterViewportManager()
 
 	RenderTargetManager = MakeShared<FDisplayClusterRenderTargetManager, ESPMode::ThreadSafe>(ViewportManagerProxy);
 	PostProcessManager  = MakeShared<FDisplayClusterViewportPostProcessManager, ESPMode::ThreadSafe>(*this);
-	LightCardManager = MakeShared<FDisplayClusterViewportLightCardManager, ESPMode::ThreadSafe>(*this);
+	// DISABLE Light Card Manager until UE-159748 is resolved.
+//	LightCardManager = MakeShared<FDisplayClusterViewportLightCardManager, ESPMode::ThreadSafe>(*this);
 
 	// initialize proxy
 	ViewportManagerProxy->Initialize(*this);
@@ -640,7 +641,7 @@ FSceneViewFamily::ConstructionValues FDisplayClusterViewportManager::CreateViewF
 	switch (InFrameTarget.CaptureMode)
 	{
 	case EDisplayClusterViewportCaptureMode::Chromakey:
-		
+
 		if (!GDisplayClusterChromaKeyAllowNanite)
 		{
 			InEngineShowFlags.SetNaniteMeshes(0);
@@ -695,7 +696,7 @@ void FDisplayClusterViewportManager::ConfigureViewFamily(const FDisplayClusterRe
 			ViewExt->SetupViewFamily(ViewFamily);
 		}
 	}
-	
+
 	// Setup capture mode:
 	{
 		ViewFamily.SceneCaptureCompositeMode = ESceneCaptureCompositeMode::SCCM_Overwrite;
@@ -930,7 +931,7 @@ TSharedPtr<IDisplayClusterProjectionPolicy, ESPMode::ThreadSafe> FDisplayCluster
 	return nullptr;
 }
 
-TSharedPtr<IDisplayClusterViewportLightCardManager, ESPMode::ThreadSafe> FDisplayClusterViewportManager::GetLightCardManager() const 
+TSharedPtr<IDisplayClusterViewportLightCardManager, ESPMode::ThreadSafe> FDisplayClusterViewportManager::GetLightCardManager() const
 {
 	return LightCardManager;
 }
@@ -950,7 +951,7 @@ void FDisplayClusterViewportManager::MarkComponentGeometryDirty(const FName InCo
 				TSharedPtr<IDisplayClusterWarpBlend, ESPMode::ThreadSafe> WarpBlendInterface;
 				if (ProjectionPolicy->GetWarpBlendInterface(WarpBlendInterface) && WarpBlendInterface.IsValid())
 				{
-					// Update only interfaces with ProceduralMesh as geometry source 
+					// Update only interfaces with ProceduralMesh as geometry source
 					if (WarpBlendInterface->GetWarpGeometryType() == EDisplayClusterWarpGeometryType::WarpProceduralMesh)
 					{
 						// Set the ProceduralMeshComponent geometry dirty for all valid WarpBlendInterface
