- Updating display cluster media capture to rely on immediate rhi resource capture

diff --git a/Engine/Plugins/Runtime/nDisplay/Source/DisplayCluster/Private/Render/Viewport/DisplayClusterViewportManagerViewExtension.cpp b/Engine/Plugins/Runtime/nDisplay/Source/DisplayCluster/Private/Render/Viewport/DisplayClusterViewportManagerViewExtension.cpp
--- a/Engine/Plugins/Runtime/nDisplay/Source/DisplayCluster/Private/Render/Viewport/DisplayClusterViewportManagerViewExtension.cpp
+++ b/Engine/Plugins/Runtime/nDisplay/Source/DisplayCluster/Private/Render/Viewport/DisplayClusterViewportManagerViewExtension.cpp
@@ -55,8 +55,8 @@ void FDisplayClusterViewportManagerViewExtension::PostRenderViewFamily_RenderThr
 			const TArrayView<IDisplayClusterViewportProxy*> ViewportProxies = ViewportMgrPoxy->GetViewports_RenderThread();
 
 			// Filter those viewports that refer to the RenderTarget used in the ViewFamily
-			TArray<IDisplayClusterViewportProxy*> FoundViewportProxies = ViewportProxies.FilterByPredicate([&InViewFamily](const IDisplayClusterViewportProxy* ViewportProxy)
-			{
+			const TArray<IDisplayClusterViewportProxy*> FoundViewportProxies = ViewportProxies.FilterByPredicate([&InViewFamily](const IDisplayClusterViewportProxy* ViewportProxy)
+			{	
 				TArray<FRHITexture*> RenderTargets;
 				if (ViewportProxy->GetResources_RenderThread(EDisplayClusterViewportResourceType::InternalRenderTargetResource, RenderTargets))
 				{
@@ -68,15 +68,12 @@ void FDisplayClusterViewportManagerViewExtension::PostRenderViewFamily_RenderThr
 
 				return false;
 			});
-
-			AddPass(GraphBuilder, RDG_EVENT_NAME("ViewportProxy_PostRenderViewFamily"), [&InViewFamily, FoundViewportProxies = MoveTemp(FoundViewportProxies)] (FRHICommandListImmediate& RHICmdList)
+			
+			// Now we can perform per-viewport notification
+			for (const IDisplayClusterViewportProxy* ViewportProxy : FoundViewportProxies)
 			{
-				// Now we can perform per-viewport notification
-				for (const IDisplayClusterViewportProxy* ViewportProxy : FoundViewportProxies)
-				{
-					IDisplayCluster::Get().GetCallbacks().OnDisplayClusterPostRenderViewFamily_RenderThread().Broadcast(RHICmdList, InViewFamily, ViewportProxy);
-				}
-			});
+				IDisplayCluster::Get().GetCallbacks().OnDisplayClusterPostRenderViewFamily_RenderThread().Broadcast(GraphBuilder, InViewFamily, ViewportProxy);
+			}
 		}
 	}
 }
diff --git a/Engine/Plugins/Runtime/nDisplay/Source/DisplayCluster/Public/IDisplayClusterCallbacks.h b/Engine/Plugins/Runtime/nDisplay/Source/DisplayCluster/Public/IDisplayClusterCallbacks.h
--- a/Engine/Plugins/Runtime/nDisplay/Source/DisplayCluster/Public/IDisplayClusterCallbacks.h
+++ b/Engine/Plugins/Runtime/nDisplay/Source/DisplayCluster/Public/IDisplayClusterCallbacks.h
@@ -7,6 +7,7 @@
 class IDisplayClusterViewportProxy;
 class IDisplayClusterViewportManagerProxy;
 class FRHICommandListImmediate;
+class FRDGBuilder;
 class FSceneViewFamily;
 class FViewport;
 
@@ -73,7 +74,7 @@ public:
 	virtual FDisplayClusterFailoverNodeDown& OnDisplayClusterFailoverNodeDown() = 0;
 
 	/** Called once the ViewFamily of this viewport is rendered **/
-	DECLARE_EVENT_ThreeParams(IDisplayClusterCallbacks, FDisplayClusterPostRenderViewFamily_RenderThread, FRHICommandListImmediate&, const FSceneViewFamily&, const IDisplayClusterViewportProxy*);
+	DECLARE_EVENT_ThreeParams(IDisplayClusterCallbacks, FDisplayClusterPostRenderViewFamily_RenderThread, FRDGBuilder&, const FSceneViewFamily&, const IDisplayClusterViewportProxy*);
 	virtual FDisplayClusterPostRenderViewFamily_RenderThread& OnDisplayClusterPostRenderViewFamily_RenderThread() = 0;
 
 	/** Called once before warping all available viewports **/
diff --git a/Engine/Plugins/Runtime/nDisplay/Source/DisplayClusterMedia/Private/Capture/DisplayClusterMediaCaptureBase.cpp b/Engine/Plugins/Runtime/nDisplay/Source/DisplayClusterMedia/Private/Capture/DisplayClusterMediaCaptureBase.cpp
--- a/Engine/Plugins/Runtime/nDisplay/Source/DisplayClusterMedia/Private/Capture/DisplayClusterMediaCaptureBase.cpp
+++ b/Engine/Plugins/Runtime/nDisplay/Source/DisplayClusterMedia/Private/Capture/DisplayClusterMediaCaptureBase.cpp
@@ -10,9 +10,11 @@
 #include "MediaOutput.h"
 #include "Engine/TextureRenderTarget2D.h"
 
+#include "RenderGraphBuilder.h"
 #include "RHIResources.h"
 
 
+
 FDisplayClusterMediaCaptureBase::FDisplayClusterMediaCaptureBase(const FString& InMediaId, const FString& InClusterNodeId, UMediaOutput* InMediaOutput, UTextureRenderTarget2D* InRenderTarget)
 	: FDisplayClusterMediaBase(InMediaId, InClusterNodeId)
 	, MediaOutput(InMediaOutput)
@@ -68,31 +70,17 @@ void FDisplayClusterMediaCaptureBase::StopCapture()
 	}
 }
 
-void FDisplayClusterMediaCaptureBase::ExportMediaData(FRHICommandListImmediate& RHICmdList, const FMediaTextureInfo& TextureInfo)
+void FDisplayClusterMediaCaptureBase::ExportMediaData(FRDGBuilder& GraphBuilder, const FMediaTextureInfo& TextureInfo)
 {
 	FRHITexture* const SrcTexture = TextureInfo.Texture;
 	FRHITexture* const DstTexture = GetRenderTarget()->GetRenderTargetResource()->GetTextureRHI();
 
 	if (SrcTexture && DstTexture)
 	{
-		MediaCapture->SetValidSourceGPUMask(RHICmdList.GetGPUMask());
+		MediaCapture->SetValidSourceGPUMask(GraphBuilder.RHICmdList.GetGPUMask());
 
 		const FIntPoint SrcRegionSize = TextureInfo.Region.Size();
-
-		if (SrcTexture->GetDesc().Format == DstTexture->GetDesc().Format &&
-			SrcRegionSize == DstTexture->GetDesc().Extent)
-		{
-			FRHICopyTextureInfo CopyInfo;
-
-			CopyInfo.SourcePosition = FIntVector(TextureInfo.Region.Min.X, TextureInfo.Region.Min.Y, 0);
-			CopyInfo.Size = FIntVector(SrcRegionSize.X, SrcRegionSize.Y, 0);
-
-			TransitionAndCopyTexture(RHICmdList, SrcTexture, DstTexture, CopyInfo);
-		}
-		else
-		{
-			DisplayClusterMediaHelpers::ResampleTexture_RenderThread(RHICmdList, SrcTexture, DstTexture, TextureInfo.Region, FIntRect(FIntPoint::ZeroValue, DstTexture->GetDesc().Extent));
-		}
+		MediaCapture->CaptureImmediate_RenderThread(GraphBuilder, SrcTexture);
 	}
 }
 
@@ -117,13 +105,15 @@ void FDisplayClusterMediaCaptureBase::OnPostClusterTick()
 
 bool FDisplayClusterMediaCaptureBase::StartMediaCapture()
 {
+	FRHICaptureResourceDescription Descriptor;
+	Descriptor.ResourceSize = GetCaptureSize();
+	
 	FMediaCaptureOptions MediaCaptureOptions;
-
 	MediaCaptureOptions.NumberOfFramesToCapture = -1;
-	MediaCaptureOptions.bResizeSourceBuffer = true;
 	MediaCaptureOptions.bSkipFrameWhenRunningExpensiveTasks = false;
+	MediaCaptureOptions.OverrunAction = EMediaCaptureOverrunAction::Skip;
 
-	const bool bCaptureStarted = MediaCapture->CaptureTextureRenderTarget2D(RenderTarget, MediaCaptureOptions);
+	const bool bCaptureStarted = MediaCapture->CaptureRHITexture(Descriptor, MediaCaptureOptions);
 	if (bCaptureStarted)
 	{
 		UE_LOG(LogDisplayClusterMedia, Log, TEXT("Started media capture: '%s'"), *GetMediaId());
diff --git a/Engine/Plugins/Runtime/nDisplay/Source/DisplayClusterMedia/Private/Capture/DisplayClusterMediaCaptureBase.h b/Engine/Plugins/Runtime/nDisplay/Source/DisplayClusterMedia/Private/Capture/DisplayClusterMediaCaptureBase.h
--- a/Engine/Plugins/Runtime/nDisplay/Source/DisplayClusterMedia/Private/Capture/DisplayClusterMediaCaptureBase.h
+++ b/Engine/Plugins/Runtime/nDisplay/Source/DisplayClusterMedia/Private/Capture/DisplayClusterMediaCaptureBase.h
@@ -6,6 +6,7 @@
 #include "DisplayClusterMediaBase.h"
 #include "UObject/GCObject.h"
 
+class FRDGBuilder;
 class FRHICommandListImmediate;
 class UMediaCapture;
 class UMediaOutput;
@@ -42,10 +43,12 @@ public:
 	}
 
 protected:
-	void ExportMediaData(FRHICommandListImmediate& RHICmdList, const FMediaTextureInfo& TextureInfo);
+	void ExportMediaData(FRDGBuilder& GraphBuilder, const FMediaTextureInfo& TextureInfo);
 	void OnPostClusterTick();
 	bool StartMediaCapture();
 
+	virtual FIntPoint GetCaptureSize() const = 0;
+
 private:
 	//~ Begin GC by AddReferencedObjects
 	UMediaOutput*           MediaOutput  = nullptr;
diff --git a/Engine/Plugins/Runtime/nDisplay/Source/DisplayClusterMedia/Private/Capture/DisplayClusterMediaCaptureCamera.cpp b/Engine/Plugins/Runtime/nDisplay/Source/DisplayClusterMedia/Private/Capture/DisplayClusterMediaCaptureCamera.cpp
new file mode 100644
--- /dev/null
+++ b/Engine/Plugins/Runtime/nDisplay/Source/DisplayClusterMedia/Private/Capture/DisplayClusterMediaCaptureCamera.cpp
@@ -0,0 +1,43 @@
+// Copyright Epic Games, Inc. All Rights Reserved.
+
+#include "Capture/DisplayClusterMediaCaptureCamera.h"
+
+#include "Components/DisplayClusterICVFXCameraComponent.h"
+#include "Config/IDisplayClusterConfigManager.h"
+#include "DisplayClusterConfigurationTypes_Viewport.h"
+#include "DisplayClusterRootActor.h"
+#include "IDisplayCluster.h"
+#include "IDisplayClusterCallbacks.h"
+#include "Game/IDisplayClusterGameManager.h"
+
+#include "Render/Viewport/IDisplayClusterViewportManagerProxy.h"
+#include "RHICommandList.h"
+#include "RHIResources.h"
+
+
+FDisplayClusterMediaCaptureCamera::FDisplayClusterMediaCaptureCamera(const FString& InMediaId, const FString& InClusterNodeId, const FString& InCameraId, const FString& InViewportId, UMediaOutput* InMediaOutput, UTextureRenderTarget2D* InRenderTarget)
+	: FDisplayClusterMediaCaptureViewport(InMediaId, InClusterNodeId, InViewportId, InMediaOutput, InRenderTarget)
+	, CameraId(InCameraId)
+{
+	if (const ADisplayClusterRootActor* const RootActor = IDisplayCluster::Get().GetGameMgr()->GetRootActor())
+	{
+		TArray<UActorComponent*> ICVFXCameraComponents;
+		RootActor->GetComponents(UDisplayClusterICVFXCameraComponent::StaticClass(), ICVFXCameraComponents);
+
+		for (const UActorComponent* const Component : ICVFXCameraComponents)
+		{
+			if (const UDisplayClusterICVFXCameraComponent* const ICVFXCamera = Cast<UDisplayClusterICVFXCameraComponent>(Component))
+			{
+				if (ICVFXCamera->GetName() == CameraId)
+				{
+					CameraResolution = ICVFXCamera->GetCameraSettingsICVFX().CustomFrustum.InnerFrustumResolution;
+				}
+			}
+		}
+	}
+}
+
+FIntPoint FDisplayClusterMediaCaptureCamera::GetCaptureSize() const
+{
+	return CameraResolution;
+}
diff --git a/Engine/Plugins/Runtime/nDisplay/Source/DisplayClusterMedia/Private/Capture/DisplayClusterMediaCaptureCamera.h b/Engine/Plugins/Runtime/nDisplay/Source/DisplayClusterMedia/Private/Capture/DisplayClusterMediaCaptureCamera.h
new file mode 100644
--- /dev/null
+++ b/Engine/Plugins/Runtime/nDisplay/Source/DisplayClusterMedia/Private/Capture/DisplayClusterMediaCaptureCamera.h
@@ -0,0 +1,28 @@
+// Copyright Epic Games, Inc. All Rights Reserved.
+
+#pragma once
+
+#include "Capture/DisplayClusterMediaCaptureViewport.h"
+
+class FRDGBuilder;
+class FSceneViewFamily;
+class IDisplayClusterViewportProxy;
+
+
+/**
+ * Viewport media capture
+ */
+class FDisplayClusterMediaCaptureCamera
+	: public FDisplayClusterMediaCaptureViewport
+{
+public:
+	FDisplayClusterMediaCaptureCamera(const FString& MediaId, const FString& ClusterNodeId, const FString& CameraId, const FString& ViewportId, UMediaOutput* MediaOutput, UTextureRenderTarget2D* RenderTarget);
+
+protected:
+	
+	virtual FIntPoint GetCaptureSize() const override;
+
+private:
+	const FString CameraId;
+	FIntPoint CameraResolution = FIntPoint::ZeroValue;
+};
diff --git a/Engine/Plugins/Runtime/nDisplay/Source/DisplayClusterMedia/Private/Capture/DisplayClusterMediaCaptureNode.cpp b/Engine/Plugins/Runtime/nDisplay/Source/DisplayClusterMedia/Private/Capture/DisplayClusterMediaCaptureNode.cpp
--- a/Engine/Plugins/Runtime/nDisplay/Source/DisplayClusterMedia/Private/Capture/DisplayClusterMediaCaptureNode.cpp
+++ b/Engine/Plugins/Runtime/nDisplay/Source/DisplayClusterMedia/Private/Capture/DisplayClusterMediaCaptureNode.cpp
@@ -2,9 +2,12 @@
 
 #include "Capture/DisplayClusterMediaCaptureNode.h"
 
+#include "Config/IDisplayClusterConfigManager.h"
+#include "DisplayClusterConfigurationTypes.h"
 #include "IDisplayCluster.h"
 #include "IDisplayClusterCallbacks.h"
 
+#include "RenderGraphBuilder.h"
 #include "Render/Viewport/IDisplayClusterViewportManagerProxy.h"
 #include "RHICommandList.h"
 #include "RHIResources.h"
@@ -38,6 +41,17 @@ void FDisplayClusterMediaCaptureNode::StopCapture()
 	FDisplayClusterMediaCaptureBase::StopCapture();
 }
 
+FIntPoint FDisplayClusterMediaCaptureNode::GetCaptureSize() const
+{
+	if (IDisplayCluster::Get().GetConfigMgr() && IDisplayCluster::Get().GetConfigMgr()->GetLocalNode())
+	{
+		const FDisplayClusterConfigurationRectangle& Window = IDisplayCluster::Get().GetConfigMgr()->GetLocalNode()->WindowRect;
+		return { Window.W, Window.H };
+	}
+
+	return FIntPoint();
+}
+
 void FDisplayClusterMediaCaptureNode::OnPostFrameRender_RenderThread(FRHICommandListImmediate& RHICmdList, const IDisplayClusterViewportManagerProxy* ViewportManagerProxy, FViewport* Viewport)
 {
 	ensure(ViewportManagerProxy);
@@ -50,7 +64,10 @@ void FDisplayClusterMediaCaptureNode::OnPostFrameRender_RenderThread(FRHICommand
 		if (Textures.Num() > 0 && Regions.Num() > 0 && Textures[0])
 		{
 			FMediaTextureInfo TextureInfo{ Textures[0], FIntRect(Regions[0], Regions[0] + Textures[0]->GetDesc().Extent) };
-			ExportMediaData(RHICmdList, TextureInfo);
+
+			FRDGBuilder Builder(RHICmdList);
+			ExportMediaData(Builder, TextureInfo);
+			Builder.Execute();
 		}
 	}
 }
diff --git a/Engine/Plugins/Runtime/nDisplay/Source/DisplayClusterMedia/Private/Capture/DisplayClusterMediaCaptureNode.h b/Engine/Plugins/Runtime/nDisplay/Source/DisplayClusterMedia/Private/Capture/DisplayClusterMediaCaptureNode.h
--- a/Engine/Plugins/Runtime/nDisplay/Source/DisplayClusterMedia/Private/Capture/DisplayClusterMediaCaptureNode.h
+++ b/Engine/Plugins/Runtime/nDisplay/Source/DisplayClusterMedia/Private/Capture/DisplayClusterMediaCaptureNode.h
@@ -22,6 +22,9 @@ public:
 	virtual bool StartCapture() override;
 	virtual void StopCapture() override;
 
+protected:
+	virtual FIntPoint GetCaptureSize() const override;
+
 private:
 	void OnPostFrameRender_RenderThread(FRHICommandListImmediate& RHICmdList, const IDisplayClusterViewportManagerProxy* ViewportManagerProxy, FViewport* Viewport);
 };
diff --git a/Engine/Plugins/Runtime/nDisplay/Source/DisplayClusterMedia/Private/Capture/DisplayClusterMediaCaptureViewport.cpp b/Engine/Plugins/Runtime/nDisplay/Source/DisplayClusterMedia/Private/Capture/DisplayClusterMediaCaptureViewport.cpp
--- a/Engine/Plugins/Runtime/nDisplay/Source/DisplayClusterMedia/Private/Capture/DisplayClusterMediaCaptureViewport.cpp
+++ b/Engine/Plugins/Runtime/nDisplay/Source/DisplayClusterMedia/Private/Capture/DisplayClusterMediaCaptureViewport.cpp
@@ -2,6 +2,8 @@
 
 #include "Capture/DisplayClusterMediaCaptureViewport.h"
 
+#include "Config/IDisplayClusterConfigManager.h"
+#include "DisplayClusterConfigurationTypes_Viewport.h"
 #include "IDisplayCluster.h"
 #include "IDisplayClusterCallbacks.h"
 
@@ -37,7 +39,21 @@ void FDisplayClusterMediaCaptureViewport::StopCapture()
 	FDisplayClusterMediaCaptureBase::StopCapture();
 }
 
-void FDisplayClusterMediaCaptureViewport::OnPostRenderViewFamily_RenderThread(FRHICommandListImmediate& RHICmdList, const FSceneViewFamily& ViewFamily, const IDisplayClusterViewportProxy* ViewportProxy)
+FIntPoint FDisplayClusterMediaCaptureViewport::GetCaptureSize() const
+{
+	if (IDisplayCluster::Get().GetConfigMgr())
+	{
+		const UDisplayClusterConfigurationViewport* Viewport = IDisplayCluster::Get().GetConfigMgr()->GetLocalViewport(ViewportId);
+		if (ensure(Viewport))
+		{
+			return { Viewport->Region.W, Viewport->Region.H };
+		}
+	}
+
+	return FIntPoint();
+}
+
+void FDisplayClusterMediaCaptureViewport::OnPostRenderViewFamily_RenderThread(FRDGBuilder& GraphBuilder, const FSceneViewFamily& ViewFamily, const IDisplayClusterViewportProxy* ViewportProxy)
 {
 	ensure(ViewportProxy);
 
@@ -52,7 +68,7 @@ void FDisplayClusterMediaCaptureViewport::OnPostRenderViewFamily_RenderThread(FR
 			if (Textures.Num() > 0 && Regions.Num() > 0)
 			{
 				FMediaTextureInfo TextureInfo{ Textures[0], Regions[0] };
-				ExportMediaData(RHICmdList, TextureInfo);
+				ExportMediaData(GraphBuilder, TextureInfo);
 			}
 		}
 	}
diff --git a/Engine/Plugins/Runtime/nDisplay/Source/DisplayClusterMedia/Private/Capture/DisplayClusterMediaCaptureViewport.h b/Engine/Plugins/Runtime/nDisplay/Source/DisplayClusterMedia/Private/Capture/DisplayClusterMediaCaptureViewport.h
--- a/Engine/Plugins/Runtime/nDisplay/Source/DisplayClusterMedia/Private/Capture/DisplayClusterMediaCaptureViewport.h
+++ b/Engine/Plugins/Runtime/nDisplay/Source/DisplayClusterMedia/Private/Capture/DisplayClusterMediaCaptureViewport.h
@@ -5,7 +5,7 @@
 #include "CoreMinimal.h"
 #include "Capture/DisplayClusterMediaCaptureBase.h"
 
-class FRHICommandListImmediate;
+class FRDGBuilder;
 class FSceneViewFamily;
 class IDisplayClusterViewportProxy;
 
@@ -28,8 +28,12 @@ public:
 		return ViewportId;
 	}
 
+protected:
+	
+	virtual FIntPoint GetCaptureSize() const override;
+
 private:
-	void OnPostRenderViewFamily_RenderThread(FRHICommandListImmediate& RHICmdList, const FSceneViewFamily& ViewFamily, const IDisplayClusterViewportProxy* ViewportProxy);
+	void OnPostRenderViewFamily_RenderThread(FRDGBuilder& GraphBuilder, const FSceneViewFamily& ViewFamily, const IDisplayClusterViewportProxy* ViewportProxy);
 
 private:
 	const FString ViewportId;
diff --git a/Engine/Plugins/Runtime/nDisplay/Source/DisplayClusterMedia/Private/DisplayClusterMediaModule.cpp b/Engine/Plugins/Runtime/nDisplay/Source/DisplayClusterMedia/Private/DisplayClusterMediaModule.cpp
--- a/Engine/Plugins/Runtime/nDisplay/Source/DisplayClusterMedia/Private/DisplayClusterMediaModule.cpp
+++ b/Engine/Plugins/Runtime/nDisplay/Source/DisplayClusterMedia/Private/DisplayClusterMediaModule.cpp
@@ -15,6 +15,7 @@
 #include "Cluster/IDisplayClusterClusterManager.h"
 #include "Game/IDisplayClusterGameManager.h"
 
+#include "Capture/DisplayClusterMediaCaptureCamera.h"
 #include "Capture/DisplayClusterMediaCaptureNode.h"
 #include "Capture/DisplayClusterMediaCaptureViewport.h"
 #include "Input/DisplayClusterMediaInputNode.h"
@@ -186,9 +187,9 @@ void FDisplayClusterMediaModule::InitializeMedia()
 								{
 									UE_LOG(LogDisplayClusterMedia, Log, TEXT("Initializing ICVFX in-cluster capture for camera '%s' (in-cluster shared)..."), *ICVFXCameraName);
 
-									TUniquePtr<FDisplayClusterMediaCaptureViewport> NewICVFXCapture = MakeUnique<FDisplayClusterMediaCaptureViewport>(
+									TUniquePtr<FDisplayClusterMediaCaptureViewport> NewICVFXCapture = MakeUnique<FDisplayClusterMediaCaptureCamera>(
 										MediaCaptureId, ClusterNodeId,
-										ICVFXViewportId,
+										ICVFXCameraName, ICVFXViewportId,
 										MediaSettings.MediaOutput.MediaOutput, MediaSettings.MediaOutput.RenderTarget);
 
 									AllCaptures.Add(NewICVFXCapture.Get());
@@ -226,9 +227,9 @@ void FDisplayClusterMediaModule::InitializeMedia()
 							{
 								UE_LOG(LogDisplayClusterMedia, Log, TEXT("Initializing ICVFX capture for camera '%s'..."), *ICVFXCameraName);
 
-								TUniquePtr<FDisplayClusterMediaCaptureViewport> NewICVFXCapture = MakeUnique<FDisplayClusterMediaCaptureViewport>(
+								TUniquePtr<FDisplayClusterMediaCaptureViewport> NewICVFXCapture = MakeUnique<FDisplayClusterMediaCaptureCamera>(
 									MediaCaptureId, ClusterNodeId,
-									ICVFXViewportId,
+									ICVFXCameraName, ICVFXViewportId,
 									MediaSettings.MediaOutput.MediaOutput, MediaSettings.MediaOutput.RenderTarget);
 
 								AllCaptures.Add(NewICVFXCapture.Get());
