Callsite upgrade pass for TObjectPtr use on plugins used when building UnrealEditor with "-allmodules" for Linux.

diff --git a/Engine/Plugins/Importers/USDImporter/Source/USDStage/Private/USDPrimTwin.cpp b/Engine/Plugins/Importers/USDImporter/Source/USDStage/Private/USDPrimTwin.cpp
--- a/Engine/Plugins/Importers/USDImporter/Source/USDStage/Private/USDPrimTwin.cpp
+++ b/Engine/Plugins/Importers/USDImporter/Source/USDStage/Private/USDPrimTwin.cpp
@@ -53,7 +53,7 @@ void UUsdPrimTwin::Clear()
 
 	Modify();
 
-	for (const TPair< FString, UUsdPrimTwin* >& Pair : Children)
+	for (const UE_TRANSITIONAL_OBJECT_PTR_TEMPLATE2_ARG2(TPair, FString, UUsdPrimTwin)& Pair : Children)
 	{
 		// Apparently when changing levels it is possible for these objects to already be nullptr by the time we try clearing them,
 		// so its safer to check
@@ -160,7 +160,7 @@ UUsdPrimTwin* UUsdPrimTwin::Find( const USceneComponent* InSceneComponent )
 		return this;
 	}
 
-	for ( const TPair< FString, UUsdPrimTwin* >& Child : Children )
+	for ( const UE_TRANSITIONAL_OBJECT_PTR_TEMPLATE2_ARG2(TPair, FString, UUsdPrimTwin)& Child : Children )
 	{
 		UUsdPrimTwin* FoundPrimTwin = Child.Value->Find( InSceneComponent );
 		if ( FoundPrimTwin )
diff --git a/Engine/Plugins/Importers/USDImporter/Source/USDStage/Private/USDStageActor.cpp b/Engine/Plugins/Importers/USDImporter/Source/USDStage/Private/USDStageActor.cpp
--- a/Engine/Plugins/Importers/USDImporter/Source/USDStage/Private/USDStageActor.cpp
+++ b/Engine/Plugins/Importers/USDImporter/Source/USDStage/Private/USDStageActor.cpp
@@ -484,7 +484,7 @@ struct FUsdStageActorImpl
 				Component->SetHiddenInGame( !bPrimIsVisible, bPropagateToChildren );
 			}
 
-			for ( const TPair<FString, UUsdPrimTwin*>& ChildPair : PrimTwin.GetChildren() )
+			for ( const UE_TRANSITIONAL_OBJECT_PTR_TEMPLATE2_ARG2(TPair, FString, UUsdPrimTwin)& ChildPair : PrimTwin.GetChildren() )
 			{
 				if ( UUsdPrimTwin* ChildTwin = ChildPair.Value )
 				{
@@ -1517,7 +1517,7 @@ TArray<UObject*> AUsdStageActor::GetGeneratedAssets( const FString& PrimPath )
 
 	// Collect any other asset that claims they came from the same prim (e.g. also return the skeleton if we query
 	// the SkelRoot, or return textures if we query a material prim that used them, etc.)
-	for ( const TPair<FString, UObject*>& HashToAsset : AssetCache->GetCachedAssets() )
+	for ( const UE_TRANSITIONAL_OBJECT_PTR_TEMPLATE2_ARG2(TPair, FString, UObject)& HashToAsset : AssetCache->GetCachedAssets() )
 	{
 		if ( UUsdAssetImportData* ImportData = UsdUtils::GetAssetImportData( HashToAsset.Value ) )
 		{
@@ -1551,7 +1551,7 @@ FString AUsdStageActor::GetSourcePrimPath( UObject* Object )
 			}
 		}
 
-		for ( const TPair<FString, UObject*>& HashToAsset : AssetCache->GetCachedAssets() )
+		for ( const UE_TRANSITIONAL_OBJECT_PTR_TEMPLATE2_ARG2(TPair, FString, UObject)& HashToAsset : AssetCache->GetCachedAssets() )
 		{
 			if ( HashToAsset.Value == Object )
 			{
diff --git a/Engine/Plugins/Protocols/MQTT/Source/MQTTCore/Private/MQTTClientObject.cpp b/Engine/Plugins/Protocols/MQTT/Source/MQTTCore/Private/MQTTClientObject.cpp
--- a/Engine/Plugins/Protocols/MQTT/Source/MQTTCore/Private/MQTTClientObject.cpp
+++ b/Engine/Plugins/Protocols/MQTT/Source/MQTTCore/Private/MQTTClientObject.cpp
@@ -1,4 +1,4 @@
-﻿// Copyright Epic Games, Inc. All Rights Reserved.
+// Copyright Epic Games, Inc. All Rights Reserved.
 
 #include "MQTTClientObject.h"
 
@@ -156,7 +156,7 @@ UMQTTSubscriptionObject* UMQTTClientObject::Subscribe(const FString& InTopic, EM
 	if (const TSharedPtr<IMQTTClient, ESPMode::ThreadSafe> MQTTClient = MqttClientPtr)
 	{
 		UMQTTSubscriptionObject* SubscriptionObject = NewObject<UMQTTSubscriptionObject>(this);
-		UMQTTSubscriptionObject*& SubscriptionObjectRef = this->Subscriptions.Add_GetRef(SubscriptionObject);
+		UE_TRANSITIONAL_OBJECT_PTR(UMQTTSubscriptionObject)& SubscriptionObjectRef = this->Subscriptions.Add_GetRef(SubscriptionObject);
 
 		MQTTClient->Subscribe({MakeTuple(InTopic, InQoS)})
 		.Next([&](TArray<FMQTTSubscribeResult> InSubscribeResults)
diff --git a/Engine/Plugins/Runtime/nDisplay/Source/DisplayCluster/Private/Game/EngineClasses/Scene/DisplayClusterRootActor.cpp b/Engine/Plugins/Runtime/nDisplay/Source/DisplayCluster/Private/Game/EngineClasses/Scene/DisplayClusterRootActor.cpp
--- a/Engine/Plugins/Runtime/nDisplay/Source/DisplayCluster/Private/Game/EngineClasses/Scene/DisplayClusterRootActor.cpp
+++ b/Engine/Plugins/Runtime/nDisplay/Source/DisplayCluster/Private/Game/EngineClasses/Scene/DisplayClusterRootActor.cpp
@@ -1010,7 +1010,7 @@ static void PropagateDataFromDefaultConfig(UDisplayClusterConfigurationData* InD
 	
 	PropagateDefaultMapToInstancedMap(ClusterNodesMapProperty, InDefaultConfigData->Cluster, InInstanceConfigData->Cluster);
 	
-	for (const TTuple<FString, UDisplayClusterConfigurationClusterNode*>& ClusterKeyVal : InDefaultConfigData->Cluster->Nodes)
+	for (const UE_TRANSITIONAL_OBJECT_PTR_TEMPLATE2_ARG2(TTuple, FString, UDisplayClusterConfigurationClusterNode)& ClusterKeyVal : InDefaultConfigData->Cluster->Nodes)
 	{
 		UDisplayClusterConfigurationClusterNode* DestinationValue = InInstanceConfigData->Cluster->Nodes.FindChecked(
 			ClusterKeyVal.Key);
@@ -1056,14 +1056,14 @@ USceneComponent* ADisplayClusterRootActor::GetCommonViewPoint() const
 			TMap<FString, int32> ViewPointCounts;
 
 			// Get all the camera names used by viewports in this cluster
-			for (const TPair<FString, UDisplayClusterConfigurationClusterNode*>& NodePair : ConfigData->Cluster->Nodes)
+			for (const UE_TRANSITIONAL_OBJECT_PTR_TEMPLATE2_ARG2(TPair, FString, UDisplayClusterConfigurationClusterNode)& NodePair : ConfigData->Cluster->Nodes)
 			{
 				if (!NodePair.Value)
 				{
 					continue;
 				}
 
-				for (const TPair<FString, UDisplayClusterConfigurationViewport*>& ViewportPair : NodePair.Value->Viewports)
+				for (const UE_TRANSITIONAL_OBJECT_PTR_TEMPLATE2_ARG2(TPair, FString, UDisplayClusterConfigurationViewport)& ViewportPair : NodePair.Value->Viewports)
 				{
 					if (!ViewportPair.Value)
 					{
@@ -1150,7 +1150,7 @@ bool ADisplayClusterRootActor::SetReplaceTextureFlagForAllViewports(bool bReplac
 	{
 		check(Node != nullptr);
 
-		for (const TPair<FString, UDisplayClusterConfigurationViewport*>& ViewportItem : Node->Viewports)
+		for (const UE_TRANSITIONAL_OBJECT_PTR_TEMPLATE2_ARG2(TPair, FString, UDisplayClusterConfigurationViewport)& ViewportItem : Node->Viewports)
 		{
 			if (ViewportItem.Value)
 			{
@@ -1158,14 +1158,14 @@ bool ADisplayClusterRootActor::SetReplaceTextureFlagForAllViewports(bool bReplac
 			}
 		}
 
-		for (const TPair<FString, UDisplayClusterConfigurationClusterNode*>& NodeItem : ConfigData->Cluster->Nodes)
+		for (const UE_TRANSITIONAL_OBJECT_PTR_TEMPLATE2_ARG2(TPair, FString, UDisplayClusterConfigurationClusterNode)& NodeItem : ConfigData->Cluster->Nodes)
 		{
 			if (!NodeItem.Value)
 			{
 				continue;
 			}
 
-			for (const TPair<FString, UDisplayClusterConfigurationViewport*>& ViewportItem : NodeItem.Value->Viewports)
+			for (const UE_TRANSITIONAL_OBJECT_PTR_TEMPLATE2_ARG2(TPair, FString, UDisplayClusterConfigurationViewport)& ViewportItem : NodeItem.Value->Viewports)
 			{
 				if (ViewportItem.Value)
 				{
diff --git a/Engine/Plugins/Runtime/nDisplay/Source/DisplayCluster/Private/Game/EngineClasses/Scene/DisplayClusterRootActorEditor.cpp b/Engine/Plugins/Runtime/nDisplay/Source/DisplayCluster/Private/Game/EngineClasses/Scene/DisplayClusterRootActorEditor.cpp
--- a/Engine/Plugins/Runtime/nDisplay/Source/DisplayCluster/Private/Game/EngineClasses/Scene/DisplayClusterRootActorEditor.cpp
+++ b/Engine/Plugins/Runtime/nDisplay/Source/DisplayCluster/Private/Game/EngineClasses/Scene/DisplayClusterRootActorEditor.cpp
@@ -216,7 +216,7 @@ void ADisplayClusterRootActor::GetPreviewRenderTargetableTextures(const TArray<F
 {
 	check(IsInGameThread());
 
-	for(const TPair<FString, UDisplayClusterPreviewComponent*>& PreviewComponentIt : PreviewComponents)
+	for(const UE_TRANSITIONAL_OBJECT_PTR_TEMPLATE2_ARG2(TPair, FString, UDisplayClusterPreviewComponent)& PreviewComponentIt : PreviewComponents)
 	{
 		if (PreviewComponentIt.Value)
 		{
@@ -353,7 +353,7 @@ bool ADisplayClusterRootActor::ImplUpdatePreviewRenderFrame_Editor(const FString
 		if (ImplUpdatePreviewConfiguration_Editor(InClusterNodeId))
 		{
 			// Update all preview components resources before render
-			for (const TTuple<FString, UDisplayClusterPreviewComponent*>& PreviewComponentIt : PreviewComponents)
+			for (const UE_TRANSITIONAL_OBJECT_PTR_TEMPLATE2_ARG2(TTuple, FString, UDisplayClusterPreviewComponent)& PreviewComponentIt : PreviewComponents)
 			{
 				if (PreviewComponentIt.Value && PreviewComponentIt.Value->GetClusterNodeId() == InClusterNodeId)
 				{
@@ -507,7 +507,7 @@ void ADisplayClusterRootActor::ImplRenderPreviewFrustums_Editor()
 	// frustum preview viewports
 	TArray<IDisplayClusterViewport*> FrustumPreviewViewports;
 
-	for (const TPair<FString, UDisplayClusterConfigurationClusterNode*>& Node : CurrentConfigData->Cluster->Nodes)
+	for (const UE_TRANSITIONAL_OBJECT_PTR_TEMPLATE2_ARG2(TPair, FString, UDisplayClusterConfigurationClusterNode)& Node : CurrentConfigData->Cluster->Nodes)
 	{
 		if (Node.Value == nullptr)
 		{
@@ -515,7 +515,7 @@ void ADisplayClusterRootActor::ImplRenderPreviewFrustums_Editor()
 		}
 
 		// collect node viewports
-		for (const TPair<FString, UDisplayClusterConfigurationViewport*>& ViewportConfig : Node.Value->Viewports)
+		for (const UE_TRANSITIONAL_OBJECT_PTR_TEMPLATE2_ARG2(TPair, FString, UDisplayClusterConfigurationViewport)& ViewportConfig : Node.Value->Viewports)
 		{
 			if (ViewportConfig.Value->bAllowPreviewFrustumRendering == false || ViewportConfig.Value->bIsVisible == false || ViewportConfig.Value == nullptr)
 			{
@@ -819,14 +819,14 @@ void ADisplayClusterRootActor::UpdatePreviewComponents()
 	
 	if (CurrentConfigData != nullptr && bPreviewEnable)
 	{
-		for (const TPair<FString, UDisplayClusterConfigurationClusterNode*>& Node : CurrentConfigData->Cluster->Nodes)
+		for (const UE_TRANSITIONAL_OBJECT_PTR_TEMPLATE2_ARG2(TPair, FString, UDisplayClusterConfigurationClusterNode)& Node : CurrentConfigData->Cluster->Nodes)
 		{
 			if (Node.Value == nullptr)
 			{
 				continue;
 			}
 
-			for (const TPair<FString, UDisplayClusterConfigurationViewport*>& Viewport : Node.Value->Viewports)
+			for (const UE_TRANSITIONAL_OBJECT_PTR_TEMPLATE2_ARG2(TPair, FString, UDisplayClusterConfigurationViewport)& Viewport : Node.Value->Viewports)
 			{
 				const FString PreviewCompId = GeneratePreviewComponentName_Editor(Node.Key, Viewport.Key);
 				UDisplayClusterPreviewComponent* PreviewComp = PreviewComponents.FindRef(PreviewCompId);
@@ -875,7 +875,7 @@ void ADisplayClusterRootActor::UpdatePreviewComponents()
 
 void ADisplayClusterRootActor::ReleasePreviewComponents()
 {
-	for (const TPair<FString, UDisplayClusterPreviewComponent*>& CompPair : PreviewComponents)
+	for (const UE_TRANSITIONAL_OBJECT_PTR_TEMPLATE2_ARG2(TPair, FString, UDisplayClusterPreviewComponent)& CompPair : PreviewComponents)
 	{
 		if (CompPair.Value)
 		{
diff --git a/Engine/Plugins/Runtime/nDisplay/Source/DisplayCluster/Private/Game/EngineClasses/Scene/DisplayClusterTestPatternsActor.cpp b/Engine/Plugins/Runtime/nDisplay/Source/DisplayCluster/Private/Game/EngineClasses/Scene/DisplayClusterTestPatternsActor.cpp
--- a/Engine/Plugins/Runtime/nDisplay/Source/DisplayCluster/Private/Game/EngineClasses/Scene/DisplayClusterTestPatternsActor.cpp
+++ b/Engine/Plugins/Runtime/nDisplay/Source/DisplayCluster/Private/Game/EngineClasses/Scene/DisplayClusterTestPatternsActor.cpp
@@ -141,7 +141,7 @@ void ADisplayClusterTestPatternsActor::InitializeInternals()
 		if (Node)
 		{
 			// For each local viewport create a PP settings structure with no PP material assigned
-			for (const TPair<FString, UDisplayClusterConfigurationViewport*>& it : Node->Viewports)
+			for (const UE_TRANSITIONAL_OBJECT_PTR_TEMPLATE2_ARG2(TPair, FString, UDisplayClusterConfigurationViewport)& it : Node->Viewports)
 			{
 				ViewportPPSettings.Emplace(it.Key, CreatePPSettings(nullptr));
 			}
diff --git a/Engine/Plugins/Runtime/nDisplay/Source/DisplayCluster/Private/Render/Viewport/Configuration/DisplayClusterViewportConfigurationBase.cpp b/Engine/Plugins/Runtime/nDisplay/Source/DisplayCluster/Private/Render/Viewport/Configuration/DisplayClusterViewportConfigurationBase.cpp
--- a/Engine/Plugins/Runtime/nDisplay/Source/DisplayCluster/Private/Render/Viewport/Configuration/DisplayClusterViewportConfigurationBase.cpp
+++ b/Engine/Plugins/Runtime/nDisplay/Source/DisplayCluster/Private/Render/Viewport/Configuration/DisplayClusterViewportConfigurationBase.cpp
@@ -28,7 +28,7 @@ void FDisplayClusterViewportConfigurationBase::Update(const FString& ClusterNode
 	const UDisplayClusterConfigurationClusterNode* ClusterNodeConfiguration = ConfigurationData.Cluster->GetNode(ClusterNodeId);
 	if (ClusterNodeConfiguration)
 	{
-		for (const TPair<FString, UDisplayClusterConfigurationViewport*>& ViewportIt : ClusterNodeConfiguration->Viewports)
+		for (const UE_TRANSITIONAL_OBJECT_PTR_TEMPLATE2_ARG2(TPair, FString, UDisplayClusterConfigurationViewport)& ViewportIt : ClusterNodeConfiguration->Viewports)
 		{
 			if (ViewportIt.Key.Len() && ViewportIt.Value)
 			{
@@ -114,11 +114,11 @@ void FDisplayClusterViewportConfigurationBase::Update(const TArray<FString>& InV
 	}
 
 	// Update and Create new viewports
-	for (const TPair<FString, UDisplayClusterConfigurationClusterNode*>& ClusterNodeConfigurationIt : ConfigurationData.Cluster->Nodes)
+	for (const UE_TRANSITIONAL_OBJECT_PTR_TEMPLATE2_ARG2(TPair, FString, UDisplayClusterConfigurationClusterNode)& ClusterNodeConfigurationIt : ConfigurationData.Cluster->Nodes)
 	{
 		if (const UDisplayClusterConfigurationClusterNode* ClusterNodeConfiguration = ClusterNodeConfigurationIt.Value)
 		{
-			for (const TPair<FString, UDisplayClusterConfigurationViewport*>& ViewportIt : ClusterNodeConfiguration->Viewports)
+			for (const UE_TRANSITIONAL_OBJECT_PTR_TEMPLATE2_ARG2(TPair, FString, UDisplayClusterConfigurationViewport)& ViewportIt : ClusterNodeConfiguration->Viewports)
 			{
 				if (ViewportIt.Key.Len() && ViewportIt.Value && InViewportNames.Find(ViewportIt.Key) != INDEX_NONE)
 				{
diff --git a/Engine/Plugins/Runtime/nDisplay/Source/DisplayClusterConfiguration/Private/DisplayClusterConfigurationTypes.cpp b/Engine/Plugins/Runtime/nDisplay/Source/DisplayClusterConfiguration/Private/DisplayClusterConfigurationTypes.cpp
--- a/Engine/Plugins/Runtime/nDisplay/Source/DisplayClusterConfiguration/Private/DisplayClusterConfigurationTypes.cpp
+++ b/Engine/Plugins/Runtime/nDisplay/Source/DisplayClusterConfiguration/Private/DisplayClusterConfigurationTypes.cpp
@@ -393,7 +393,7 @@ UDisplayClusterConfigurationViewport* UDisplayClusterConfigurationClusterNode::G
 
 void UDisplayClusterConfigurationClusterNode::GetReferencedMeshNames(TArray<FString>& OutMeshNames) const
 {
-	for (const TPair<FString, UDisplayClusterConfigurationViewport*>& It : Viewports)
+	for (const UE_TRANSITIONAL_OBJECT_PTR_TEMPLATE2_ARG2(TPair, FString, UDisplayClusterConfigurationViewport)& It : Viewports)
 	{
 		if (It.Value == nullptr)
 		{
@@ -417,7 +417,7 @@ UDisplayClusterConfigurationClusterNode* UDisplayClusterConfigurationCluster::Ge
 
 void UDisplayClusterConfigurationCluster::GetReferencedMeshNames(TArray<FString>& OutMeshNames) const
 {
-	for (const TPair<FString, UDisplayClusterConfigurationClusterNode*>& It : Nodes)
+	for (const UE_TRANSITIONAL_OBJECT_PTR_TEMPLATE2_ARG2(TPair, FString, UDisplayClusterConfigurationClusterNode)& It : Nodes)
 	{
 		if (It.Value == nullptr)
 		{
diff --git a/Engine/Plugins/Runtime/nDisplay/Source/DisplayClusterConfiguration/Private/Formats/JSON500/DisplayClusterConfigurationJsonParser_500.cpp b/Engine/Plugins/Runtime/nDisplay/Source/DisplayClusterConfiguration/Private/Formats/JSON500/DisplayClusterConfigurationJsonParser_500.cpp
--- a/Engine/Plugins/Runtime/nDisplay/Source/DisplayClusterConfiguration/Private/Formats/JSON500/DisplayClusterConfigurationJsonParser_500.cpp
+++ b/Engine/Plugins/Runtime/nDisplay/Source/DisplayClusterConfiguration/Private/Formats/JSON500/DisplayClusterConfigurationJsonParser_500.cpp
@@ -311,7 +311,7 @@ namespace JSON500
 		// Scene
 		{
 			// Cameras
-			for (const TPair<FString, UDisplayClusterConfigurationSceneComponentCamera*>& Comp : Config->Scene->Cameras)
+			for (const UE_TRANSITIONAL_OBJECT_PTR_TEMPLATE2_ARG2(TPair, FString, UDisplayClusterConfigurationSceneComponentCamera)& Comp : Config->Scene->Cameras)
 			{
 				FDisplayClusterConfigurationJsonSceneComponentCamera_500 CfgComp;
 
@@ -328,7 +328,7 @@ namespace JSON500
 			}
 
 			// Screens
-			for (const TPair<FString, UDisplayClusterConfigurationSceneComponentScreen*>& Comp : Config->Scene->Screens)
+			for (const UE_TRANSITIONAL_OBJECT_PTR_TEMPLATE2_ARG2(TPair, FString, UDisplayClusterConfigurationSceneComponentScreen)& Comp : Config->Scene->Screens)
 			{
 				FDisplayClusterConfigurationJsonSceneComponentScreen_500 CfgComp;
 
@@ -343,7 +343,7 @@ namespace JSON500
 			}
 
 			// Xforms
-			for (const TPair<FString, UDisplayClusterConfigurationSceneComponentXform*>& Comp : Config->Scene->Xforms)
+			for (const UE_TRANSITIONAL_OBJECT_PTR_TEMPLATE2_ARG2(TPair, FString, UDisplayClusterConfigurationSceneComponentXform)& Comp : Config->Scene->Xforms)
 			{
 				FDisplayClusterConfigurationJsonSceneComponentScreen_500 CfgComp;
 
@@ -404,7 +404,7 @@ namespace JSON500
 			}
 
 			// Cluster nodes
-			for (const TPair<FString, UDisplayClusterConfigurationClusterNode*>& CfgNode : Config->Cluster->Nodes)
+			for (const UE_TRANSITIONAL_OBJECT_PTR_TEMPLATE2_ARG2(TPair, FString, UDisplayClusterConfigurationClusterNode)& CfgNode : Config->Cluster->Nodes)
 			{
 				FDisplayClusterConfigurationJsonClusterNode_500 Node;
 
@@ -415,7 +415,7 @@ namespace JSON500
 				Node.Window = FDisplayClusterConfigurationJsonRectangle_500(CfgNode.Value->WindowRect.X, CfgNode.Value->WindowRect.Y, CfgNode.Value->WindowRect.W, CfgNode.Value->WindowRect.H);
 
 				// Viewports
-				for (const TPair<FString, UDisplayClusterConfigurationViewport*>& CfgViewport : CfgNode.Value->Viewports)
+				for (const UE_TRANSITIONAL_OBJECT_PTR_TEMPLATE2_ARG2(TPair, FString, UDisplayClusterConfigurationViewport)& CfgViewport : CfgNode.Value->Viewports)
 				{
 					FDisplayClusterConfigurationJsonViewport_500 Viewport;
 
diff --git a/Engine/Plugins/Runtime/nDisplay/Source/DisplayClusterConfigurator/Private/ClusterConfiguration/DisplayClusterConfiguratorClusterUtils.cpp b/Engine/Plugins/Runtime/nDisplay/Source/DisplayClusterConfigurator/Private/ClusterConfiguration/DisplayClusterConfiguratorClusterUtils.cpp
--- a/Engine/Plugins/Runtime/nDisplay/Source/DisplayClusterConfigurator/Private/ClusterConfiguration/DisplayClusterConfiguratorClusterUtils.cpp
+++ b/Engine/Plugins/Runtime/nDisplay/Source/DisplayClusterConfigurator/Private/ClusterConfiguration/DisplayClusterConfiguratorClusterUtils.cpp
@@ -139,7 +139,7 @@ UDisplayClusterConfigurationViewport* FDisplayClusterConfiguratorClusterUtils::C
 	UDisplayClusterConfigurationCluster* Cluster = Toolkit->GetEditorData()->Cluster;
 
 	TArray<FString> ParentItems;
-	for (const TPair<FString, UDisplayClusterConfigurationClusterNode*>& Node : Cluster->Nodes)
+	for (const UE_TRANSITIONAL_OBJECT_PTR_TEMPLATE2_ARG2(TPair, FString, UDisplayClusterConfigurationClusterNode)& Node : Cluster->Nodes)
 	{
 		ParentItems.Add(Node.Key);
 	}
@@ -422,7 +422,7 @@ FVector2D FDisplayClusterConfiguratorClusterUtils::FindNextAvailablePositionForV
 
 	FBox2D DesiredBounds = FBox2D(DesiredPosition, DesiredPosition + DesiredSize);
 	TArray<FBox2D> ViewportBounds;
-	for (const TPair<FString, UDisplayClusterConfigurationViewport*>& ViewportKeyPair : ClusterNode->Viewports)
+	for (const UE_TRANSITIONAL_OBJECT_PTR_TEMPLATE2_ARG2(TPair, FString, UDisplayClusterConfigurationViewport)& ViewportKeyPair : ClusterNode->Viewports)
 	{
 		const FDisplayClusterConfigurationRectangle& Region = ViewportKeyPair.Value->Region;
 		ViewportBounds.Add(FBox2D(FVector2D(Region.X, Region.Y), FVector2D(Region.X + Region.W, Region.Y + Region.H)));
@@ -507,7 +507,7 @@ bool FDisplayClusterConfiguratorClusterUtils::RemoveUnusedHostDisplayData(UDispl
 	TArray<FString> UnusedHosts;
 	Cluster->HostDisplayData.GetKeys(UnusedHosts);
 
-	for (const TPair<FString, UDisplayClusterConfigurationClusterNode*>& NodePair : Cluster->Nodes)
+	for (const UE_TRANSITIONAL_OBJECT_PTR_TEMPLATE2_ARG2(TPair, FString, UDisplayClusterConfigurationClusterNode)& NodePair : Cluster->Nodes)
 	{
 		UnusedHosts.Remove(NodePair.Value->Host);
 	}
@@ -568,7 +568,7 @@ bool FDisplayClusterConfiguratorClusterUtils::RemoveHost(UDisplayClusterConfigur
 	}
 
 	TArray<FString> ClusterNodesToRemove;
-	for (const TPair<FString, UDisplayClusterConfigurationClusterNode*>& Pair : Cluster->Nodes)
+	for (const UE_TRANSITIONAL_OBJECT_PTR_TEMPLATE2_ARG2(TPair, FString, UDisplayClusterConfigurationClusterNode)& Pair : Cluster->Nodes)
 	{
 		if (Pair.Value->Host == Host)
 		{
@@ -704,7 +704,7 @@ bool FDisplayClusterConfiguratorClusterUtils::RemoveClusterNodeFromCluster(UDisp
 
 			if (bRemovingPrimaryNode)
 			{
-				for (const TTuple<FString, UDisplayClusterConfigurationClusterNode*>& NodesKeyVal : ClusterNodeParent->Nodes)
+				for (const UE_TRANSITIONAL_OBJECT_PTR_TEMPLATE2_ARG2(TTuple, FString, UDisplayClusterConfigurationClusterNode)& NodesKeyVal : ClusterNodeParent->Nodes)
 				{
 					SetClusterNodeAsPrimary(NodesKeyVal.Value);
 					break;
@@ -801,19 +801,19 @@ FString FDisplayClusterConfiguratorClusterUtils::GetUniqueNameForViewport(FStrin
 	// in the cluster to check for uniqueness. Add the parent cluster node's viewports first, in case we can't get to the root cluster through
 	// the cluster node's Outer (i.e. the cluster node has not been added to the cluster yet)
 	TSet<FString> UsedNames;
-	for (const TPair<FString, UDisplayClusterConfigurationViewport*>& ViewportKeyPair : ParentClusterNode->Viewports)
+	for (const UE_TRANSITIONAL_OBJECT_PTR_TEMPLATE2_ARG2(TPair, FString, UDisplayClusterConfigurationViewport)& ViewportKeyPair : ParentClusterNode->Viewports)
 	{
 		UsedNames.Add(ViewportKeyPair.Key);
 	}
 
 	if (UDisplayClusterConfigurationCluster* Cluster = Cast<UDisplayClusterConfigurationCluster>(ParentClusterNode->GetOuter()))
 	{
-		for (const TPair<FString, UDisplayClusterConfigurationClusterNode*>& ClusterNodeKeyPair : Cluster->Nodes)
+		for (const UE_TRANSITIONAL_OBJECT_PTR_TEMPLATE2_ARG2(TPair, FString, UDisplayClusterConfigurationClusterNode)& ClusterNodeKeyPair : Cluster->Nodes)
 		{
 			UDisplayClusterConfigurationClusterNode* ClusterNode = ClusterNodeKeyPair.Value;
 			if (ClusterNode != ParentClusterNode)
 			{
-				for (const TPair<FString, UDisplayClusterConfigurationViewport*>& ViewportKeyPair : ClusterNode->Viewports)
+				for (const UE_TRANSITIONAL_OBJECT_PTR_TEMPLATE2_ARG2(TPair, FString, UDisplayClusterConfigurationViewport)& ViewportKeyPair : ClusterNode->Viewports)
 				{
 					UsedNames.Add(ViewportKeyPair.Key);
 				}
diff --git a/Engine/Plugins/Runtime/nDisplay/Source/DisplayClusterConfigurator/Private/Views/Details/Cluster/DisplayClusterConfiguratorNodeSelectionCustomization.cpp b/Engine/Plugins/Runtime/nDisplay/Source/DisplayClusterConfigurator/Private/Views/Details/Cluster/DisplayClusterConfiguratorNodeSelectionCustomization.cpp
--- a/Engine/Plugins/Runtime/nDisplay/Source/DisplayClusterConfigurator/Private/Views/Details/Cluster/DisplayClusterConfiguratorNodeSelectionCustomization.cpp
+++ b/Engine/Plugins/Runtime/nDisplay/Source/DisplayClusterConfigurator/Private/Views/Details/Cluster/DisplayClusterConfiguratorNodeSelectionCustomization.cpp
@@ -158,14 +158,14 @@ void FDisplayClusterConfiguratorNodeSelection::ResetOptions()
 	Options.Reset();
 	if (UDisplayClusterConfigurationData* ConfigData = GetConfigData())
 	{
-		for (const TTuple<FString, UDisplayClusterConfigurationClusterNode*>& Node : ConfigData->Cluster->Nodes)
+		for (const UE_TRANSITIONAL_OBJECT_PTR_TEMPLATE2_ARG2(TTuple, FString, UDisplayClusterConfigurationClusterNode)& Node : ConfigData->Cluster->Nodes)
 		{
 			if (OperationMode == ClusterNodes)
 			{
 				Options.Add(MakeShared<FString>(Node.Value->GetName()));
 				continue;
 			}
-			for (const TTuple<FString, UDisplayClusterConfigurationViewport*>& Viewport : Node.Value->Viewports)
+			for (const UE_TRANSITIONAL_OBJECT_PTR_TEMPLATE2_ARG2(TTuple, FString, UDisplayClusterConfigurationViewport)& Viewport : Node.Value->Viewports)
 			{
 				Options.Add(MakeShared<FString>(Viewport.Value->GetName()));
 			}
diff --git a/Engine/Plugins/Runtime/nDisplay/Source/DisplayClusterConfigurator/Private/Views/Details/DisplayClusterRootActorDetailsCustomization.cpp b/Engine/Plugins/Runtime/nDisplay/Source/DisplayClusterConfigurator/Private/Views/Details/DisplayClusterRootActorDetailsCustomization.cpp
--- a/Engine/Plugins/Runtime/nDisplay/Source/DisplayClusterConfigurator/Private/Views/Details/DisplayClusterRootActorDetailsCustomization.cpp
+++ b/Engine/Plugins/Runtime/nDisplay/Source/DisplayClusterConfigurator/Private/Views/Details/DisplayClusterRootActorDetailsCustomization.cpp
@@ -190,7 +190,7 @@ bool FDisplayClusterRootActorDetailsCustomization::RebuildNodeIdOptionsList()
 	// Fill combobox with the options
 	NodeIdOptions.Reset();
 	NodeIdOptions.Emplace(NodeIdOptionNone);
-	for (const TPair<FString, UDisplayClusterConfigurationClusterNode*>& Node : ConfigurationData->Cluster->Nodes)
+	for (const UE_TRANSITIONAL_OBJECT_PTR_TEMPLATE2_ARG2(TPair, FString, UDisplayClusterConfigurationClusterNode)& Node : ConfigurationData->Cluster->Nodes)
 	{
 		if (!Node.Key.IsEmpty())
 		{
diff --git a/Engine/Plugins/Runtime/nDisplay/Source/DisplayClusterConfigurator/Private/Views/OutputMapping/DisplayClusterConfiguratorGraph.cpp b/Engine/Plugins/Runtime/nDisplay/Source/DisplayClusterConfigurator/Private/Views/OutputMapping/DisplayClusterConfiguratorGraph.cpp
--- a/Engine/Plugins/Runtime/nDisplay/Source/DisplayClusterConfigurator/Private/Views/OutputMapping/DisplayClusterConfiguratorGraph.cpp
+++ b/Engine/Plugins/Runtime/nDisplay/Source/DisplayClusterConfigurator/Private/Views/OutputMapping/DisplayClusterConfiguratorGraph.cpp
@@ -80,7 +80,7 @@ void UDisplayClusterConfiguratorGraph::RebuildGraph()
 				UDisplayClusterConfiguratorWindowNode* WindowNode = BuildWindowNode(HostNode, ClusterNodePair.Key, WindowIndex, ClusterNodePair.Value);
 				++WindowIndex;
 
-				for (const TPair<FString, UDisplayClusterConfigurationViewport*>& ViewportPair : ClusterNodePair.Value->Viewports)
+				for (const UE_TRANSITIONAL_OBJECT_PTR_TEMPLATE2_ARG2(TPair, FString, UDisplayClusterConfigurationViewport)& ViewportPair : ClusterNodePair.Value->Viewports)
 				{
 					BuildViewportNode(WindowNode, ViewportPair.Key, ViewportIndex, ViewportPair.Value);
 					++ViewportIndex;
diff --git a/Engine/Plugins/Runtime/nDisplay/Source/DisplayClusterLightCardEditor/Private/Viewport/DisplayClusterLightCardEditorViewportClient.cpp b/Engine/Plugins/Runtime/nDisplay/Source/DisplayClusterLightCardEditor/Private/Viewport/DisplayClusterLightCardEditorViewportClient.cpp
--- a/Engine/Plugins/Runtime/nDisplay/Source/DisplayClusterLightCardEditor/Private/Viewport/DisplayClusterLightCardEditorViewportClient.cpp
+++ b/Engine/Plugins/Runtime/nDisplay/Source/DisplayClusterLightCardEditor/Private/Viewport/DisplayClusterLightCardEditorViewportClient.cpp
@@ -178,10 +178,10 @@ void FDisplayClusterLightCardEditorViewportClient::Tick(float DeltaSeconds)
 		// Pass the preview render targets from the level instance root actor to the preview root actor
 		UDisplayClusterConfigurationData* Config = RootActorLevelInstance->GetConfigData();
 
-		for (const TPair<FString, UDisplayClusterConfigurationClusterNode*>& NodePair : Config->Cluster->Nodes)
+		for (const UE_TRANSITIONAL_OBJECT_PTR_TEMPLATE2_ARG2(TPair, FString, UDisplayClusterConfigurationClusterNode)& NodePair : Config->Cluster->Nodes)
 		{
 			const UDisplayClusterConfigurationClusterNode* Node = NodePair.Value;
-			for (const TPair<FString, UDisplayClusterConfigurationViewport*>& ViewportPair : Node->Viewports)
+			for (const UE_TRANSITIONAL_OBJECT_PTR_TEMPLATE2_ARG2(TPair, FString, UDisplayClusterConfigurationViewport)& ViewportPair : Node->Viewports)
 			{
 				UDisplayClusterPreviewComponent* LevelInstancePreviewComp = RootActorLevelInstance->GetPreviewComponent(NodePair.Key, ViewportPair.Key);
 				UDisplayClusterPreviewComponent* PreviewComp = RootActorProxy->GetPreviewComponent(NodePair.Key, ViewportPair.Key);
@@ -1644,10 +1644,10 @@ UDisplayClusterConfigurationViewport* FDisplayClusterLightCardEditorViewportClie
 		const FString PrimitiveComponentName = PrimitiveComponent->GetName();
 		UDisplayClusterConfigurationData* Config = RootActorProxy->GetConfigData();
 		
-		for (const TPair<FString, UDisplayClusterConfigurationClusterNode*>& NodePair : Config->Cluster->Nodes)
+		for (const UE_TRANSITIONAL_OBJECT_PTR_TEMPLATE2_ARG2(TPair, FString, UDisplayClusterConfigurationClusterNode)& NodePair : Config->Cluster->Nodes)
 		{
 			const UDisplayClusterConfigurationClusterNode* Node = NodePair.Value;
-			for (const TPair<FString, UDisplayClusterConfigurationViewport*>& ViewportPair : Node->Viewports)
+			for (const UE_TRANSITIONAL_OBJECT_PTR_TEMPLATE2_ARG2(TPair, FString, UDisplayClusterConfigurationViewport)& ViewportPair : Node->Viewports)
 			{
 				UDisplayClusterConfigurationViewport* CfgViewport = ViewportPair.Value;
 
diff --git a/Engine/Plugins/Runtime/nDisplay/Source/DisplayClusterMedia/Private/DisplayClusterMediaModule.cpp b/Engine/Plugins/Runtime/nDisplay/Source/DisplayClusterMedia/Private/DisplayClusterMediaModule.cpp
--- a/Engine/Plugins/Runtime/nDisplay/Source/DisplayClusterMedia/Private/DisplayClusterMediaModule.cpp
+++ b/Engine/Plugins/Runtime/nDisplay/Source/DisplayClusterMedia/Private/DisplayClusterMediaModule.cpp
@@ -112,7 +112,7 @@ void FDisplayClusterMediaModule::InitializeMedia()
 
 			///////////////////////////////
 			// Viewports media setup
-			for (const TPair<FString, UDisplayClusterConfigurationViewport*>& ViewportIt : ClusterNode->Viewports)
+			for (const UE_TRANSITIONAL_OBJECT_PTR_TEMPLATE2_ARG2(TPair, FString, UDisplayClusterConfigurationViewport)& ViewportIt : ClusterNode->Viewports)
 			{
 				if (const UDisplayClusterConfigurationViewport* const Viewport = ViewportIt.Value)
 				{
diff --git a/Engine/Plugins/Runtime/nDisplay/Source/DisplayClusterMoviePipeline/Private/DisplayClusterMoviePipelineSettings.cpp b/Engine/Plugins/Runtime/nDisplay/Source/DisplayClusterMoviePipeline/Private/DisplayClusterMoviePipelineSettings.cpp
--- a/Engine/Plugins/Runtime/nDisplay/Source/DisplayClusterMoviePipeline/Private/DisplayClusterMoviePipelineSettings.cpp
+++ b/Engine/Plugins/Runtime/nDisplay/Source/DisplayClusterMoviePipeline/Private/DisplayClusterMoviePipelineSettings.cpp
@@ -35,12 +35,12 @@ bool UDisplayClusterMoviePipelineSettings::GetViewports(const UWorld* InWorld, T
 		{
 			if (const UDisplayClusterConfigurationCluster* InClusterCfg =  InConfigurationData->Cluster)
 			{
-				for (const TPair<FString, UDisplayClusterConfigurationClusterNode*>& NodeIt : InClusterCfg->Nodes)
+				for (const UE_TRANSITIONAL_OBJECT_PTR_TEMPLATE2_ARG2(TPair, FString, UDisplayClusterConfigurationClusterNode)& NodeIt : InClusterCfg->Nodes)
 				{
 					if (const UDisplayClusterConfigurationClusterNode* InConfigurationClusterNode = NodeIt.Value)
 					{
 						const FString& InClusterNodeId = NodeIt.Key;
-						for (const TPair<FString, UDisplayClusterConfigurationViewport*>& InConfigurationViewportIt : InConfigurationClusterNode->Viewports)
+						for (const UE_TRANSITIONAL_OBJECT_PTR_TEMPLATE2_ARG2(TPair, FString, UDisplayClusterConfigurationViewport)& InConfigurationViewportIt : InConfigurationClusterNode->Viewports)
 						{
 							if (const UDisplayClusterConfigurationViewport* InConfigurationViewport = InConfigurationViewportIt.Value)
 							{
diff --git a/Engine/Plugins/Runtime/nDisplay/Source/DisplayClusterScenePreview/Private/DisplayClusterScenePreviewModule.cpp b/Engine/Plugins/Runtime/nDisplay/Source/DisplayClusterScenePreview/Private/DisplayClusterScenePreviewModule.cpp
--- a/Engine/Plugins/Runtime/nDisplay/Source/DisplayClusterScenePreview/Private/DisplayClusterScenePreviewModule.cpp
+++ b/Engine/Plugins/Runtime/nDisplay/Source/DisplayClusterScenePreview/Private/DisplayClusterScenePreviewModule.cpp
@@ -312,10 +312,10 @@ bool FDisplayClusterScenePreviewModule::InternalRenderImmediate(FRendererConfig&
 	{
 		UDisplayClusterConfigurationData* Config = RendererConfig.RootActor->GetConfigData();
 
-		for (const TPair<FString, UDisplayClusterConfigurationClusterNode*>& NodePair : Config->Cluster->Nodes)
+		for (const UE_TRANSITIONAL_OBJECT_PTR_TEMPLATE2_ARG2(TPair, FString, UDisplayClusterConfigurationClusterNode)& NodePair : Config->Cluster->Nodes)
 		{
 			const UDisplayClusterConfigurationClusterNode* Node = NodePair.Value;
-			for (const TPair<FString, UDisplayClusterConfigurationViewport*>& ViewportPair : Node->Viewports)
+			for (const UE_TRANSITIONAL_OBJECT_PTR_TEMPLATE2_ARG2(TPair, FString, UDisplayClusterConfigurationViewport)& ViewportPair : Node->Viewports)
 			{
 				UDisplayClusterPreviewComponent* PreviewComp = RendererConfig.RootActor->GetPreviewComponent(NodePair.Key, ViewportPair.Key);
 				if (PreviewComp)
diff --git a/Engine/Plugins/Runtime/nDisplay/Source/DisplayClusterTests/Private/DisplayCluster.spec.cpp b/Engine/Plugins/Runtime/nDisplay/Source/DisplayClusterTests/Private/DisplayCluster.spec.cpp
--- a/Engine/Plugins/Runtime/nDisplay/Source/DisplayClusterTests/Private/DisplayCluster.spec.cpp
+++ b/Engine/Plugins/Runtime/nDisplay/Source/DisplayClusterTests/Private/DisplayCluster.spec.cpp
@@ -355,7 +355,7 @@ void FDisplayClusterSpec::Define()
 			DisplayClusterTestUtils::AddClusterNodeToCluster(ClusterAsset, AssetRootCluster);
 
 			TestEqual(TEXT("instance key count (compared to CDO key count)"), InstancedRootCluster->Nodes.Num(), AssetRootCluster->Nodes.Num());
-			for (const TPair<FString, UDisplayClusterConfigurationClusterNode*>& Pair : AssetRootCluster->Nodes)
+			for (const UE_TRANSITIONAL_OBJECT_PTR_TEMPLATE2_ARG2(TPair, FString, UDisplayClusterConfigurationClusterNode)& Pair : AssetRootCluster->Nodes)
 			{
 				if (TestTrue("Instance has entry corresponding to CDO key", InstancedRootCluster->Nodes.Contains(Pair.Key)))
 				{
@@ -647,17 +647,17 @@ void FDisplayClusterSpec::RegisterPropagationTests(TFunction<void()> BeforeTestA
 
 		if (TestEqual(TEXT("Both nodes are present in asset's cluster"), AssetRootCluster->Nodes.Num(), 2))
 		{
-			for (const TPair<FString, UDisplayClusterConfigurationClusterNode*>& Pair : AssetRootCluster->Nodes)
+			for (const UE_TRANSITIONAL_OBJECT_PTR_TEMPLATE2_ARG2(TPair, FString, UDisplayClusterConfigurationClusterNode)& Pair : AssetRootCluster->Nodes)
 			{
-				TestNotNull(TEXT("All asset nodes are valid"), Pair.Value);
+				TestNotNull(TEXT("All asset nodes are valid"), ToRawPtr(Pair.Value));
 			}
 		}
 		
 		if (TestEqual(TEXT("Both nodes are present in instanced cluster"), InstancedRootCluster->Nodes.Num(), 2))
 		{
-			for (const TPair<FString, UDisplayClusterConfigurationClusterNode*>& Pair : AssetRootCluster->Nodes)
+			for (const UE_TRANSITIONAL_OBJECT_PTR_TEMPLATE2_ARG2(TPair, FString, UDisplayClusterConfigurationClusterNode)& Pair : AssetRootCluster->Nodes)
 			{
-				TestNotNull(TEXT("All instance nodes are valid"), Pair.Value);
+				TestNotNull(TEXT("All instance nodes are valid"), ToRawPtr(Pair.Value));
 			}
 		}
 	});
@@ -670,17 +670,17 @@ void FDisplayClusterSpec::RegisterPropagationTests(TFunction<void()> BeforeTestA
 
 		if (TestEqual(TEXT("Both viewports are present in asset's node"), AssetClusterNode->Viewports.Num(), 2))
 		{
-			for (const TPair<FString, UDisplayClusterConfigurationViewport*>& Pair : AssetClusterNode->Viewports)
+			for (const UE_TRANSITIONAL_OBJECT_PTR_TEMPLATE2_ARG2(TPair, FString, UDisplayClusterConfigurationViewport)& Pair : AssetClusterNode->Viewports)
 			{
-				TestNotNull(TEXT("All asset viewports are valid"), Pair.Value);
+				TestNotNull(TEXT("All asset viewports are valid"), ToRawPtr(Pair.Value));
 			}
 		}
 		
 		if (TestEqual(TEXT("Both viewports are present in instanced node"), InstancedClusterNode->Viewports.Num(), 2))
 		{
-			for (const TPair<FString, UDisplayClusterConfigurationViewport*>& Pair : AssetClusterNode->Viewports)
+			for (const UE_TRANSITIONAL_OBJECT_PTR_TEMPLATE2_ARG2(TPair, FString, UDisplayClusterConfigurationViewport)& Pair : AssetClusterNode->Viewports)
 			{
-				TestNotNull(TEXT("All instanced viewports are valid"), Pair.Value);
+				TestNotNull(TEXT("All instanced viewports are valid"), ToRawPtr(Pair.Value));
 			}
 		}
 	});
