[EpicStageApp] Use interactive changes for light card property updates

diff --git a/Engine/Plugins/Runtime/nDisplay/Source/DisplayClusterLightCardEditor/Private/Viewport/DisplayClusterLightCardEditorViewportClient.cpp b/Engine/Plugins/Runtime/nDisplay/Source/DisplayClusterLightCardEditor/Private/Viewport/DisplayClusterLightCardEditorViewportClient.cpp
--- a/Engine/Plugins/Runtime/nDisplay/Source/DisplayClusterLightCardEditor/Private/Viewport/DisplayClusterLightCardEditorViewportClient.cpp
+++ b/Engine/Plugins/Runtime/nDisplay/Source/DisplayClusterLightCardEditor/Private/Viewport/DisplayClusterLightCardEditorViewportClient.cpp
@@ -1319,6 +1319,16 @@ void FDisplayClusterLightCardEditorViewportClient::UpdateProxyTransforms()
 				LightCardProxy.Proxy->SetActorLocation(LightCardProxy.LevelInstance->GetActorLocation() - RootActorLevelInstance->GetActorLocation());
 				LightCardProxy.Proxy->SetActorRotation(LightCardProxy.LevelInstance->GetActorRotation() - RootActorLevelInstance->GetActorRotation());
 				LightCardProxy.Proxy->SetActorScale3D(LightCardProxy.LevelInstance->GetActorScale3D());
+
+				// Need to update these manually or the proxy's position will be out of sync next update
+				LightCardProxy.Proxy->Longitude = LightCardProxy.LevelInstance->Longitude;
+				LightCardProxy.Proxy->Latitude = LightCardProxy.LevelInstance->Latitude;
+				LightCardProxy.Proxy->DistanceFromCenter = LightCardProxy.LevelInstance->DistanceFromCenter;
+				LightCardProxy.Proxy->Spin = LightCardProxy.LevelInstance->Spin;
+				LightCardProxy.Proxy->Pitch = LightCardProxy.LevelInstance->Pitch;
+				LightCardProxy.Proxy->Yaw = LightCardProxy.LevelInstance->Yaw;
+				LightCardProxy.Proxy->Scale = LightCardProxy.LevelInstance->Scale;
+				LightCardProxy.Proxy->UVCoordinates = LightCardProxy.LevelInstance->UVCoordinates;
 			}
 		}
 	}
diff --git a/Engine/Plugins/Runtime/nDisplay/Source/DisplayClusterScenePreview/Private/DisplayClusterLightCardEditorHelper.cpp b/Engine/Plugins/Runtime/nDisplay/Source/DisplayClusterScenePreview/Private/DisplayClusterLightCardEditorHelper.cpp
--- a/Engine/Plugins/Runtime/nDisplay/Source/DisplayClusterScenePreview/Private/DisplayClusterLightCardEditorHelper.cpp
+++ b/Engine/Plugins/Runtime/nDisplay/Source/DisplayClusterScenePreview/Private/DisplayClusterLightCardEditorHelper.cpp
@@ -1515,25 +1515,29 @@ void FDisplayClusterLightCardEditorHelper::PostEditChangePropertiesForMovedLight
 
 	LightCard.Modify();
 
-	TArray<const FProperty*> ChangedProperties;
 	auto ModifyLightCardProperty = [&](const FName& PropertyName)
 	{
 		FProperty* Property = FindFProperty<FProperty>(LightCard.GetClass(), PropertyName);
-		FPropertyChangedEvent PropertyEvent(Property);
-		LightCard.PostEditChangeProperty(PropertyEvent);
 
-		ChangedProperties.Add(Property);
+		// Broadcast the event directly instead of PostEditChangeProperty on the object itself, which would attempt to create unnecessary snapshots for each call
+		FPropertyChangedEvent PropertyChangedEvent(Property, EPropertyChangeType::Interactive);
+		FCoreUObjectDelegates::OnObjectPropertyChanged.Broadcast(&LightCard, PropertyChangedEvent);
+
+		return Property;
 	};
 
-	ModifyLightCardProperty(GET_MEMBER_NAME_CHECKED(ADisplayClusterLightCardActor, Longitude));
-	ModifyLightCardProperty(GET_MEMBER_NAME_CHECKED(ADisplayClusterLightCardActor, Latitude));
-	ModifyLightCardProperty(GET_MEMBER_NAME_CHECKED(ADisplayClusterLightCardActor, DistanceFromCenter));
-	ModifyLightCardProperty(GET_MEMBER_NAME_CHECKED(ADisplayClusterLightCardActor, Spin));
-	ModifyLightCardProperty(GET_MEMBER_NAME_CHECKED(ADisplayClusterLightCardActor, Pitch));
-	ModifyLightCardProperty(GET_MEMBER_NAME_CHECKED(ADisplayClusterLightCardActor, Yaw));
-	ModifyLightCardProperty(GET_MEMBER_NAME_CHECKED(ADisplayClusterLightCardActor, Scale));
-	ModifyLightCardProperty(GET_MEMBER_NAME_CHECKED(ADisplayClusterLightCardActor, UVCoordinates));
+	const FProperty* ChangedProperties[] =
+	{
+		ModifyLightCardProperty(GET_MEMBER_NAME_CHECKED(ADisplayClusterLightCardActor, Longitude)),
+		ModifyLightCardProperty(GET_MEMBER_NAME_CHECKED(ADisplayClusterLightCardActor, Latitude)),
+		ModifyLightCardProperty(GET_MEMBER_NAME_CHECKED(ADisplayClusterLightCardActor, DistanceFromCenter)),
+		ModifyLightCardProperty(GET_MEMBER_NAME_CHECKED(ADisplayClusterLightCardActor, Spin)),
+		ModifyLightCardProperty(GET_MEMBER_NAME_CHECKED(ADisplayClusterLightCardActor, Pitch)),
+		ModifyLightCardProperty(GET_MEMBER_NAME_CHECKED(ADisplayClusterLightCardActor, Yaw)),
+		ModifyLightCardProperty(GET_MEMBER_NAME_CHECKED(ADisplayClusterLightCardActor, Scale)),
+		ModifyLightCardProperty(GET_MEMBER_NAME_CHECKED(ADisplayClusterLightCardActor, UVCoordinates)),
+	};
 
-	SnapshotTransactionBuffer(&LightCard, MakeArrayView(ChangedProperties));
+	SnapshotTransactionBuffer(&LightCard, ChangedProperties);
 }
 #endif
