nDisplay: Fix saving a sequence with nDisplay not working. Preview components and their properties are meant to be transient and created for each use. Setting the properties to transient prevents sequencer from reloading the components and also fixes a crash on level load. The component class itself isn't marked transient to avoid warnings when loading existing projects.

diff --git a/Engine/Plugins/Runtime/nDisplay/Source/DisplayCluster/Public/Components/DisplayClusterPreviewComponent.h b/Engine/Plugins/Runtime/nDisplay/Source/DisplayCluster/Public/Components/DisplayClusterPreviewComponent.h
--- a/Engine/Plugins/Runtime/nDisplay/Source/DisplayCluster/Public/Components/DisplayClusterPreviewComponent.h
+++ b/Engine/Plugins/Runtime/nDisplay/Source/DisplayCluster/Public/Components/DisplayClusterPreviewComponent.h
@@ -120,43 +120,43 @@ protected:
 #if WITH_EDITORONLY_DATA
 protected:
 	// Texture for preview material
-	UPROPERTY(VisibleAnywhere, BlueprintReadOnly, Category = "Preview", meta = (DisplayName = "Render Target"))
+	UPROPERTY(Transient, VisibleAnywhere, BlueprintReadOnly, Category = "Preview", meta = (DisplayName = "Render Target"))
 	TObjectPtr<UTextureRenderTarget2D> RenderTarget;
 
 	// Texture when DCRA has post process disabled but is requesting a post process render target.
-	UPROPERTY(VisibleAnywhere, BlueprintReadOnly, Category = "Preview")
+	UPROPERTY(Transient, VisibleAnywhere, BlueprintReadOnly, Category = "Preview")
 	TObjectPtr<UTextureRenderTarget2D> RenderTargetPostProcess;
 	
 private:
 	// Saved mesh policy params
-	UPROPERTY()
+	UPROPERTY(Transient)
 	FDisplayClusterConfigurationProjection WarpMeshSavedProjectionPolicy;
 
-	UPROPERTY()
+	UPROPERTY(Transient)
 	TObjectPtr<ADisplayClusterRootActor> RootActor = nullptr;
 
-	UPROPERTY()
+	UPROPERTY(Transient)
 	FString ViewportId;
 
-	UPROPERTY()
+	UPROPERTY(Transient)
 	FString ClusterNodeId;
 
-	UPROPERTY()
+	UPROPERTY(Transient)
 	TObjectPtr<UDisplayClusterConfigurationViewport> ViewportConfig = nullptr;
 
-	UPROPERTY()
+	UPROPERTY(Transient)
 	TObjectPtr<UMeshComponent> PreviewMesh = nullptr;
 
-	UPROPERTY()
+	UPROPERTY(Transient)
 	bool bIsRootActorPreviewMesh = false;
 
-	UPROPERTY()
+	UPROPERTY(Transient)
 	TObjectPtr<UMaterial> OriginalMaterial = nullptr;
 
-	UPROPERTY()
+	UPROPERTY(Transient)
 	TObjectPtr<UMaterial> PreviewMaterial = nullptr;
 
-	UPROPERTY()
+	UPROPERTY(Transient)
 	TObjectPtr<UMaterialInstanceDynamic> PreviewMaterialInstance = nullptr;
 
 	UPROPERTY(Transient)
