nDisplay: Fixed issue where the preview world created for rending UV light cards was causing nodes to crash when MU was used. Issue was caused by the engine considering GamePreview world types to be "leaked" worlds even if those worlds have been marked for GC and no references to them exist.

diff --git a/Engine/Plugins/Runtime/nDisplay/Source/DisplayCluster/Private/Render/Viewport/DisplayClusterViewportManager.cpp b/Engine/Plugins/Runtime/nDisplay/Source/DisplayCluster/Private/Render/Viewport/DisplayClusterViewportManager.cpp
--- a/Engine/Plugins/Runtime/nDisplay/Source/DisplayCluster/Private/Render/Viewport/DisplayClusterViewportManager.cpp
+++ b/Engine/Plugins/Runtime/nDisplay/Source/DisplayCluster/Private/Render/Viewport/DisplayClusterViewportManager.cpp
@@ -64,8 +64,7 @@ FDisplayClusterViewportManager::FDisplayClusterViewportManager()
 
 	RenderTargetManager = MakeShared<FDisplayClusterRenderTargetManager, ESPMode::ThreadSafe>(ViewportManagerProxy);
 	PostProcessManager  = MakeShared<FDisplayClusterViewportPostProcessManager, ESPMode::ThreadSafe>(*this);
-	// DISABLE Light Card Manager until UE-159748 is resolved.
-//	LightCardManager = MakeShared<FDisplayClusterViewportLightCardManager, ESPMode::ThreadSafe>(*this);
+	LightCardManager = MakeShared<FDisplayClusterViewportLightCardManager, ESPMode::ThreadSafe>(*this);
 
 	// initialize proxy
 	ViewportManagerProxy->Initialize(*this);
diff --git a/Engine/Plugins/Runtime/nDisplay/Source/DisplayCluster/Private/Render/Viewport/LightCard/DisplayClusterViewportLightCardManager.cpp b/Engine/Plugins/Runtime/nDisplay/Source/DisplayCluster/Private/Render/Viewport/LightCard/DisplayClusterViewportLightCardManager.cpp
--- a/Engine/Plugins/Runtime/nDisplay/Source/DisplayCluster/Private/Render/Viewport/LightCard/DisplayClusterViewportLightCardManager.cpp
+++ b/Engine/Plugins/Runtime/nDisplay/Source/DisplayCluster/Private/Render/Viewport/LightCard/DisplayClusterViewportLightCardManager.cpp
@@ -57,15 +57,20 @@ void FDisplayClusterViewportLightCardManager::Release()
 {
 	// The destructor is usually called from the rendering thread, so Release() must be called first from the game thread.
 	const bool bIsInRenderingThread = IsInRenderingThread();
-	check(!bIsInRenderingThread || (bIsInRenderingThread && PreviewScene.IsValid() == false));
+	check(!bIsInRenderingThread || (bIsInRenderingThread && PreviewWorld == nullptr));
 
 	// Deleting PreviewScene is only called from the game thread
-	PreviewScene.Reset();
+	DestroyPreviewWorld();
 
 	UVLightCards.Empty();
 	ReleaseUVLightCardMap();
 }
 
+void FDisplayClusterViewportLightCardManager::AddReferencedObjects(FReferenceCollector& Collector)
+{
+	Collector.AddReferencedObject(PreviewWorld);
+}
+
 FRHITexture* FDisplayClusterViewportLightCardManager::GetUVLightCardMap_RenderThread() const
 { 
 	check(IsInRenderingThread());
@@ -94,22 +99,19 @@ void FDisplayClusterViewportLightCardManager::UpdateConfiguration()
 
 void FDisplayClusterViewportLightCardManager::HandleStartScene()
 {
-	PreviewScene = MakeShared<FPreviewScene>(FPreviewScene::ConstructionValues()
-		.SetEditor(false)
-		.SetCreatePhysicsScene(false)
-		.SetCreateDefaultLighting(false));
+	InitializePreviewWorld();
 }
 
 void FDisplayClusterViewportLightCardManager::HandleEndScene()
 {
-	PreviewScene.Reset();
+	DestroyPreviewWorld();
 }
 
 void FDisplayClusterViewportLightCardManager::RenderFrame()
 {
-	if (PreviewScene.IsValid() && ProxyData.IsValid())
+	if (PreviewWorld && ProxyData.IsValid())
 	{
-		if (FSceneInterface* SceneInterface = PreviewScene->GetScene())
+		if (FSceneInterface* SceneInterface = PreviewWorld->Scene)
 		{
 			InitializeUVLightCardMap();
 
@@ -143,6 +145,44 @@ void FDisplayClusterViewportLightCardManager::RenderFrame()
 	}
 }
 
+void FDisplayClusterViewportLightCardManager::InitializePreviewWorld()
+{
+	if (!PreviewWorld)
+	{
+		FName UniqueWorldName = MakeUniqueObjectName(GetTransientPackage(), UWorld::StaticClass(), FName(TEXT("DisplayClusterLightCardManager_PreviewWorld")));
+		PreviewWorld = NewObject<UWorld>(GetTransientPackage(), UniqueWorldName);
+		PreviewWorld->WorldType = EWorldType::GamePreview;
+
+		FWorldContext& WorldContext = GEngine->CreateNewWorldContext(PreviewWorld->WorldType);
+		WorldContext.SetCurrentWorld(PreviewWorld);
+
+		PreviewWorld->InitializeNewWorld(UWorld::InitializationValues()
+			.AllowAudioPlayback(false)
+			.CreatePhysicsScene(false)
+			.RequiresHitProxies(false)
+			.CreateNavigation(false)
+			.CreateAISystem(false)
+			.ShouldSimulatePhysics(false)
+			.SetTransactional(false));
+	}
+}
+
+void FDisplayClusterViewportLightCardManager::DestroyPreviewWorld()
+{
+	if (PreviewWorld)
+	{
+		// Hack to avoid issue where the engine considers this world a leaked object; When UEngine loads a new map, it checks to see if there are any UWorlds
+		// still in memory that aren't what it considers "persistent" worlds, worlds with type Inactive or EditorPreview. Eeven if the UWorld object has been marked for
+		// GC and has no references to it, UEngine will still flag it as "leaked" unless it is one of these two types.
+		PreviewWorld->WorldType = EWorldType::Inactive;
+
+		GEngine->DestroyWorldContext(PreviewWorld);
+		PreviewWorld->DestroyWorld(true);
+		PreviewWorld->MarkObjectsPendingKill();
+		PreviewWorld = nullptr;
+	}
+}
+
 void FDisplayClusterViewportLightCardManager::InitializeUVLightCardMap()
 {
 	const uint32 LightCardTextureSize = CVarUVLightCardTextureSize.GetValueOnAnyThread();
diff --git a/Engine/Plugins/Runtime/nDisplay/Source/DisplayCluster/Private/Render/Viewport/LightCard/DisplayClusterViewportLightCardManager.h b/Engine/Plugins/Runtime/nDisplay/Source/DisplayCluster/Private/Render/Viewport/LightCard/DisplayClusterViewportLightCardManager.h
--- a/Engine/Plugins/Runtime/nDisplay/Source/DisplayCluster/Private/Render/Viewport/LightCard/DisplayClusterViewportLightCardManager.h
+++ b/Engine/Plugins/Runtime/nDisplay/Source/DisplayCluster/Private/Render/Viewport/LightCard/DisplayClusterViewportLightCardManager.h
@@ -4,6 +4,7 @@
 
 #include "CoreMinimal.h"
 #include "Render/Viewport/IDisplayClusterViewportLightCardManager.h"
+#include "UObject/GCObject.h"
 
 #include "RenderResource.h"
 #include "UnrealClient.h"
@@ -12,6 +13,7 @@ class ADisplayClusterLightCardActor;
 class FDisplayClusterViewportManager;
 class FPreviewScene;
 class FSceneInterface;
+class UWorld;
 
 /** A render targetable texture resource used to render the UV light cards to */
 class FDisplayClusterLightCardMap : public FTexture, public FRenderTarget
@@ -36,7 +38,7 @@ private:
 };
 
 /** Manages the rendering of UV light cards for the viewport manager */
-class FDisplayClusterViewportLightCardManager : public IDisplayClusterViewportLightCardManager
+class FDisplayClusterViewportLightCardManager : public IDisplayClusterViewportLightCardManager, public FGCObject
 {
 public:
 	FDisplayClusterViewportLightCardManager(FDisplayClusterViewportManager& InViewportManager);
@@ -44,6 +46,11 @@ public:
 
 	void Release();
 
+	//~ Begin FGCObject interface
+	virtual void AddReferencedObjects(FReferenceCollector& Collector) override;
+	virtual FString GetReferencerName() const override { return TEXT("FDisplayClusterViewportLightCardManager"); }
+	//~ End FGCObject interface
+
 public:
 	//~ Begin IDisplayClusterViewportLightCardManager interface
 	virtual FRHITexture* GetUVLightCardMap_RenderThread() const override;
@@ -58,6 +65,12 @@ public:
 	void RenderFrame();
 
 private:
+	/** Initializes the preview world used to render the UV light cards*/
+	void InitializePreviewWorld();
+
+	/** Destroys the preview world used to render the UV light cards */
+	void DestroyPreviewWorld();
+
 	/** Initializes the UV light card map texture */
 	void InitializeUVLightCardMap();
 
@@ -68,8 +81,8 @@ private:
 	/** A reference to the owning viewport manager */
 	FDisplayClusterViewportManager& ViewportManager;
 
-	/** The preview scene the UV light card proxies live in for rendering */
-	TSharedPtr<FPreviewScene> PreviewScene = nullptr;
+	/** The preview world the UV light card proxies live in for rendering */
+	UWorld* PreviewWorld = nullptr;
 
 	/** The list of UV light card actors that are referenced by the root actor */
 	TArray<ADisplayClusterLightCardActor*> UVLightCards;
