nDisplay: Fixed issue with the previous fix for UE-159748 that was causing a crash when exiting the editor with the nDisplay config editor open

diff --git a/Engine/Plugins/Runtime/nDisplay/Source/DisplayCluster/Private/Render/Viewport/LightCard/DisplayClusterViewportLightCardManager.cpp b/Engine/Plugins/Runtime/nDisplay/Source/DisplayCluster/Private/Render/Viewport/LightCard/DisplayClusterViewportLightCardManager.cpp
--- a/Engine/Plugins/Runtime/nDisplay/Source/DisplayCluster/Private/Render/Viewport/LightCard/DisplayClusterViewportLightCardManager.cpp
+++ b/Engine/Plugins/Runtime/nDisplay/Source/DisplayCluster/Private/Render/Viewport/LightCard/DisplayClusterViewportLightCardManager.cpp
@@ -172,12 +172,12 @@ void FDisplayClusterViewportLightCardManager::DestroyPreviewWorld()
 	if (PreviewWorld)
 	{
 		// Hack to avoid issue where the engine considers this world a leaked object; When UEngine loads a new map, it checks to see if there are any UWorlds
-		// still in memory that aren't what it considers "persistent" worlds, worlds with type Inactive or EditorPreview. Eeven if the UWorld object has been marked for
+		// still in memory that aren't what it considers "persistent" worlds, worlds with type Inactive or EditorPreview. Even if the UWorld object has been marked for
 		// GC and has no references to it, UEngine will still flag it as "leaked" unless it is one of these two types.
 		PreviewWorld->WorldType = EWorldType::Inactive;
 
 		GEngine->DestroyWorldContext(PreviewWorld);
-		PreviewWorld->DestroyWorld(true);
+		PreviewWorld->DestroyWorld(false);
 		PreviewWorld->MarkObjectsPendingKill();
 		PreviewWorld = nullptr;
 	}
