nDisplay: new immediate media capture cleaning (use ImmediateCapture always, no more intermediate RTT)

diff --git a/Engine/Plugins/Runtime/nDisplay/Source/DisplayClusterConfiguration/Public/DisplayClusterConfigurationTypes_Media.h b/Engine/Plugins/Runtime/nDisplay/Source/DisplayClusterConfiguration/Public/DisplayClusterConfigurationTypes_Media.h
--- a/Engine/Plugins/Runtime/nDisplay/Source/DisplayClusterConfiguration/Public/DisplayClusterConfigurationTypes_Media.h
+++ b/Engine/Plugins/Runtime/nDisplay/Source/DisplayClusterConfiguration/Public/DisplayClusterConfigurationTypes_Media.h
@@ -9,8 +9,6 @@
 #include "MediaTexture.h"
 #include "MediaOutput.h"
 
-#include "Engine/TextureRenderTarget2D.h"
-
 #include "DisplayClusterConfigurationTypes_Media.generated.h"
 
 
@@ -57,10 +55,6 @@ public:
 	/** Media output to use */
 	UPROPERTY(EditAnywhere, BlueprintReadWrite, Category = "Media Output", meta = (EditCondition = "bEnabled"))
 	TObjectPtr<UMediaOutput> MediaOutput = nullptr;
-
-	/** Render target to use */
-	UPROPERTY(EditAnywhere, BlueprintReadWrite, Category = "Media Output", meta = (EditCondition = "bEnabled"))
-	TObjectPtr<UTextureRenderTarget2D> RenderTarget = nullptr;
 };
 
 
diff --git a/Engine/Plugins/Runtime/nDisplay/Source/DisplayClusterMedia/Private/Capture/DisplayClusterMediaCaptureBase.cpp b/Engine/Plugins/Runtime/nDisplay/Source/DisplayClusterMedia/Private/Capture/DisplayClusterMediaCaptureBase.cpp
--- a/Engine/Plugins/Runtime/nDisplay/Source/DisplayClusterMedia/Private/Capture/DisplayClusterMediaCaptureBase.cpp
+++ b/Engine/Plugins/Runtime/nDisplay/Source/DisplayClusterMedia/Private/Capture/DisplayClusterMediaCaptureBase.cpp
@@ -8,20 +8,17 @@
 
 #include "MediaCapture.h"
 #include "MediaOutput.h"
-#include "Engine/TextureRenderTarget2D.h"
 
 #include "RenderGraphBuilder.h"
 #include "RHIResources.h"
 
 
 
-FDisplayClusterMediaCaptureBase::FDisplayClusterMediaCaptureBase(const FString& InMediaId, const FString& InClusterNodeId, UMediaOutput* InMediaOutput, UTextureRenderTarget2D* InRenderTarget)
+FDisplayClusterMediaCaptureBase::FDisplayClusterMediaCaptureBase(const FString& InMediaId, const FString& InClusterNodeId, UMediaOutput* InMediaOutput)
 	: FDisplayClusterMediaBase(InMediaId, InClusterNodeId)
 	, MediaOutput(InMediaOutput)
-	, RenderTarget(InRenderTarget)
 {
-	// We expect these to always be valid
-	check(InMediaOutput && InRenderTarget);
+	check(InMediaOutput);
 
 	IDisplayCluster::Get().GetCallbacks().OnDisplayClusterPostTick().AddRaw(this, &FDisplayClusterMediaCaptureBase::OnPostClusterTick);
 }
@@ -35,7 +32,6 @@ FDisplayClusterMediaCaptureBase::~FDisplayClusterMediaCaptureBase()
 void FDisplayClusterMediaCaptureBase::AddReferencedObjects(FReferenceCollector& Collector)
 {
 	Collector.AddReferencedObject(MediaOutput);
-	Collector.AddReferencedObject(RenderTarget);
 
 	if (MediaCapture)
 	{
@@ -45,7 +41,7 @@ void FDisplayClusterMediaCaptureBase::AddReferencedObjects(FReferenceCollector&
 
 bool FDisplayClusterMediaCaptureBase::StartCapture()
 {
-	if (MediaOutput && RenderTarget && !MediaCapture)
+	if (MediaOutput && !MediaCapture)
 	{
 		MediaCapture = MediaOutput->CreateMediaCapture();
 		if (MediaCapture)
@@ -73,9 +69,8 @@ void FDisplayClusterMediaCaptureBase::StopCapture()
 void FDisplayClusterMediaCaptureBase::ExportMediaData(FRDGBuilder& GraphBuilder, const FMediaTextureInfo& TextureInfo)
 {
 	FRHITexture* const SrcTexture = TextureInfo.Texture;
-	FRHITexture* const DstTexture = GetRenderTarget()->GetRenderTargetResource()->GetTextureRHI();
 
-	if (SrcTexture && DstTexture)
+	if (SrcTexture)
 	{
 		MediaCapture->SetValidSourceGPUMask(GraphBuilder.RHICmdList.GetGPUMask());
 
diff --git a/Engine/Plugins/Runtime/nDisplay/Source/DisplayClusterMedia/Private/Capture/DisplayClusterMediaCaptureBase.h b/Engine/Plugins/Runtime/nDisplay/Source/DisplayClusterMedia/Private/Capture/DisplayClusterMediaCaptureBase.h
--- a/Engine/Plugins/Runtime/nDisplay/Source/DisplayClusterMedia/Private/Capture/DisplayClusterMediaCaptureBase.h
+++ b/Engine/Plugins/Runtime/nDisplay/Source/DisplayClusterMedia/Private/Capture/DisplayClusterMediaCaptureBase.h
@@ -10,7 +10,6 @@ class FRDGBuilder;
 class FRHICommandListImmediate;
 class UMediaCapture;
 class UMediaOutput;
-class UTextureRenderTarget2D;
 
 
 /**
@@ -21,7 +20,7 @@ class FDisplayClusterMediaCaptureBase
 	, public FGCObject
 {
 public:
-	FDisplayClusterMediaCaptureBase(const FString& MediaId, const FString& ClusterNodeId, UMediaOutput* MediaOutput, UTextureRenderTarget2D* RenderTarget);
+	FDisplayClusterMediaCaptureBase(const FString& MediaId, const FString& ClusterNodeId, UMediaOutput* MediaOutput);
 	virtual ~FDisplayClusterMediaCaptureBase();
 
 public:
@@ -37,11 +36,6 @@ public:
 	virtual bool StartCapture();
 	virtual void StopCapture();
 
-	UTextureRenderTarget2D* GetRenderTarget()
-	{
-		return RenderTarget;
-	}
-
 protected:
 	void ExportMediaData(FRDGBuilder& GraphBuilder, const FMediaTextureInfo& TextureInfo);
 	void OnPostClusterTick();
@@ -53,7 +47,6 @@ private:
 	//~ Begin GC by AddReferencedObjects
 	UMediaOutput*           MediaOutput  = nullptr;
 	UMediaCapture*          MediaCapture = nullptr;
-	UTextureRenderTarget2D* RenderTarget = nullptr;
 	//~ End GC by AddReferencedObjects
 
 	// Used to restart media capture in the case it falls in error
diff --git a/Engine/Plugins/Runtime/nDisplay/Source/DisplayClusterMedia/Private/Capture/DisplayClusterMediaCaptureCamera.cpp b/Engine/Plugins/Runtime/nDisplay/Source/DisplayClusterMedia/Private/Capture/DisplayClusterMediaCaptureCamera.cpp
--- a/Engine/Plugins/Runtime/nDisplay/Source/DisplayClusterMedia/Private/Capture/DisplayClusterMediaCaptureCamera.cpp
+++ b/Engine/Plugins/Runtime/nDisplay/Source/DisplayClusterMedia/Private/Capture/DisplayClusterMediaCaptureCamera.cpp
@@ -15,8 +15,8 @@
 #include "RHIResources.h"
 
 
-FDisplayClusterMediaCaptureCamera::FDisplayClusterMediaCaptureCamera(const FString& InMediaId, const FString& InClusterNodeId, const FString& InCameraId, const FString& InViewportId, UMediaOutput* InMediaOutput, UTextureRenderTarget2D* InRenderTarget)
-	: FDisplayClusterMediaCaptureViewport(InMediaId, InClusterNodeId, InViewportId, InMediaOutput, InRenderTarget)
+FDisplayClusterMediaCaptureCamera::FDisplayClusterMediaCaptureCamera(const FString& InMediaId, const FString& InClusterNodeId, const FString& InCameraId, const FString& InViewportId, UMediaOutput* InMediaOutput)
+	: FDisplayClusterMediaCaptureViewport(InMediaId, InClusterNodeId, InViewportId, InMediaOutput)
 	, CameraId(InCameraId)
 {
 	if (const ADisplayClusterRootActor* const RootActor = IDisplayCluster::Get().GetGameMgr()->GetRootActor())
diff --git a/Engine/Plugins/Runtime/nDisplay/Source/DisplayClusterMedia/Private/Capture/DisplayClusterMediaCaptureCamera.h b/Engine/Plugins/Runtime/nDisplay/Source/DisplayClusterMedia/Private/Capture/DisplayClusterMediaCaptureCamera.h
--- a/Engine/Plugins/Runtime/nDisplay/Source/DisplayClusterMedia/Private/Capture/DisplayClusterMediaCaptureCamera.h
+++ b/Engine/Plugins/Runtime/nDisplay/Source/DisplayClusterMedia/Private/Capture/DisplayClusterMediaCaptureCamera.h
@@ -16,7 +16,7 @@ class FDisplayClusterMediaCaptureCamera
 	: public FDisplayClusterMediaCaptureViewport
 {
 public:
-	FDisplayClusterMediaCaptureCamera(const FString& MediaId, const FString& ClusterNodeId, const FString& CameraId, const FString& ViewportId, UMediaOutput* MediaOutput, UTextureRenderTarget2D* RenderTarget);
+	FDisplayClusterMediaCaptureCamera(const FString& MediaId, const FString& ClusterNodeId, const FString& CameraId, const FString& ViewportId, UMediaOutput* MediaOutput);
 
 protected:
 	
diff --git a/Engine/Plugins/Runtime/nDisplay/Source/DisplayClusterMedia/Private/Capture/DisplayClusterMediaCaptureNode.cpp b/Engine/Plugins/Runtime/nDisplay/Source/DisplayClusterMedia/Private/Capture/DisplayClusterMediaCaptureNode.cpp
--- a/Engine/Plugins/Runtime/nDisplay/Source/DisplayClusterMedia/Private/Capture/DisplayClusterMediaCaptureNode.cpp
+++ b/Engine/Plugins/Runtime/nDisplay/Source/DisplayClusterMedia/Private/Capture/DisplayClusterMediaCaptureNode.cpp
@@ -15,8 +15,8 @@
 #include "UnrealClient.h"
 
 
-FDisplayClusterMediaCaptureNode::FDisplayClusterMediaCaptureNode(const FString& InMediaId, const FString& InClusterNodeId, UMediaOutput* InMediaOutput, UTextureRenderTarget2D* InRenderTarget)
-	: FDisplayClusterMediaCaptureBase(InMediaId, InClusterNodeId, InMediaOutput, InRenderTarget)
+FDisplayClusterMediaCaptureNode::FDisplayClusterMediaCaptureNode(const FString& InMediaId, const FString& InClusterNodeId, UMediaOutput* InMediaOutput)
+	: FDisplayClusterMediaCaptureBase(InMediaId, InClusterNodeId, InMediaOutput)
 {
 }
 
diff --git a/Engine/Plugins/Runtime/nDisplay/Source/DisplayClusterMedia/Private/Capture/DisplayClusterMediaCaptureNode.h b/Engine/Plugins/Runtime/nDisplay/Source/DisplayClusterMedia/Private/Capture/DisplayClusterMediaCaptureNode.h
--- a/Engine/Plugins/Runtime/nDisplay/Source/DisplayClusterMedia/Private/Capture/DisplayClusterMediaCaptureNode.h
+++ b/Engine/Plugins/Runtime/nDisplay/Source/DisplayClusterMedia/Private/Capture/DisplayClusterMediaCaptureNode.h
@@ -16,7 +16,7 @@ class FDisplayClusterMediaCaptureNode
 	: public FDisplayClusterMediaCaptureBase
 {
 public:
-	FDisplayClusterMediaCaptureNode(const FString& MediaId, const FString& ClusterNodeId, UMediaOutput* MediaOutput, UTextureRenderTarget2D* RenderTarget);
+	FDisplayClusterMediaCaptureNode(const FString& MediaId, const FString& ClusterNodeId, UMediaOutput* MediaOutput);
 
 public:
 	virtual bool StartCapture() override;
diff --git a/Engine/Plugins/Runtime/nDisplay/Source/DisplayClusterMedia/Private/Capture/DisplayClusterMediaCaptureViewport.cpp b/Engine/Plugins/Runtime/nDisplay/Source/DisplayClusterMedia/Private/Capture/DisplayClusterMediaCaptureViewport.cpp
--- a/Engine/Plugins/Runtime/nDisplay/Source/DisplayClusterMedia/Private/Capture/DisplayClusterMediaCaptureViewport.cpp
+++ b/Engine/Plugins/Runtime/nDisplay/Source/DisplayClusterMedia/Private/Capture/DisplayClusterMediaCaptureViewport.cpp
@@ -12,8 +12,8 @@
 #include "RHIResources.h"
 
 
-FDisplayClusterMediaCaptureViewport::FDisplayClusterMediaCaptureViewport(const FString& InMediaId, const FString& InClusterNodeId, const FString& InViewportId, UMediaOutput* InMediaOutput, UTextureRenderTarget2D* InRenderTarget)
-	: FDisplayClusterMediaCaptureBase(InMediaId, InClusterNodeId, InMediaOutput, InRenderTarget)
+FDisplayClusterMediaCaptureViewport::FDisplayClusterMediaCaptureViewport(const FString& InMediaId, const FString& InClusterNodeId, const FString& InViewportId, UMediaOutput* InMediaOutput)
+	: FDisplayClusterMediaCaptureBase(InMediaId, InClusterNodeId, InMediaOutput)
 	, ViewportId(InViewportId)
 {
 }
diff --git a/Engine/Plugins/Runtime/nDisplay/Source/DisplayClusterMedia/Private/Capture/DisplayClusterMediaCaptureViewport.h b/Engine/Plugins/Runtime/nDisplay/Source/DisplayClusterMedia/Private/Capture/DisplayClusterMediaCaptureViewport.h
--- a/Engine/Plugins/Runtime/nDisplay/Source/DisplayClusterMedia/Private/Capture/DisplayClusterMediaCaptureViewport.h
+++ b/Engine/Plugins/Runtime/nDisplay/Source/DisplayClusterMedia/Private/Capture/DisplayClusterMediaCaptureViewport.h
@@ -17,7 +17,7 @@ class FDisplayClusterMediaCaptureViewport
 	: public FDisplayClusterMediaCaptureBase
 {
 public:
-	FDisplayClusterMediaCaptureViewport(const FString& MediaId, const FString& ClusterNodeId, const FString& ViewportId, UMediaOutput* MediaOutput, UTextureRenderTarget2D* RenderTarget);
+	FDisplayClusterMediaCaptureViewport(const FString& MediaId, const FString& ClusterNodeId, const FString& ViewportId, UMediaOutput* MediaOutput);
 
 public:
 	virtual bool StartCapture() override;
diff --git a/Engine/Plugins/Runtime/nDisplay/Source/DisplayClusterMedia/Private/DisplayClusterMediaModule.cpp b/Engine/Plugins/Runtime/nDisplay/Source/DisplayClusterMedia/Private/DisplayClusterMediaModule.cpp
--- a/Engine/Plugins/Runtime/nDisplay/Source/DisplayClusterMedia/Private/DisplayClusterMediaModule.cpp
+++ b/Engine/Plugins/Runtime/nDisplay/Source/DisplayClusterMedia/Private/DisplayClusterMediaModule.cpp
@@ -79,7 +79,7 @@ void FDisplayClusterMediaModule::InitializeMedia()
 				const FDisplayClusterConfigurationMedia& MediaSettings = ClusterNode->Media;
 
 				const bool bInputDataValid   = MediaSettings.MediaInput.MediaPlayer && MediaSettings.MediaInput.MediaSource && MediaSettings.MediaInput.MediaTexture;
-				const bool bOutputDataValid  = MediaSettings.MediaOutput.MediaOutput && MediaSettings.MediaOutput.RenderTarget;
+				const bool bOutputDataValid  = !!MediaSettings.MediaOutput.MediaOutput;
 
 				// Media input
 				if (MediaSettings.MediaInput.bEnabled && bInputDataValid)
@@ -104,7 +104,7 @@ void FDisplayClusterMediaModule::InitializeMedia()
 
 					CaptureNode = MakeUnique<FDisplayClusterMediaCaptureNode>(
 						MediaCaptureId, ClusterNodeId,
-						MediaSettings.MediaOutput.MediaOutput, MediaSettings.MediaOutput.RenderTarget);
+						MediaSettings.MediaOutput.MediaOutput);
 
 					AllCaptures.Add(CaptureNode.Get());
 				}
@@ -120,7 +120,7 @@ void FDisplayClusterMediaModule::InitializeMedia()
 					if (MediaSettings.bEnabled)
 					{
 						const bool bInputDataValid   = MediaSettings.MediaInput.MediaPlayer && MediaSettings.MediaInput.MediaSource && MediaSettings.MediaInput.MediaTexture;
-						const bool bOutputDataValid  = MediaSettings.MediaOutput.MediaOutput && MediaSettings.MediaOutput.RenderTarget;
+						const bool bOutputDataValid  = !!MediaSettings.MediaOutput.MediaOutput;
 
 						// Media input
 						if (MediaSettings.MediaInput.bEnabled && bInputDataValid)
@@ -148,7 +148,7 @@ void FDisplayClusterMediaModule::InitializeMedia()
 							TUniquePtr<FDisplayClusterMediaCaptureViewport> NewViewportCapture = MakeUnique<FDisplayClusterMediaCaptureViewport>(
 								MediaCaptureId, ClusterNodeId,
 								ViewportIt.Key,
-								MediaSettings.MediaOutput.MediaOutput, MediaSettings.MediaOutput.RenderTarget);
+								MediaSettings.MediaOutput.MediaOutput);
 
 							AllCaptures.Add(NewViewportCapture.Get());
 							CaptureViewports.Emplace(MediaCaptureId, MoveTemp(NewViewportCapture));
@@ -176,7 +176,7 @@ void FDisplayClusterMediaModule::InitializeMedia()
 							const FString MediaCaptureId   = RootActor->GetName() + FString("_") + ICVFXCameraName + FString("_icvfx_capture");
 							const FString MediaInputId     = RootActor->GetName() + FString("_") + ICVFXCameraName + FString("_icvfx_input");
 							const bool bInputDataValid     = MediaSettings.MediaInput.MediaPlayer && MediaSettings.MediaInput.MediaSource && MediaSettings.MediaInput.MediaTexture;
-							const bool bOutputDataValid    = MediaSettings.MediaOutput.MediaOutput && MediaSettings.MediaOutput.RenderTarget;
+							const bool bOutputDataValid    = !!MediaSettings.MediaOutput.MediaOutput;
 							const bool bClusterNodeMatches = MediaSettings.MediaOutputNode.Equals(ClusterNodeId, ESearchCase::IgnoreCase);
 
 							// Media sharing
@@ -190,7 +190,7 @@ void FDisplayClusterMediaModule::InitializeMedia()
 									TUniquePtr<FDisplayClusterMediaCaptureViewport> NewICVFXCapture = MakeUnique<FDisplayClusterMediaCaptureCamera>(
 										MediaCaptureId, ClusterNodeId,
 										ICVFXCameraName, ICVFXViewportId,
-										MediaSettings.MediaOutput.MediaOutput, MediaSettings.MediaOutput.RenderTarget);
+										MediaSettings.MediaOutput.MediaOutput);
 
 									AllCaptures.Add(NewICVFXCapture.Get());
 									CaptureViewports.Emplace(MediaCaptureId, MoveTemp(NewICVFXCapture));
@@ -230,7 +230,7 @@ void FDisplayClusterMediaModule::InitializeMedia()
 								TUniquePtr<FDisplayClusterMediaCaptureViewport> NewICVFXCapture = MakeUnique<FDisplayClusterMediaCaptureCamera>(
 									MediaCaptureId, ClusterNodeId,
 									ICVFXCameraName, ICVFXViewportId,
-									MediaSettings.MediaOutput.MediaOutput, MediaSettings.MediaOutput.RenderTarget);
+									MediaSettings.MediaOutput.MediaOutput);
 
 								AllCaptures.Add(NewICVFXCapture.Get());
 								CaptureViewports.Emplace(MediaCaptureId, MoveTemp(NewICVFXCapture));
