nDispaly: Added the color grading drawer to the nDisplay operator menu, which displays a list of color gradable objects relative to the active nDispaly stage actor, and displays the color grading properties and color wheels of the selected object.

diff --git a/Engine/Plugins/Runtime/nDisplay/Source/DisplayClusterColorGrading/DisplayClusterColorGrading.Build.cs b/Engine/Plugins/Runtime/nDisplay/Source/DisplayClusterColorGrading/DisplayClusterColorGrading.Build.cs
new file mode 100644
--- /dev/null
+++ b/Engine/Plugins/Runtime/nDisplay/Source/DisplayClusterColorGrading/DisplayClusterColorGrading.Build.cs
@@ -0,0 +1,33 @@
+// Copyright Epic Games, Inc. All Rights Reserved.
+
+using UnrealBuildTool;
+using System.IO;
+
+public class DisplayClusterColorGrading : ModuleRules
+{
+	public DisplayClusterColorGrading(ReadOnlyTargetRules ROTargetRules) : base(ROTargetRules)
+	{
+		PrivateDependencyModuleNames.AddRange(
+			new string[] {
+				"DisplayCluster",
+				"DisplayClusterConfiguration",
+				"DisplayClusterOperator",
+
+				"AppFramework",
+				"ColorCorrectRegions",
+				"Core",
+				"CoreUObject",
+				"DetailCustomizations",
+				"EditorStyle",
+				"Engine",
+				"InputCore",
+				"Kismet",
+				"PropertyEditor",
+				"Slate",
+				"SlateCore",
+				"ToolWidgets",
+				"UnrealEd",
+				"WorkspaceMenuStructure",
+			});
+	}
+}
diff --git a/Engine/Plugins/Runtime/nDisplay/Source/DisplayClusterColorGrading/Private/DataModelGenerators/DisplayClusterColorGradingGenerator_ColorCorrectRegion.cpp b/Engine/Plugins/Runtime/nDisplay/Source/DisplayClusterColorGrading/Private/DataModelGenerators/DisplayClusterColorGradingGenerator_ColorCorrectRegion.cpp
new file mode 100644
--- /dev/null
+++ b/Engine/Plugins/Runtime/nDisplay/Source/DisplayClusterColorGrading/Private/DataModelGenerators/DisplayClusterColorGradingGenerator_ColorCorrectRegion.cpp
@@ -0,0 +1,134 @@
+// Copyright Epic Games, Inc. All Rights Reserved.
+
+#include "DisplayClusterColorGradingGenerator_ColorCorrectRegion.h"
+
+#include "ColorCorrectRegion.h"
+#include "IDetailTreeNode.h"
+#include "IPropertyRowGenerator.h"
+#include "PropertyHandle.h"
+#include "Widgets/Input/SButton.h"
+#include "Widgets/Text/SInlineEditableTextBlock.h"
+
+#define LOCTEXT_NAMESPACE "DisplayClusterColorGrading"
+
+TSharedRef<IDisplayClusterColorGradingDataModelGenerator> FDisplayClusterColorGradingGenerator_ColorCorrectRegion::MakeInstance()
+{
+	return MakeShareable(new FDisplayClusterColorGradingGenerator_ColorCorrectRegion());
+}
+
+void FDisplayClusterColorGradingGenerator_ColorCorrectRegion::GenerateDataModel(IPropertyRowGenerator& PropertyRowGenerator, FDisplayClusterColorGradingDataModel& OutColorGradingDataModel)
+{
+	const TArray<TSharedRef<IDetailTreeNode>>& RootNodes = PropertyRowGenerator.GetRootTreeNodes();
+
+	const TSharedRef<IDetailTreeNode>* ColorCorrectionNodePtr = RootNodes.FindByPredicate([](const TSharedRef<IDetailTreeNode>& Node)
+	{
+		return Node->GetNodeName() == TEXT("Color Correction");
+	});
+
+	if (ColorCorrectionNodePtr)
+	{
+		const TSharedRef<IDetailTreeNode> ColorCorrectionNode = *ColorCorrectionNodePtr;
+		FDisplayClusterColorGradingDataModel::FColorGradingGroup ColorGradingGroup;
+
+		AddPropertiesToDetailsView(ColorCorrectionNode, ColorGradingGroup);
+
+		TArray<TSharedRef<IDetailTreeNode>> ColorCorrectionPropertyNodes;
+		ColorCorrectionNode->GetChildren(ColorCorrectionPropertyNodes);
+
+		const TSharedRef<IDetailTreeNode>* ColorGradingSettingsNodePtr = ColorCorrectionPropertyNodes.FindByPredicate([](const TSharedRef<IDetailTreeNode>& Node)
+		{
+			return Node->GetNodeName() == GET_MEMBER_NAME_CHECKED(AColorCorrectRegion, ColorGradingSettings);
+		});
+
+		if (ColorGradingSettingsNodePtr)
+		{
+			const TSharedRef<IDetailTreeNode> ColorGradingSettingsNode = *ColorGradingSettingsNodePtr;
+
+			TArray<TSharedRef<IDetailTreeNode>> ColorGradingPropertyNodes;
+			ColorGradingSettingsNode->GetChildren(ColorGradingPropertyNodes);
+
+			for (const TSharedRef<IDetailTreeNode>& PropertyNode : ColorGradingPropertyNodes)
+			{
+				FName NodeName = PropertyNode->GetNodeName();
+				if (NodeName == GET_MEMBER_NAME_CHECKED(FColorGradingSettings, Global) ||
+					NodeName == GET_MEMBER_NAME_CHECKED(FColorGradingSettings, Shadows) ||
+					NodeName == GET_MEMBER_NAME_CHECKED(FColorGradingSettings, Midtones) ||
+					NodeName == GET_MEMBER_NAME_CHECKED(FColorGradingSettings, Highlights))
+				{
+					TSharedPtr<IPropertyHandle> PropertyHandle = PropertyNode->CreatePropertyHandle();
+					PropertyHandle->MarkHiddenByCustomization();
+
+					FDisplayClusterColorGradingDataModel::FColorGradingElement ColorGradingElement = CreateColorGradingElement(PropertyNode, FText::FromName(PropertyNode->GetNodeName()));
+					ColorGradingGroup.ColorGradingElements.Add(ColorGradingElement);
+				}
+			}
+		}
+
+		OutColorGradingDataModel.ColorGradingGroups.Add(ColorGradingGroup);
+	}
+}
+
+FDisplayClusterColorGradingDataModel::FColorGradingElement FDisplayClusterColorGradingGenerator_ColorCorrectRegion::CreateColorGradingElement(const TSharedRef<IDetailTreeNode>& GroupNode, FText ElementLabel)
+{
+	FDisplayClusterColorGradingDataModel::FColorGradingElement ColorGradingElement;
+	ColorGradingElement.DisplayName = ElementLabel;
+
+	TArray<TSharedRef<IDetailTreeNode>> ChildNodes;
+	GroupNode->GetChildren(ChildNodes);
+
+	for (const TSharedRef<IDetailTreeNode>& ChildNode : ChildNodes)
+	{
+		TSharedPtr<IPropertyHandle> PropertyHandle = ChildNode->CreatePropertyHandle();
+		if (PropertyHandle.IsValid() && PropertyHandle->IsValidHandle())
+		{
+			const FString ColorGradingModeString = PropertyHandle->GetProperty()->GetMetaData(TEXT("ColorGradingMode")).ToLower();
+
+			if (!ColorGradingModeString.IsEmpty())
+			{
+				if (ColorGradingModeString.Compare(TEXT("saturation")) == 0)
+				{
+					ColorGradingElement.SaturationPropertyHandle = PropertyHandle;
+				}
+				else if (ColorGradingModeString.Compare(TEXT("contrast")) == 0)
+				{
+					ColorGradingElement.ContrastPropertyHandle = PropertyHandle;
+				}
+				else if (ColorGradingModeString.Compare(TEXT("gamma")) == 0)
+				{
+					ColorGradingElement.GammaPropertyHandle = PropertyHandle;
+				}
+				else if (ColorGradingModeString.Compare(TEXT("gain")) == 0)
+				{
+					ColorGradingElement.GainPropertyHandle = PropertyHandle;
+				}
+				else if (ColorGradingModeString.Compare(TEXT("offset")) == 0)
+				{
+					ColorGradingElement.OffsetPropertyHandle = PropertyHandle;
+				}
+			}
+		}
+	}
+
+	return ColorGradingElement;
+}
+
+void FDisplayClusterColorGradingGenerator_ColorCorrectRegion::AddPropertiesToDetailsView(const TSharedRef<IDetailTreeNode>& GroupNode, FDisplayClusterColorGradingDataModel::FColorGradingGroup& ColorGradingGroup)
+{
+	TArray<TSharedRef<IDetailTreeNode>> ChildNodes;
+	GroupNode->GetChildren(ChildNodes);
+
+	for (const TSharedRef<IDetailTreeNode>& ChildNode : ChildNodes)
+	{
+		TSharedPtr<IPropertyHandle> PropertyHandle = ChildNode->CreatePropertyHandle();
+		if (PropertyHandle.IsValid() && PropertyHandle->IsValidHandle())
+		{
+			const bool bIsColorProperty = PropertyHandle->GetProperty()->HasMetaData(TEXT("ColorGradingMode"));
+
+			if (!bIsColorProperty)
+			{
+				PropertyHandle->SetInstanceMetaData(TEXT("CategoryOverride"), GroupNode->GetNodeName().ToString());
+				ColorGradingGroup.DetailsViewPropertyHandles.Add(PropertyHandle);
+			}
+		}
+	}
+}
\ No newline at end of file
diff --git a/Engine/Plugins/Runtime/nDisplay/Source/DisplayClusterColorGrading/Private/DataModelGenerators/DisplayClusterColorGradingGenerator_ColorCorrectRegion.h b/Engine/Plugins/Runtime/nDisplay/Source/DisplayClusterColorGrading/Private/DataModelGenerators/DisplayClusterColorGradingGenerator_ColorCorrectRegion.h
new file mode 100644
--- /dev/null
+++ b/Engine/Plugins/Runtime/nDisplay/Source/DisplayClusterColorGrading/Private/DataModelGenerators/DisplayClusterColorGradingGenerator_ColorCorrectRegion.h
@@ -0,0 +1,26 @@
+// Copyright Epic Games, Inc. All Rights Reserved.
+
+#pragma once
+
+#include "CoreMinimal.h"
+#include "DisplayClusterColorGradingDataModel.h"
+
+class IDetailTreeNode;
+
+/** Color Grading Data Model Generator for the AColorCorrectionRegion actor class */
+class FDisplayClusterColorGradingGenerator_ColorCorrectRegion : public IDisplayClusterColorGradingDataModelGenerator
+{
+public:
+	static TSharedRef<IDisplayClusterColorGradingDataModelGenerator> MakeInstance();
+
+	//~ IDisplayClusterColorGradingDataModelGenerator interface
+	virtual void GenerateDataModel(IPropertyRowGenerator& PropertyRowGenerator, FDisplayClusterColorGradingDataModel& OutColorGradingDataModel) override;
+	//~ End IDisplayClusterColorGradingDataModelGenerator interface
+
+private:
+	/** Creates a new color grading element structure for the specified detail tree node, which is expected to have child color properties with the ColorGradingMode metadata set */
+	FDisplayClusterColorGradingDataModel::FColorGradingElement CreateColorGradingElement(const TSharedRef<IDetailTreeNode>& GroupNode, FText ElementLabel);
+
+	/** Adds all child properties of the specified detail tree node to the color grading group's list of properties to display in the details view*/
+	void AddPropertiesToDetailsView(const TSharedRef<IDetailTreeNode>& GroupNode, FDisplayClusterColorGradingDataModel::FColorGradingGroup& ColorGradingGroup);
+};
\ No newline at end of file
diff --git a/Engine/Plugins/Runtime/nDisplay/Source/DisplayClusterColorGrading/Private/DataModelGenerators/DisplayClusterColorGradingGenerator_PostProcessVolume.cpp b/Engine/Plugins/Runtime/nDisplay/Source/DisplayClusterColorGrading/Private/DataModelGenerators/DisplayClusterColorGradingGenerator_PostProcessVolume.cpp
new file mode 100644
--- /dev/null
+++ b/Engine/Plugins/Runtime/nDisplay/Source/DisplayClusterColorGrading/Private/DataModelGenerators/DisplayClusterColorGradingGenerator_PostProcessVolume.cpp
@@ -0,0 +1,124 @@
+// Copyright Epic Games, Inc. All Rights Reserved.
+
+#include "DisplayClusterColorGradingGenerator_PostProcessVolume.h"
+
+#include "Engine/PostProcessVolume.h"
+#include "IDetailTreeNode.h"
+#include "IPropertyRowGenerator.h"
+#include "PropertyHandle.h"
+#include "Widgets/Input/SButton.h"
+#include "Widgets/Text/SInlineEditableTextBlock.h"
+
+#define LOCTEXT_NAMESPACE "DisplayClusterColorGrading"
+
+TSharedRef<IDisplayClusterColorGradingDataModelGenerator> FDisplayClusterColorGradingGenerator_PostProcessVolume::MakeInstance()
+{
+	return MakeShareable(new FDisplayClusterColorGradingGenerator_PostProcessVolume());
+}
+
+void FDisplayClusterColorGradingGenerator_PostProcessVolume::GenerateDataModel(IPropertyRowGenerator& PropertyRowGenerator, FDisplayClusterColorGradingDataModel& OutColorGradingDataModel)
+{
+	const TArray<TSharedRef<IDetailTreeNode>>& RootNodes = PropertyRowGenerator.GetRootTreeNodes();
+
+	const TSharedRef<IDetailTreeNode>* ColorGradingNodePtr = RootNodes.FindByPredicate([](const TSharedRef<IDetailTreeNode>& Node)
+	{
+		return Node->GetNodeName() == TEXT("Color Grading");
+	});
+
+	if (ColorGradingNodePtr)
+	{
+		const TSharedRef<IDetailTreeNode> ColorGradingNode = *ColorGradingNodePtr;
+		FDisplayClusterColorGradingDataModel::FColorGradingGroup ColorGradingGroup;
+
+		TArray<TSharedRef<IDetailTreeNode>> PropertyGroupNodes;
+		ColorGradingNode->GetChildren(PropertyGroupNodes);
+
+		for (const TSharedRef<IDetailTreeNode>& PropertyGroupNode : PropertyGroupNodes)
+		{
+			FString CategoryName = TEXT(""); 
+			FString GroupName = TEXT("");
+			PropertyGroupNode->GetNodeName().ToString().Split(TEXT("|"), &CategoryName, &GroupName);
+
+			if (GroupName == TEXT("Global") || GroupName == TEXT("Shadows") || GroupName == TEXT("Midtones") || GroupName == TEXT("Highlights"))
+			{
+				FDisplayClusterColorGradingDataModel::FColorGradingElement ColorGradingElement = CreateColorGradingElement(PropertyGroupNode, FText::FromString(GroupName));
+				ColorGradingGroup.ColorGradingElements.Add(ColorGradingElement);
+			}
+
+			AddPropertiesToDetailsView(PropertyGroupNode, ColorGradingGroup);
+		}
+
+		OutColorGradingDataModel.ColorGradingGroups.Add(ColorGradingGroup);
+	}
+}
+
+FDisplayClusterColorGradingDataModel::FColorGradingElement FDisplayClusterColorGradingGenerator_PostProcessVolume::CreateColorGradingElement(const TSharedRef<IDetailTreeNode>& GroupNode, FText ElementLabel)
+{
+	FDisplayClusterColorGradingDataModel::FColorGradingElement ColorGradingElement;
+	ColorGradingElement.DisplayName = ElementLabel;
+
+	TArray<TSharedRef<IDetailTreeNode>> ChildNodes;
+	GroupNode->GetChildren(ChildNodes);
+
+	for (const TSharedRef<IDetailTreeNode>& ChildNode : ChildNodes)
+	{
+		TSharedPtr<IPropertyHandle> PropertyHandle = ChildNode->CreatePropertyHandle();
+		if (PropertyHandle.IsValid() && PropertyHandle->IsValidHandle())
+		{
+			const FString ColorGradingModeString = PropertyHandle->GetProperty()->GetMetaData(TEXT("ColorGradingMode")).ToLower();
+
+			if (!ColorGradingModeString.IsEmpty())
+			{
+				if (ColorGradingModeString.Compare(TEXT("saturation")) == 0)
+				{
+					ColorGradingElement.SaturationPropertyHandle = PropertyHandle;
+				}
+				else if (ColorGradingModeString.Compare(TEXT("contrast")) == 0)
+				{
+					ColorGradingElement.ContrastPropertyHandle = PropertyHandle;
+				}
+				else if (ColorGradingModeString.Compare(TEXT("gamma")) == 0)
+				{
+					ColorGradingElement.GammaPropertyHandle = PropertyHandle;
+				}
+				else if (ColorGradingModeString.Compare(TEXT("gain")) == 0)
+				{
+					ColorGradingElement.GainPropertyHandle = PropertyHandle;
+				}
+				else if (ColorGradingModeString.Compare(TEXT("offset")) == 0)
+				{
+					ColorGradingElement.OffsetPropertyHandle = PropertyHandle;
+				}
+			}
+		}
+	}
+
+	return ColorGradingElement;
+}
+
+void FDisplayClusterColorGradingGenerator_PostProcessVolume::AddPropertiesToDetailsView(const TSharedRef<IDetailTreeNode>& GroupNode, FDisplayClusterColorGradingDataModel::FColorGradingGroup& ColorGradingGroup)
+{
+	FString CategoryName = TEXT("");
+	FString GroupName = TEXT("");
+	GroupNode->GetNodeName().ToString().Split(TEXT("|"), &CategoryName, &GroupName);
+
+	TArray<TSharedRef<IDetailTreeNode>> ChildNodes;
+	GroupNode->GetChildren(ChildNodes);
+
+	for (const TSharedRef<IDetailTreeNode>& ChildNode : ChildNodes)
+	{
+		TSharedPtr<IPropertyHandle> PropertyHandle = ChildNode->CreatePropertyHandle();
+		if (PropertyHandle.IsValid() && PropertyHandle->IsValidHandle())
+		{
+			const bool bIsColorProperty = PropertyHandle->GetProperty()->HasMetaData(TEXT("ColorGradingMode"));
+
+			if (!bIsColorProperty)
+			{
+				PropertyHandle->SetInstanceMetaData(TEXT("CategoryOverride"), GroupName);
+				ColorGradingGroup.DetailsViewPropertyHandles.Add(PropertyHandle);
+			}
+		}
+	}
+}
+
+#undef LOCTEXT_NAMESPACE
\ No newline at end of file
diff --git a/Engine/Plugins/Runtime/nDisplay/Source/DisplayClusterColorGrading/Private/DataModelGenerators/DisplayClusterColorGradingGenerator_PostProcessVolume.h b/Engine/Plugins/Runtime/nDisplay/Source/DisplayClusterColorGrading/Private/DataModelGenerators/DisplayClusterColorGradingGenerator_PostProcessVolume.h
new file mode 100644
--- /dev/null
+++ b/Engine/Plugins/Runtime/nDisplay/Source/DisplayClusterColorGrading/Private/DataModelGenerators/DisplayClusterColorGradingGenerator_PostProcessVolume.h
@@ -0,0 +1,26 @@
+// Copyright Epic Games, Inc. All Rights Reserved.
+
+#pragma once
+
+#include "CoreMinimal.h"
+#include "DisplayClusterColorGradingDataModel.h"
+
+class IDetailTreeNode;
+
+/** Color Grading Data Model Generator for the APostProcessVolume actor class */
+class FDisplayClusterColorGradingGenerator_PostProcessVolume: public IDisplayClusterColorGradingDataModelGenerator
+{
+public:
+	static TSharedRef<IDisplayClusterColorGradingDataModelGenerator> MakeInstance();
+
+	//~ IDisplayClusterColorGradingDataModelGenerator interface
+	virtual void GenerateDataModel(IPropertyRowGenerator& PropertyRowGenerator, FDisplayClusterColorGradingDataModel& OutColorGradingDataModel) override;
+	//~ End IDisplayClusterColorGradingDataModelGenerator interface
+
+private:
+	/** Creates a new color grading element structure for the specified detail tree node, which is expected to have child color properties with the ColorGradingMode metadata set */
+	FDisplayClusterColorGradingDataModel::FColorGradingElement CreateColorGradingElement(const TSharedRef<IDetailTreeNode>& GroupNode, FText ElementLabel);
+
+	/** Adds all child properties of the specified detail tree node to the color grading group's list of properties to display in the details view*/
+	void AddPropertiesToDetailsView(const TSharedRef<IDetailTreeNode>& GroupNode, FDisplayClusterColorGradingDataModel::FColorGradingGroup& ColorGradingGroup);
+};
\ No newline at end of file
diff --git a/Engine/Plugins/Runtime/nDisplay/Source/DisplayClusterColorGrading/Private/DataModelGenerators/DisplayClusterColorGradingGenerator_RootActor.cpp b/Engine/Plugins/Runtime/nDisplay/Source/DisplayClusterColorGrading/Private/DataModelGenerators/DisplayClusterColorGradingGenerator_RootActor.cpp
new file mode 100644
--- /dev/null
+++ b/Engine/Plugins/Runtime/nDisplay/Source/DisplayClusterColorGrading/Private/DataModelGenerators/DisplayClusterColorGradingGenerator_RootActor.cpp
@@ -0,0 +1,720 @@
+// Copyright Epic Games, Inc. All Rights Reserved.
+
+#include "DisplayClusterColorGradingGenerator_RootActor.h"
+
+#include "DisplayClusterColorGradingStyle.h"
+#include "IDisplayClusterColorGrading.h"
+#include "IDisplayClusterColorGradingDrawerSingleton.h"
+
+#include "DisplayClusterRootActor.h"
+#include "Components/DisplayClusterICVFXCameraComponent.h"
+
+#include "DisplayClusterConfigurationTypes.h"
+#include "DisplayClusterConfigurationTypes_ICVFX.h"
+#include "DisplayClusterConfigurationTypes_Postprocess.h"
+
+#include "Framework/MultiBox/MultiBoxBuilder.h"
+#include "IDetailTreeNode.h"
+#include "IPropertyRowGenerator.h"
+#include "PropertyHandle.h"
+#include "ScopedTransaction.h"
+#include "Widgets/Input/SButton.h"
+#include "Widgets/Text/SInlineEditableTextBlock.h"
+#include "Widgets/Views/SListView.h"
+
+#define LOCTEXT_NAMESPACE "DisplayClusterColorGrading"
+
+FDisplayClusterColorGradingDataModel::FColorGradingGroup FDisplayClusterColorGradingGenerator_ColorGradingRenderingSettings::CreateColorGradingGroup(const TSharedPtr<IPropertyHandle>& GroupPropertyHandle)
+{
+	FDisplayClusterColorGradingDataModel::FColorGradingGroup ColorGradingGroup;
+	ColorGradingGroup.DisplayName = GroupPropertyHandle->GetPropertyDisplayName();
+	ColorGradingGroup.GroupPropertyHandle = GroupPropertyHandle;
+
+	ColorGradingGroup.ColorGradingElements.Add(CreateColorGradingElement(GroupPropertyHandle, TEXT("Global"), LOCTEXT("ColorGrading_GlobalLabel", "Global")));
+	ColorGradingGroup.ColorGradingElements.Add(CreateColorGradingElement(GroupPropertyHandle, TEXT("Shadows"), LOCTEXT("ColorGrading_ShadowsLabel", "Shadows")));
+	ColorGradingGroup.ColorGradingElements.Add(CreateColorGradingElement(GroupPropertyHandle, TEXT("Midtones"), LOCTEXT("ColorGrading_MidtonesLabel", "Midtones")));
+	ColorGradingGroup.ColorGradingElements.Add(CreateColorGradingElement(GroupPropertyHandle, TEXT("Highlights"), LOCTEXT("ColorGrading_HighlightsLabel", "Highlights")));
+
+	AddDetailsViewPropertyToGroup(GroupPropertyHandle, ColorGradingGroup, GET_MEMBER_NAME_CHECKED(FDisplayClusterConfigurationViewport_ColorGradingRenderingSettings, AutoExposureBias), TEXT("Exposure"));
+	AddDetailsViewPropertyToGroup(GroupPropertyHandle, ColorGradingGroup, GET_MEMBER_NAME_CHECKED(FDisplayClusterConfigurationViewport_ColorGradingRenderingSettings, ColorCorrectionShadowsMax), TEXT("Color Grading"));
+	AddDetailsViewPropertyToGroup(GroupPropertyHandle, ColorGradingGroup, GET_MEMBER_NAME_CHECKED(FDisplayClusterConfigurationViewport_ColorGradingRenderingSettings, ColorCorrectionHighlightsMin), TEXT("Color Grading"));
+	AddDetailsViewPropertyToGroup(GroupPropertyHandle, ColorGradingGroup, GET_MEMBER_NAME_CHECKED(FDisplayClusterConfigurationViewport_ColorGradingWhiteBalanceSettings, TemperatureType), TEXT("White Balance"));
+	AddDetailsViewPropertyToGroup(GroupPropertyHandle, ColorGradingGroup, GET_MEMBER_NAME_CHECKED(FDisplayClusterConfigurationViewport_ColorGradingWhiteBalanceSettings, WhiteTemp), TEXT("White Balance"));
+	AddDetailsViewPropertyToGroup(GroupPropertyHandle, ColorGradingGroup, GET_MEMBER_NAME_CHECKED(FDisplayClusterConfigurationViewport_ColorGradingWhiteBalanceSettings, WhiteTint), TEXT("White Balance"));
+	AddDetailsViewPropertyToGroup(GroupPropertyHandle, ColorGradingGroup, GET_MEMBER_NAME_CHECKED(FDisplayClusterConfigurationViewport_ColorGradingMiscSettings, BlueCorrection), TEXT("Misc"));
+	AddDetailsViewPropertyToGroup(GroupPropertyHandle, ColorGradingGroup, GET_MEMBER_NAME_CHECKED(FDisplayClusterConfigurationViewport_ColorGradingMiscSettings, ExpandGamut), TEXT("Misc"));
+	AddDetailsViewPropertyToGroup(GroupPropertyHandle, ColorGradingGroup, GET_MEMBER_NAME_CHECKED(FDisplayClusterConfigurationViewport_ColorGradingMiscSettings, SceneColorTint), TEXT("Misc"));
+
+	return ColorGradingGroup;
+}
+
+FDisplayClusterColorGradingDataModel::FColorGradingElement FDisplayClusterColorGradingGenerator_ColorGradingRenderingSettings::CreateColorGradingElement(
+	const TSharedPtr<IPropertyHandle>& GroupPropertyHandle,
+	FName ElementPropertyName,
+	FText ElementLabel)
+{
+	FDisplayClusterColorGradingDataModel::FColorGradingElement ColorGradingElement;
+	ColorGradingElement.DisplayName = ElementLabel;
+
+	TSharedPtr<IPropertyHandle> ElementPropertyHandle = GroupPropertyHandle->GetChildHandle(ElementPropertyName);
+	if (ElementPropertyHandle.IsValid() && ElementPropertyHandle->IsValidHandle())
+	{
+		ColorGradingElement.SaturationPropertyHandle = ElementPropertyHandle->GetChildHandle(GET_MEMBER_NAME_CHECKED(FDisplayClusterConfigurationViewport_ColorGradingSettings, Saturation));
+		ColorGradingElement.ContrastPropertyHandle = ElementPropertyHandle->GetChildHandle(GET_MEMBER_NAME_CHECKED(FDisplayClusterConfigurationViewport_ColorGradingSettings, Contrast));
+		ColorGradingElement.GammaPropertyHandle = ElementPropertyHandle->GetChildHandle(GET_MEMBER_NAME_CHECKED(FDisplayClusterConfigurationViewport_ColorGradingSettings, Gamma));
+		ColorGradingElement.GainPropertyHandle = ElementPropertyHandle->GetChildHandle(GET_MEMBER_NAME_CHECKED(FDisplayClusterConfigurationViewport_ColorGradingSettings, Gain));
+		ColorGradingElement.OffsetPropertyHandle = ElementPropertyHandle->GetChildHandle(GET_MEMBER_NAME_CHECKED(FDisplayClusterConfigurationViewport_ColorGradingSettings, Offset));
+	}
+
+	return ColorGradingElement;
+}
+
+void FDisplayClusterColorGradingGenerator_ColorGradingRenderingSettings::AddDetailsViewPropertyToGroup(
+	const TSharedPtr<IPropertyHandle>& GroupPropertyHandle,
+	FDisplayClusterColorGradingDataModel::FColorGradingGroup& Group,
+	FName PropertyName,
+	const FString& CategoryOverride)
+{
+	TSharedPtr<IPropertyHandle> PropertyHandle = GroupPropertyHandle->GetChildHandle(PropertyName);
+	if (PropertyHandle.IsValid() && PropertyHandle->IsValidHandle())
+	{
+		PropertyHandle->SetInstanceMetaData(TEXT("CategoryOverride"), CategoryOverride);
+
+		Group.DetailsViewPropertyHandles.Add(PropertyHandle);
+	}
+}
+
+TSharedPtr<IDetailTreeNode> FDisplayClusterColorGradingGenerator_ColorGradingRenderingSettings::FindPropertyTreeNode(const TSharedRef<IDetailTreeNode>& Node, const FName& PropertyName)
+{
+	if (Node->GetNodeType() == EDetailNodeType::Item)
+	{
+		FName NodeName = Node->GetNodeName();
+		if (Node->GetNodeName() == PropertyName)
+		{
+			return Node;
+		}
+		else
+		{
+			return nullptr;
+		}
+	}
+	else
+	{
+		TArray<TSharedRef<IDetailTreeNode>> Children;
+		Node->GetChildren(Children);
+		for (const TSharedRef<IDetailTreeNode>& Child : Children)
+		{
+			if (TSharedPtr<IDetailTreeNode> PropertyTreeNode = FindPropertyTreeNode(Child, PropertyName))
+			{
+				return PropertyTreeNode;
+			}
+		}
+
+		return nullptr;
+	}
+}
+
+TSharedPtr<IPropertyHandle> FDisplayClusterColorGradingGenerator_ColorGradingRenderingSettings::FindPropertyHandle(IPropertyRowGenerator& PropertyRowGenerator, const FName& PropertyNodeName)
+{
+	const TArray<TSharedRef<IDetailTreeNode>>& RootNodes = PropertyRowGenerator.GetRootTreeNodes();
+
+	for (const TSharedRef<IDetailTreeNode>& RootNode : RootNodes)
+	{
+		if (TSharedPtr<IDetailTreeNode> PropertyTreeNode = FindPropertyTreeNode(RootNode, PropertyNodeName))
+		{
+			return PropertyTreeNode->CreatePropertyHandle();
+		}
+	}
+
+	return nullptr;
+}
+
+
+TSharedRef<IDisplayClusterColorGradingDataModelGenerator> FDisplayClusterColorGradingGenerator_RootActor::MakeInstance()
+{
+	return MakeShareable(new FDisplayClusterColorGradingGenerator_RootActor());
+}
+
+void FDisplayClusterColorGradingGenerator_RootActor::GenerateDataModel(IPropertyRowGenerator& PropertyRowGenerator, FDisplayClusterColorGradingDataModel& OutColorGradingDataModel)
+{
+	const TArray<TWeakObjectPtr<UObject>>& SelectedObjects = PropertyRowGenerator.GetSelectedObjects();
+	for (const TWeakObjectPtr<UObject>& SelectedObject : SelectedObjects)
+	{
+		if (SelectedObject.IsValid() && SelectedObject->IsA<ADisplayClusterRootActor>())
+		{
+			TWeakObjectPtr<ADisplayClusterRootActor> SelectedRootActor = CastChecked<ADisplayClusterRootActor>(SelectedObject.Get());
+			RootActors.Add(SelectedRootActor);
+		}
+	}
+
+	// Add a color grading group for the root actor's "EntireClusterColorGrading" property
+	if (TSharedPtr<IPropertyHandle> EntireClusterColorGradingHandle = FindPropertyHandle(PropertyRowGenerator, GET_MEMBER_NAME_CHECKED(FDisplayClusterConfigurationICVFX_StageSettings, EntireClusterColorGrading)))
+	{
+		FDisplayClusterColorGradingDataModel::FColorGradingGroup EntireClusterGroup = CreateColorGradingGroup(EntireClusterColorGradingHandle);
+		EntireClusterGroup.EditConditionPropertyHandle = EntireClusterColorGradingHandle->GetChildHandle(GET_MEMBER_NAME_CHECKED(FDisplayClusterConfigurationViewport_EntireClusterColorGrading, bEnableEntireClusterColorGrading));
+
+		EntireClusterGroup.GroupHeaderWidget = SNew(SHorizontalBox)
+					+ SHorizontalBox::Slot()
+					.AutoWidth()
+					.Padding(0, 0, 4, 0)
+					.VAlign(VAlign_Center)
+					[
+						SNew(STextBlock)
+						.Text(EntireClusterGroup.DisplayName)
+						.Font(FAppStyle::Get().GetFontStyle("NormalFontBold"))
+					]
+
+					+ SHorizontalBox::Slot()
+					.AutoWidth()
+					.VAlign(VAlign_Center)
+					[
+						CreateViewportComboBox(INDEX_NONE)
+					];
+
+		OutColorGradingDataModel.ColorGradingGroups.Add(EntireClusterGroup);
+	}
+
+	// Add a color grading group for each element in the root actor's "PerViewportColorGrading" array
+	if (TSharedPtr<IPropertyHandle> PerViewportColorGradingHandle = FindPropertyHandle(PropertyRowGenerator, GET_MEMBER_NAME_CHECKED(FDisplayClusterConfigurationICVFX_StageSettings, PerViewportColorGrading)))
+	{
+		check(PerViewportColorGradingHandle->AsArray().IsValid());
+
+		uint32 NumGroups;
+		if (PerViewportColorGradingHandle->AsArray()->GetNumElements(NumGroups) == FPropertyAccess::Success)
+		{
+			for (int32 Index = 0; Index < (int32)NumGroups; ++Index)
+			{
+				TSharedRef<IPropertyHandle> PerViewportElementHandle = PerViewportColorGradingHandle->AsArray()->GetElement(Index);
+
+				FDisplayClusterColorGradingDataModel::FColorGradingGroup PerViewportGroup = CreateColorGradingGroup(PerViewportElementHandle);
+				PerViewportGroup.EditConditionPropertyHandle = PerViewportElementHandle->GetChildHandle(GET_MEMBER_NAME_CHECKED(FDisplayClusterConfigurationViewport_PerViewportColorGrading, bIsEnabled));
+
+				TSharedPtr<IPropertyHandle> NamePropertyHandle = PerViewportElementHandle->GetChildHandle(GET_MEMBER_NAME_CHECKED(FDisplayClusterConfigurationViewport_PerViewportColorGrading, Name));
+				if (NamePropertyHandle.IsValid() && NamePropertyHandle->IsValidHandle())
+				{
+					NamePropertyHandle->GetValue(PerViewportGroup.DisplayName);
+				}
+
+				PerViewportGroup.GroupHeaderWidget = SNew(SHorizontalBox)
+					+ SHorizontalBox::Slot()
+					.AutoWidth()
+					.Padding(0, 0, 4, 0)
+					.VAlign(VAlign_Center)
+					[
+						SNew(SInlineEditableTextBlock)
+						.Font(FAppStyle::Get().GetFontStyle("NormalFontBold"))
+						.OnTextCommitted_Lambda([NamePropertyHandle](const FText& InText, ETextCommit::Type TextCommitType)
+						{
+							if (NamePropertyHandle.IsValid() && NamePropertyHandle->IsValidHandle())
+							{
+								NamePropertyHandle->SetValue(InText);
+								IDisplayClusterColorGrading::Get().GetColorGradingDrawerSingleton().RefreshColorGradingDrawers(true);
+							}
+						})
+						.Text_Lambda([NamePropertyHandle]()
+						{
+							FText Name = FText::GetEmpty();
+
+							if (NamePropertyHandle.IsValid() && NamePropertyHandle->IsValidHandle())
+							{
+								NamePropertyHandle->GetValue(Name);
+							}
+
+							return Name;
+						})
+					]
+
+					+ SHorizontalBox::Slot()
+					.AutoWidth()
+					.VAlign(VAlign_Center)
+					[
+						CreateViewportComboBox(Index)
+					];
+
+				OutColorGradingDataModel.ColorGradingGroups.Add(PerViewportGroup);
+			}
+		}
+	}
+
+	OutColorGradingDataModel.bShowColorGradingGroupToolBar = true;
+	OutColorGradingDataModel.ColorGradingGroupToolBarWidget = SNew(SButton)
+		.ButtonStyle(FAppStyle::Get(), "SimpleButton")
+		.OnClicked(this, &FDisplayClusterColorGradingGenerator_RootActor::AddColorGradingGroup)
+		.ContentPadding(FMargin(1.0f, 0.0f))
+		.VAlign(VAlign_Center)
+		[
+			SNew(SImage)
+			.Image(FAppStyle::GetBrush("Icons.PlusCircle"))
+			.ColorAndOpacity(FSlateColor::UseForeground())
+		];
+}
+
+FReply FDisplayClusterColorGradingGenerator_RootActor::AddColorGradingGroup()
+{
+	for (const TWeakObjectPtr<ADisplayClusterRootActor>& RootActor : RootActors)
+	{
+		if (RootActor.IsValid())
+		{
+			if (UDisplayClusterConfigurationData* ConfigData = RootActor->GetConfigData())
+			{
+				FScopedTransaction Transaction(LOCTEXT("AddViewportColorGradingGroupTransaction", "Add Viewport Group"));
+
+				RootActor->Modify();
+				ConfigData->Modify();
+
+				FDisplayClusterConfigurationViewport_PerViewportColorGrading NewColorGradingGroup;
+				NewColorGradingGroup.Name = LOCTEXT("NewViewportColorGradingGroupName", "NewViewportGroup");
+
+				ConfigData->StageSettings.PerViewportColorGrading.Add(NewColorGradingGroup);
+
+				IDisplayClusterColorGrading::Get().GetColorGradingDrawerSingleton().RefreshColorGradingDrawers(true);
+			}
+		}
+	}
+
+	return FReply::Handled();
+}
+
+TSharedRef<SWidget> FDisplayClusterColorGradingGenerator_RootActor::CreateViewportComboBox(int32 PerViewportColorGradingIndex) const
+{
+	return SNew(SComboButton)
+		.HasDownArrow(true)
+		.OnGetMenuContent(this, &FDisplayClusterColorGradingGenerator_RootActor::GetViewportComboBoxMenu, PerViewportColorGradingIndex)
+		.ButtonContent()
+		[
+			SNew(SHorizontalBox)
+
+			+ SHorizontalBox::Slot()
+			.AutoWidth()
+			.Padding(0, 0, 4, 0)
+			.VAlign(VAlign_Center)
+			[
+				SNew(SImage)
+				.Image(FDisplayClusterColorGradingStyle::Get().GetBrush("ColorGradingDrawer.Viewports"))
+				.ColorAndOpacity(FSlateColor::UseForeground())
+			]
+
+			+ SHorizontalBox::Slot()
+			.FillWidth(1.0f)
+			.VAlign(VAlign_Center)
+			[
+				SNew(STextBlock)
+				.Text(this, &FDisplayClusterColorGradingGenerator_RootActor::GetViewportComboBoxText, PerViewportColorGradingIndex)
+			]
+		];
+}
+
+FText FDisplayClusterColorGradingGenerator_RootActor::GetViewportComboBoxText(int32 PerViewportColorGradingIndex) const
+{
+	// For now, only support displaying actual data when a single root actor is selected
+	if (RootActors.Num() == 1)
+	{
+		TWeakObjectPtr<ADisplayClusterRootActor> RootActor = RootActors[0];
+
+		if (RootActor.IsValid())
+		{
+			if (UDisplayClusterConfigurationData* ConfigData = RootActor->GetConfigData())
+			{
+				// If a valid per-viewport color grading group is passed in, determine the number of viewports associated with that group; otherwise,
+				// count the total viewports in the configuration
+				int32 NumViewports = 0;
+				if (PerViewportColorGradingIndex > INDEX_NONE && PerViewportColorGradingIndex < ConfigData->StageSettings.PerViewportColorGrading.Num())
+				{
+					FDisplayClusterConfigurationViewport_PerViewportColorGrading& PerViewportColorGrading = ConfigData->StageSettings.PerViewportColorGrading[PerViewportColorGradingIndex];
+					NumViewports = PerViewportColorGrading.ApplyPostProcessToObjects.Num();
+				}
+				else
+				{
+					for (const UE_TRANSITIONAL_OBJECT_PTR_TEMPLATE2_ARG2(TTuple, FString, UDisplayClusterConfigurationClusterNode)& Node : ConfigData->Cluster->Nodes)
+					{
+						if (Node.Value)
+						{
+							NumViewports += Node.Value->Viewports.Num();
+						}
+					}
+				}
+
+				return FText::Format(LOCTEXT("PerViewportColorGradingGroup_NumViewports", "{0} {0}|plural(one=Viewport,other=Viewports)"), FText::AsNumber(NumViewports));
+			}
+		}
+	}
+	else if (RootActors.Num() > 1)
+	{
+		return LOCTEXT("MultipleValuesSelectedLabel", "Multiple Values");
+	}
+
+	return FText::GetEmpty();
+}
+
+TSharedRef<SWidget> FDisplayClusterColorGradingGenerator_RootActor::GetViewportComboBoxMenu(int32 PerViewportColorGradingIndex) const
+{
+	FMenuBuilder MenuBuilder(false, nullptr);
+
+	// For now, only support displaying actual data when a single root actor is selected
+	if (RootActors.Num() == 1)
+	{
+		TWeakObjectPtr<ADisplayClusterRootActor> RootActor = RootActors[0];
+		if (RootActor.IsValid())
+		{
+			if (UDisplayClusterConfigurationData* ConfigData = RootActor->GetConfigData())
+			{
+				// Extract all viewport names from the configuration data, so they can be sorted alphabetically
+				TArray<FString> ViewportNames;
+
+				for (const UE_TRANSITIONAL_OBJECT_PTR_TEMPLATE2_ARG2(TTuple, FString, UDisplayClusterConfigurationClusterNode)& Node : ConfigData->Cluster->Nodes)
+				{
+					if (Node.Value)
+					{
+						for (const UE_TRANSITIONAL_OBJECT_PTR_TEMPLATE2_ARG2(TTuple, FString, UDisplayClusterConfigurationViewport)& Viewport : Node.Value->Viewports)
+						{
+							ViewportNames.Add(Viewport.Key);
+						}
+					}
+				}
+
+				ViewportNames.Sort();
+
+				const bool bForEntireCluster = PerViewportColorGradingIndex == INDEX_NONE;
+				FDisplayClusterConfigurationViewport_PerViewportColorGrading* PerViewportColorGrading = !bForEntireCluster ? 
+					&ConfigData->StageSettings.PerViewportColorGrading[PerViewportColorGradingIndex] : nullptr;
+
+				MenuBuilder.BeginSection(TEXT("ViewportSection"), LOCTEXT("ViewportsMenuSectionLabel", "Viewports"));
+				for (const FString& ViewportName : ViewportNames)
+				{
+					MenuBuilder.AddMenuEntry(FText::FromString(ViewportName),
+						FText::GetEmpty(),
+						FSlateIcon(),
+						FUIAction(
+							FExecuteAction::CreateLambda([ConfigData , PerViewportColorGrading, ViewportName, InRootActor=RootActor.Get()]
+							{
+								if (PerViewportColorGrading)
+								{
+									if (PerViewportColorGrading->ApplyPostProcessToObjects.Contains(ViewportName))
+									{
+										FScopedTransaction Transaction(LOCTEXT("RemoveViewportFromColorGradingGroupTransaction", "Remove Viewport from Group"));
+										InRootActor->Modify();
+										ConfigData->Modify();
+
+										PerViewportColorGrading->ApplyPostProcessToObjects.Remove(ViewportName);
+									}
+									else
+									{
+										FScopedTransaction Transaction(LOCTEXT("AddViewportToColorGradingGroupTransaction", "Add Viewport to Group"));
+										InRootActor->Modify();
+										ConfigData->Modify();
+
+										PerViewportColorGrading->ApplyPostProcessToObjects.Add(ViewportName);
+									}
+								}
+							}),
+							FCanExecuteAction::CreateLambda([PerViewportColorGrading] { return PerViewportColorGrading != nullptr; }),
+							FGetActionCheckState::CreateLambda([PerViewportColorGrading, ViewportName]
+							{
+								// If the menu is for the EntireCluster group (PerViewportColorGrading is null), all viewport list items should be checked
+								if (PerViewportColorGrading)
+								{
+									const bool bIsViewportInGroup = PerViewportColorGrading->ApplyPostProcessToObjects.Contains(ViewportName);
+									return bIsViewportInGroup ? ECheckBoxState::Checked : ECheckBoxState::Unchecked;
+								}
+								else
+								{
+									return ECheckBoxState::Checked;
+								}
+							})
+						),
+						NAME_None,
+						EUserInterfaceActionType::ToggleButton);
+				}
+				MenuBuilder.EndSection();
+			}
+		}
+	}
+
+	return MenuBuilder.MakeWidget();
+}
+
+TSharedRef<IDisplayClusterColorGradingDataModelGenerator> FDisplayClusterColorGradingGenerator_ICVFXCamera::MakeInstance()
+{
+	return MakeShareable(new FDisplayClusterColorGradingGenerator_ICVFXCamera());
+}
+
+void FDisplayClusterColorGradingGenerator_ICVFXCamera::GenerateDataModel(IPropertyRowGenerator& PropertyRowGenerator, FDisplayClusterColorGradingDataModel& OutColorGradingDataModel)
+{
+	const TArray<TWeakObjectPtr<UObject>>& SelectedObjects = PropertyRowGenerator.GetSelectedObjects();
+	for (const TWeakObjectPtr<UObject>& SelectedObject : SelectedObjects)
+	{
+		if (SelectedObject.IsValid() && SelectedObject->IsA<UDisplayClusterICVFXCameraComponent>())
+		{
+			TWeakObjectPtr<UDisplayClusterICVFXCameraComponent> SelectedCameraComponent = CastChecked<UDisplayClusterICVFXCameraComponent>(SelectedObject.Get());
+			CameraComponents.Add(SelectedCameraComponent);
+		}
+	}
+
+	// Add a color grading group for the camera's "AllNodesColorGrading" property
+	if (TSharedPtr<IPropertyHandle> AllNodesColorGradingHandle = FindPropertyHandle(PropertyRowGenerator, GET_MEMBER_NAME_CHECKED(FDisplayClusterConfigurationICVFX_CameraSettings, AllNodesColorGrading)))
+	{
+		FDisplayClusterColorGradingDataModel::FColorGradingGroup AllNodesGroup = CreateColorGradingGroup(AllNodesColorGradingHandle);
+		AllNodesGroup.EditConditionPropertyHandle = AllNodesColorGradingHandle->GetChildHandle(GET_MEMBER_NAME_CHECKED(FDisplayClusterConfigurationViewport_AllNodesColorGrading, bEnableInnerFrustumAllNodesColorGrading));
+
+		AllNodesGroup.GroupHeaderWidget = SNew(SHorizontalBox)
+			+ SHorizontalBox::Slot()
+			.AutoWidth()
+			.Padding(0, 0, 4, 0)
+			.VAlign(VAlign_Center)
+			[
+				SNew(STextBlock)
+				.Text(AllNodesGroup.DisplayName)
+				.Font(FAppStyle::Get().GetFontStyle("NormalFontBold"))
+			]
+
+			+ SHorizontalBox::Slot()
+			.AutoWidth()
+			.VAlign(VAlign_Center)
+			[
+				CreateNodeComboBox(INDEX_NONE)
+			];
+
+		OutColorGradingDataModel.ColorGradingGroups.Add(AllNodesGroup);
+	}
+
+	// Add a color grading group for each element in the camera's "PerNodeColorGrading" array
+	if (TSharedPtr<IPropertyHandle> PerNodeColorGradingHandle = FindPropertyHandle(PropertyRowGenerator, GET_MEMBER_NAME_CHECKED(FDisplayClusterConfigurationICVFX_CameraSettings, PerNodeColorGrading)))
+	{
+		check(PerNodeColorGradingHandle->AsArray().IsValid());
+
+		uint32 NumGroups;
+		if (PerNodeColorGradingHandle->AsArray()->GetNumElements(NumGroups) == FPropertyAccess::Success)
+		{
+			for (int32 Index = 0; Index < (int32)NumGroups; ++Index)
+			{
+				TSharedRef<IPropertyHandle> PerNodeElementHandle = PerNodeColorGradingHandle->AsArray()->GetElement(Index);
+
+				FDisplayClusterColorGradingDataModel::FColorGradingGroup PerNodeGroup = CreateColorGradingGroup(PerNodeElementHandle);
+				PerNodeGroup.EditConditionPropertyHandle = PerNodeElementHandle->GetChildHandle(GET_MEMBER_NAME_CHECKED(FDisplayClusterConfigurationViewport_PerNodeColorGrading, bIsEnabled));
+
+				TSharedPtr<IPropertyHandle> NamePropertyHandle = PerNodeElementHandle->GetChildHandle(GET_MEMBER_NAME_CHECKED(FDisplayClusterConfigurationViewport_PerNodeColorGrading, Name));
+				if (NamePropertyHandle.IsValid() && NamePropertyHandle->IsValidHandle())
+				{
+					NamePropertyHandle->GetValue(PerNodeGroup.DisplayName);
+				}
+
+				PerNodeGroup.GroupHeaderWidget = SNew(SHorizontalBox)
+					+ SHorizontalBox::Slot()
+					.AutoWidth()
+					.Padding(0, 0, 4, 0)
+					.VAlign(VAlign_Center)
+					[
+						SNew(SInlineEditableTextBlock)
+						.Font(FAppStyle::Get().GetFontStyle("NormalFontBold"))
+						.OnTextCommitted_Lambda([NamePropertyHandle](const FText& InText, ETextCommit::Type TextCommitType)
+						{
+							if (NamePropertyHandle.IsValid() && NamePropertyHandle->IsValidHandle())
+							{
+								NamePropertyHandle->SetValue(InText);
+								IDisplayClusterColorGrading::Get().GetColorGradingDrawerSingleton().RefreshColorGradingDrawers(true);
+							}
+						})
+						.Text_Lambda([NamePropertyHandle]()
+						{
+							FText Name = FText::GetEmpty();
+
+							if (NamePropertyHandle.IsValid() && NamePropertyHandle->IsValidHandle())
+							{
+								NamePropertyHandle->GetValue(Name);
+							}
+
+							return Name;
+						})
+					]
+
+					+ SHorizontalBox::Slot()
+					.AutoWidth()
+					.VAlign(VAlign_Center)
+					[
+						CreateNodeComboBox(Index)
+					];
+
+				OutColorGradingDataModel.ColorGradingGroups.Add(PerNodeGroup);
+			}
+		}
+	}
+
+	OutColorGradingDataModel.bShowColorGradingGroupToolBar = true;OutColorGradingDataModel.bShowColorGradingGroupToolBar = true;
+	OutColorGradingDataModel.ColorGradingGroupToolBarWidget = SNew(SButton)
+		.ButtonStyle(FAppStyle::Get(), "SimpleButton")
+		.OnClicked(this, &FDisplayClusterColorGradingGenerator_ICVFXCamera::AddColorGradingGroup)
+		.ContentPadding(FMargin(1.0f, 0.0f))
+		.VAlign(VAlign_Center)
+		[
+			SNew(SImage)
+			.Image(FAppStyle::GetBrush("Icons.PlusCircle"))
+			.ColorAndOpacity(FSlateColor::UseForeground())
+		];
+}
+
+FReply FDisplayClusterColorGradingGenerator_ICVFXCamera::AddColorGradingGroup()
+{
+	for (const TWeakObjectPtr<UDisplayClusterICVFXCameraComponent>& CameraComponent : CameraComponents)
+	{
+		if (CameraComponent.IsValid())
+		{
+			FScopedTransaction Transaction(LOCTEXT("AddNodeColorGradingGroupTransaction", "Add Node Group"));
+			CameraComponent->Modify();
+
+			FDisplayClusterConfigurationViewport_PerNodeColorGrading NewColorGradingGroup;
+			NewColorGradingGroup.Name = LOCTEXT("NewNodeColorGradingGroupName", "NewNodeGroup");
+
+			CameraComponent->CameraSettings.PerNodeColorGrading.Add(NewColorGradingGroup);
+
+			IDisplayClusterColorGrading::Get().GetColorGradingDrawerSingleton().RefreshColorGradingDrawers(true);
+		}
+	}
+
+	return FReply::Handled();
+}
+
+TSharedRef<SWidget> FDisplayClusterColorGradingGenerator_ICVFXCamera::CreateNodeComboBox(int32 PerNodeColorGradingIndex) const
+{
+	return SNew(SComboButton)
+		.HasDownArrow(true)
+		.OnGetMenuContent(this, &FDisplayClusterColorGradingGenerator_ICVFXCamera::GetNodeComboBoxMenu, PerNodeColorGradingIndex)
+		.ButtonContent()
+		[
+			SNew(SHorizontalBox)
+
+			+ SHorizontalBox::Slot()
+			.AutoWidth()
+			.Padding(0, 0, 4, 0)
+			.VAlign(VAlign_Center)
+			[
+				SNew(SImage)
+				.Image(FDisplayClusterColorGradingStyle::Get().GetBrush("ColorGradingDrawer.Nodes"))
+				.ColorAndOpacity(FSlateColor::UseForeground())
+			]
+
+			+ SHorizontalBox::Slot()
+			.FillWidth(1.0f)
+			.VAlign(VAlign_Center)
+			[
+				SNew(STextBlock)
+				.Text(this, &FDisplayClusterColorGradingGenerator_ICVFXCamera::GetNodeComboBoxText, PerNodeColorGradingIndex)
+			]
+		];
+}
+
+FText FDisplayClusterColorGradingGenerator_ICVFXCamera::GetNodeComboBoxText(int32 PerNodeColorGradingIndex) const
+{
+	// For now, only support displaying actual data when a single camera component is selected
+	if (CameraComponents.Num() == 1)
+	{
+		TWeakObjectPtr<UDisplayClusterICVFXCameraComponent> CameraComponent = CameraComponents[0];
+		if (CameraComponent.IsValid())
+		{
+			int32 NumNodes = 0;
+			if (PerNodeColorGradingIndex > INDEX_NONE && PerNodeColorGradingIndex < CameraComponent->CameraSettings.PerNodeColorGrading.Num())
+			{
+				const FDisplayClusterConfigurationViewport_PerNodeColorGrading& PerNodeColorGrading = CameraComponent->CameraSettings.PerNodeColorGrading[PerNodeColorGradingIndex];
+				NumNodes = PerNodeColorGrading.ApplyPostProcessToObjects.Num();
+			}
+			else
+			{
+				if (ADisplayClusterRootActor* ParentRootActor = Cast<ADisplayClusterRootActor>(CameraComponent->GetOwner()))
+				{
+					if (UDisplayClusterConfigurationData* ConfigData = ParentRootActor->GetConfigData())
+					{
+						NumNodes = ConfigData->Cluster->Nodes.Num();
+					}
+				}
+			}
+
+			return FText::Format(LOCTEXT("PerNodeColorGradingGroup_NumViewports", "{0} {0}|plural(one=Node,other=Nodes)"), FText::AsNumber(NumNodes));
+		}
+	}
+	else if (CameraComponents.Num() > 1)
+	{
+		return LOCTEXT("MultipleValuesSelectedLabel", "Multiple Values");
+	}
+
+	return FText::GetEmpty();
+}
+
+
+TSharedRef<SWidget> FDisplayClusterColorGradingGenerator_ICVFXCamera::GetNodeComboBoxMenu(int32 PerNodeColorGradingIndex) const
+{
+
+	FMenuBuilder MenuBuilder(false, nullptr);
+
+	// For now, only support displaying actual data when a single camera component is selected
+	if (CameraComponents.Num() == 1)
+	{
+		TWeakObjectPtr<UDisplayClusterICVFXCameraComponent> CameraComponent = CameraComponents[0];
+		if (CameraComponent.IsValid())
+		{
+			if (ADisplayClusterRootActor* ParentRootActor = Cast<ADisplayClusterRootActor>(CameraComponent->GetOwner()))
+			{
+				if (UDisplayClusterConfigurationData* ConfigData = ParentRootActor->GetConfigData())
+				{
+					TArray<FString> NodeNames;
+					ConfigData->Cluster->Nodes.GetKeys(NodeNames);
+
+					NodeNames.Sort();
+
+					const bool bForAllNodes = PerNodeColorGradingIndex == INDEX_NONE;
+					FDisplayClusterConfigurationViewport_PerNodeColorGrading* PerNodeColorGrading = !bForAllNodes ?
+						&CameraComponent->CameraSettings.PerNodeColorGrading[PerNodeColorGradingIndex] : nullptr;
+
+					MenuBuilder.BeginSection(TEXT("NodeSection"), LOCTEXT("NodeMenuSectionLabel", "Nodes"));
+					for (const FString& NodeName : NodeNames)
+					{
+						MenuBuilder.AddMenuEntry(FText::FromString(NodeName),
+							FText::GetEmpty(),
+							FSlateIcon(),
+							FUIAction(
+								FExecuteAction::CreateLambda([PerNodeColorGrading, NodeName, InCameraComponent=CameraComponent.Get()]()
+								{
+									if (PerNodeColorGrading)
+									{
+										if (PerNodeColorGrading->ApplyPostProcessToObjects.Contains(NodeName))
+										{
+											FScopedTransaction Transaction(LOCTEXT("RemoveNodeFromColorGradingGroupTransaction", "Remove Node from Group"));
+											InCameraComponent->Modify();
+
+											PerNodeColorGrading->ApplyPostProcessToObjects.Remove(NodeName);
+										}
+										else
+										{
+											FScopedTransaction Transaction(LOCTEXT("AddNodeToColorGradingGroupTransaction", "Add Node to Group"));
+											InCameraComponent->Modify();
+
+											PerNodeColorGrading->ApplyPostProcessToObjects.Add(NodeName);
+										}
+									}
+								}),
+								FCanExecuteAction::CreateLambda([PerNodeColorGrading] { return PerNodeColorGrading != nullptr; }),
+								FGetActionCheckState::CreateLambda([PerNodeColorGrading, NodeName]()
+								{
+									// If the menu is for the AllNodes group (PerNodeColorGrading is null), all node list items should be checked
+									if (PerNodeColorGrading)
+									{
+										const bool bIsNodeInGroup = PerNodeColorGrading->ApplyPostProcessToObjects.Contains(NodeName);
+										return bIsNodeInGroup ? ECheckBoxState::Checked : ECheckBoxState::Unchecked;
+									}
+									else
+									{
+										return ECheckBoxState::Checked;
+									}
+								})
+							),
+							NAME_None,
+							EUserInterfaceActionType::ToggleButton);
+					}
+					MenuBuilder.EndSection();
+				}
+			}
+		}
+	}
+
+	return MenuBuilder.MakeWidget();
+}
+
+#undef LOCTEXT_NAMESPACE
\ No newline at end of file
diff --git a/Engine/Plugins/Runtime/nDisplay/Source/DisplayClusterColorGrading/Private/DataModelGenerators/DisplayClusterColorGradingGenerator_RootActor.h b/Engine/Plugins/Runtime/nDisplay/Source/DisplayClusterColorGrading/Private/DataModelGenerators/DisplayClusterColorGradingGenerator_RootActor.h
new file mode 100644
--- /dev/null
+++ b/Engine/Plugins/Runtime/nDisplay/Source/DisplayClusterColorGrading/Private/DataModelGenerators/DisplayClusterColorGradingGenerator_RootActor.h
@@ -0,0 +1,91 @@
+// Copyright Epic Games, Inc. All Rights Reserved.
+
+#pragma once
+
+#include "CoreMinimal.h"
+#include "DisplayClusterColorGradingDataModel.h"
+
+#include "Input/Reply.h"
+
+class ADisplayClusterRootActor;
+class IDetailTreeNode;
+class IPropertyHandle;
+class IPropertyRowGenerator;
+class FDisplayClusterColorGradingDataModel;
+class UDisplayClusterICVFXCameraComponent;
+
+/** Base generator for any object that needs a color grading data model generated from an FColorGradingRenderingSettings struct */
+class FDisplayClusterColorGradingGenerator_ColorGradingRenderingSettings : public IDisplayClusterColorGradingDataModelGenerator
+{
+protected:
+	/** Creates a new color grading group structure from the specified group property handle, which finds and connects the appropriate property handles for the color wheels and details view */
+	FDisplayClusterColorGradingDataModel::FColorGradingGroup CreateColorGradingGroup(const TSharedPtr<IPropertyHandle>& GroupPropertyHandle);
+
+	/** Creates a color grading element structure from the specified property handle, whose child properties are expected to be colors with the ColorGradingMode metadata set */
+	FDisplayClusterColorGradingDataModel::FColorGradingElement CreateColorGradingElement(const TSharedPtr<IPropertyHandle>& GroupPropertyHandle, FName ElementPropertyName, FText ElementLabel);
+
+	/** Adds the specified property handle to the specified color grading group's list of properties to display in the details view */
+	void AddDetailsViewPropertyToGroup(const TSharedPtr<IPropertyHandle>& GroupPropertyHandle, FDisplayClusterColorGradingDataModel::FColorGradingGroup& Group, FName PropertyName, const FString& CategoryOverride);
+
+	/** Recursively searches the detail tree hierarchy for a property detail tree node whose name matches the specified name */
+	TSharedPtr<IDetailTreeNode> FindPropertyTreeNode(const TSharedRef<IDetailTreeNode>& Node, const FName& PropertyName);
+
+	/** Finds a property handle in the specified property row generator whose name matches the specified name */
+	TSharedPtr<IPropertyHandle> FindPropertyHandle(IPropertyRowGenerator& PropertyRowGenerator, const FName& PropertyNodeName);
+};
+
+/** Color grading data model generator for an nDisplay root actor */
+class FDisplayClusterColorGradingGenerator_RootActor : public FDisplayClusterColorGradingGenerator_ColorGradingRenderingSettings
+{
+public:
+	static TSharedRef<IDisplayClusterColorGradingDataModelGenerator> MakeInstance();
+
+	//~ IDisplayClusterColorGradingDataModelGenerator interface
+	virtual void GenerateDataModel(IPropertyRowGenerator& PropertyRowGenerator, FDisplayClusterColorGradingDataModel& OutColorGradingDataModel) override;
+	//~ End IDisplayClusterColorGradingDataModelGenerator interface
+
+private:
+	/** OnClick delegate handler for the "+" button that is added to the color grading group toolbar */
+	FReply AddColorGradingGroup();
+
+	/** Creates the combo box used to display which viewports are part of the specified color grading group */
+	TSharedRef<SWidget> CreateViewportComboBox(int32 PerViewportColorGradingIndex) const;
+
+	/** Gets the display text for the viewports list combo box for the PerViewportsGroup at the specified index */
+	FText GetViewportComboBoxText(int32 PerViewportColorGradingIndex) const;
+
+	/** Generates a menu widget for the viewports list combo box for the PerViewportsGroup at the specified index */
+	TSharedRef<SWidget> GetViewportComboBoxMenu(int32 PerViewportColorGradingIndex) const;
+
+private:
+	/** A list of root actors that are being represented by the data model */
+	TArray<TWeakObjectPtr<ADisplayClusterRootActor>> RootActors;
+};
+
+/** Color grading data model generator for an nDisplay ICVFX camera component */
+class FDisplayClusterColorGradingGenerator_ICVFXCamera : public FDisplayClusterColorGradingGenerator_ColorGradingRenderingSettings
+{
+public:
+	static TSharedRef<IDisplayClusterColorGradingDataModelGenerator> MakeInstance();
+
+	//~ IDisplayClusterColorGradingDataModelGenerator interface
+	virtual void GenerateDataModel(IPropertyRowGenerator& PropertyRowGenerator, FDisplayClusterColorGradingDataModel& OutColorGradingDataModel) override;
+	//~ End IDisplayClusterColorGradingDataModelGenerator interface
+
+private:
+	/** OnClick delegate handler for the "+" button that is added to the color grading group toolbar */
+	FReply AddColorGradingGroup();
+
+	/** Creates the combo box used to display which nodes are part of the specified color grading group */
+	TSharedRef<SWidget> CreateNodeComboBox(int32 PerNodeColorGradingIndex) const;
+
+	/** Gets the display text for the nodes list combo box from the specified PerNodesGroup property handle */
+	FText GetNodeComboBoxText(int32 PerNodeColorGradingIndex) const;
+
+	/** Generates a menu widget for the nodes list combo box for the PerNodesGroup at the specified index */
+	TSharedRef<SWidget> GetNodeComboBoxMenu(int32 PerNodeColorGradingIndex) const;
+
+private:
+	/** A list of root actors that are being represented by the data model */
+	TArray<TWeakObjectPtr<UDisplayClusterICVFXCameraComponent>> CameraComponents;
+};
\ No newline at end of file
diff --git a/Engine/Plugins/Runtime/nDisplay/Source/DisplayClusterColorGrading/Private/DisplayClusterColorGradingCommands.cpp b/Engine/Plugins/Runtime/nDisplay/Source/DisplayClusterColorGrading/Private/DisplayClusterColorGradingCommands.cpp
new file mode 100644
--- /dev/null
+++ b/Engine/Plugins/Runtime/nDisplay/Source/DisplayClusterColorGrading/Private/DisplayClusterColorGradingCommands.cpp
@@ -0,0 +1,15 @@
+// Copyright Epic Games, Inc. All Rights Reserved.
+
+#include "DisplayClusterColorGradingCommands.h"
+
+#define LOCTEXT_NAMESPACE "DisplayClusterColorGrading"
+
+void FDisplayClusterColorGradingCommands::RegisterCommands()
+{
+	UI_COMMAND(SaturationColorWheelVisibility, "Saturation", "Show or hide the saturation color wheel", EUserInterfaceActionType::ToggleButton, FInputChord());
+	UI_COMMAND(ContrastColorWheelVisibility, "Constrast", "Show or hide the constrast color wheel", EUserInterfaceActionType::ToggleButton, FInputChord());
+	UI_COMMAND(ColorWheelSliderOrientationHorizontal, "Right", "Puts the color wheel sliders to the right of the color wheel", EUserInterfaceActionType::RadioButton, FInputChord());
+	UI_COMMAND(ColorWheelSliderOrientationVertical, "Below", "Puts the color wheel sliders below the color wheel", EUserInterfaceActionType::RadioButton, FInputChord());
+}
+
+#undef LOCTEXT_NAMESPACE
diff --git a/Engine/Plugins/Runtime/nDisplay/Source/DisplayClusterColorGrading/Private/DisplayClusterColorGradingCommands.h b/Engine/Plugins/Runtime/nDisplay/Source/DisplayClusterColorGrading/Private/DisplayClusterColorGradingCommands.h
new file mode 100644
--- /dev/null
+++ b/Engine/Plugins/Runtime/nDisplay/Source/DisplayClusterColorGrading/Private/DisplayClusterColorGradingCommands.h
@@ -0,0 +1,26 @@
+// Copyright Epic Games, Inc. All Rights Reserved.
+
+#pragma once
+
+#include "CoreMinimal.h"
+#include "Framework/Commands/Commands.h"
+#include "Styling/AppStyle.h"
+
+/** Command list for the color grading drawer */
+class FDisplayClusterColorGradingCommands
+	: public TCommands<FDisplayClusterColorGradingCommands>
+{
+public:
+	FDisplayClusterColorGradingCommands()
+		: TCommands<FDisplayClusterColorGradingCommands>(TEXT("DisplayClusterColorGrading"),
+			NSLOCTEXT("Contexts", "DisplayClusterColorGrading", "Display Cluster Color Grading"), NAME_None, FAppStyle::GetAppStyleSetName())
+	{ }
+
+	virtual void RegisterCommands() override;
+
+public:
+	TSharedPtr<FUICommandInfo> SaturationColorWheelVisibility;
+	TSharedPtr<FUICommandInfo> ContrastColorWheelVisibility;
+	TSharedPtr<FUICommandInfo> ColorWheelSliderOrientationHorizontal;
+	TSharedPtr<FUICommandInfo> ColorWheelSliderOrientationVertical;
+};
\ No newline at end of file
diff --git a/Engine/Plugins/Runtime/nDisplay/Source/DisplayClusterColorGrading/Private/DisplayClusterColorGradingDataModel.cpp b/Engine/Plugins/Runtime/nDisplay/Source/DisplayClusterColorGrading/Private/DisplayClusterColorGradingDataModel.cpp
new file mode 100644
--- /dev/null
+++ b/Engine/Plugins/Runtime/nDisplay/Source/DisplayClusterColorGrading/Private/DisplayClusterColorGradingDataModel.cpp
@@ -0,0 +1,163 @@
+// Copyright Epic Games, Inc. All Rights Reserved.
+
+#include "DisplayClusterColorGradingDataModel.h"
+
+#include "DisplayClusterRootActor.h"
+
+#include "IDetailTreeNode.h"
+#include "IPropertyRowGenerator.h"
+#include "Modules/ModuleManager.h"
+#include "PropertyEditorModule.h"
+#include "PropertyHandle.h"
+#include "Widgets/Input/SCheckBox.h"
+
+#define LOCTEXT_NAMESPACE "DisplayClusterColorGrading"
+
+TMap<FName, FGetColorGradingDataModelGenerator> FDisplayClusterColorGradingDataModel::RegisteredDataModelGenerators;
+
+FDisplayClusterColorGradingDataModel::FDisplayClusterColorGradingDataModel()
+{
+	FPropertyEditorModule& PropertyEditorModule = FModuleManager::GetModuleChecked<FPropertyEditorModule>("PropertyEditor");
+
+	FPropertyRowGeneratorArgs Args;
+	PropertyRowGenerator = PropertyEditorModule.CreatePropertyRowGenerator(Args);
+}
+
+TArray<TWeakObjectPtr<UObject>> FDisplayClusterColorGradingDataModel::GetObjects() const
+{
+	if (PropertyRowGenerator.IsValid())
+	{
+		return PropertyRowGenerator->GetSelectedObjects();
+	}
+
+	return TArray<TWeakObjectPtr<UObject>>();
+}
+
+void FDisplayClusterColorGradingDataModel::SetObjects(const TArray<UObject*>& InObjects)
+{
+	Reset();
+
+	if (PropertyRowGenerator.IsValid())
+	{
+		PropertyRowGenerator->SetObjects(InObjects);
+
+		const TArray<TWeakObjectPtr<UObject>>& SelectedObjects = PropertyRowGenerator->GetSelectedObjects();
+
+		if (SelectedObjects.Num() == 1)
+		{
+			TWeakObjectPtr<UObject> SelectedObject = SelectedObjects[0];
+
+			if (SelectedObject.IsValid())
+			{
+				if (TSharedPtr<IDisplayClusterColorGradingDataModelGenerator> Generator = GetDataModelGenerator(SelectedObject->GetClass()))
+				{
+					Generator->GenerateDataModel(*PropertyRowGenerator, *this);
+				}
+			}
+		}
+
+		// TODO: Figure out what needs to be done to support multiple disparate types of objects being color graded at the same time
+	}
+
+	SelectedColorGradingGroupIndex = ColorGradingGroups.Num() ? 0 : INDEX_NONE;
+	SelectedColorGradingElementIndex = 0;
+}
+
+bool FDisplayClusterColorGradingDataModel::HasObjectOfType(const UClass* InClass) const
+{
+	if (PropertyRowGenerator.IsValid())
+	{
+		const TArray<TWeakObjectPtr<UObject>> SelectedObjects = PropertyRowGenerator->GetSelectedObjects();
+
+		for (const TWeakObjectPtr<UObject>& Object : SelectedObjects)
+		{
+			if (Object.IsValid() && Object->GetClass()->IsChildOf(InClass))
+			{
+				return true;
+			}
+		}
+	}
+
+	return false;
+}
+
+void FDisplayClusterColorGradingDataModel::Reset()
+{
+	DataModelGeneratorInstances.Empty();
+	ColorGradingGroups.Empty();
+	SelectedColorGradingGroupIndex = INDEX_NONE;
+	SelectedColorGradingElementIndex = INDEX_NONE;
+	ColorGradingGroupToolBarWidget = nullptr;
+	bShowColorGradingGroupToolBar = false;
+
+	if (PropertyRowGenerator.IsValid())
+	{
+		PropertyRowGenerator->SetObjects(TArray<UObject*>());
+	}
+}
+
+FDisplayClusterColorGradingDataModel::FColorGradingGroup* FDisplayClusterColorGradingDataModel::GetSelectedColorGradingGroup()
+{
+	if (SelectedColorGradingGroupIndex > INDEX_NONE && SelectedColorGradingGroupIndex < ColorGradingGroups.Num())
+	{
+		return &ColorGradingGroups[SelectedColorGradingGroupIndex];
+	}
+
+	return nullptr;
+}
+
+void FDisplayClusterColorGradingDataModel::SetSelectedColorGradingGroup(int32 InColorGradingGroupIndex)
+{
+	SelectedColorGradingGroupIndex = InColorGradingGroupIndex <= ColorGradingGroups.Num() ? InColorGradingGroupIndex : INDEX_NONE;
+	OnColorGradingGroupSelectionChangedDelegate.Broadcast();
+}
+
+FDisplayClusterColorGradingDataModel::FColorGradingElement* FDisplayClusterColorGradingDataModel::GetSelectedColorGradingElement()
+{
+	if (SelectedColorGradingGroupIndex > INDEX_NONE && SelectedColorGradingGroupIndex < ColorGradingGroups.Num())
+	{
+		FDisplayClusterColorGradingDataModel::FColorGradingGroup& SelectedGroup = ColorGradingGroups[SelectedColorGradingGroupIndex];
+		if (SelectedColorGradingElementIndex > INDEX_NONE && SelectedColorGradingElementIndex < SelectedGroup.ColorGradingElements.Num())
+		{
+			return &SelectedGroup.ColorGradingElements[SelectedColorGradingElementIndex];
+		}
+	}
+
+	return nullptr;
+}
+
+void FDisplayClusterColorGradingDataModel::SetSelectedColorGradingElement(int32 InColorGradingElementIndex)
+{
+	SelectedColorGradingElementIndex = InColorGradingElementIndex;
+	OnColorGradingElementSelectionChangedDelegate.Broadcast();
+}
+
+TSharedPtr<IDisplayClusterColorGradingDataModelGenerator> FDisplayClusterColorGradingDataModel::GetDataModelGenerator(const UClass* InClass)
+{
+	const UClass* CurrentClass = InClass;
+	while (CurrentClass)
+	{
+		const FName ClassName = CurrentClass->GetFName();
+
+		if (DataModelGeneratorInstances.Contains(ClassName))
+		{
+			return DataModelGeneratorInstances[ClassName];
+		}
+		else if (RegisteredDataModelGenerators.Contains(ClassName))
+		{
+			if (RegisteredDataModelGenerators[ClassName].IsBound())
+			{
+				TSharedRef<IDisplayClusterColorGradingDataModelGenerator> Generator = RegisteredDataModelGenerators[ClassName].Execute();
+				DataModelGeneratorInstances.Add(ClassName, Generator);
+
+				return Generator;
+			}
+		}
+
+		CurrentClass = CurrentClass->GetSuperClass();
+	}
+
+	return nullptr;
+}
+
+#undef LOCTEXT_NAMESPACE
\ No newline at end of file
diff --git a/Engine/Plugins/Runtime/nDisplay/Source/DisplayClusterColorGrading/Private/DisplayClusterColorGradingDataModel.h b/Engine/Plugins/Runtime/nDisplay/Source/DisplayClusterColorGrading/Private/DisplayClusterColorGradingDataModel.h
new file mode 100644
--- /dev/null
+++ b/Engine/Plugins/Runtime/nDisplay/Source/DisplayClusterColorGrading/Private/DisplayClusterColorGradingDataModel.h
@@ -0,0 +1,164 @@
+// Copyright Epic Games, Inc. All Rights Reserved.
+
+#pragma once
+
+#include "CoreMinimal.h"
+
+class IPropertyHandle;
+class IPropertyRowGenerator;
+class SWidget;
+
+/** Interface that allows color grading data models to be generated for specific UObjects */
+class IDisplayClusterColorGradingDataModelGenerator : public TSharedFromThis<IDisplayClusterColorGradingDataModelGenerator>
+{
+public:
+	virtual ~IDisplayClusterColorGradingDataModelGenerator() {}
+
+	/** Called when the color grading data model should be generated by this generator */
+	virtual void GenerateDataModel(IPropertyRowGenerator& PropertyRowGenerator, class FDisplayClusterColorGradingDataModel& OutColorGradingDataModel) = 0;
+};
+
+DECLARE_DELEGATE_RetVal(TSharedRef<IDisplayClusterColorGradingDataModelGenerator>, FGetColorGradingDataModelGenerator);
+
+/**
+ * A data model that stores the color grading properties from a UObject, which allows UObjects with widely varying color grading
+ * interfaces to be used with the color grading panel.
+ */
+class FDisplayClusterColorGradingDataModel
+{
+public:
+	/** Registers a new data model generator used to populate a color grading data model for the specified class */
+	template<class T>
+	static void RegisterColorGradingDataModelGenerator(FGetColorGradingDataModelGenerator CreateGeneratorDelegate)
+	{
+		const FName ClassName = T::StaticClass()->GetFName();
+		RegisteredDataModelGenerators.Add(ClassName, CreateGeneratorDelegate);
+	}
+
+private:
+	/** A list of data model generators, indexed by class name, that have been registered */
+	static TMap<FName, FGetColorGradingDataModelGenerator> RegisteredDataModelGenerators;
+
+
+public:
+	/** Stores the property handles for the color properties for a single color grading element (Shadows, Midtones, etc.) */
+	struct FColorGradingElement
+	{
+		/** The display name of the color grading element */
+		FText DisplayName;
+
+		/** Property handle of the color property used for saturation */
+		TSharedPtr<IPropertyHandle> SaturationPropertyHandle;
+
+		/** Property handle of the color property used for contrast */
+		TSharedPtr<IPropertyHandle> ContrastPropertyHandle;
+
+		/** Property handle of the color property used for gamma */
+		TSharedPtr<IPropertyHandle> GammaPropertyHandle;
+
+		/** Property handle of the color property used for gain */
+		TSharedPtr<IPropertyHandle> GainPropertyHandle;
+
+		/** Property handle of the color property used for offset */
+		TSharedPtr<IPropertyHandle> OffsetPropertyHandle;
+	};
+
+	/** Stores a list of color grading elements that form a single color grading group */
+	struct FColorGradingGroup
+	{
+		/** The display name of the color grading group */
+		FText DisplayName = FText::GetEmpty();
+		
+		/** The handle of the property that the color grading properties are generated from */
+		TSharedPtr<IPropertyHandle> GroupPropertyHandle = nullptr;
+
+		/** The widget displayed in the color wheel panel's header for this group, usually just a text block to show the group's name */
+		TSharedPtr<SWidget> GroupHeaderWidget = nullptr;
+
+		/** The handle of the property that controls the editing condition of the properties in the color grading group */
+		TSharedPtr<IPropertyHandle> EditConditionPropertyHandle = nullptr;
+
+		/** The list of color grading elements (Shadows, Midtones, etc.) that belong in the group */
+		TArray<FColorGradingElement> ColorGradingElements;
+
+		/** A list of properties to display in the details view of the color grading panel */
+		TArray<TSharedPtr<IPropertyHandle>> DetailsViewPropertyHandles;
+	};
+
+	DECLARE_MULTICAST_DELEGATE(FOnColorGradingSelectionChanged);
+
+public:
+	FDisplayClusterColorGradingDataModel();
+
+	/** Gets the property row generator that was used to generate this data model */
+	const TSharedRef<IPropertyRowGenerator> GetPropertyRowGenerator() const { return PropertyRowGenerator.ToSharedRef(); }
+
+	/** Gets the list of objects from which the color grading properties were generated from */
+	TArray<TWeakObjectPtr<UObject>> GetObjects() const;
+
+	/** Sets the objects from which the color grading properties will be generated from */
+	void SetObjects(const TArray<UObject*>& InObjects);
+
+	/** Gets whether the data model has objects of the specified class */
+	bool HasObjectOfType(const UClass* InClass) const;
+
+	/** Resets the data model to be empty */
+	void Reset();
+
+	/** Gets the index of the selected color grading group */
+	int32 GetSelectedColorGradingGroupIndex() const { return SelectedColorGradingGroupIndex; }
+
+	/** Gets a pointer to the selected color grading group, or null if there isn't a valid selection */
+	FColorGradingGroup* GetSelectedColorGradingGroup();
+
+	/** Sets the selected color grading group */
+	void SetSelectedColorGradingGroup(int32 InColorGradingGroupIndex);
+
+	/** Gets the delegate that is raised when the selected color grading group has changed */
+	FOnColorGradingSelectionChanged& OnColorGradingGroupSelectionChanged() { return OnColorGradingGroupSelectionChangedDelegate; }
+
+	/** Gets the index of the selected color grading element */
+	int32 GetSelectedColorGradingElementIndex() const { return SelectedColorGradingElementIndex; }
+
+	/** Gets a pointer to the selected color grading element, or null if there isn't a valid selection */
+	FColorGradingElement* GetSelectedColorGradingElement();
+
+	/** Sets the selected color grading element */
+	void SetSelectedColorGradingElement(int32 InColorGradingElementIndex);
+
+	/** Gets the delegate that is raised when the selected color grading element has changed */
+	FOnColorGradingSelectionChanged& OnColorGradingElementSelectionChanged() { return OnColorGradingElementSelectionChangedDelegate; }
+
+private:
+	/** Attempts to find an instance of a data model generator that can be used for the specified class */
+	TSharedPtr<IDisplayClusterColorGradingDataModelGenerator> GetDataModelGenerator(const UClass* InClass);
+
+public:
+	/** A list of color grading groups in the color grading data model */
+	TArray<FColorGradingGroup> ColorGradingGroups;
+
+	/** An optional widget to append to the end of the color grading group toolbar */
+	TSharedPtr<SWidget> ColorGradingGroupToolBarWidget = nullptr;
+
+	/** Indicates whether the color grading groups should be visible/selectable in a toolbar */
+	bool bShowColorGradingGroupToolBar = false;
+
+private:
+	/** A list of data model generators that have been created to generate the data model for the current objects */
+	TMap<FName, TSharedPtr<IDisplayClusterColorGradingDataModelGenerator>> DataModelGeneratorInstances;
+
+	/** The property row generator which manages the property handles of the color grading objects */
+	TSharedPtr<IPropertyRowGenerator> PropertyRowGenerator;
+
+	/** The currently selected color grading element */
+	int32 SelectedColorGradingGroupIndex = INDEX_NONE;
+
+	/** The currently selected color grading element */
+	int32 SelectedColorGradingElementIndex = INDEX_NONE;
+
+	/** Delegate that is raised when the selected color grading group has been changed */
+	FOnColorGradingSelectionChanged OnColorGradingGroupSelectionChangedDelegate;
+
+	/** Delegate that is raised when the selected color grading element has been changed */
+	FOnColorGradingSelectionChanged OnColorGradingElementSelectionChangedDelegate;
+};
\ No newline at end of file
diff --git a/Engine/Plugins/Runtime/nDisplay/Source/DisplayClusterColorGrading/Private/DisplayClusterColorGradingModule.cpp b/Engine/Plugins/Runtime/nDisplay/Source/DisplayClusterColorGrading/Private/DisplayClusterColorGradingModule.cpp
new file mode 100644
--- /dev/null
+++ b/Engine/Plugins/Runtime/nDisplay/Source/DisplayClusterColorGrading/Private/DisplayClusterColorGradingModule.cpp
@@ -0,0 +1,51 @@
+// Copyright Epic Games, Inc. All Rights Reserved.
+
+#include "DisplayClusterColorGradingModule.h"
+
+#include "DisplayClusterColorGradingCommands.h"
+#include "DataModelGenerators/DisplayClusterColorGradingGenerator_RootActor.h"
+#include "DataModelGenerators/DisplayClusterColorGradingGenerator_PostProcessVolume.h"
+#include "DataModelGenerators/DisplayClusterColorGradingGenerator_ColorCorrectRegion.h"
+#include "Drawer/DisplayClusterColorGradingDrawerSingleton.h"
+#include "Drawer/SDisplayClusterColorGradingDrawer.h"
+
+#include "DisplayClusterRootActor.h"
+#include "Components/DisplayClusterICVFXCameraComponent.h"
+
+#include "ColorCorrectRegion.h"
+#include "Engine/PostProcessVolume.h"
+
+#define LOCTEXT_NAMESPACE "DisplayClusterColorGrading"
+
+void FDisplayClusterColorGradingModule::StartupModule()
+{
+	ColorGradingDrawerSingleton = MakeUnique<FDisplayClusterColorGradingDrawerSingleton>();
+
+	FDisplayClusterColorGradingDataModel::RegisterColorGradingDataModelGenerator<ADisplayClusterRootActor>(
+		FGetColorGradingDataModelGenerator::CreateStatic(&FDisplayClusterColorGradingGenerator_RootActor::MakeInstance));
+
+	FDisplayClusterColorGradingDataModel::RegisterColorGradingDataModelGenerator<UDisplayClusterICVFXCameraComponent>(
+		FGetColorGradingDataModelGenerator::CreateStatic(&FDisplayClusterColorGradingGenerator_ICVFXCamera::MakeInstance));
+
+	FDisplayClusterColorGradingDataModel::RegisterColorGradingDataModelGenerator<APostProcessVolume>(
+		FGetColorGradingDataModelGenerator::CreateStatic(&FDisplayClusterColorGradingGenerator_PostProcessVolume::MakeInstance));
+
+	FDisplayClusterColorGradingDataModel::RegisterColorGradingDataModelGenerator<AColorCorrectRegion>(
+		FGetColorGradingDataModelGenerator::CreateStatic(&FDisplayClusterColorGradingGenerator_ColorCorrectRegion::MakeInstance));
+
+	FDisplayClusterColorGradingCommands::Register();
+}
+
+void FDisplayClusterColorGradingModule::ShutdownModule()
+{
+	ColorGradingDrawerSingleton.Reset();
+}
+
+IDisplayClusterColorGradingDrawerSingleton& FDisplayClusterColorGradingModule::GetColorGradingDrawerSingleton() const
+{
+	return *ColorGradingDrawerSingleton.Get();
+}
+
+IMPLEMENT_MODULE(FDisplayClusterColorGradingModule, DisplayClusterColorGrading);
+
+#undef LOCTEXT_NAMESPACE
diff --git a/Engine/Plugins/Runtime/nDisplay/Source/DisplayClusterColorGrading/Private/DisplayClusterColorGradingModule.h b/Engine/Plugins/Runtime/nDisplay/Source/DisplayClusterColorGrading/Private/DisplayClusterColorGradingModule.h
new file mode 100644
--- /dev/null
+++ b/Engine/Plugins/Runtime/nDisplay/Source/DisplayClusterColorGrading/Private/DisplayClusterColorGradingModule.h
@@ -0,0 +1,27 @@
+// Copyright Epic Games, Inc. All Rights Reserved.
+
+#pragma once
+
+#include "IDisplayClusterColorGrading.h"
+
+class FDisplayClusterColorGradingDrawerSingleton;
+
+/**
+ * Display Cluster Color Grading module
+ */
+class FDisplayClusterColorGradingModule : public IDisplayClusterColorGrading
+{
+public:
+	//~ IModuleInterface interface
+	virtual void StartupModule() override;
+	virtual void ShutdownModule() override;
+	//~ End IModuleInterface interface
+
+	//~ IDisplayClusterColorGrading interface
+	virtual IDisplayClusterColorGradingDrawerSingleton& GetColorGradingDrawerSingleton() const override;
+	//~ End IDisplayClusterColorGrading interface
+
+private:
+	/** The color grading drawer singleton, which manages the color grading drawer widget */
+	TUniquePtr<FDisplayClusterColorGradingDrawerSingleton> ColorGradingDrawerSingleton;
+};
diff --git a/Engine/Plugins/Runtime/nDisplay/Source/DisplayClusterColorGrading/Private/DisplayClusterColorGradingStyle.h b/Engine/Plugins/Runtime/nDisplay/Source/DisplayClusterColorGrading/Private/DisplayClusterColorGradingStyle.h
new file mode 100644
--- /dev/null
+++ b/Engine/Plugins/Runtime/nDisplay/Source/DisplayClusterColorGrading/Private/DisplayClusterColorGradingStyle.h
@@ -0,0 +1,42 @@
+// Copyright Epic Games, Inc. All Rights Reserved.
+
+#pragma once
+
+#include "CoreMinimal.h"
+#include "Styling/SlateStyleRegistry.h"
+#include "Styling/SlateStyleMacros.h"
+#include "Styling/SlateStyle.h"
+
+
+/** Styleset for the nDisplay color grading UI elements */
+class FDisplayClusterColorGradingStyle final : public FSlateStyleSet
+{
+public:
+
+	FDisplayClusterColorGradingStyle()
+		: FSlateStyleSet("DisplayClusterColorGradingStyle")
+	{
+		const FVector2D Icon16x16(16.0f, 16.0f);
+
+		// Set miscellaneous icons
+		SetContentRoot(FPaths::EnginePluginsDir() / TEXT("Runtime/nDisplay/Content/Icons/"));
+		SetCoreContentRoot(FPaths::EngineContentDir() / TEXT("Editor/Slate"));
+
+		Set("ColorGradingDrawer.Icon", new IMAGE_BRUSH_SVG("Components/nDisplayCamera", Icon16x16));
+		Set("ColorGradingDrawer.Viewports", new CORE_IMAGE_BRUSH_SVG("Starship/Common/Viewports", Icon16x16));
+		Set("ColorGradingDrawer.Nodes", new IMAGE_BRUSH_SVG("Cluster/ClusterNode", Icon16x16));
+
+		FSlateStyleRegistry::RegisterSlateStyle(*this);
+	}
+
+	virtual ~FDisplayClusterColorGradingStyle()
+	{
+		FSlateStyleRegistry::UnRegisterSlateStyle(*this);
+	}
+
+	static FDisplayClusterColorGradingStyle& Get()
+	{
+		static FDisplayClusterColorGradingStyle Inst;
+		return Inst;
+	}
+};
\ No newline at end of file
diff --git a/Engine/Plugins/Runtime/nDisplay/Source/DisplayClusterColorGrading/Private/Drawer/DisplayClusterColorGradingDrawerSingleton.cpp b/Engine/Plugins/Runtime/nDisplay/Source/DisplayClusterColorGrading/Private/Drawer/DisplayClusterColorGradingDrawerSingleton.cpp
new file mode 100644
--- /dev/null
+++ b/Engine/Plugins/Runtime/nDisplay/Source/DisplayClusterColorGrading/Private/Drawer/DisplayClusterColorGradingDrawerSingleton.cpp
@@ -0,0 +1,155 @@
+// Copyright Epic Games, Inc. All Rights Reserved.
+
+#include "DisplayClusterColorGradingDrawerSingleton.h"
+
+#include "DisplayClusterColorGradingStyle.h"
+
+#include "DisplayClusterRootActor.h"
+
+#include "IDisplayClusterOperator.h"
+#include "IDisplayClusterOperatorViewModel.h"
+#include "DisplayClusterOperatorStatusBarExtender.h"
+#include "Drawer/SDisplayClusterColorGradingDrawer.h"
+
+#include "Editor.h"
+#include "Framework/Application/SlateApplication.h"
+#include "Framework/Docking/LayoutExtender.h"
+#include "Framework/Docking/TabManager.h"
+#include "IPropertyRowGenerator.h"
+#include "Modules/ModuleManager.h"
+#include "PropertyEditorModule.h"
+#include "Widgets/DeclarativeSyntaxSupport.h"
+#include "Widgets/SCompoundWidget.h"
+#include "Widgets/Docking/SDockTab.h"
+
+#define LOCTEXT_NAMESPACE "DisplayClusterColorGrading"
+
+const FName FDisplayClusterColorGradingDrawerSingleton::ColorGradingDrawerId = TEXT("ColorGradingDrawer");
+const FName FDisplayClusterColorGradingDrawerSingleton::ColorGradingDrawerTab = TEXT("ColorGradingDrawerTab");
+
+FDisplayClusterColorGradingDrawerSingleton::FDisplayClusterColorGradingDrawerSingleton()
+{
+	FGlobalTabmanager::Get()->RegisterNomadTabSpawner(ColorGradingDrawerTab, FOnSpawnTab::CreateRaw(this, &FDisplayClusterColorGradingDrawerSingleton::SpawnColorGradingDrawerTab))
+		.SetDisplayName(LOCTEXT("ColorGradingDarwerTab_DisplayName", "In-Camera VFX"))
+		.SetTooltipText(LOCTEXT("ColorGradingDarwerTab_Tooltip", "Editing tools for in-camera VFX."))
+		.SetMenuType(ETabSpawnerMenuType::Hidden);
+
+	IDisplayClusterOperator::Get().OnRegisterLayoutExtensions().AddRaw(this, &FDisplayClusterColorGradingDrawerSingleton::ExtendOperatorTabLayout);
+	IDisplayClusterOperator::Get().OnRegisterStatusBarExtensions().AddRaw(this, &FDisplayClusterColorGradingDrawerSingleton::ExtendOperatorStatusBar);
+}
+
+FDisplayClusterColorGradingDrawerSingleton::~FDisplayClusterColorGradingDrawerSingleton()
+{
+	if (FSlateApplication::IsInitialized())
+	{
+		FGlobalTabmanager::Get()->UnregisterNomadTabSpawner(ColorGradingDrawerTab);
+	}
+}
+
+void FDisplayClusterColorGradingDrawerSingleton::DockColorGradingDrawer()
+{
+	if (TSharedPtr<FTabManager> OperatorPanelTabManager = IDisplayClusterOperator::Get().GetOperatorViewModel()->GetTabManager())
+	{
+		if (TSharedPtr<SDockTab> ExistingTab = OperatorPanelTabManager->FindExistingLiveTab(ColorGradingDrawerTab))
+		{
+			IDisplayClusterOperator::Get().ForceDismissDrawers();
+			ExistingTab->ActivateInParent(ETabActivationCause::SetDirectly);
+		}
+		else
+		{
+			OperatorPanelTabManager->TryInvokeTab(ColorGradingDrawerTab);
+		}
+	}
+}
+
+void FDisplayClusterColorGradingDrawerSingleton::RefreshColorGradingDrawers(bool bPreserveDrawerState)
+{
+	if (ColorGradingDrawer.IsValid())
+	{
+		ColorGradingDrawer.Pin()->Refresh(bPreserveDrawerState);
+	}
+
+	if (TSharedPtr<FTabManager> OperatorPanelTabManager = IDisplayClusterOperator::Get().GetOperatorViewModel()->GetTabManager())
+	{
+		if (TSharedPtr<SDockTab> ExistingTab = OperatorPanelTabManager->FindExistingLiveTab(ColorGradingDrawerTab))
+		{
+			TSharedRef<SDisplayClusterColorGradingDrawer> DockedDrawer = StaticCastSharedRef<SDisplayClusterColorGradingDrawer>(ExistingTab->GetContent());
+			DockedDrawer->Refresh(bPreserveDrawerState);
+		}
+	}
+}
+
+TSharedRef<SWidget> FDisplayClusterColorGradingDrawerSingleton::CreateDrawerContent(bool bIsInDrawer, bool bCopyStateFromActiveDrawer)
+{
+	if (bIsInDrawer)
+	{
+		TSharedPtr<SDisplayClusterColorGradingDrawer> Drawer = ColorGradingDrawer.IsValid() ? ColorGradingDrawer.Pin() : nullptr;
+
+		if (!Drawer.IsValid())
+		{
+			Drawer = SNew(SDisplayClusterColorGradingDrawer, true);
+			ColorGradingDrawer = Drawer;
+		}
+
+		if (PreviousDrawerState.IsSet())
+		{
+			ColorGradingDrawer.Pin()->SetDrawerState(PreviousDrawerState.GetValue());
+			PreviousDrawerState.Reset();
+		}
+
+		return Drawer.ToSharedRef();
+	}
+	else
+	{
+		TSharedRef<SDisplayClusterColorGradingDrawer> NewDrawer = SNew(SDisplayClusterColorGradingDrawer, false);
+
+		if (bCopyStateFromActiveDrawer && ColorGradingDrawer.IsValid())
+		{
+			NewDrawer->SetDrawerState(ColorGradingDrawer.Pin()->GetDrawerState());
+		}
+
+		return NewDrawer;
+	}
+}
+
+TSharedRef<SDockTab> FDisplayClusterColorGradingDrawerSingleton::SpawnColorGradingDrawerTab(const FSpawnTabArgs& SpawnTabArgs)
+{
+	const TSharedRef<SDockTab> MajorTab = SNew(SDockTab)
+		.TabRole(ETabRole::NomadTab);
+
+	MajorTab->SetContent(CreateDrawerContent(false, true));
+
+	return MajorTab;
+}
+
+void FDisplayClusterColorGradingDrawerSingleton::ExtendOperatorTabLayout(FLayoutExtender& InExtender)
+{
+	FTabManager::FTab NewTab(FTabId(ColorGradingDrawerTab, ETabIdFlags::SaveLayout), ETabState::ClosedTab);
+	InExtender.ExtendStack(IDisplayClusterOperator::Get().GetAuxilliaryOperatorExtensionId(), ELayoutExtensionPosition::After, NewTab);
+}
+
+void FDisplayClusterColorGradingDrawerSingleton::ExtendOperatorStatusBar(FDisplayClusterOperatorStatusBarExtender& StatusBarExtender)
+{
+	FWidgetDrawerConfig ColorGradingDrawerConfig(ColorGradingDrawerId);
+
+	ColorGradingDrawerConfig.GetDrawerContentDelegate.BindRaw(this, &FDisplayClusterColorGradingDrawerSingleton::CreateDrawerContent, true, false);
+	ColorGradingDrawerConfig.OnDrawerDismissedDelegate.BindRaw(this, &FDisplayClusterColorGradingDrawerSingleton::SaveDrawerState);
+	ColorGradingDrawerConfig.ButtonText = LOCTEXT("ColorGradingDrawer_ButtonText", "In-Camera VFX");
+	ColorGradingDrawerConfig.Icon = FDisplayClusterColorGradingStyle::Get().GetBrush("ColorGradingDrawer.Icon");
+
+	StatusBarExtender.AddWidgetDrawer(ColorGradingDrawerConfig);
+}
+
+void FDisplayClusterColorGradingDrawerSingleton::SaveDrawerState(const TSharedPtr<SWidget>& DrawerContent)
+{
+	if (ColorGradingDrawer.IsValid())
+	{
+		PreviousDrawerState = ColorGradingDrawer.Pin()->GetDrawerState();
+	}
+	else
+	{
+		PreviousDrawerState.Reset();
+	}
+}
+
+#undef LOCTEXT_NAMESPACE
\ No newline at end of file
diff --git a/Engine/Plugins/Runtime/nDisplay/Source/DisplayClusterColorGrading/Private/Drawer/DisplayClusterColorGradingDrawerSingleton.h b/Engine/Plugins/Runtime/nDisplay/Source/DisplayClusterColorGrading/Private/Drawer/DisplayClusterColorGradingDrawerSingleton.h
new file mode 100644
--- /dev/null
+++ b/Engine/Plugins/Runtime/nDisplay/Source/DisplayClusterColorGrading/Private/Drawer/DisplayClusterColorGradingDrawerSingleton.h
@@ -0,0 +1,57 @@
+// Copyright Epic Games, Inc. All Rights Reserved.
+
+#pragma once
+
+#include "CoreMinimal.h"
+#include "IDisplayClusterColorGradingDrawerSingleton.h"
+#include "Drawer/SDisplayClusterColorGradingDrawer.h"
+
+class FDisplayClusterOperatorStatusBarExtender;
+class FLayoutExtender;
+class FSpawnTabArgs;
+class SWidget;
+class SDockTab;
+
+/** A singleton used to manage and store persistent state for the color grading drawer */
+class FDisplayClusterColorGradingDrawerSingleton : public IDisplayClusterColorGradingDrawerSingleton
+{
+public:
+	/** The ID of the color grading drawer when registered with the nDisplay operator panel's status bar */
+	static const FName ColorGradingDrawerId;
+
+	/** The ID of the color grading drawer when docked in the nDisplay operator panel's tab manager */
+	static const FName ColorGradingDrawerTab;
+
+public:
+	FDisplayClusterColorGradingDrawerSingleton();
+	virtual ~FDisplayClusterColorGradingDrawerSingleton();
+
+	/** Docks the color grading drawer in the nDisplay operator window */
+	virtual void DockColorGradingDrawer() override;
+
+	/** Refreshes the UI of any open color grading drawers */
+	virtual void RefreshColorGradingDrawers(bool bPreserveDrawerState) override;
+
+private:
+	/** Creates a new drawer widget to place in a drawer or in a tab */
+	TSharedRef<SWidget> CreateDrawerContent(bool bIsInDrawer, bool bCopyStateFromActiveDrawer);
+
+	/** Tab spawn delegate handler used to create the drawer tab when the drawer is docked in the operator panel */
+	TSharedRef<SDockTab> SpawnColorGradingDrawerTab(const FSpawnTabArgs& SpawnTabArgs);
+
+	/** Tab extender delegate callback that registers the tab spawner with the operator panel's tab manager */
+	void ExtendOperatorTabLayout(FLayoutExtender& InExtender);
+
+	/** Status bar extender delegate callback that registers the drawer spawner with the operator panel's status bar */
+	void ExtendOperatorStatusBar(FDisplayClusterOperatorStatusBarExtender& StatusBarExtender);
+
+	/** Delegate callback when the drawer is closed to save the drawer state */
+	void SaveDrawerState(const TSharedPtr<SWidget>& DrawerContent);
+
+private:
+	/** A weak pointer to the active color grading drawer that is open */
+	TWeakPtr<SDisplayClusterColorGradingDrawer> ColorGradingDrawer;
+
+	/** The drawer state when the last instance of the color grading drawer was dismissed */
+	TOptional<FDisplayClusterColorGradingDrawerState> PreviousDrawerState;
+};
\ No newline at end of file
diff --git a/Engine/Plugins/Runtime/nDisplay/Source/DisplayClusterColorGrading/Private/Drawer/SDisplayClusterColorGradingDrawer.cpp b/Engine/Plugins/Runtime/nDisplay/Source/DisplayClusterColorGrading/Private/Drawer/SDisplayClusterColorGradingDrawer.cpp
new file mode 100644
--- /dev/null
+++ b/Engine/Plugins/Runtime/nDisplay/Source/DisplayClusterColorGrading/Private/Drawer/SDisplayClusterColorGradingDrawer.cpp
@@ -0,0 +1,862 @@
+// Copyright Epic Games, Inc. All Rights Reserved.
+
+#include "SDisplayClusterColorGradingDrawer.h"
+
+#include "IDisplayClusterColorGrading.h"
+#include "IDisplayClusterColorGradingDrawerSingleton.h"
+#include "SDisplayClusterColorGradingObjectList.h"
+#include "SDisplayClusterColorGradingColorWheelPanel.h"
+
+#include "IDisplayClusterOperator.h"
+#include "IDisplayClusterOperatorViewModel.h"
+#include "DisplayClusterOperatorStatusBarExtender.h"
+
+#include "DisplayClusterRootActor.h"
+#include "Components/DisplayClusterICVFXCameraComponent.h"
+#include "DisplayClusterConfigurationTypes.h"
+
+#include "ColorCorrectRegion.h"
+#include "Engine/PostProcessVolume.h"
+#include "Modules/ModuleManager.h"
+#include "ScopedTransaction.h"
+#include "Styling/AppStyle.h"
+#include "WidgetDrawerConfig.h"
+#include "Widgets/Input/SButton.h"
+#include "Widgets/Layout/SBorder.h"
+#include "Widgets/Layout/SBox.h"
+#include "Widgets/Layout/SExpandableArea.h"
+#include "Widgets/Layout/SSeparator.h"
+#include "Widgets/Layout/SSpacer.h"
+#include "Widgets/Layout/SSplitter.h"
+#include "Widgets/SBoxPanel.h"
+
+#define LOCTEXT_NAMESPACE "DisplayClusterColorGrading"
+
+SDisplayClusterColorGradingDrawer::~SDisplayClusterColorGradingDrawer()
+{
+	OperatorViewModel = IDisplayClusterOperator::Get().GetOperatorViewModel();
+	OperatorViewModel->OnActiveRootActorChanged().RemoveAll(this);
+
+	FCoreUObjectDelegates::OnObjectsReplaced.RemoveAll(this);
+	GEngine->OnLevelActorAdded().RemoveAll(this);
+	GEngine->OnLevelActorDeleted().RemoveAll(this);
+
+	GEditor->UnregisterForUndo(this);
+
+	for (const FDisplayClusterColorGradingListItemRef& LevelColorGradingItem : LevelColorGradingItems)
+	{
+		if (LevelColorGradingItem->Component.IsValid())
+		{
+			UnbindBlueprintCompiledDelegate(LevelColorGradingItem->Component->GetClass());
+		}
+
+		if (LevelColorGradingItem->Actor.IsValid())
+		{
+			UnbindBlueprintCompiledDelegate(LevelColorGradingItem->Actor->GetClass());
+		}
+	}
+
+	for (const FDisplayClusterColorGradingListItemRef& RootActorColorGradingItem : RootActorColorGradingItems)
+	{
+		if (RootActorColorGradingItem->Component.IsValid())
+		{
+			UnbindBlueprintCompiledDelegate(RootActorColorGradingItem->Component->GetClass());
+		}
+
+		if (RootActorColorGradingItem->Actor.IsValid())
+		{
+			UnbindBlueprintCompiledDelegate(RootActorColorGradingItem->Actor->GetClass());
+		}
+	}
+}
+
+void SDisplayClusterColorGradingDrawer::Construct(const FArguments& InArgs, bool bInIsInDrawer)
+{
+	ColorGradingDataModel = MakeShared<FDisplayClusterColorGradingDataModel>();
+
+	bIsInDrawer = bInIsInDrawer;
+	OperatorViewModel = IDisplayClusterOperator::Get().GetOperatorViewModel();
+	OperatorViewModel->OnActiveRootActorChanged().AddSP(this, &SDisplayClusterColorGradingDrawer::OnActiveRootActorChanged);
+
+	FCoreUObjectDelegates::OnObjectsReplaced.AddSP(this, &SDisplayClusterColorGradingDrawer::OnObjectsReplaced);
+	GEngine->OnLevelActorAdded().AddSP(this, &SDisplayClusterColorGradingDrawer::OnLevelActorAdded);
+	GEngine->OnLevelActorDeleted().AddSP(this, &SDisplayClusterColorGradingDrawer::OnLevelActorDeleted);
+
+	GEditor->RegisterForUndo(this);
+
+	FillLevelColorGradingList();
+	FillRootActorColorGradingList();
+
+	ChildSlot
+	[
+		SNew(SBorder)
+		.BorderImage(FAppStyle::Get().GetBrush("Brushes.Panel"))
+		.HAlign(HAlign_Fill)
+		.VAlign(VAlign_Fill)
+		.Padding(FMargin(0.0f, 0.0f))
+		[
+			// Splitter to divide the object list and the color panel
+			SNew(SSplitter)
+			.Orientation(Orient_Horizontal)
+			.PhysicalSplitterHandleSize(2.0f)
+
+			+SSplitter::Slot()
+			.Value(0.12f)
+			[
+				SNew(SVerticalBox)
+
+				// Toolbar slot for the object list 
+				+ SVerticalBox::Slot()
+				.AutoHeight()
+				.Padding(0, 0, 0, 0)
+				[
+					SNew(SBorder)
+					.Padding(FMargin(3))
+					.BorderImage(bIsInDrawer ? FStyleDefaults::GetNoBrush() : FAppStyle::Get().GetBrush("Brushes.Panel"))
+					[
+						SNew(SBox)
+						.HeightOverride(24.0f)
+					]
+				]
+
+				+SVerticalBox::Slot()
+				.AutoHeight()
+				[
+					SNew(SSeparator)
+					.Thickness(2.0f)
+				]
+
+				// Slot for the object lists, wrapped in expandable areas
+				+ SVerticalBox::Slot()
+				.FillHeight(1.0f)
+				[
+					SNew(SBox)
+					.Padding(FMargin(4.f))
+					[
+						SNew(SBorder)
+						.Padding(FMargin(0.0f))
+						.BorderImage(FAppStyle::GetBrush("Brushes.Recessed"))
+						[
+							SNew(SVerticalBox)
+
+							+SVerticalBox::Slot()
+							.AutoHeight()
+							[
+								SNew(SExpandableArea)
+								.BorderImage(FAppStyle::Get().GetBrush("Brushes.Header"))
+								.BodyBorderImage(FAppStyle::Get().GetBrush("Brushes.Recessed"))
+								.HeaderPadding(FMargin(4.0f, 2.0f))
+								.InitiallyCollapsed(false)
+								.AllowAnimatedTransition(false)
+								.Visibility_Lambda([this]() { return LevelColorGradingItems.Num() ? EVisibility::Visible : EVisibility::Collapsed; })
+								.HeaderContent()
+								[
+									SNew(STextBlock)
+									.Text(this, &SDisplayClusterColorGradingDrawer::GetCurrentLevelName)
+									.TextStyle(FAppStyle::Get(), "ButtonText")
+									.Font(FAppStyle::Get().GetFontStyle("NormalFontBold"))
+								]
+								.BodyContent()
+								[
+									SAssignNew(LevelActorsList, SDisplayClusterColorGradingObjectList)
+									.ColorGradingItemsSource(&LevelColorGradingItems)
+									.OnSelectionChanged(this, &SDisplayClusterColorGradingDrawer::OnListSelectionChanged)
+								]
+							]
+
+							+SVerticalBox::Slot()
+							.AutoHeight()
+							[
+								SNew(SExpandableArea)
+								.BorderImage(FAppStyle::Get().GetBrush("Brushes.Header"))
+								.BodyBorderImage(FAppStyle::Get().GetBrush("Brushes.Recessed"))
+								.HeaderPadding(FMargin(4.0f, 2.0f))
+								.InitiallyCollapsed(false)
+								.AllowAnimatedTransition(false)
+								.Visibility_Lambda([this]() { return RootActorColorGradingItems.Num() ? EVisibility::Visible : EVisibility::Collapsed; })
+								.HeaderContent()
+								[
+									SNew(STextBlock)
+									.Text(this, &SDisplayClusterColorGradingDrawer::GetCurrentRootActorName)
+									.TextStyle(FAppStyle::Get(), "ButtonText")
+									.Font(FAppStyle::Get().GetFontStyle("NormalFontBold"))
+								]
+								.BodyContent()
+								[
+									SAssignNew(RootActorList, SDisplayClusterColorGradingObjectList)
+									.ColorGradingItemsSource(&RootActorColorGradingItems)
+									.OnSelectionChanged(this, &SDisplayClusterColorGradingDrawer::OnListSelectionChanged)
+								]
+							]
+						]
+					]
+				]
+			]
+
+			+SSplitter::Slot()
+			.Value(0.88f)
+			[
+				SNew(SVerticalBox)
+
+				// Toolbar slot for the main drawer toolbar
+				+ SVerticalBox::Slot()
+				.AutoHeight()
+				.Padding(0, 0, 0, 0)
+				[
+					SNew(SBorder)
+					.Padding(FMargin(3))
+					.BorderImage(bIsInDrawer ? FStyleDefaults::GetNoBrush() : FAppStyle::Get().GetBrush("Brushes.Panel"))
+					[
+						SNew(SBox)
+						.HeightOverride(24.0f)
+						[
+							SNew(SHorizontalBox)
+
+							// Slot for the color grading group toolbar
+							+ SHorizontalBox::Slot()
+							.AutoWidth()
+							.HAlign(HAlign_Left)
+							.VAlign(VAlign_Center)
+							[
+								SAssignNew(ColorGradingGroupToolBarBox, SHorizontalBox)
+								.Visibility(this, &SDisplayClusterColorGradingDrawer::GetColorGradingGroupToolBarVisibility)
+							]
+							
+
+							+ SHorizontalBox::Slot()
+							.FillWidth(1.0f)
+							[
+								SNew(SSpacer)
+							]
+
+							// Slot for the "Dock in Layout" button
+							+ SHorizontalBox::Slot()
+							.AutoWidth()
+							.HAlign(HAlign_Right)
+							.VAlign(VAlign_Center)
+							[
+								CreateDockInLayoutButton()
+							]
+
+							// Slot for the Settings button
+							+ SHorizontalBox::Slot()
+							.AutoWidth()
+							.Padding(5.0f, 0.0f, 0.0f, 0.0f)
+							.HAlign(HAlign_Right)
+							.VAlign(VAlign_Center)
+							[
+								SNew(SComboButton)
+								.ComboButtonStyle(&FAppStyle::Get().GetWidgetStyle<FComboButtonStyle>("SimpleComboButton"))
+								.HasDownArrow(false)
+								.ButtonContent()
+								[
+									SNew(SHorizontalBox)
+									+ SHorizontalBox::Slot()
+									.AutoWidth()
+									.VAlign(VAlign_Center)
+									.Padding(4.0, 0.0f)
+									[
+										SNew(SImage)
+										.ColorAndOpacity(FSlateColor::UseForeground())
+										.Image(FAppStyle::Get().GetBrush("Icons.Settings"))
+									]
+									+ SHorizontalBox::Slot()
+									.VAlign(VAlign_Center)
+									.Padding(4.0, 0.0f)
+									[
+										SNew(STextBlock)
+										.Text(LOCTEXT("Settings", "Settings"))
+										.ColorAndOpacity(FSlateColor::UseForeground())
+									]
+								]
+							]
+						]
+					]
+				]
+
+				+SVerticalBox::Slot()
+				.AutoHeight()
+				[
+					SNew(SSeparator)
+					.Thickness(2.0f)
+				]
+
+				// Slot for the color panel
+				+SVerticalBox::Slot()
+				[
+					SNew(SBorder)
+					.BorderImage(FAppStyle::Get().GetBrush("Brushes.Panel"))
+					.Padding(FMargin(2.0f, 2.0f, 2.0f, 0.0f))
+					[
+						SAssignNew(ColorWheelPanel, SDisplayClusterColorGradingColorWheelPanel)
+						.ColorGradingDataModelSource(ColorGradingDataModel)
+					]
+				]
+			]
+		]
+	];
+}
+
+void SDisplayClusterColorGradingDrawer::Refresh(bool bPreserveDrawerState)
+{
+	FDisplayClusterColorGradingDrawerState DrawerState = GetDrawerState();
+
+	ColorGradingDataModel->Reset();
+
+	FillLevelColorGradingList();
+	FillRootActorColorGradingList();
+
+	if (ColorWheelPanel.IsValid())
+	{
+		ColorWheelPanel->Refresh();
+	}
+
+	if (bPreserveDrawerState)
+	{
+		SetDrawerState(DrawerState);
+	}
+}
+
+void SDisplayClusterColorGradingDrawer::Tick(const FGeometry& AllottedGeometry, const double InCurrentTime, const float InDeltaTime)
+{
+	SCompoundWidget::Tick(AllottedGeometry, InCurrentTime, InDeltaTime);
+
+	if (bRefreshOnNextTick)
+	{
+		const bool bPreserveDrawerState = true;
+		Refresh(bPreserveDrawerState);
+
+		bRefreshOnNextTick = false;
+	}
+}
+
+void SDisplayClusterColorGradingDrawer::PostUndo(bool bSuccess)
+{
+	if (bSuccess)
+	{
+		const bool bPreserveDrawerState = true;
+		Refresh(bPreserveDrawerState);
+	}
+}
+
+void SDisplayClusterColorGradingDrawer::PostRedo(bool bSuccess)
+{
+	if (bSuccess)
+	{
+		const bool bPreserveDrawerState = true;
+		Refresh(bPreserveDrawerState);
+	}
+}
+
+FDisplayClusterColorGradingDrawerState SDisplayClusterColorGradingDrawer::GetDrawerState() const
+{
+	FDisplayClusterColorGradingDrawerState DrawerState;
+
+	DrawerState.SelectedColorGradingGroup = ColorGradingDataModel->GetSelectedColorGradingGroupIndex();
+	DrawerState.SelectedColorGradingElement = ColorGradingDataModel->GetSelectedColorGradingElementIndex();
+
+	if (ColorWheelPanel.IsValid())
+	{
+		ColorWheelPanel->GetDrawerState(DrawerState);
+	}
+
+	if (LevelActorsList.IsValid())
+	{
+		TArray<FDisplayClusterColorGradingListItemRef> SelectedLevelItems = LevelActorsList->GetSelectedItems();
+
+		for (const FDisplayClusterColorGradingListItemRef& SelectedLevelItem : SelectedLevelItems)
+		{
+			if (SelectedLevelItem.IsValid())
+			{
+				if (SelectedLevelItem->Component.IsValid())
+				{
+					DrawerState.SelectedObjects.Add(SelectedLevelItem->Component);
+				}
+				else if (SelectedLevelItem->Actor.IsValid())
+				{
+					DrawerState.SelectedObjects.Add(SelectedLevelItem->Actor);
+				}
+			}
+		}
+	}
+
+	if (RootActorList.IsValid())
+	{
+		TArray<FDisplayClusterColorGradingListItemRef> SelectedRootActorItems = RootActorList->GetSelectedItems();
+
+		for (const FDisplayClusterColorGradingListItemRef& SelectedRootActorItem : SelectedRootActorItems)
+		{
+			if (SelectedRootActorItem.IsValid())
+			{
+				if (SelectedRootActorItem->Component.IsValid())
+				{
+					DrawerState.SelectedObjects.Add(SelectedRootActorItem->Component);
+				}
+				else if (SelectedRootActorItem->Actor.IsValid())
+				{
+					DrawerState.SelectedObjects.Add(SelectedRootActorItem->Actor);
+				}
+			}
+		}
+	}
+
+	return DrawerState;
+}
+
+void SDisplayClusterColorGradingDrawer::SetDrawerState(const FDisplayClusterColorGradingDrawerState& InDrawerState)
+{
+	TArray<FDisplayClusterColorGradingListItemRef> LevelItemsToSelect;
+	TArray<FDisplayClusterColorGradingListItemRef> RootActorItemsToSelect;
+	for (const TWeakObjectPtr<UObject>& SelectedObject : InDrawerState.SelectedObjects)
+	{
+		if (SelectedObject.IsValid())
+		{
+			FDisplayClusterColorGradingListItemRef* FoundColorGradingItemPtr = LevelColorGradingItems.FindByPredicate([&SelectedObject](const FDisplayClusterColorGradingListItemRef& ColorGradingItem)
+			{
+				return ColorGradingItem->Actor == SelectedObject || ColorGradingItem->Component == SelectedObject;
+			});
+
+			if (FoundColorGradingItemPtr)
+			{
+				LevelItemsToSelect.Add(*FoundColorGradingItemPtr);
+			}
+			else
+			{
+				FoundColorGradingItemPtr = RootActorColorGradingItems.FindByPredicate([&SelectedObject](const FDisplayClusterColorGradingListItemRef& ColorGradingItem)
+				{
+					return ColorGradingItem->Actor == SelectedObject || ColorGradingItem->Component == SelectedObject;
+				});
+
+				if (FoundColorGradingItemPtr)
+				{
+					RootActorItemsToSelect.Add(*FoundColorGradingItemPtr);
+				}
+			}
+		}
+	}
+
+	// TODO: For now, since we don't support multiple color grading items selected at once, ensure either a level item or a root actor item is seleccted, not both
+	if (LevelItemsToSelect.Num())
+	{
+		if (LevelActorsList.IsValid())
+		{
+			LevelActorsList->SetSelectedItems(LevelItemsToSelect);
+		}
+	}
+	else if (RootActorItemsToSelect.Num())
+	{
+		if (RootActorList.IsValid())
+		{
+			RootActorList->SetSelectedItems(RootActorItemsToSelect);
+		}
+	}
+
+	ColorGradingDataModel->SetSelectedColorGradingGroup(InDrawerState.SelectedColorGradingGroup);
+	ColorGradingDataModel->SetSelectedColorGradingElement(InDrawerState.SelectedColorGradingElement);
+
+	if (ColorWheelPanel.IsValid())
+	{
+		ColorWheelPanel->SetDrawerState(InDrawerState);
+	}
+}
+
+TSharedRef<SWidget> SDisplayClusterColorGradingDrawer::CreateDockInLayoutButton()
+{
+	if (bIsInDrawer)
+	{
+		return SNew(SButton)
+			.ButtonStyle(FAppStyle::Get(), "SimpleButton")
+			.ToolTipText(LOCTEXT("DockInLayout_Tooltip", "Docks this panel in the current operator window, copying all settings from the drawer.\nThe drawer will still be usable."))
+			.OnClicked(this, &SDisplayClusterColorGradingDrawer::DockInLayout)
+			[
+				SNew(SHorizontalBox)
+
+				+ SHorizontalBox::Slot()
+				.AutoWidth()
+				.VAlign(VAlign_Center)
+				.Padding(4.0, 0.0f)
+				[
+					SNew(SImage)
+					.ColorAndOpacity(FSlateColor::UseForeground())
+					.Image(FAppStyle::Get().GetBrush("Icons.Layout"))
+				]
+
+				+ SHorizontalBox::Slot()
+				.VAlign(VAlign_Center)
+				.Padding(4.0, 0.0f)
+				[
+					SNew(STextBlock)
+					.Text(LOCTEXT("DockInLayout", "Dock in Layout"))
+					.ColorAndOpacity(FSlateColor::UseForeground())
+				]
+			];
+	}
+	
+	return SNullWidget::NullWidget;
+}
+
+FText SDisplayClusterColorGradingDrawer::GetCurrentLevelName() const
+{
+	if (OperatorViewModel->HasRootActor())
+	{
+		if (UWorld* World = OperatorViewModel->GetRootActor()->GetWorld())
+		{
+			return FText::FromString(World->GetMapName());
+		}
+	}
+
+	return FText::GetEmpty();
+}
+
+FText SDisplayClusterColorGradingDrawer::GetCurrentRootActorName() const
+{
+	if (OperatorViewModel->HasRootActor())
+	{
+		return FText::FromString(OperatorViewModel->GetRootActor()->GetActorLabel());
+	}
+
+	return FText::GetEmpty();
+}
+
+void SDisplayClusterColorGradingDrawer::BindBlueprintCompiledDelegate(const UClass* Class)
+{
+	if (UBlueprint* Blueprint = UBlueprint::GetBlueprintFromClass(Class))
+	{
+		if (!Blueprint->OnCompiled().IsBoundToObject(this))
+		{
+			Blueprint->OnCompiled().AddSP(this, &SDisplayClusterColorGradingDrawer::OnBlueprintCompiled);
+		}
+	}
+}
+
+void SDisplayClusterColorGradingDrawer::UnbindBlueprintCompiledDelegate(const UClass* Class)
+{
+	if (UBlueprint* Blueprint = UBlueprint::GetBlueprintFromClass(Class))
+	{
+		Blueprint->OnCompiled().RemoveAll(this);
+	}
+}
+
+// Macros to create the lambdas used for a specified UObject's color grading enabled property in the color grading object list
+#define CREATE_IS_ENABLED_LAMBDA(Object, IsEnabledProperty) TAttribute<bool>::CreateLambda([Object]() { return (bool)IsEnabledProperty; })
+#define CREATE_ON_ENABLED_CHANGED_LAMBDA(Object, IsEnabledProperty) FOnColorGradingItemEnabledChanged::CreateLambda([Object](FDisplayClusterColorGradingListItemRef ListItem, bool bIsEnabled) \
+	{ \
+		FScopedTransaction Transaction(LOCTEXT("ColorGradingToggledTransaction", "Color Grading Toggled")); \
+		Object->Modify(); \
+		IsEnabledProperty = bIsEnabled; \
+	})
+
+void SDisplayClusterColorGradingDrawer::FillLevelColorGradingList()
+{
+	for (const FDisplayClusterColorGradingListItemRef& LevelColorGradingItem : LevelColorGradingItems)
+	{
+		if (LevelColorGradingItem->Component.IsValid())
+		{
+			UnbindBlueprintCompiledDelegate(LevelColorGradingItem->Component->GetClass());
+		}
+
+		if (LevelColorGradingItem->Actor.IsValid())
+		{
+			UnbindBlueprintCompiledDelegate(LevelColorGradingItem->Actor->GetClass());
+		}
+	}
+
+	LevelColorGradingItems.Empty();
+
+	if (ADisplayClusterRootActor* RootActor = OperatorViewModel->GetRootActor())
+	{
+		if (UWorld* World = RootActor->GetWorld())
+		{
+			for (TActorIterator<APostProcessVolume> PPVIter(World); PPVIter; ++PPVIter)
+			{
+				APostProcessVolume* PostProcessVolume = *PPVIter;
+				BindBlueprintCompiledDelegate(PostProcessVolume->GetClass());
+
+				FDisplayClusterColorGradingListItemRef PPVListItemRef = MakeShared<FDisplayClusterColorGradingListItem>(PostProcessVolume);
+				PPVListItemRef->IsItemEnabled = CREATE_IS_ENABLED_LAMBDA(PostProcessVolume, PostProcessVolume->bEnabled);
+				PPVListItemRef->OnItemEnabledChanged = CREATE_ON_ENABLED_CHANGED_LAMBDA(PostProcessVolume, PostProcessVolume->bEnabled);
+
+				LevelColorGradingItems.Add(PPVListItemRef);
+			}
+
+			for (TActorIterator<AColorCorrectRegion> CCRIter(World); CCRIter; ++CCRIter)
+			{
+				AColorCorrectRegion* ColorCorrectRegion = *CCRIter;
+				BindBlueprintCompiledDelegate(ColorCorrectRegion->GetClass());
+
+				FDisplayClusterColorGradingListItemRef CCRListItemRef = MakeShared<FDisplayClusterColorGradingListItem>(ColorCorrectRegion);
+				CCRListItemRef->IsItemEnabled = CREATE_IS_ENABLED_LAMBDA(ColorCorrectRegion, ColorCorrectRegion->Enabled);
+				CCRListItemRef->OnItemEnabledChanged = CREATE_ON_ENABLED_CHANGED_LAMBDA(ColorCorrectRegion, ColorCorrectRegion->Enabled);
+
+				LevelColorGradingItems.Add(CCRListItemRef);
+			}
+		}
+	}
+
+	if (LevelActorsList.IsValid())
+	{
+		LevelActorsList->RefreshList();
+	}
+}
+
+void SDisplayClusterColorGradingDrawer::FillRootActorColorGradingList()
+{
+	for (const FDisplayClusterColorGradingListItemRef& RootActorColorGradingItem : RootActorColorGradingItems)
+	{
+		if (RootActorColorGradingItem->Component.IsValid())
+		{
+			UnbindBlueprintCompiledDelegate(RootActorColorGradingItem->Component->GetClass());
+		}
+
+		if (RootActorColorGradingItem->Actor.IsValid())
+		{
+			UnbindBlueprintCompiledDelegate(RootActorColorGradingItem->Actor->GetClass());
+		}
+	}
+
+	RootActorColorGradingItems.Empty();
+
+	if (ADisplayClusterRootActor* RootActor = OperatorViewModel->GetRootActor())
+	{
+		BindBlueprintCompiledDelegate(RootActor->GetClass());
+
+		FDisplayClusterColorGradingListItemRef RootActorListItemRef = MakeShared<FDisplayClusterColorGradingListItem>(RootActor);
+		RootActorListItemRef->IsItemEnabled = CREATE_IS_ENABLED_LAMBDA(RootActor, RootActor->GetConfigData()->StageSettings.EnableColorGrading);
+		RootActorListItemRef->OnItemEnabledChanged = CREATE_ON_ENABLED_CHANGED_LAMBDA(RootActor, RootActor->GetConfigData()->StageSettings.EnableColorGrading);
+
+		RootActorColorGradingItems.Add(RootActorListItemRef);
+
+		RootActor->ForEachComponent<UDisplayClusterICVFXCameraComponent>(false, [this, RootActor](UDisplayClusterICVFXCameraComponent* ICVFXCameraComponent)
+		{
+			BindBlueprintCompiledDelegate(ICVFXCameraComponent->GetClass());
+
+			FDisplayClusterColorGradingListItemRef ICVFXCameraListItemRef = MakeShared<FDisplayClusterColorGradingListItem>(RootActor, ICVFXCameraComponent);
+			ICVFXCameraListItemRef->IsItemEnabled = CREATE_IS_ENABLED_LAMBDA(ICVFXCameraComponent, ICVFXCameraComponent->CameraSettings.EnableInnerFrustumColorGrading);
+			ICVFXCameraListItemRef->OnItemEnabledChanged = CREATE_ON_ENABLED_CHANGED_LAMBDA(ICVFXCameraComponent, ICVFXCameraComponent->CameraSettings.EnableInnerFrustumColorGrading);
+
+			RootActorColorGradingItems.Add(ICVFXCameraListItemRef);
+		});
+	}
+
+	if (RootActorList.IsValid())
+	{
+		RootActorList->RefreshList();
+	}
+}
+
+void SDisplayClusterColorGradingDrawer::FillColorGradingGroupToolBar()
+{
+	if (ColorGradingGroupToolBarBox.IsValid())
+	{
+		ColorGradingGroupToolBarBox->ClearChildren();
+
+		for (int32 Index = 0; Index < ColorGradingDataModel->ColorGradingGroups.Num(); ++Index)
+		{
+			const FDisplayClusterColorGradingDataModel::FColorGradingGroup& Group = ColorGradingDataModel->ColorGradingGroups[Index];
+			ColorGradingGroupToolBarBox->AddSlot()
+			.AutoWidth()
+			.Padding(2.0f, 0.0f, 2.0f, 0.0f)
+			[
+				SNew(SCheckBox)
+				.Style(FAppStyle::Get(), "DetailsView.SectionButton")
+				.OnCheckStateChanged(this, &SDisplayClusterColorGradingDrawer::OnColorGradingGroupCheckedChanged, Index)
+				.IsChecked(this, &SDisplayClusterColorGradingDrawer::IsColorGradingGroupSelected, Index)
+				[
+					SNew(STextBlock)
+					.TextStyle(FAppStyle::Get(), "SmallText")
+					.Text(Group.DisplayName)
+				]
+			];
+		}
+
+		if (ColorGradingDataModel->ColorGradingGroupToolBarWidget.IsValid())
+		{
+			ColorGradingGroupToolBarBox->AddSlot()
+			.AutoWidth()
+			.Padding(2.0f, 0.0f, 2.0f, 0.0f)
+			[
+				ColorGradingDataModel->ColorGradingGroupToolBarWidget.ToSharedRef()
+			];
+		}
+	}
+}
+
+EVisibility SDisplayClusterColorGradingDrawer::GetColorGradingGroupToolBarVisibility() const
+{
+	if (ColorGradingDataModel->bShowColorGradingGroupToolBar)
+	{
+		return EVisibility::Visible;
+	}
+
+	return EVisibility::Hidden;
+}
+
+ECheckBoxState SDisplayClusterColorGradingDrawer::IsColorGradingGroupSelected(int32 GroupIndex) const
+{
+	return ColorGradingDataModel->GetSelectedColorGradingGroupIndex() == GroupIndex ? ECheckBoxState::Checked : ECheckBoxState::Unchecked;
+}
+
+void SDisplayClusterColorGradingDrawer::OnColorGradingGroupCheckedChanged(ECheckBoxState State, int32 GroupIndex)
+{
+	if (State == ECheckBoxState::Checked)
+	{
+		ColorGradingDataModel->SetSelectedColorGradingGroup(GroupIndex);
+	}
+}
+
+void SDisplayClusterColorGradingDrawer::OnObjectsReplaced(const TMap<UObject*, UObject*>& OldToNewInstanceMap)
+{
+	bool bNeedsFullRefresh = false;
+	bool bNeedsListRefresh = false;
+
+	const TArray<TWeakObjectPtr<UObject>> SelectedObjects = ColorGradingDataModel->GetObjects();
+
+	for (const TPair<UObject*, UObject*>& Pair : OldToNewInstanceMap)
+	{
+		if (Pair.Key && Pair.Value)
+		{
+			// Must use GetEvenIfUnreachable on the weak pointers here because most of the time, the objects being replaced have already been marked for GC, and TWeakObjectPtr
+			// will return nullptr from Get on GC-marked objects
+			FDisplayClusterColorGradingListItemRef* FoundColorGradingItemPtr =
+				RootActorColorGradingItems.FindByPredicate([&Pair](const FDisplayClusterColorGradingListItemRef& ColorGradingItem)
+				{
+					return ColorGradingItem->Actor.GetEvenIfUnreachable() == Pair.Key || ColorGradingItem->Component.GetEvenIfUnreachable() == Pair.Key;
+				});
+
+			if (FoundColorGradingItemPtr)
+			{
+				FDisplayClusterColorGradingListItemRef FoundColorGradingItem = *FoundColorGradingItemPtr;
+				if (FoundColorGradingItem->Actor.GetEvenIfUnreachable() == Pair.Key)
+				{
+					FoundColorGradingItem->Actor = Cast<AActor>(Pair.Value);
+				}
+				else if (FoundColorGradingItem->Component.GetEvenIfUnreachable() == Pair.Key)
+				{
+					FoundColorGradingItem->Component = Cast<UActorComponent>(Pair.Value);
+				}
+
+				bNeedsListRefresh = true;
+			}
+
+			if (SelectedObjects.Contains(Pair.Key))
+			{
+				bNeedsFullRefresh = true;
+			}
+		}
+	}
+
+	if (bNeedsFullRefresh)
+	{
+		const bool bPreserveDrawerState = true;
+		Refresh(bPreserveDrawerState);
+	}
+	else if (bNeedsListRefresh)
+	{
+		if (LevelActorsList.IsValid())
+		{
+			LevelActorsList->RefreshList();
+		}
+
+		if (RootActorList.IsValid())
+		{
+			RootActorList->RefreshList();
+		}
+	}
+}
+
+void SDisplayClusterColorGradingDrawer::OnLevelActorAdded(AActor* Actor)
+{
+	// Only refresh when the actor being added is being added to the root actor's world and is of a type this drawer cares about
+	if (OperatorViewModel->HasRootActor())
+	{
+		if (UWorld* World = OperatorViewModel->GetRootActor()->GetWorld())
+		{
+			if (World == Actor->GetWorld())
+			{
+				if (Actor->IsA<ADisplayClusterRootActor>() || Actor->IsA<APostProcessVolume>() || Actor->IsA<AColorCorrectRegion>())
+				{
+					// Wait to refresh, as this event can be fired off for several actors in a row in certain cases, such as when the root actor is recompiled after a property change
+					bRefreshOnNextTick = true;
+				}
+			}
+		}
+	}
+}
+
+void SDisplayClusterColorGradingDrawer::OnLevelActorDeleted(AActor* Actor)
+{
+	auto ContainsActorRef = [Actor](const FDisplayClusterColorGradingListItemRef& ColorGradingItem)
+	{
+		return ColorGradingItem->Actor.GetEvenIfUnreachable() == Actor;
+	};
+
+	// Only refresh if the actor being deleted is being referenced by the drawer
+	if (RootActorColorGradingItems.ContainsByPredicate(ContainsActorRef) || LevelColorGradingItems.ContainsByPredicate(ContainsActorRef))
+	{
+		// Must wait for next tick to refresh because the actor has not actually been removed from the level at this point
+		bRefreshOnNextTick = true;
+	}
+}
+
+void SDisplayClusterColorGradingDrawer::OnBlueprintCompiled(UBlueprint* Blueprint)
+{
+	const bool bPreserveDrawerState = true;
+	Refresh(bPreserveDrawerState);
+}
+
+void SDisplayClusterColorGradingDrawer::OnActiveRootActorChanged(ADisplayClusterRootActor* NewRootActor)
+{
+	const bool bPreserveDrawerState = true;
+	Refresh(bPreserveDrawerState);
+}
+
+void SDisplayClusterColorGradingDrawer::OnListSelectionChanged(TSharedRef<SDisplayClusterColorGradingObjectList> SourceList, FDisplayClusterColorGradingListItemRef SelectedItem, ESelectInfo::Type SelectInfo)
+{
+	TArray<FDisplayClusterColorGradingListItemRef> SelectedObjects = SourceList->GetSelectedItems();
+
+	// When an item on one list is selected, clear the selection of the other list
+	if (LevelActorsList == SourceList && SelectedObjects.Num())
+	{
+		if (RootActorList->GetSelectedItems().Num())
+		{
+			RootActorList->SetSelectedItems(TArray<FDisplayClusterColorGradingListItemRef>());
+		}
+	}
+	else if (RootActorList == SourceList && SelectedObjects.Num())
+	{
+		if (LevelActorsList->GetSelectedItems().Num())
+		{
+			LevelActorsList->SetSelectedItems(TArray<FDisplayClusterColorGradingListItemRef>());
+		}
+	}
+
+	TArray<UObject*> ObjectsToColorGrade;
+	for (const FDisplayClusterColorGradingListItemRef& SelectedObject : SelectedObjects)
+	{
+		if (SelectedObject->Component.IsValid())
+		{
+			ObjectsToColorGrade.Add(SelectedObject->Component.Get());
+		}
+		else if (SelectedObject->Actor.IsValid())
+		{
+			ObjectsToColorGrade.Add(SelectedObject->Actor.Get());
+		}
+	}
+
+	ColorGradingDataModel->SetObjects(ObjectsToColorGrade);
+
+	FillColorGradingGroupToolBar();
+
+	if (ColorWheelPanel.IsValid())
+	{
+		ColorWheelPanel->Refresh();
+	}
+}
+
+FReply SDisplayClusterColorGradingDrawer::DockInLayout()
+{
+	IDisplayClusterColorGrading::Get().GetColorGradingDrawerSingleton().DockColorGradingDrawer();
+	return FReply::Handled();
+}
+
+#undef LOCTEXT_NAMESPACE
\ No newline at end of file
diff --git a/Engine/Plugins/Runtime/nDisplay/Source/DisplayClusterColorGrading/Private/Drawer/SDisplayClusterColorGradingDrawer.h b/Engine/Plugins/Runtime/nDisplay/Source/DisplayClusterColorGrading/Private/Drawer/SDisplayClusterColorGradingDrawer.h
new file mode 100644
--- /dev/null
+++ b/Engine/Plugins/Runtime/nDisplay/Source/DisplayClusterColorGrading/Private/Drawer/SDisplayClusterColorGradingDrawer.h
@@ -0,0 +1,155 @@
+// Copyright Epic Games, Inc. All Rights Reserved.
+
+#pragma once
+
+#include "CoreMinimal.h"
+#include "EditorUndoClient.h"
+#include "Styling/SlateTypes.h"
+#include "Widgets/DeclarativeSyntaxSupport.h"
+#include "Widgets/SCompoundWidget.h"
+
+#include "SDisplayClusterColorGradingObjectList.h"
+#include "SDisplayClusterColorGradingColorWheel.h"
+
+class ADisplayClusterRootActor;
+class FDisplayClusterOperatorStatusBarExtender;
+class FDisplayClusterColorGradingDataModel;
+class IDisplayClusterOperatorViewModel;
+class IPropertyRowGenerator;
+class SHorizontalBox;
+class SDisplayClusterColorGradingColorWheelPanel;
+
+/** Stores the state of the drawer UI that can be reloaded in cases where the drawer or any of its elements are reloaded (such as when the drawer is reopened or docked) */
+struct FDisplayClusterColorGradingDrawerState
+{
+	/** The objects that are selected in the list */
+	TArray<TWeakObjectPtr<UObject>> SelectedObjects;
+
+	/** The color grading group that is selected */
+	int32 SelectedColorGradingGroup = INDEX_NONE;
+
+	/** The color grading element that is selected */
+	int32 SelectedColorGradingElement = INDEX_NONE;
+
+	/** Indicates which color wheels are hidden */
+	TArray<bool> HiddenColorWheels;
+
+	/** The selected orientation of the color wheels */
+	EOrientation ColorWheelOrientation = EOrientation::Orient_Vertical;
+
+	/** The color display mode of the color wheels */
+	SDisplayClusterColorGradingColorWheel::EColorDisplayMode ColorDisplayMode;
+};
+
+/** Color grading drawer widget, which displays a list of color gradable items, and the color wheel panel */
+class SDisplayClusterColorGradingDrawer : public SCompoundWidget, public FEditorUndoClient
+{
+public:
+	~SDisplayClusterColorGradingDrawer();
+
+	SLATE_BEGIN_ARGS(SDisplayClusterColorGradingDrawer)
+	{}
+	SLATE_END_ARGS()
+
+	void Construct(const FArguments& InArgs, bool bInIsInDrawer);
+
+	/** Refreshes the drawer's UI to match the current state of the level and active root actor, optionally preserving UI state */
+	void Refresh(bool bPreserveDrawerState = false);
+
+	//~ SWidget interface
+	virtual void Tick(const FGeometry& AllottedGeometry, const double InCurrentTime, const float InDeltaTime) override;
+	//~ End SWidget interface
+
+	//~ FEditorUndoClient interface
+	virtual void PostUndo(bool bSuccess) override;
+	virtual void PostRedo(bool bSuccess) override;
+	//~ End FEditorUndoClient interface
+
+	/** Gets the state of the drawer UI */
+	FDisplayClusterColorGradingDrawerState GetDrawerState() const;
+
+	/** Sets the state of the drawer UI */
+	void SetDrawerState(const FDisplayClusterColorGradingDrawerState& InDrawerState);
+
+private:
+	/** Creates the button used to dock the drawer in the operator panel */
+	TSharedRef<SWidget> CreateDockInLayoutButton();
+
+	/** Gets the name of the current level the active root actor is in */
+	FText GetCurrentLevelName() const;
+
+	/** Gets the name of the active root actor */
+	FText GetCurrentRootActorName() const;
+
+	/** Binds a callback to the BlueprintCompiled delegate of the specified class */
+	void BindBlueprintCompiledDelegate(const UClass* Class);
+
+	/** Unbinds a callback to the BlueprintCompiled delegate of the specified class */
+	void UnbindBlueprintCompiledDelegate(const UClass* Class);
+
+	/** Fills the level color grading list with any color gradable items found in the currently loaded level */
+	void FillLevelColorGradingList();
+
+	/** Fills the root actor color grading list with any color gradable components of the active root actor */
+	void FillRootActorColorGradingList();
+
+	/** Fills the color grading group toolbar using the color grading data model */
+	void FillColorGradingGroupToolBar();
+
+	/** Gets the visibility state of the color grading group toolbar */
+	EVisibility GetColorGradingGroupToolBarVisibility() const;
+
+	/** Gets whether the color grading group at the specified index is currently selected */
+	ECheckBoxState IsColorGradingGroupSelected(int32 GroupIndex) const;
+
+	/** Raised when the user has selected the specified color grading group */
+	void OnColorGradingGroupCheckedChanged(ECheckBoxState State, int32 GroupIndex);
+
+	/** Raised when the editor replaces any UObjects with new instantiations, usually when actors have been recompiled from blueprints */
+	void OnObjectsReplaced(const TMap<UObject*, UObject*>& OldToNewInstanceMap);
+
+	/** Raised when an actor is added to the current level */
+	void OnLevelActorAdded(AActor* Actor);
+
+	/** Raised when an actor has been deleted from the currnent level */
+	void OnLevelActorDeleted(AActor* Actor);
+
+	/** Raised when the specified blueprint has been recompiled */
+	void OnBlueprintCompiled(UBlueprint* Blueprint);
+
+	/** Raised when the user has changed the active root actor selected in the nDisplay operator panel */
+	void OnActiveRootActorChanged(ADisplayClusterRootActor* NewRootActor);
+
+	/** Raised when the user has selected a new item in any of the drawer's list views */
+	void OnListSelectionChanged(TSharedRef<SDisplayClusterColorGradingObjectList> SourceList, FDisplayClusterColorGradingListItemRef SelectedItem, ESelectInfo::Type SelectInfo);
+
+	/** Raised when the "Dock in Layout" button has been clicked */
+	FReply DockInLayout();
+
+private:
+	/** The operator panel's view model */
+	TSharedPtr<IDisplayClusterOperatorViewModel> OperatorViewModel;
+
+	TSharedPtr<SDisplayClusterColorGradingObjectList> LevelActorsList;
+	TSharedPtr<SDisplayClusterColorGradingObjectList> RootActorList;
+	TSharedPtr<SHorizontalBox> ColorGradingGroupToolBarBox;
+	TSharedPtr<SDisplayClusterColorGradingColorWheelPanel> ColorWheelPanel;
+
+	/** List of color gradable items found in the current level */
+	TArray<FDisplayClusterColorGradingListItemRef> LevelColorGradingItems;
+
+	/** List of color gradable actors and components attached to the root actor */
+	TArray<FDisplayClusterColorGradingListItemRef> RootActorColorGradingItems;
+
+	/** The color grading data model for the currently selected objects */
+	TSharedPtr<FDisplayClusterColorGradingDataModel> ColorGradingDataModel;
+
+	/** Gets whether this widget is in a drawer or docked in a tab */
+	bool bIsInDrawer = false;
+
+	/** Indicates that the drawer should refresh itself on the next tick */
+	bool bRefreshOnNextTick = false;
+
+	/** Delegate handle for the OnActiveRootActorChanged delegate */
+	FDelegateHandle ActiveRootActorChangedHandle;
+};
\ No newline at end of file
diff --git a/Engine/Plugins/Runtime/nDisplay/Source/DisplayClusterColorGrading/Private/SDisplayClusterColorGradingColorWheel.cpp b/Engine/Plugins/Runtime/nDisplay/Source/DisplayClusterColorGrading/Private/SDisplayClusterColorGradingColorWheel.cpp
new file mode 100644
--- /dev/null
+++ b/Engine/Plugins/Runtime/nDisplay/Source/DisplayClusterColorGrading/Private/SDisplayClusterColorGradingColorWheel.cpp
@@ -0,0 +1,804 @@
+// Copyright Epic Games, Inc. All Rights Reserved.
+
+#include "SDisplayClusterColorGradingColorWheel.h"
+
+#include "Customizations/MathStructCustomizations.h"
+#include "DetailLayoutBuilder.h"
+#include "Editor.h"
+#include "PropertyHandle.h"
+#include "ScopedTransaction.h"
+#include "Widgets/SBoxPanel.h"
+#include "Widgets/Colors/SComplexGradient.h"
+#include "Widgets/Layout/SBox.h"
+
+#define LOCTEXT_NAMESPACE "DisplayClusterColorGrading"
+
+void SDisplayClusterColorGradingColorWheel::Construct(const FArguments& InArgs)
+{
+	Orientation = InArgs._Orientation;
+	ColorDisplayMode = InArgs._ColorDisplayMode;
+
+	ChildSlot
+	[
+		SNew(SVerticalBox)
+
+		+ SVerticalBox::Slot()
+		.AutoHeight()
+		.HAlign(HAlign_Center)
+		[
+			SAssignNew(HeaderBox, SBox)
+		]
+
+		+ SVerticalBox::Slot()
+		.FillHeight(1.0f)
+		[
+			SAssignNew(ColorWheelBox, SBox)
+		]
+	];
+
+	ColorPickerBox = SNew(SBox);
+	ColorSlidersBox = SNew(SBox);
+
+	SetOrientation(Orientation);
+
+	if (InArgs._HeaderContent.IsValid())
+	{
+		HeaderBox->SetContent(InArgs._HeaderContent.ToSharedRef());
+	}
+}
+
+void SDisplayClusterColorGradingColorWheel::SetColorPropertyHandle(TSharedPtr<IPropertyHandle> InColorPropertyHandle)
+{
+	ColorGradingPicker.Reset();
+	ColorPropertyMetadata.Reset();
+	ComponentSliderDynamicMinValue.Reset();
+	ComponentSliderDynamicMaxValue.Reset();
+
+	ColorPropertyHandle = InColorPropertyHandle;
+
+	if (ColorPropertyHandle.IsValid() && ColorPropertyHandle.Pin()->IsValidHandle())
+	{
+		ColorPropertyMetadata = GetColorPropertyMetadata();
+		ComponentSliderDynamicMinValue = ColorPropertyMetadata->MinValue;
+		ComponentSliderDynamicMaxValue = ColorPropertyMetadata->MaxValue;
+	}
+
+	// Since many of the color picker slate properties are not attributes, we need to recreate the widget every time the property handle changes
+	// to ensure the picker is configured correctly for the color property
+	if (ColorPickerBox.IsValid())
+	{
+		ColorPickerBox->SetContent(CreateColorGradingPicker());
+	}
+
+	if (ColorSlidersBox.IsValid())
+	{
+		ColorSlidersBox->SetContent(CreateColorComponentSliders());
+	}
+}
+
+void SDisplayClusterColorGradingColorWheel::SetHeaderContent(const TSharedRef<SWidget>& HeaderContent)
+{
+	if (HeaderBox.IsValid())
+	{
+		HeaderBox->SetContent(HeaderContent);
+	}
+}
+
+void SDisplayClusterColorGradingColorWheel::SetOrientation(EOrientation NewOrientation)
+{
+	Orientation = NewOrientation;
+
+	if (ColorWheelBox.IsValid())
+	{
+		TSharedPtr<SWidget> OrientedBox;
+
+		if (Orientation == EOrientation::Orient_Vertical)
+		{
+			OrientedBox = SNew(SVerticalBox)
+
+			+ SVerticalBox::Slot()
+			.AutoHeight()
+			[
+				ColorPickerBox.ToSharedRef()
+			]
+
+			+ SVerticalBox::Slot()
+			.AutoHeight()
+			[
+				ColorSlidersBox.ToSharedRef()
+			];
+		}
+		else
+		{
+			OrientedBox = SNew(SHorizontalBox)
+
+			+ SHorizontalBox::Slot()
+			.AutoWidth()
+			[
+				ColorPickerBox.ToSharedRef()
+			]
+
+			+ SHorizontalBox::Slot()
+			.FillWidth(1.0f)
+			[
+				ColorSlidersBox.ToSharedRef()
+			];
+		}
+
+		ColorWheelBox->SetContent(OrientedBox.ToSharedRef());
+		ColorWheelBox->SetHAlign(Orientation == EOrientation::Orient_Vertical ? HAlign_Center : HAlign_Fill);
+	}
+}
+
+TSharedRef<SWidget> SDisplayClusterColorGradingColorWheel::SDisplayClusterColorGradingColorWheel::CreateColorGradingPicker()
+{
+	if (ColorPropertyHandle.IsValid() && ColorPropertyHandle.Pin()->IsValidHandle())
+	{
+		return SNew(SColorGradingPicker)
+			.ValueMin(GetMetadataMinValue())
+			.ValueMax(GetMetadataMaxValue())
+			.SliderValueMin(GetMetadataSliderMinValue())
+			.SliderValueMax(GetMetadataSliderMaxValue())
+			.MainDelta(GetMetadataDelta())
+			.SupportDynamicSliderMinValue(GetMetadataSupportDynamicSliderMinValue())
+			.SupportDynamicSliderMaxValue(GetMetadataSupportDynamicSliderMaxValue())
+			.MainShiftMouseMovePixelPerDelta(GetMetadataShiftMouseMovePixelPerDelta())
+			.ColorGradingModes(ColorPropertyMetadata->ColorGradingMode)
+			.OnColorCommitted(this, &SDisplayClusterColorGradingColorWheel::CommitColor)
+			.OnQueryCurrentColor(this, &SDisplayClusterColorGradingColorWheel::GetColor)
+			.AllowSpin(ColorPropertyHandle.Pin()->GetNumOuterObjects() == 1) // TODO: May want to find a way to support multiple objects
+			.OnBeginSliderMovement(this, &SDisplayClusterColorGradingColorWheel::BeginUsingColorPickerSlider)
+			.OnEndSliderMovement(this, &SDisplayClusterColorGradingColorWheel::EndUsingColorPickerSlider)
+			.OnBeginMouseCapture(this, &SDisplayClusterColorGradingColorWheel::BeginUsingColorPickerSlider)
+			.OnEndMouseCapture(this, &SDisplayClusterColorGradingColorWheel::EndUsingColorPickerSlider)
+			.IsEnabled(this, &SDisplayClusterColorGradingColorWheel::IsPropertyEnabled);
+	}
+
+	return SNullWidget::NullWidget;
+}
+
+TSharedRef<SWidget> SDisplayClusterColorGradingColorWheel::CreateColorComponentSliders()
+{
+	if (ColorPropertyHandle.IsValid() && ColorPropertyHandle.Pin()->IsValidHandle())
+	{
+		TSharedRef<SVerticalBox> ColorSlidersVerticalBox = SNew(SVerticalBox);
+
+		uint32 NumComponents = 4;
+		for (uint32 ComponentIndex = 0; ComponentIndex < NumComponents; ++ComponentIndex)
+		{
+			TAttribute<FText> TextGetter = TAttribute<FText>::Create(TAttribute<FText>::FGetter::CreateSP(this, &SDisplayClusterColorGradingColorWheel::GetComponentLabelText, ComponentIndex));
+			TSharedRef<SWidget> LabelWidget = SNumericEntryBox<float>::BuildLabel(TextGetter, FLinearColor::White, FLinearColor(0.2f, 0.2f, 0.2f));
+
+			TSharedRef<SNumericEntryBox<float>> NumericEntryBox = SNew(SNumericEntryBox<float>)
+				.SpinBoxStyle(&FCoreStyle::Get().GetWidgetStyle<FSpinBoxStyle>("NumericEntrySpinBox_Dark"))
+				.EditableTextBoxStyle(&FCoreStyle::Get().GetWidgetStyle<FEditableTextBoxStyle>("DarkEditableTextBox"))
+				.Font(IDetailLayoutBuilder::GetDetailFont())
+				.UndeterminedString(NSLOCTEXT("PropertyEditor", "MultipleValues", "Multiple Values"))
+				.Value(this, &SDisplayClusterColorGradingColorWheel::GetComponentValue, ComponentIndex)
+				.OnValueChanged(this, &SDisplayClusterColorGradingColorWheel::SetComponentValue, ComponentIndex)
+				.OnBeginSliderMovement(this, &SDisplayClusterColorGradingColorWheel::BeginUsingComponentSlider, ComponentIndex)
+				.OnEndSliderMovement(this, &SDisplayClusterColorGradingColorWheel::EndUsingComponentSlider, ComponentIndex)
+				.AllowSpin(ColorPropertyHandle.Pin()->GetNumOuterObjects() == 1)
+				.ShiftMouseMovePixelPerDelta(ColorPropertyMetadata->ShiftMouseMovePixelPerDelta)
+				.SupportDynamicSliderMinValue(this, &SDisplayClusterColorGradingColorWheel::ComponentSupportsDynamicSliderValue, ColorPropertyMetadata->bSupportDynamicSliderMinValue, ComponentIndex)
+				.SupportDynamicSliderMaxValue(this, &SDisplayClusterColorGradingColorWheel::ComponentSupportsDynamicSliderValue, ColorPropertyMetadata->bSupportDynamicSliderMaxValue, ComponentIndex)
+				.OnDynamicSliderMinValueChanged(this, &SDisplayClusterColorGradingColorWheel::UpdateComponentDynamicSliderMinValue)
+				.OnDynamicSliderMaxValueChanged(this, &SDisplayClusterColorGradingColorWheel::UpdateComponentDynamicSliderMaxValue)
+				.MinValue(ColorPropertyMetadata->MinValue)
+				.MaxValue(this, &SDisplayClusterColorGradingColorWheel::GetComponentMaxValue, ColorPropertyMetadata->MaxValue, ComponentIndex)
+				.MinSliderValue(this, &SDisplayClusterColorGradingColorWheel::GetComponentMinSliderValue, ColorPropertyMetadata->SliderMinValue, ComponentIndex)
+				.MaxSliderValue(this, &SDisplayClusterColorGradingColorWheel::GetComponentMaxSliderValue, ColorPropertyMetadata->SliderMaxValue, ComponentIndex)
+				.SliderExponent(ColorPropertyMetadata->SliderExponent)
+				.SliderExponentNeutralValue(ColorPropertyMetadata->SliderMinValue.GetValue() + (ColorPropertyMetadata->SliderMaxValue.GetValue() - ColorPropertyMetadata->SliderMinValue.GetValue()) / 2.0f)
+				.Delta(this, &SDisplayClusterColorGradingColorWheel::GetComponentSliderDeltaValue, ColorPropertyMetadata->Delta, ComponentIndex)
+				.ToolTipText(this, &SDisplayClusterColorGradingColorWheel::GetComponentToolTipText, ComponentIndex)
+				.LabelPadding(FMargin(0))
+				.IsEnabled(this, &SDisplayClusterColorGradingColorWheel::IsPropertyEnabled)
+				.Label()
+				[
+					LabelWidget
+				];
+
+			ColorSlidersVerticalBox->AddSlot()
+				.Padding(FMargin(0.0f, 2.0f, 3.0f, 0.0f))
+				.AutoHeight()
+				.HAlign(HAlign_Fill)
+				.VAlign(VAlign_Fill)
+				[
+					NumericEntryBox
+				];
+
+			ColorSlidersVerticalBox->AddSlot()
+				.Padding(FMargin(15.0f, 0.0f, 3.0f, 2.0f))
+				.AutoHeight()
+				.HAlign(HAlign_Fill)
+				.VAlign(VAlign_Bottom)
+				[
+					SNew(SBox)
+					.HeightOverride(6.0f)
+					[
+						SNew(SComplexGradient)
+						.GradientColors(TAttribute<TArray<FLinearColor>>::Create(TAttribute<TArray<FLinearColor>>::FGetter::CreateSP(this, &SDisplayClusterColorGradingColorWheel::GetGradientColor, ComponentIndex)))
+						.Visibility(this, &SDisplayClusterColorGradingColorWheel::GetGradientVisibility)
+					]
+				];
+		}
+
+		return ColorSlidersVerticalBox;
+	}
+
+	return SNullWidget::NullWidget;
+}
+
+SDisplayClusterColorGradingColorWheel::FColorPropertyMetadata SDisplayClusterColorGradingColorWheel::GetColorPropertyMetadata() const
+{
+	FColorPropertyMetadata PropertyMetadata;
+
+	if (ColorPropertyHandle.IsValid() && ColorPropertyHandle.Pin()->IsValidHandle())
+	{
+		FProperty* Property = ColorPropertyHandle.Pin()->GetProperty();
+		const FString ColorGradingModeString = Property->GetMetaData(TEXT("ColorGradingMode")).ToLower();
+
+		if (!ColorGradingModeString.IsEmpty())
+		{
+			if (ColorGradingModeString.Compare(TEXT("saturation")) == 0)
+			{
+				PropertyMetadata.ColorGradingMode = EColorGradingModes::Saturation;
+			}
+			else if (ColorGradingModeString.Compare(TEXT("contrast")) == 0)
+			{
+				PropertyMetadata.ColorGradingMode = EColorGradingModes::Contrast;
+			}
+			else if (ColorGradingModeString.Compare(TEXT("gamma")) == 0)
+			{
+				PropertyMetadata.ColorGradingMode = EColorGradingModes::Gamma;
+			}
+			else if (ColorGradingModeString.Compare(TEXT("gain")) == 0)
+			{
+				PropertyMetadata.ColorGradingMode = EColorGradingModes::Gain;
+			}
+			else if (ColorGradingModeString.Compare(TEXT("offset")) == 0)
+			{
+				PropertyMetadata.ColorGradingMode = EColorGradingModes::Offset;
+			}
+		}
+
+		const FString& MetaUIMinString = Property->GetMetaData(TEXT("UIMin"));
+		const FString& MetaUIMaxString = Property->GetMetaData(TEXT("UIMax"));
+		const FString& SliderExponentString = Property->GetMetaData(TEXT("SliderExponent"));
+		const FString& DeltaString = Property->GetMetaData(TEXT("Delta"));
+		const FString& LinearDeltaSensitivityString = Property->GetMetaData(TEXT("LinearDeltaSensitivity"));
+		const FString& ShiftMouseMovePixelPerDeltaString = Property->GetMetaData(TEXT("ShiftMouseMovePixelPerDelta"));
+		const FString& SupportDynamicSliderMaxValueString = Property->GetMetaData(TEXT("SupportDynamicSliderMaxValue"));
+		const FString& SupportDynamicSliderMinValueString = Property->GetMetaData(TEXT("SupportDynamicSliderMinValue"));
+		const FString& ClampMinString = Property->GetMetaData(TEXT("ClampMin"));
+		const FString& ClampMaxString = Property->GetMetaData(TEXT("ClampMax"));
+
+		const FString& UIMinString = MetaUIMinString.Len() ? MetaUIMinString : ClampMinString;
+		const FString& UIMaxString = MetaUIMaxString.Len() ? MetaUIMaxString : ClampMaxString;
+
+		float ClampMin = TNumericLimits<float>::Lowest();
+		float ClampMax = TNumericLimits<float>::Max();
+
+		if (!ClampMinString.IsEmpty())
+		{
+			TTypeFromString<float>::FromString(ClampMin, *ClampMinString);
+		}
+
+		if (!ClampMaxString.IsEmpty())
+		{
+			TTypeFromString<float>::FromString(ClampMax, *ClampMaxString);
+		}
+
+		float UIMin = TNumericLimits<float>::Lowest();
+		float UIMax = TNumericLimits<float>::Max();
+		TTypeFromString<float>::FromString(UIMin, *UIMinString);
+		TTypeFromString<float>::FromString(UIMax, *UIMaxString);
+
+		const float ActualUIMin = FMath::Max(UIMin, ClampMin);
+		const float ActualUIMax = FMath::Min(UIMax, ClampMax);
+
+		PropertyMetadata.MinValue = ClampMinString.Len() ? ClampMin : TOptional<float>();
+		PropertyMetadata.MaxValue = ClampMaxString.Len() ? ClampMax : TOptional<float>();
+		PropertyMetadata.SliderMinValue = (UIMinString.Len()) ? ActualUIMin : TOptional<float>();
+		PropertyMetadata.SliderMaxValue = (UIMaxString.Len()) ? ActualUIMax : TOptional<float>();
+
+		if (SliderExponentString.Len())
+		{
+			TTypeFromString<float>::FromString(PropertyMetadata.SliderExponent, *SliderExponentString);
+		}
+
+		if (DeltaString.Len())
+		{
+			TTypeFromString<float>::FromString(PropertyMetadata.Delta, *DeltaString);
+		}
+
+		if (LinearDeltaSensitivityString.Len())
+		{
+			TTypeFromString<int32>::FromString(PropertyMetadata.LinearDeltaSensitivity, *LinearDeltaSensitivityString);
+			PropertyMetadata.Delta = (PropertyMetadata.Delta == 0.0f) ? 1.0f : PropertyMetadata.Delta;
+		}
+
+		if (ShiftMouseMovePixelPerDeltaString.Len())
+		{
+			TTypeFromString<int32>::FromString(PropertyMetadata.ShiftMouseMovePixelPerDelta, *ShiftMouseMovePixelPerDeltaString);
+			PropertyMetadata.ShiftMouseMovePixelPerDelta = FMath::Max(PropertyMetadata.ShiftMouseMovePixelPerDelta, 1);
+		}
+
+		PropertyMetadata.bSupportDynamicSliderMaxValue = SupportDynamicSliderMaxValueString.Len() > 0 && SupportDynamicSliderMaxValueString.ToBool();
+		PropertyMetadata.bSupportDynamicSliderMinValue = SupportDynamicSliderMinValueString.Len() > 0 && SupportDynamicSliderMinValueString.ToBool();
+	}
+
+	return PropertyMetadata;
+}
+
+bool SDisplayClusterColorGradingColorWheel::IsPropertyEnabled() const
+{
+	if (ColorPropertyHandle.IsValid() && ColorPropertyHandle.Pin()->IsValidHandle())
+	{
+		return ColorPropertyHandle.Pin()->IsEditable();
+	}
+
+	return false;
+}
+
+bool SDisplayClusterColorGradingColorWheel::GetColor(FVector4& OutCurrentColor)
+{
+	if (ColorPropertyHandle.IsValid())
+	{
+		return ColorPropertyHandle.Pin()->GetValue(OutCurrentColor) == FPropertyAccess::Success;
+	}
+
+	return false;
+}
+
+void SDisplayClusterColorGradingColorWheel::CommitColor(FVector4& NewValue, bool bShouldCommitValueChanges)
+{
+	FScopedTransaction Transaction(LOCTEXT("ColorWheel_TransactionName", "Color Grading Main Value"), bShouldCommitValueChanges);
+
+	if (ColorPropertyHandle.IsValid())
+	{
+		// Always perform a purely interactive change. We do this because it won't invoke reconstruction, which may cause that only the first 
+		// element gets updated due to its change causing a component reconstruction and the remaining vector element property handles updating 
+		// the trashed component.
+		ColorPropertyHandle.Pin()->SetValue(NewValue, EPropertyValueSetFlags::InteractiveChange | EPropertyValueSetFlags::NotTransactable);
+
+		// If not purely interactive, set the value with default flags.
+		if (bShouldCommitValueChanges || !bIsUsingColorPickerSlider)
+		{
+			ColorPropertyHandle.Pin()->SetValue(NewValue, EPropertyValueSetFlags::DefaultFlags);
+		}
+	}
+}
+
+void SDisplayClusterColorGradingColorWheel::BeginUsingColorPickerSlider()
+{
+	bIsUsingColorPickerSlider = true;
+	GEditor->BeginTransaction(LOCTEXT("ColorWheel_TransactionName", "Color Grading Main Value"));
+}
+
+void SDisplayClusterColorGradingColorWheel::EndUsingColorPickerSlider()
+{
+	bIsUsingColorPickerSlider = false;
+	GEditor->EndTransaction();
+}
+
+void SDisplayClusterColorGradingColorWheel::BeginUsingComponentSlider(uint32 ComponentIndex)
+{
+	bIsUsingComponentSlider = true;
+	GEditor->BeginTransaction(LOCTEXT("ColorWheel_TransactionName", "Color Grading Main Value"));
+}
+
+void SDisplayClusterColorGradingColorWheel::EndUsingComponentSlider(float NewValue, uint32 ComponentIndex)
+{
+	bIsUsingComponentSlider = false;
+	SetComponentValue(NewValue, ComponentIndex);
+	GEditor->EndTransaction();
+}
+
+FText SDisplayClusterColorGradingColorWheel::GetComponentLabelText(uint32 ComponentIndex) const
+{
+	static const FText RGBLabelText[] = {
+		LOCTEXT("ColorWheel_RedComponentLabel", "R"),
+		LOCTEXT("ColorWheel_GreenComponentLabel", "G"),
+		LOCTEXT("ColorWheel_BlueComponentLabel", "B")
+	};
+
+	static const FText HSVLabelText[] = {
+		LOCTEXT("ColorWheel_HueComponentLabel", "H"),
+		LOCTEXT("ColorWheel_SatComponentLabel", "S"),
+		LOCTEXT("ColorWheel_ValComponentLabel", "V")
+	};
+
+	if (ComponentIndex >= 0 && ComponentIndex < 3)
+	{ 
+		if (ColorDisplayMode.Get(EColorDisplayMode::RGB) == EColorDisplayMode::RGB)
+		{
+			return RGBLabelText[ComponentIndex];
+		}
+		else
+		{
+			return HSVLabelText[ComponentIndex];
+		}
+	}
+	else if (ComponentIndex == 3)
+	{
+		return LOCTEXT("ColorWheel_LuminanceComponentLabel", "Y");
+	}
+
+	return FText::GetEmpty();
+}
+
+TOptional<float> SDisplayClusterColorGradingColorWheel::GetComponentValue(uint32 ComponentIndex) const
+{
+	if (ColorPropertyHandle.IsValid())
+	{
+		FVector4 ColorValue;
+		if (ColorPropertyHandle.Pin()->GetValue(ColorValue) == FPropertyAccess::Success)
+		{
+			float Value = 0.0f;
+
+			if (ColorDisplayMode.Get(EColorDisplayMode::RGB) == EColorDisplayMode::RGB)
+			{
+				Value = ColorValue[ComponentIndex];
+			}
+			else
+			{
+				FLinearColor HSVColor = FLinearColor(ColorValue.X, ColorValue.Y, ColorValue.Z, ColorValue.W).LinearRGBToHSV();
+				Value = HSVColor.Component(ComponentIndex);
+			}
+
+			return Value;
+		}
+	}
+
+	return TOptional<float>();
+}
+
+void SDisplayClusterColorGradingColorWheel::SetComponentValue(float NewValue, uint32 ComponentIndex)
+{
+	if (ColorPropertyHandle.IsValid())
+	{
+		FVector4 CurrentColorValue;
+		if (ColorPropertyHandle.Pin()->GetValue(CurrentColorValue) == FPropertyAccess::Success)
+		{
+			FVector4 NewColorValue = CurrentColorValue;
+
+			if (ColorDisplayMode.Get(EColorDisplayMode::RGB) == EColorDisplayMode::RGB)
+			{
+				NewColorValue[ComponentIndex] = NewValue;
+			}
+			else
+			{
+				FLinearColor HSVColor = FLinearColor(CurrentColorValue.X, CurrentColorValue.Y, CurrentColorValue.Z, CurrentColorValue.W).LinearRGBToHSV();
+				HSVColor.Component(ComponentIndex) = NewValue;
+				NewColorValue = (FVector4)HSVColor.HSVToLinearRGB();
+			}
+
+			ColorPropertyHandle.Pin()->SetValue(NewColorValue, bIsUsingComponentSlider ? EPropertyValueSetFlags::InteractiveChange : EPropertyValueSetFlags::DefaultFlags);
+		}
+	}
+}
+
+bool SDisplayClusterColorGradingColorWheel::ComponentSupportsDynamicSliderValue(bool bDefaultValue, uint32 ComponentIndex) const
+{
+	if (bDefaultValue)
+	{
+		if (ColorDisplayMode.Get(EColorDisplayMode::RGB) != EColorDisplayMode::RGB)
+		{
+			return ComponentIndex >= 2;
+		}
+	}
+
+	return bDefaultValue;
+}
+
+void SDisplayClusterColorGradingColorWheel::UpdateComponentDynamicSliderMinValue(float NewValue, TWeakPtr<SWidget> SourceWidget, bool bIsOriginator, bool bUpdateOnlyIfLower)
+{
+	if (!ComponentSliderDynamicMinValue.IsSet() || (NewValue < ComponentSliderDynamicMinValue.GetValue() && bUpdateOnlyIfLower) || !bUpdateOnlyIfLower)
+	{
+		ComponentSliderDynamicMinValue = NewValue;
+	}
+}
+
+void SDisplayClusterColorGradingColorWheel::UpdateComponentDynamicSliderMaxValue(float NewValue, TWeakPtr<SWidget> SourceWidget, bool bIsOriginator, bool bUpdateOnlyIfHigher)
+{
+	if (!ComponentSliderDynamicMaxValue.IsSet() || (NewValue > ComponentSliderDynamicMaxValue.GetValue() && bUpdateOnlyIfHigher) || !bUpdateOnlyIfHigher)
+	{
+		ComponentSliderDynamicMaxValue = NewValue;
+	}
+}
+
+TOptional<float> SDisplayClusterColorGradingColorWheel::GetComponentMaxValue(TOptional<float> DefaultValue, uint32 ComponentIndex) const
+{
+	if (ColorDisplayMode.Get(EColorDisplayMode::RGB) != EColorDisplayMode::RGB)
+	{
+		if (ComponentIndex == 0)
+		{
+			return 359.0f;
+		}
+		else if (ComponentIndex == 1)
+		{
+			return 1.0f;
+		}
+	}
+
+	return DefaultValue;
+}
+
+TOptional<float> SDisplayClusterColorGradingColorWheel::GetComponentMinSliderValue(TOptional<float> DefaultValue, uint32 ComponentIndex) const
+{
+	if (ColorDisplayMode.Get(EColorDisplayMode::RGB) != EColorDisplayMode::RGB)
+	{
+		return 0.0f;
+	}
+
+	return ComponentSliderDynamicMinValue.IsSet() ? ComponentSliderDynamicMinValue.GetValue() : DefaultValue;
+}
+
+TOptional<float> SDisplayClusterColorGradingColorWheel::GetComponentMaxSliderValue(TOptional<float> DefaultValue, uint32 ComponentIndex) const
+{
+	if (ColorDisplayMode.Get(EColorDisplayMode::RGB) != EColorDisplayMode::RGB)
+	{
+		if (ComponentIndex == 0)
+		{
+			return 359.0f;
+		}
+		else if (ComponentIndex == 1)
+		{
+			return 1.0f;
+		}
+	}
+
+	return ComponentSliderDynamicMaxValue.IsSet() ? ComponentSliderDynamicMaxValue : DefaultValue;
+}
+
+float SDisplayClusterColorGradingColorWheel::GetComponentSliderDeltaValue(float DefaultValue, uint32 ComponentIndex) const
+{
+	if (ComponentIndex == 0 && ColorDisplayMode.Get(EColorDisplayMode::RGB) != EColorDisplayMode::RGB)
+	{
+		return 1.0f;
+	}
+
+	return DefaultValue;
+}
+
+FText SDisplayClusterColorGradingColorWheel::GetComponentToolTipText(uint32 ComponentIndex) const
+{
+	static const FText RGBToolTipText[] =
+	{
+		LOCTEXT("ColorWheel_RedComponentToolTip", "Red"),
+		LOCTEXT("ColorWheel_GreenComponentToolTip", "Green"),
+		LOCTEXT("ColorWheel_BlueComponentToolTip", "Blue")
+	};
+
+	static const FText HSVToolTipText[] =
+	{
+		LOCTEXT("ColorWheel_HueComponentToolTip", "Hue"),
+		LOCTEXT("ColorWheel_SaturationComponentToolTip", "Saturation"),
+		LOCTEXT("ColorWheel_ValueComponentToolTip", "Value")
+	};
+
+	if (ComponentIndex >= 0 && ComponentIndex < 3)
+	{
+		if (ColorDisplayMode.Get(EColorDisplayMode::RGB) != EColorDisplayMode::RGB)
+		{
+			return RGBToolTipText[ComponentIndex];
+		}
+		else
+		{
+			return HSVToolTipText[ComponentIndex];
+		}
+	}
+	else if (ComponentIndex == 3)
+	{
+		return LOCTEXT("ColorWheel_LuminanceComponentToolTip", "Luminance");
+	}
+
+	return FText::GetEmpty();
+}
+
+TArray<FLinearColor> SDisplayClusterColorGradingColorWheel::GetGradientColor(uint32 ComponentIndex)
+{
+	TArray<FLinearColor> GradientColors;
+
+	if (ColorPropertyHandle.IsValid())
+	{
+		const bool bIsRGB = ColorDisplayMode.Get(EColorDisplayMode::RGB) == EColorDisplayMode::RGB;
+		FVector4 ColorValue;
+			
+		if (ColorPropertyHandle.Pin()->GetValue(ColorValue) != FPropertyAccess::Success)
+		{
+			ColorValue = FVector4(0.0, 0.0, 0.0, 0.0);
+		}
+
+		if (bIsRGB || ComponentIndex > 0)
+		{
+			GradientColors.Add(GetGradientStartColor(ColorValue, bIsRGB, ComponentIndex));
+			GradientColors.Add(GetGradientEndColor(ColorValue, bIsRGB, ComponentIndex));
+			GradientColors.Add(GetGradientFillerColor(ColorValue, bIsRGB, ComponentIndex));
+		}
+		else // HSV Hue handling
+		{
+			for (int32 i = 0; i < 7; ++i)
+			{
+				GradientColors.Add(FLinearColor((i % 6) * 60.f, 1.f, 1.f).HSVToLinearRGB());
+			}
+		}
+	}
+
+	return GradientColors;
+}
+
+FLinearColor SDisplayClusterColorGradingColorWheel::GetGradientStartColor(const FVector4& ColorValue, bool bIsRGB, uint32 ComponentIndex) const
+{
+	if (bIsRGB)
+	{
+		switch (ComponentIndex)
+		{
+		case 0:		return FLinearColor(0.0f, ColorValue.Y, ColorValue.Z, 1.0f);
+		case 1:		return FLinearColor(ColorValue.X, 0.0f, ColorValue.Z, 1.0f);
+		case 2:		return FLinearColor(ColorValue.X, ColorValue.Y, 0.0f, 1.0f);
+		case 3:		return FLinearColor(0.0f, 0.0f, 0.0f, 1.0f);
+		default:	return FLinearColor(ForceInit);
+		}
+	}
+	else
+	{
+		FLinearColor HSVColor = FLinearColor(ColorValue.X, ColorValue.Y, ColorValue.Z, ColorValue.W).LinearRGBToHSV();
+		switch (ComponentIndex)
+		{
+		case 0:		return FLinearColor(HSVColor.R, HSVColor.G, HSVColor.B, 1.0f);
+		case 1:		return FLinearColor(HSVColor.R, 0.0f, HSVColor.B, 1.0f).HSVToLinearRGB();
+		case 2:		return FLinearColor(HSVColor.R, HSVColor.G, 0.0f, 1.0f).HSVToLinearRGB();
+		case 3:		return FLinearColor(0.0f, 0.0f, 0.0f, 1.0f);
+		default:	return FLinearColor(ForceInit);
+		}
+	}
+}
+
+FLinearColor SDisplayClusterColorGradingColorWheel::GetGradientEndColor(const FVector4& ColorValue, bool bIsRGB, uint32 ComponentIndex) const
+{
+	if (bIsRGB)
+	{
+		switch (ComponentIndex)
+		{
+		case 0:		return FLinearColor(1.0f, ColorValue.Y, ColorValue.Z, 1.0f);
+		case 1:		return FLinearColor(ColorValue.X, 1.0f, ColorValue.Z, 1.0f);
+		case 2:		return FLinearColor(ColorValue.X, ColorValue.Y, 1.0f, 1.0f);
+		case 3:		return FLinearColor(ColorValue.X, ColorValue.Y, ColorValue.Z, 1.0f);
+		default:	return FLinearColor(ForceInit);
+		}
+	}
+	else
+	{
+		FLinearColor HSVColor = FLinearColor(ColorValue.X, ColorValue.Y, ColorValue.Z, ColorValue.W).LinearRGBToHSV();
+		switch (ComponentIndex)
+		{
+		case 0:		return FLinearColor(HSVColor.R, HSVColor.G, HSVColor.B, 1.0f);
+		case 1:		return FLinearColor(HSVColor.R, 1.0f, HSVColor.B, 1.0f).HSVToLinearRGB();
+		case 2:		return FLinearColor(HSVColor.R, HSVColor.G, 1.0f, 1.0f).HSVToLinearRGB();
+		case 3:		return FLinearColor(HSVColor.R, HSVColor.G, HSVColor.B, 1.0f).HSVToLinearRGB();
+		default:	return FLinearColor(ForceInit);
+		}
+	}
+}
+
+FLinearColor SDisplayClusterColorGradingColorWheel::GetGradientFillerColor(const FVector4& ColorValue, bool bIsRGB, uint32 ComponentIndex) const
+{
+	const float MaxValue = 1.0f/*SpinBoxMinMaxSliderValues.CurrentMaxSliderValue.GetValue()*/;
+	if (bIsRGB)
+	{
+		switch (ComponentIndex)
+		{
+		case 0:		return FLinearColor(MaxValue, ColorValue.Y, ColorValue.Z, 1.0f);
+		case 1:		return FLinearColor(ColorValue.X, MaxValue, ColorValue.Z, 1.0f);
+		case 2:		return FLinearColor(ColorValue.X, ColorValue.Y, MaxValue, 1.0f);
+		case 3:		return FLinearColor(ColorValue.X, ColorValue.Y, ColorValue.Z, 1.0f);
+		default:	return FLinearColor(ForceInit);
+		}
+	}
+	else
+	{
+		FLinearColor HSVColor = FLinearColor(ColorValue.X, ColorValue.Y, ColorValue.Z, ColorValue.W).LinearRGBToHSV();
+		switch (ComponentIndex)
+		{
+		case 0:		return FLinearColor(HSVColor.R, HSVColor.G, HSVColor.B, 1.0f);
+		case 1:		return FLinearColor(HSVColor.R, 1.0f, HSVColor.B, 1.0f).HSVToLinearRGB();
+		case 2:		return FLinearColor(HSVColor.R, HSVColor.G, MaxValue, 1.0f).HSVToLinearRGB();
+		case 3:		return FLinearColor(HSVColor.R, HSVColor.G, HSVColor.B, 1.0f).HSVToLinearRGB();
+		default:	return FLinearColor(ForceInit);
+		}
+	}
+}
+
+EVisibility SDisplayClusterColorGradingColorWheel::GetGradientVisibility() const
+{
+	bool bShowGradient =
+		ColorPropertyMetadata->ColorGradingMode != EColorGradingModes::Offset &&
+		ColorPropertyMetadata->ColorGradingMode != EColorGradingModes::Invalid;
+
+	return bShowGradient ? EVisibility::Visible : EVisibility::Hidden;
+}
+
+TOptional<float> SDisplayClusterColorGradingColorWheel::GetMetadataMinValue() const
+{
+	if (ColorPropertyMetadata.IsSet())
+	{
+		return ColorPropertyMetadata->MinValue;
+	}
+
+	return TOptional<float>();
+}
+
+TOptional<float> SDisplayClusterColorGradingColorWheel::GetMetadataMaxValue() const
+{
+	if (ColorPropertyMetadata.IsSet())
+	{
+		return ColorPropertyMetadata->MaxValue;
+	}
+
+	return TOptional<float>();
+}
+
+TOptional<float> SDisplayClusterColorGradingColorWheel::GetMetadataSliderMinValue() const
+{
+	if (ColorPropertyMetadata.IsSet())
+	{
+		return ColorPropertyMetadata->SliderMinValue;
+	}
+
+	return TOptional<float>();
+}
+
+TOptional<float> SDisplayClusterColorGradingColorWheel::GetMetadataSliderMaxValue() const
+{
+	if (ColorPropertyMetadata.IsSet())
+	{
+		return ColorPropertyMetadata->SliderMaxValue;
+	}
+
+	return TOptional<float>();
+}
+
+float SDisplayClusterColorGradingColorWheel::GetMetadataDelta() const
+{
+	if (ColorPropertyMetadata.IsSet())
+	{
+		return ColorPropertyMetadata->Delta;
+	}
+
+	return 0.0f;
+}
+
+int32 SDisplayClusterColorGradingColorWheel::GetMetadataShiftMouseMovePixelPerDelta() const
+{
+	if (ColorPropertyMetadata.IsSet())
+	{
+		return ColorPropertyMetadata->ShiftMouseMovePixelPerDelta;
+	}
+
+	return 0;
+}
+
+bool SDisplayClusterColorGradingColorWheel::GetMetadataSupportDynamicSliderMinValue() const
+{
+	if (ColorPropertyMetadata.IsSet())
+	{
+		return ColorPropertyMetadata->bSupportDynamicSliderMinValue;
+	}
+
+	return false;
+}
+
+bool SDisplayClusterColorGradingColorWheel::GetMetadataSupportDynamicSliderMaxValue() const
+{
+	if (ColorPropertyMetadata.IsSet())
+	{
+		return ColorPropertyMetadata->bSupportDynamicSliderMaxValue;
+	}
+
+	return false;
+}
+
+#undef LOCTEXT_NAMESPACE
\ No newline at end of file
diff --git a/Engine/Plugins/Runtime/nDisplay/Source/DisplayClusterColorGrading/Private/SDisplayClusterColorGradingColorWheel.h b/Engine/Plugins/Runtime/nDisplay/Source/DisplayClusterColorGrading/Private/SDisplayClusterColorGradingColorWheel.h
new file mode 100644
--- /dev/null
+++ b/Engine/Plugins/Runtime/nDisplay/Source/DisplayClusterColorGrading/Private/SDisplayClusterColorGradingColorWheel.h
@@ -0,0 +1,135 @@
+// Copyright Epic Games, Inc. All Rights Reserved.
+
+#pragma once
+
+#include "CoreMinimal.h"
+#include "Customizations/MathStructCustomizations.h"
+#include "Widgets/DeclarativeSyntaxSupport.h"
+#include "Widgets/SCompoundWidget.h"
+#include "Widgets/Colors/SColorGradingPicker.h"
+
+class IPropertyHandle;
+class SBox;
+
+/** A widget which encapsulates a color picker and numeric sliders for each color component, hooked up to a color property handle */
+class SDisplayClusterColorGradingColorWheel : public SCompoundWidget
+{
+public:
+	enum class EColorDisplayMode
+	{
+		RGB,
+		HSV
+	};
+
+	struct FColorPropertyMetadata
+	{
+		EColorGradingModes ColorGradingMode = EColorGradingModes::Invalid;
+		TOptional<float> MinValue = TOptional<float>();
+		TOptional<float> MaxValue = TOptional<float>();
+		TOptional<float> SliderMinValue = TOptional<float>();
+		TOptional<float> SliderMaxValue = TOptional<float>();
+		float SliderExponent = 1.0f;
+		float Delta = 0.0f;
+		int32 LinearDeltaSensitivity = 0;
+		int32 ShiftMouseMovePixelPerDelta = 1;
+		bool bSupportDynamicSliderMaxValue = false;
+		bool bSupportDynamicSliderMinValue = false;
+	};
+
+public:
+	SLATE_BEGIN_ARGS(SDisplayClusterColorGradingColorWheel)
+		: _Orientation(Orient_Vertical)
+	{}
+		SLATE_ARGUMENT(EOrientation, Orientation)
+		SLATE_ATTRIBUTE(EColorDisplayMode, ColorDisplayMode)
+		SLATE_ARGUMENT(TSharedPtr<SWidget>, HeaderContent)
+	SLATE_END_ARGS()
+
+	void Construct(const FArguments& InArgs);
+
+	/** Sets the property handle for the color property to edit with this color wheel */
+	void SetColorPropertyHandle(TSharedPtr<IPropertyHandle> InColorPropertyHandle);
+
+	/** Sets the widget to display as the header of the color wheel */
+	void SetHeaderContent(const TSharedRef<SWidget>& HeaderContent);
+
+	/** Sets the orientation of the color wheel, which determines if the numeric sliders are below or to the right of the color picker wheel */
+	void SetOrientation(EOrientation NewOrientation);
+
+private:
+	TSharedRef<SWidget> CreateColorGradingPicker();
+	TSharedRef<SWidget> CreateColorComponentSliders();
+
+	FColorPropertyMetadata GetColorPropertyMetadata() const;
+
+	bool IsPropertyEnabled() const;
+
+	bool GetColor(FVector4& OutCurrentColor);
+	void CommitColor(FVector4& NewValue, bool bShouldCommitValueChanges);
+
+	void BeginUsingColorPickerSlider();
+	void EndUsingColorPickerSlider();
+	void BeginUsingComponentSlider(uint32 ComponentIndex);
+	void EndUsingComponentSlider(float NewValue, uint32 ComponentIndex);
+
+	FText GetComponentLabelText(uint32 ComponentIndex) const;
+	TOptional<float> GetComponentValue(uint32 ComponentIndex) const;
+	void SetComponentValue(float NewValue, uint32 ComponentIndex);
+
+	bool ComponentSupportsDynamicSliderValue(bool bDefaultValue, uint32 ComponentIndex) const;
+	void UpdateComponentDynamicSliderMinValue(float NewValue, TWeakPtr<SWidget> SourceWidget, bool bIsOriginator, bool bUpdateOnlyIfLower);
+	void UpdateComponentDynamicSliderMaxValue(float NewValue, TWeakPtr<SWidget> SourceWidget, bool bIsOriginator, bool bUpdateOnlyIfHigher);
+
+	TOptional<float> GetComponentMaxValue(TOptional<float> DefaultValue, uint32 ComponentIndex) const;
+	TOptional<float> GetComponentMinSliderValue(TOptional<float> DefaultValue, uint32 ComponentIndex) const;
+	TOptional<float> GetComponentMaxSliderValue(TOptional<float> DefaultValue, uint32 ComponentIndex) const;
+	float GetComponentSliderDeltaValue(float DefaultValue, uint32 ComponentIndex) const;
+	FText GetComponentToolTipText(uint32 ComponentIndex) const;
+
+	TArray<FLinearColor> GetGradientColor(uint32 ComponentIndex);
+
+	FLinearColor GetGradientStartColor(const FVector4& ColorValue, bool bIsRGB, uint32 ComponentIndex) const;
+	FLinearColor GetGradientEndColor(const FVector4& ColorValue, bool bIsRGB, uint32 ComponentIndex) const;
+	FLinearColor GetGradientFillerColor(const FVector4& ColorValue, bool bIsRGB, uint32 ComponentIndex) const;
+	EVisibility GetGradientVisibility() const;
+
+	TOptional<float> GetMetadataMinValue() const;
+	TOptional<float> GetMetadataMaxValue() const;
+	TOptional<float> GetMetadataSliderMinValue() const;
+	TOptional<float> GetMetadataSliderMaxValue() const;
+	float GetMetadataDelta() const;
+	int32 GetMetadataShiftMouseMovePixelPerDelta() const;
+	bool GetMetadataSupportDynamicSliderMinValue() const;
+	bool GetMetadataSupportDynamicSliderMaxValue() const;
+
+private:
+	TSharedPtr<SColorGradingPicker> ColorGradingPicker;
+	TSharedPtr<SBox> HeaderBox;
+	TSharedPtr<SBox> ColorWheelBox;
+	TSharedPtr<SBox> ColorPickerBox;
+	TSharedPtr<SBox> ColorSlidersBox;
+
+	/** The property handle of the linear color property being edited */
+	TWeakPtr<IPropertyHandle> ColorPropertyHandle;
+
+	/** The metadata of the color property */
+	TOptional<FColorPropertyMetadata> ColorPropertyMetadata;
+
+	/** Attribute for the color mode type the color wheel is presenting the color components in */
+	TAttribute<EColorDisplayMode> ColorDisplayMode;
+
+	/** The orientation of the color wheel, which determines whether the color component sliders are next to or above the color picker */
+	EOrientation Orientation;
+
+	/** Stored current min value of the color component numeric sliders */
+	TOptional<float> ComponentSliderDynamicMinValue;
+
+	/** Stored current max value of the color component numeric sliders */
+	TOptional<float> ComponentSliderDynamicMaxValue;
+
+	/** Indicates that the color picker slider is currently being used to change the color on the color picker */
+	bool bIsUsingColorPickerSlider = false;
+
+	/** Indicates that a component's numeric slider is currently being used to change the color */
+	bool bIsUsingComponentSlider = false;
+};
\ No newline at end of file
diff --git a/Engine/Plugins/Runtime/nDisplay/Source/DisplayClusterColorGrading/Private/SDisplayClusterColorGradingColorWheelPanel.cpp b/Engine/Plugins/Runtime/nDisplay/Source/DisplayClusterColorGrading/Private/SDisplayClusterColorGradingColorWheelPanel.cpp
new file mode 100644
--- /dev/null
+++ b/Engine/Plugins/Runtime/nDisplay/Source/DisplayClusterColorGrading/Private/SDisplayClusterColorGradingColorWheelPanel.cpp
@@ -0,0 +1,685 @@
+// Copyright Epic Games, Inc. All Rights Reserved.
+
+#include "SDisplayClusterColorGradingColorWheelPanel.h"
+
+#include "DisplayClusterColorGradingCommands.h"
+#include "SDisplayClusterColorGradingColorWheel.h"
+#include "Drawer/SDisplayClusterColorGradingDrawer.h"
+
+#include "DetailCategoryBuilder.h"
+#include "DetailLayoutBuilder.h"
+#include "Framework/MultiBox/MultiBoxBuilder.h"
+#include "IDetailCustomization.h"
+#include "IDetailTreeNode.h"
+#include "IPropertyRowGenerator.h"
+#include "Modules/ModuleManager.h"
+#include "PropertyEditorModule.h"
+#include "PropertyHandle.h"
+#include "Widgets/SBoxPanel.h"
+#include "Widgets/Input/SComboButton.h"
+#include "Widgets/Layout/SBox.h"
+#include "Widgets/Layout/SSpacer.h"
+#include "Widgets/Layout/SSplitter.h"
+ 
+#define LOCTEXT_NAMESPACE "DisplayClusterColorGrading"
+
+SDisplayClusterColorGradingColorWheelPanel::~SDisplayClusterColorGradingColorWheelPanel()
+{
+	if (ColorGradingDataModel)
+	{
+		ColorGradingDataModel->OnColorGradingGroupSelectionChanged().RemoveAll(this);
+		ColorGradingDataModel->OnColorGradingElementSelectionChanged().RemoveAll(this);
+	}
+}
+
+void SDisplayClusterColorGradingColorWheelPanel::Construct(const FArguments& InArgs)
+{
+	ColorGradingDataModel = InArgs._ColorGradingDataModelSource;
+
+	if (ColorGradingDataModel)
+	{
+		ColorGradingDataModel->OnColorGradingGroupSelectionChanged().AddSP(this, &SDisplayClusterColorGradingColorWheelPanel::OnColorGradingGroupSelectionChanged);
+		ColorGradingDataModel->OnColorGradingElementSelectionChanged().AddSP(this, &SDisplayClusterColorGradingColorWheelPanel::OnColorGradingElementSelectionChanged);
+	}
+
+	FPropertyEditorModule& EditModule = FModuleManager::Get().GetModuleChecked<FPropertyEditorModule>("PropertyEditor");
+	FDetailsViewArgs DetailsViewArgs;
+	DetailsViewArgs.bAllowSearch = false;
+	DetailsViewArgs.NameAreaSettings = FDetailsViewArgs::HideNameArea;
+	DetailsViewArgs.bHideSelectionTip = true;
+	DetailsView = EditModule.CreateDetailView(DetailsViewArgs);
+
+	CommandList = MakeShared<FUICommandList>();
+	BindCommands();
+
+	ColorWheelOrientation = EOrientation::Orient_Vertical;
+
+	ColorWheels.AddDefaulted(NumColorWheels);
+	HiddenColorWheels.AddZeroed(NumColorWheels);
+
+	for (int32 Index = 0; Index < NumColorWheels; ++Index)
+    ChildSlot
+	[
+		SNew(SSplitter)
+		.Orientation(Orient_Horizontal)
+		.PhysicalSplitterHandleSize(1.0f)
+		.HitDetectionSplitterHandleSize(5.0f)
+		.Style(FAppStyle::Get(), "DetailsView.Splitter")
+
+		+ SSplitter::Slot()
+		.Value(0.8f)
+		[
+			SNew(SVerticalBox)
+			.Visibility(this, &SDisplayClusterColorGradingColorWheelPanel::GetColorWheelPanelVisibility)
+
+			// Toolbar slot
+			+ SVerticalBox::Slot()
+			.AutoHeight()
+			.Padding(6, 4)
+			[
+				SNew(SHorizontalBox)
+
+				+ SHorizontalBox::Slot()
+				.AutoWidth()
+				[
+					SAssignNew(ColorGradingGroupPropertyBox, SBox)
+				]
+
+				+ SHorizontalBox::Slot()
+				.FillWidth(1.0f)
+				[
+					SNew(SSpacer)
+				]
+
+				+ SHorizontalBox::Slot()
+				.AutoWidth()
+				[
+					SAssignNew(ColorGradingElementsToolBarBox, SHorizontalBox)
+				]
+
+				+ SHorizontalBox::Slot()
+				.FillWidth(1.0f)
+				[
+					SNew(SSpacer)
+				]
+
+				+ SHorizontalBox::Slot()
+				.AutoWidth()
+				[
+					MakeColorDisplayModeCheckbox()
+				]
+
+				+ SHorizontalBox::Slot()
+				.AutoWidth()
+				.Padding(5.0f, 0.0f, 0.0f, 0.0f)
+				.HAlign(HAlign_Right)
+				.VAlign(VAlign_Center)
+				[
+					SNew(SComboButton)
+					.ComboButtonStyle(&FAppStyle::Get().GetWidgetStyle<FComboButtonStyle>("SimpleComboButton"))
+					.OnGetMenuContent(this, &SDisplayClusterColorGradingColorWheelPanel::MakeSettingsMenu)
+					.HasDownArrow(false)
+					.ContentPadding(FMargin(1.0f, 0.0f))
+					.ButtonContent()
+					[
+						SNew(SImage)
+						.ColorAndOpacity(FSlateColor::UseForeground())
+						.Image(FAppStyle::Get().GetBrush("Icons.Settings"))
+					]
+				]
+			]
+
+			+ SVerticalBox::Slot()
+			.FillHeight(1.0f)
+			.Padding(6, 4)
+			[
+				SNew(SHorizontalBox)
+
+				+ SHorizontalBox::Slot()
+				.FillWidth(1.0f)
+				.HAlign(HAlign_Fill)
+				.Padding(2, 0)
+				[
+					SAssignNew(ColorWheels[0], SDisplayClusterColorGradingColorWheel)
+					.ColorDisplayMode(this, &SDisplayClusterColorGradingColorWheelPanel::GetColorDisplayMode)
+					.Orientation(ColorWheelOrientation)
+					.Visibility(this, &SDisplayClusterColorGradingColorWheelPanel::GetColorWheelVisibility, 0)
+				]
+
+				+ SHorizontalBox::Slot()
+				.FillWidth(1.0f)
+				.HAlign(HAlign_Fill)
+				.Padding(2, 0)
+				[
+					SAssignNew(ColorWheels[1], SDisplayClusterColorGradingColorWheel)
+					.ColorDisplayMode(this, &SDisplayClusterColorGradingColorWheelPanel::GetColorDisplayMode)
+					.Orientation(ColorWheelOrientation)
+					.Visibility(this, &SDisplayClusterColorGradingColorWheelPanel::GetColorWheelVisibility, 1)
+				]
+
+				+ SHorizontalBox::Slot()
+				.FillWidth(1.0f)
+				.HAlign(HAlign_Fill)
+				.Padding(2, 0)
+				[
+					SAssignNew(ColorWheels[2], SDisplayClusterColorGradingColorWheel)
+					.ColorDisplayMode(this, &SDisplayClusterColorGradingColorWheelPanel::GetColorDisplayMode)
+					.Orientation(ColorWheelOrientation)
+					.Visibility(this, &SDisplayClusterColorGradingColorWheelPanel::GetColorWheelVisibility, 2)
+				]
+
+				+ SHorizontalBox::Slot()
+				.FillWidth(1.0f)
+				.HAlign(HAlign_Fill)
+				.Padding(2, 0)
+				[
+					SAssignNew(ColorWheels[3], SDisplayClusterColorGradingColorWheel)
+					.ColorDisplayMode(this, &SDisplayClusterColorGradingColorWheelPanel::GetColorDisplayMode)
+					.Orientation(ColorWheelOrientation)
+					.Visibility(this, &SDisplayClusterColorGradingColorWheelPanel::GetColorWheelVisibility, 3)
+				]
+
+				+ SHorizontalBox::Slot()
+				.FillWidth(1.0f)
+				.HAlign(HAlign_Fill)
+				.Padding(2, 0)
+				[
+					SAssignNew(ColorWheels[4], SDisplayClusterColorGradingColorWheel)
+					.ColorDisplayMode(this, &SDisplayClusterColorGradingColorWheelPanel::GetColorDisplayMode)
+					.Orientation(ColorWheelOrientation)
+					.Visibility(this, &SDisplayClusterColorGradingColorWheelPanel::GetColorWheelVisibility, 4)
+				]
+			]
+		]
+
+		+ SSplitter::Slot()
+		.Value(0.2f)
+		[
+			DetailsView.ToSharedRef()
+		]
+	];
+}
+
+void SDisplayClusterColorGradingColorWheelPanel::Refresh()
+{
+	if (ColorGradingDataModel)
+	{
+		if (FDisplayClusterColorGradingDataModel::FColorGradingGroup* ColorGradingGroup = ColorGradingDataModel->GetSelectedColorGradingGroup())
+		{
+			FillColorGradingGroupProperty(*ColorGradingGroup);
+			FillColorGradingElementsToolBar(ColorGradingGroup->ColorGradingElements);
+			ColorGradingDataModel->SetSelectedColorGradingElement(ColorGradingGroup->ColorGradingElements.Num() ? 0 : INDEX_NONE);
+
+			FillDetailsView(*ColorGradingGroup);
+		}
+		else
+		{
+			ClearColorGradingGroupProperty();
+			ClearColorGradingElementsToolBar();
+			ClearDetailsView();
+			ClearColorWheels();
+		}
+	}
+}
+
+void SDisplayClusterColorGradingColorWheelPanel::GetDrawerState(FDisplayClusterColorGradingDrawerState& OutDrawerState)
+{
+	OutDrawerState.HiddenColorWheels = HiddenColorWheels;
+	OutDrawerState.ColorDisplayMode = ColorDisplayMode;
+	OutDrawerState.ColorWheelOrientation = ColorWheelOrientation;
+}
+
+void SDisplayClusterColorGradingColorWheelPanel::SetDrawerState(const FDisplayClusterColorGradingDrawerState& InDrawerState)
+{
+	// TODO: These could also be output to a config file to be stored between runs
+	HiddenColorWheels = InDrawerState.HiddenColorWheels;
+	ColorDisplayMode = InDrawerState.ColorDisplayMode;
+	SetColorWheelOrientation(InDrawerState.ColorWheelOrientation);
+}
+
+void SDisplayClusterColorGradingColorWheelPanel::BindCommands()
+{
+	const FDisplayClusterColorGradingCommands& Commands = FDisplayClusterColorGradingCommands::Get();
+
+	CommandList->MapAction(
+		Commands.SaturationColorWheelVisibility,
+		FExecuteAction::CreateSP(this, &SDisplayClusterColorGradingColorWheelPanel::ToggleColorWheelVisible, 3),
+		FCanExecuteAction(),
+		FIsActionChecked::CreateSP(this, &SDisplayClusterColorGradingColorWheelPanel::IsColorWheelVisibile, 3)
+	);
+
+	CommandList->MapAction(
+		Commands.ContrastColorWheelVisibility,
+		FExecuteAction::CreateSP(this, &SDisplayClusterColorGradingColorWheelPanel::ToggleColorWheelVisible, 4),
+		FCanExecuteAction(),
+		FIsActionChecked::CreateSP(this, &SDisplayClusterColorGradingColorWheelPanel::IsColorWheelVisibile, 4)
+	);
+
+	CommandList->MapAction(
+		Commands.ColorWheelSliderOrientationHorizontal,
+		FExecuteAction::CreateSP(this, &SDisplayClusterColorGradingColorWheelPanel::SetColorWheelOrientation, EOrientation::Orient_Horizontal),
+		FCanExecuteAction(),
+		FIsActionChecked::CreateSP(this, &SDisplayClusterColorGradingColorWheelPanel::IsColorWheelOrientationSelected, EOrientation::Orient_Horizontal)
+	);
+
+	CommandList->MapAction(
+		Commands.ColorWheelSliderOrientationVertical,
+		FExecuteAction::CreateSP(this, &SDisplayClusterColorGradingColorWheelPanel::SetColorWheelOrientation, EOrientation::Orient_Vertical),
+		FCanExecuteAction(),
+		FIsActionChecked::CreateSP(this, &SDisplayClusterColorGradingColorWheelPanel::IsColorWheelOrientationSelected, EOrientation::Orient_Vertical)
+	);
+}
+
+TSharedRef<SWidget> SDisplayClusterColorGradingColorWheelPanel::MakeColorDisplayModeCheckbox()
+{
+	return SNew(SHorizontalBox)
+
+		+ SHorizontalBox::Slot()
+		.AutoWidth()
+		.HAlign(HAlign_Left)
+		.VAlign(VAlign_Top)
+		.Padding(FMargin(0.0f, 0.0f, 3.0f, 0.0f))
+		[
+			SNew(SCheckBox)
+			.Style(FAppStyle::Get(), "ToggleButtonCheckbox")
+			.Type(ESlateCheckBoxType::ToggleButton)
+			.IsChecked(this, &SDisplayClusterColorGradingColorWheelPanel::IsColorDisplayModeChecked, EColorDisplayMode::RGB)
+			.OnCheckStateChanged(this, &SDisplayClusterColorGradingColorWheelPanel::OnColorDisplayModeCheckedChanged, EColorDisplayMode::RGB)
+			.ToolTipText(this, &SDisplayClusterColorGradingColorWheelPanel::GetColorDisplayModeToolTip, EColorDisplayMode::RGB)
+			.Padding(4)
+			.Content()
+			[
+				SNew(STextBlock)
+				.Text(this, &SDisplayClusterColorGradingColorWheelPanel::GetColorDisplayModeLabel, EColorDisplayMode::RGB)
+				.Font(FAppStyle::GetFontStyle("StandardDialog.SmallFont"))
+			]
+		]
+		+ SHorizontalBox::Slot()
+		.AutoWidth()
+		.HAlign(HAlign_Left)			
+		.VAlign(VAlign_Top)
+		.Padding(FMargin(0.0f, 0.0f, 3.0f, 0.0f))
+		[
+			SNew(SCheckBox)
+			.Style(FAppStyle::Get(), "ToggleButtonCheckbox")
+			.Type(ESlateCheckBoxType::ToggleButton)
+			.IsChecked(this, &SDisplayClusterColorGradingColorWheelPanel::IsColorDisplayModeChecked, EColorDisplayMode::HSV)
+			.OnCheckStateChanged(this, &SDisplayClusterColorGradingColorWheelPanel::OnColorDisplayModeCheckedChanged, EColorDisplayMode::HSV)
+			.ToolTipText(this, &SDisplayClusterColorGradingColorWheelPanel::GetColorDisplayModeToolTip, EColorDisplayMode::HSV)
+			.Padding(4)
+			.Content()
+			[
+				SNew(STextBlock)
+				.Text(this, &SDisplayClusterColorGradingColorWheelPanel::GetColorDisplayModeLabel, EColorDisplayMode::HSV)
+				.Font(FAppStyle::GetFontStyle("StandardDialog.SmallFont"))
+			]
+		];
+}
+
+TSharedRef<SWidget> SDisplayClusterColorGradingColorWheelPanel::MakeSettingsMenu()
+{
+	const bool bShouldCloseWindowAfterMenuSelection = true;
+	FMenuBuilder MenuBuilder(bShouldCloseWindowAfterMenuSelection, CommandList);
+
+	const FDisplayClusterColorGradingCommands& Commands = FDisplayClusterColorGradingCommands::Get();
+
+	MenuBuilder.BeginSection(TEXT("ColorWheelVisibility"), LOCTEXT("ColorWheelPanel_ShowLabel", "Show"));
+	{
+		MenuBuilder.AddMenuEntry(Commands.SaturationColorWheelVisibility);
+		MenuBuilder.AddMenuEntry(Commands.ContrastColorWheelVisibility);
+	}
+	MenuBuilder.EndSection();
+
+	MenuBuilder.BeginSection(TEXT("ColorWheelSliders"), LOCTEXT("ColorWheelPanel_SlidersLabel", "Sliders"));
+	{
+		MenuBuilder.AddMenuEntry(Commands.ColorWheelSliderOrientationVertical);
+		MenuBuilder.AddMenuEntry(Commands.ColorWheelSliderOrientationHorizontal);
+	}
+	MenuBuilder.EndSection();
+
+	return MenuBuilder.MakeWidget();
+}
+
+void SDisplayClusterColorGradingColorWheelPanel::FillColorGradingGroupProperty(const FDisplayClusterColorGradingDataModel::FColorGradingGroup& ColorGradingGroup)
+{
+	if (ColorGradingGroupPropertyBox.IsValid())
+	{
+		TSharedRef<SHorizontalBox> PropertyNameBox = SNew(SHorizontalBox);
+
+		if (ColorGradingGroup.EditConditionPropertyHandle.IsValid())
+		{
+			if (TSharedPtr<IDetailTreeNode> EditConditionTreeNode = ColorGradingDataModel->GetPropertyRowGenerator()->FindTreeNode(ColorGradingGroup.EditConditionPropertyHandle))
+			{
+				FNodeWidgets EditConditionWidgets = EditConditionTreeNode->CreateNodeWidgets();
+
+				if (EditConditionWidgets.ValueWidget.IsValid())
+				{
+					PropertyNameBox->AddSlot()
+						.HAlign(HAlign_Left)
+						.VAlign(VAlign_Center)
+						.Padding(2, 0, 4, 0)
+						.AutoWidth()
+						[
+							EditConditionWidgets.ValueWidget.ToSharedRef()
+						];
+				}
+			}
+		}
+
+		TSharedRef<SWidget> GroupHeaderWidget = ColorGradingGroup.GroupHeaderWidget.IsValid()
+			? ColorGradingGroup.GroupHeaderWidget.ToSharedRef()
+			: SNew(STextBlock).Text(ColorGradingGroup.DisplayName).Font(FAppStyle::Get().GetFontStyle("NormalFontBold"));
+
+		PropertyNameBox->AddSlot()
+			.HAlign(HAlign_Left)
+			.VAlign(VAlign_Center)
+			.Padding(2, 0, 0, 0)
+			[
+				GroupHeaderWidget
+			];
+
+		ColorGradingGroupPropertyBox->SetContent(PropertyNameBox);
+	}
+}
+
+void SDisplayClusterColorGradingColorWheelPanel::ClearColorGradingGroupProperty()
+{
+	ColorGradingGroupPropertyBox->SetContent(SNullWidget::NullWidget);
+}
+
+void SDisplayClusterColorGradingColorWheelPanel::FillColorGradingElementsToolBar(const TArray<FDisplayClusterColorGradingDataModel::FColorGradingElement>& ColorGradingElements)
+{
+	ColorGradingElementsToolBarBox->ClearChildren();
+
+	for (const FDisplayClusterColorGradingDataModel::FColorGradingElement& Element : ColorGradingElements)
+	{
+		ColorGradingElementsToolBarBox->AddSlot()
+			.AutoWidth()
+			.Padding(2.0f, 0.0f, 2.0f, 0.0f)
+			[
+				SNew(SCheckBox)
+				.Style(FAppStyle::Get(), "DetailsView.SectionButton")
+				.OnCheckStateChanged(this, &SDisplayClusterColorGradingColorWheelPanel::OnColorGradingElementCheckedChanged, Element.DisplayName)
+				.IsChecked(this, &SDisplayClusterColorGradingColorWheelPanel::IsColorGradingElementSelected, Element.DisplayName)
+				[
+					SNew(STextBlock)
+					.TextStyle(FAppStyle::Get(), "SmallText")
+					.Text(Element.DisplayName)
+				]
+			];
+	}
+}
+
+void SDisplayClusterColorGradingColorWheelPanel::ClearColorGradingElementsToolBar()
+{
+	ColorGradingElementsToolBarBox->ClearChildren();
+}
+
+class FColorGradingDetailsCustomization : public IDetailCustomization
+{
+public:
+	FColorGradingDetailsCustomization(TArray<TSharedPtr<IPropertyHandle>> InPropertyHandles)
+		: PropertyHandles(InPropertyHandles)
+	{ }
+
+	virtual void CustomizeDetails(IDetailLayoutBuilder& DetailBuilder) override
+	{
+		TArray<FName> Categories;
+		DetailBuilder.GetCategoryNames(Categories);
+
+		for (const FName& Category : Categories)
+		{
+			DetailBuilder.HideCategory(Category);
+		}
+
+		// TransformCommon is a custom category that doesn't get returned by GetCategoryNames that also needs to be hidden
+		DetailBuilder.HideCategory(TEXT("TransformCommon"));
+
+		for (const TSharedPtr<IPropertyHandle>& PropertyHandle : PropertyHandles)
+		{
+			FName CategoryOverride = PropertyHandle->GetDefaultCategoryName();
+			FText CategoryOverrideText = PropertyHandle->GetDefaultCategoryText();
+			if (const FString* CategoryOverridePtr = PropertyHandle->GetInstanceMetaData(TEXT("CategoryOverride")))
+			{
+				CategoryOverride = FName(*CategoryOverridePtr + "_Override");
+				CategoryOverrideText = FText::FromString(*CategoryOverridePtr);
+			}
+
+			IDetailCategoryBuilder& CategoryBuilder = DetailBuilder.EditCategory(CategoryOverride, CategoryOverrideText);
+			CategoryBuilder.AddProperty(PropertyHandle);
+		}
+	}
+
+private:
+	TArray<TSharedPtr<IPropertyHandle>> PropertyHandles;
+};
+
+void SDisplayClusterColorGradingColorWheelPanel::FillDetailsView(const FDisplayClusterColorGradingDataModel::FColorGradingGroup& ColorGradingGroup)
+{
+	if (DetailsView.IsValid() && ColorGradingGroup.DetailsViewPropertyHandles.Num())
+	{
+		TArray<UObject*> OuterObjects;
+
+		if (ColorGradingGroup.GroupPropertyHandle.IsValid())
+		{
+			ColorGradingGroup.GroupPropertyHandle->GetOuterObjects(OuterObjects);
+		}
+		else if (ColorGradingGroup.DetailsViewPropertyHandles[0].IsValid())
+		{
+			ColorGradingGroup.DetailsViewPropertyHandles[0]->GetOuterObjects(OuterObjects);
+		}
+
+		DetailsView->RegisterInstancedCustomPropertyLayout(OuterObjects[0]->GetClass(), FOnGetDetailCustomizationInstance::CreateLambda([&ColorGradingGroup]()
+		{
+			return MakeShareable(new FColorGradingDetailsCustomization(ColorGradingGroup.DetailsViewPropertyHandles));
+		}));
+
+		DetailsView->SetObjects(OuterObjects);
+	}
+}
+
+void SDisplayClusterColorGradingColorWheelPanel::ClearDetailsView()
+{
+	TArray<UObject*> EmptyList;
+	DetailsView->SetObjects(EmptyList);
+}
+
+void SDisplayClusterColorGradingColorWheelPanel::FillColorWheels(const FDisplayClusterColorGradingDataModel::FColorGradingElement& ColorGradingElement)
+{
+	auto FillColorWheel = [this](const TSharedPtr<SDisplayClusterColorGradingColorWheel>& ColorWheel, const TSharedPtr<IPropertyHandle>& PropertyHandle)
+	{
+		if (ColorWheel.IsValid())
+		{
+			ColorWheel->SetColorPropertyHandle(PropertyHandle);
+
+			if (TSharedPtr<IDetailTreeNode> TreeNode = ColorGradingDataModel->GetPropertyRowGenerator()->FindTreeNode(PropertyHandle))
+			{
+				FNodeWidgets NodeWidgets = TreeNode->CreateNodeWidgets();
+
+				TSharedRef<SHorizontalBox> PropertyNameBox = SNew(SHorizontalBox);
+
+				if (NodeWidgets.EditConditionWidget.IsValid())
+				{
+					PropertyNameBox->AddSlot()
+						.VAlign(VAlign_Center)
+						.HAlign(HAlign_Left)
+						.Padding(2, 0, 0, 0)
+						.AutoWidth()
+						[
+							NodeWidgets.EditConditionWidget.ToSharedRef()
+						];
+				}
+
+				if (NodeWidgets.NameWidget.IsValid())
+				{
+					PropertyNameBox->AddSlot()
+						.HAlign(NodeWidgets.NameWidgetLayoutData.HorizontalAlignment)
+						.VAlign(NodeWidgets.NameWidgetLayoutData.VerticalAlignment)
+						.Padding(2, 0, 0, 0)
+						[
+							NodeWidgets.NameWidget.ToSharedRef()
+						];
+				}
+
+				ColorWheel->SetHeaderContent(PropertyNameBox);
+			}
+		}
+	};
+
+	FillColorWheel(ColorWheels[0], ColorGradingElement.OffsetPropertyHandle);
+	FillColorWheel(ColorWheels[1], ColorGradingElement.GammaPropertyHandle);
+	FillColorWheel(ColorWheels[2], ColorGradingElement.GainPropertyHandle);
+	FillColorWheel(ColorWheels[3], ColorGradingElement.SaturationPropertyHandle);
+	FillColorWheel(ColorWheels[4], ColorGradingElement.ContrastPropertyHandle);
+}
+
+void SDisplayClusterColorGradingColorWheelPanel::ClearColorWheels()
+{
+	for (const TSharedPtr<SDisplayClusterColorGradingColorWheel>& ColorWheel : ColorWheels)
+	{
+		if (ColorWheel.IsValid())
+		{
+			ColorWheel->SetColorPropertyHandle(nullptr);
+			ColorWheel->SetHeaderContent(SNullWidget::NullWidget);
+		}
+	};
+}
+
+void SDisplayClusterColorGradingColorWheelPanel::SetColorWheelOrientation(EOrientation NewOrientation)
+{
+	ColorWheelOrientation = NewOrientation;
+
+	for (const TSharedPtr<SDisplayClusterColorGradingColorWheel>& ColorWheel : ColorWheels)
+	{
+		if (ColorWheel.IsValid())
+		{
+			ColorWheel->SetOrientation(ColorWheelOrientation);
+		}
+	};
+}
+
+bool SDisplayClusterColorGradingColorWheelPanel::IsColorWheelOrientationSelected(EOrientation Orientation) const
+{
+	return ColorWheelOrientation == Orientation;
+}
+
+void SDisplayClusterColorGradingColorWheelPanel::ToggleColorWheelVisible(int32 ColorWheelIndex)
+{
+	if (ColorWheelIndex >= 0 && ColorWheelIndex < HiddenColorWheels.Num())
+	{
+		HiddenColorWheels[ColorWheelIndex] = !HiddenColorWheels[ColorWheelIndex];
+	}
+}
+
+bool SDisplayClusterColorGradingColorWheelPanel::IsColorWheelVisibile(int32 ColorWheelIndex)
+{
+	if (ColorWheelIndex >= 0 && ColorWheelIndex < HiddenColorWheels.Num())
+	{
+		return !HiddenColorWheels[ColorWheelIndex];
+	}
+
+	return false;
+}
+
+void SDisplayClusterColorGradingColorWheelPanel::OnColorGradingGroupSelectionChanged()
+{
+	Refresh();
+}
+
+void SDisplayClusterColorGradingColorWheelPanel::OnColorGradingElementSelectionChanged()
+{
+	if (const FDisplayClusterColorGradingDataModel::FColorGradingElement* ColorGradingElement = ColorGradingDataModel->GetSelectedColorGradingElement())
+	{
+		FillColorWheels(*ColorGradingElement);
+	}
+	else
+	{
+		ClearColorWheels();
+	}
+}
+
+void SDisplayClusterColorGradingColorWheelPanel::OnColorGradingElementCheckedChanged(ECheckBoxState State, FText ElementName)
+{
+	if (State == ECheckBoxState::Checked && ColorGradingDataModel)
+	{
+		if (FDisplayClusterColorGradingDataModel::FColorGradingGroup* ColorGradingGroup = ColorGradingDataModel->GetSelectedColorGradingGroup())
+		{
+			int32 ColorGradingElementIndex = ColorGradingGroup->ColorGradingElements.IndexOfByPredicate([=](const FDisplayClusterColorGradingDataModel::FColorGradingElement& Element)
+			{
+				return Element.DisplayName.CompareTo(ElementName) == 0;
+			});
+
+			ColorGradingDataModel->SetSelectedColorGradingElement(ColorGradingElementIndex);
+		}
+	}
+}
+
+ECheckBoxState SDisplayClusterColorGradingColorWheelPanel::IsColorGradingElementSelected(FText ElementName) const
+{
+	if (ColorGradingDataModel)
+	{
+		if (const FDisplayClusterColorGradingDataModel::FColorGradingElement* ColorGradingElement = ColorGradingDataModel->GetSelectedColorGradingElement())
+		{
+			if (ColorGradingElement->DisplayName.CompareTo(ElementName) == 0)
+			{
+				return ECheckBoxState::Checked;
+			}
+		}
+	}
+
+	return ECheckBoxState::Unchecked;
+}
+
+EVisibility SDisplayClusterColorGradingColorWheelPanel::GetColorWheelPanelVisibility() const
+{
+	bool bHasObjects = ColorGradingDataModel && ColorGradingDataModel->GetPropertyRowGenerator()->GetSelectedObjects().Num() > 0;
+	return bHasObjects ? EVisibility::Visible : EVisibility::Collapsed;
+}
+
+EVisibility SDisplayClusterColorGradingColorWheelPanel::GetColorWheelVisibility(int32 ColorWheelIndex) const
+{
+	bool bIsHidden = HiddenColorWheels[ColorWheelIndex];
+	return bIsHidden ? EVisibility::Collapsed : EVisibility::Visible;
+}
+
+ECheckBoxState SDisplayClusterColorGradingColorWheelPanel::IsColorDisplayModeChecked(EColorDisplayMode InColorDisplayMode) const
+{
+	bool bIsModeSelected = InColorDisplayMode == ColorDisplayMode;
+	return bIsModeSelected ? ECheckBoxState::Checked : ECheckBoxState::Unchecked;
+}
+
+void SDisplayClusterColorGradingColorWheelPanel::OnColorDisplayModeCheckedChanged(ECheckBoxState State, EColorDisplayMode InColorDisplayMode)
+{
+	if (State == ECheckBoxState::Checked)
+	{
+		ColorDisplayMode = InColorDisplayMode;
+	}
+}
+
+FText SDisplayClusterColorGradingColorWheelPanel::GetColorDisplayModeLabel(EColorDisplayMode InColorDisplayMode) const
+{
+	FText Text;
+
+	switch (InColorDisplayMode)
+	{
+	case EColorDisplayMode::RGB: Text = LOCTEXT("ColorWheel_RGBColorDisplayModeLabel", "RGB"); break;
+	case EColorDisplayMode::HSV: Text = LOCTEXT("ColorWheel_HSVColorDisplayModeLabel", "HSV"); break;
+	}
+
+	return Text;
+}
+
+FText SDisplayClusterColorGradingColorWheelPanel::GetColorDisplayModeToolTip(EColorDisplayMode InColorDisplayMode) const
+{
+	FText Text;
+
+	switch (InColorDisplayMode)
+	{
+	case EColorDisplayMode::RGB: Text = LOCTEXT("ColorWheel_RGBColorDisplayModeToolTip", "Change to RGB color mode"); break;
+	case EColorDisplayMode::HSV: Text = LOCTEXT("ColorWheel_HSVColorDisplayModeToolTip", "Change to HSV color mode"); break;
+	}
+
+	return Text;
+}
+
+#undef LOCTEXT_NAMESPACE
\ No newline at end of file
diff --git a/Engine/Plugins/Runtime/nDisplay/Source/DisplayClusterColorGrading/Private/SDisplayClusterColorGradingColorWheelPanel.h b/Engine/Plugins/Runtime/nDisplay/Source/DisplayClusterColorGrading/Private/SDisplayClusterColorGradingColorWheelPanel.h
new file mode 100644
--- /dev/null
+++ b/Engine/Plugins/Runtime/nDisplay/Source/DisplayClusterColorGrading/Private/SDisplayClusterColorGradingColorWheelPanel.h
@@ -0,0 +1,109 @@
+// Copyright Epic Games, Inc. All Rights Reserved.
+
+#pragma once
+
+#include "CoreMinimal.h"
+#include "DisplayClusterColorGradingDataModel.h"
+#include "SDisplayClusterColorGradingColorWheel.h"
+#include "Styling/SlateTypes.h"
+#include "Widgets/DeclarativeSyntaxSupport.h"
+#include "Widgets/SCompoundWidget.h"
+
+class IDetailTreeNode;
+class IDetaisView;
+class IPropertyHandle;
+class IPropertyRowGenerator;
+class SBox;
+class SHorizontalBox;
+struct FDisplayClusterColorGradingDrawerState;
+
+/** A panel that contains up to five color wheels (for saturation, contrast, gamma, gain, and offset) as well as a details view for extra, non-color properties */
+class SDisplayClusterColorGradingColorWheelPanel : public SCompoundWidget
+{
+private:
+	typedef SDisplayClusterColorGradingColorWheel::EColorDisplayMode EColorDisplayMode;
+
+	/** The number of color wheels the color wheel panel displays (for saturation, contrast, gamma, gain, and offset) */
+	static const uint32 NumColorWheels = 5;
+
+public:
+	~SDisplayClusterColorGradingColorWheelPanel();
+
+	SLATE_BEGIN_ARGS(SDisplayClusterColorGradingColorWheelPanel) {}
+		SLATE_ARGUMENT(TSharedPtr<FDisplayClusterColorGradingDataModel>, ColorGradingDataModelSource)
+	SLATE_END_ARGS()
+
+	void Construct(const FArguments& InArgs);
+
+	/** Regenerates the color wheel panel from the current state of the data model source */
+	void Refresh();
+
+	/** Adds the state of the color wheel panel to the specified drawer state */
+	void GetDrawerState(FDisplayClusterColorGradingDrawerState& OutDrawerState);
+
+	/** Sets the state of the color wheel panel from the specified drawer state */
+	void SetDrawerState(const FDisplayClusterColorGradingDrawerState& InDrawerState);
+
+private:
+	void BindCommands();
+
+	TSharedRef<SWidget> MakeColorDisplayModeCheckbox();
+	TSharedRef<SWidget> MakeSettingsMenu();
+
+	void FillColorGradingGroupProperty(const FDisplayClusterColorGradingDataModel::FColorGradingGroup& ColorGradingGroup);
+	void ClearColorGradingGroupProperty();
+
+	void FillColorGradingElementsToolBar(const TArray<FDisplayClusterColorGradingDataModel::FColorGradingElement>& ColorGradingElements);
+	void ClearColorGradingElementsToolBar();
+
+	void FillDetailsView(const FDisplayClusterColorGradingDataModel::FColorGradingGroup& ColorGradingGroup);
+	void ClearDetailsView();
+
+	void FillColorWheels(const FDisplayClusterColorGradingDataModel::FColorGradingElement& ColorGradingElement);
+	void ClearColorWheels();
+
+	void SetColorWheelOrientation(EOrientation NewOrientation);
+	bool IsColorWheelOrientationSelected(EOrientation Orientation) const;
+
+	void ToggleColorWheelVisible(int32 ColorWheelIndex);
+	bool IsColorWheelVisibile(int32 ColorWheelIndex);
+
+	void OnColorGradingGroupSelectionChanged();
+	void OnColorGradingElementSelectionChanged();
+
+	void OnColorGradingElementCheckedChanged(ECheckBoxState State, FText ElementName);
+
+	ECheckBoxState IsColorGradingElementSelected(FText ElementName) const;
+	EVisibility GetColorWheelPanelVisibility() const;
+	EVisibility GetColorWheelVisibility(int32 ColorWheelIndex) const;
+
+	EColorDisplayMode GetColorDisplayMode() const { return ColorDisplayMode; }
+	ECheckBoxState IsColorDisplayModeChecked(EColorDisplayMode InColorDisplayMode) const;
+	void OnColorDisplayModeCheckedChanged(ECheckBoxState State, EColorDisplayMode InColorDisplayMode);
+	FText GetColorDisplayModeLabel(EColorDisplayMode InColorDisplayMode) const;
+	FText GetColorDisplayModeToolTip(EColorDisplayMode InColorDisplayMode) const;
+
+private:
+	/** The color grading data model that the panel is displaying */
+	TSharedPtr<FDisplayClusterColorGradingDataModel> ColorGradingDataModel;
+
+	TSharedPtr<SBox> ColorGradingGroupPropertyBox;
+	TSharedPtr<SHorizontalBox> ColorGradingElementsToolBarBox;
+
+	TArray<TSharedPtr<SDisplayClusterColorGradingColorWheel>> ColorWheels;
+	TArray<bool> HiddenColorWheels;
+
+	TSharedPtr<IDetailsView> DetailsView;
+
+	/** Commands used by the color wheel panel */
+	TSharedPtr<FUICommandList> CommandList;
+
+	/** The currently selected color grading group */
+	int32 SelectedColorGradingGroup = INDEX_NONE;
+
+	/** The current color display mode for the color wheels */
+	EColorDisplayMode ColorDisplayMode;
+
+	/** The current orientation for the color wheeels */
+	EOrientation ColorWheelOrientation;
+};
\ No newline at end of file
diff --git a/Engine/Plugins/Runtime/nDisplay/Source/DisplayClusterColorGrading/Private/SDisplayClusterColorGradingObjectList.cpp b/Engine/Plugins/Runtime/nDisplay/Source/DisplayClusterColorGrading/Private/SDisplayClusterColorGradingObjectList.cpp
new file mode 100644
--- /dev/null
+++ b/Engine/Plugins/Runtime/nDisplay/Source/DisplayClusterColorGrading/Private/SDisplayClusterColorGradingObjectList.cpp
@@ -0,0 +1,214 @@
+// Copyright Epic Games, Inc. All Rights Reserved.
+
+#include "SDisplayClusterColorGradingObjectList.h"
+
+#include "ClassIconFinder.h"
+#include "Components/ActorComponent.h"
+#include "GameFramework/Actor.h"
+#include "Styling/AppStyle.h"
+#include "Widgets/Images/SImage.h"
+#include "Widgets/Input/SCheckBox.h"
+#include "Widgets/Views/SListView.h"
+#include "Widgets/Views/STableRow.h"
+
+#define LOCTEXT_NAMESPACE "DisplayClusterColorGrading"
+
+namespace DisplayClusterColorGradingObjectListColumnNames
+{
+	const static FName ItemEnabled(TEXT("ItemEnabled"));
+	const static FName ItemLabel(TEXT("ItemLabel"));
+};
+
+class SColorGradingListItemRow : public SMultiColumnTableRow<FDisplayClusterColorGradingListItemRef>
+{
+public:
+	SLATE_BEGIN_ARGS(SColorGradingListItemRow) {}
+	SLATE_END_ARGS()
+
+	void Construct(const FArguments& InArgs, const TSharedRef<STableViewBase>& InOwnerTable, const FDisplayClusterColorGradingListItemRef InListItem)
+	{
+		ListItem = InListItem;
+
+		SMultiColumnTableRow<FDisplayClusterColorGradingListItemRef>::Construct(
+			SMultiColumnTableRow<FDisplayClusterColorGradingListItemRef>::FArguments()
+			.Style(&FAppStyle::Get().GetWidgetStyle<FTableRowStyle>("SceneOutliner.TableViewRow"))
+			, InOwnerTable);
+	}
+
+	virtual TSharedRef<SWidget> GenerateWidgetForColumn(const FName& ColumnName) override
+	{
+		if (ColumnName == DisplayClusterColorGradingObjectListColumnNames::ItemEnabled)
+		{
+			return SNew(SCheckBox)
+				.IsChecked(this, &SColorGradingListItemRow::GetCheckBoxState)
+				.OnCheckStateChanged(this, &SColorGradingListItemRow::OnCheckBoxStateChanged);
+		}
+		else if (ColumnName == DisplayClusterColorGradingObjectListColumnNames::ItemLabel)
+		{
+			return SNew(SHorizontalBox)
+
+				+ SHorizontalBox::Slot()
+				.AutoWidth()
+				.Padding(FMargin(0, 1, 6, 1))
+				.VAlign(VAlign_Center)
+				[
+					SNew(SBox)
+					.WidthOverride(16)
+					.HeightOverride(16)
+					[
+						SNew(SImage)
+						.ColorAndOpacity(FSlateColor::UseForeground())
+						.Image(this, &SColorGradingListItemRow::GetItemIcon)
+					]
+				]
+
+				+ SHorizontalBox::Slot()
+				.AutoWidth()
+				.VAlign(VAlign_Center)
+				.HAlign(HAlign_Left)
+				[
+					SNew(STextBlock)
+					.Text(this, &SColorGradingListItemRow::GetItemLabel)
+				];
+		}
+
+		return SNullWidget::NullWidget;
+	}
+
+private:
+	FText GetItemLabel() const
+	{
+		FString ItemName = TEXT("");
+
+		if (ListItem.IsValid())
+		{
+			if (ListItem->Component.IsValid())
+			{
+				ItemName = ListItem->Component->GetName();
+			}
+			else if (ListItem->Actor.IsValid())
+			{
+				ItemName = ListItem->Actor->GetActorLabel();
+			}
+		}
+
+		return FText::FromString(*ItemName);
+	}
+
+	const FSlateBrush* GetItemIcon() const
+	{
+		if (ListItem.IsValid())
+		{
+			if (ListItem->Component.IsValid())
+			{
+				return FSlateIconFinder::FindIconBrushForClass(ListItem->Component->GetClass(), TEXT("SCS.Component"));
+			}
+			else if (ListItem->Actor.IsValid())
+			{
+				return FClassIconFinder::FindIconForActor(ListItem->Actor);
+			}
+		}
+
+		return nullptr;
+	}
+
+	EVisibility GetCheckBoxVisibility() const
+	{
+		if (ListItem.IsValid())
+		{
+			bool bEnabledIsSet = ListItem->IsItemEnabled.IsSet();
+			return bEnabledIsSet ? EVisibility::Visible : EVisibility::Collapsed;
+		}
+
+		return EVisibility::Collapsed;
+	}
+
+	ECheckBoxState GetCheckBoxState() const
+	{
+		if (ListItem.IsValid())
+		{
+			bool bIsEnabled = ListItem->IsItemEnabled.Get(false);
+			return bIsEnabled ? ECheckBoxState::Checked : ECheckBoxState::Unchecked;
+		}
+
+		return ECheckBoxState::Unchecked;
+	}
+
+	void OnCheckBoxStateChanged(ECheckBoxState NewState)
+	{
+		if (ListItem.IsValid())
+		{
+			const bool bNewEnabledState = NewState == ECheckBoxState::Checked ? true : false;
+			ListItem->OnItemEnabledChanged.ExecuteIfBound(ListItem, bNewEnabledState);
+		}
+	}
+
+private:
+	FDisplayClusterColorGradingListItemRef ListItem;
+};
+
+void SDisplayClusterColorGradingObjectList::Construct(const FArguments& InArgs)
+{
+	OnSelectionChangedDelegate = InArgs._OnSelectionChanged;
+
+	RefreshList();
+
+	ChildSlot
+	[
+		SNew(SVerticalBox)
+
+		+ SVerticalBox::Slot()
+		.FillHeight(1.0f)
+		[
+			SAssignNew(ListView, SListView<FDisplayClusterColorGradingListItemRef>)
+			.ListItemsSource(InArgs._ColorGradingItemsSource)
+			.ItemHeight(28)
+			.SelectionMode(ESelectionMode::Single)
+			.OnGenerateRow(this, &SDisplayClusterColorGradingObjectList::GenerateListItemRow)
+			.OnSelectionChanged(this, &SDisplayClusterColorGradingObjectList::OnSelectionChanged)
+			.HeaderRow
+			(
+				SNew(SHeaderRow)
+
+				+ SHeaderRow::Column(DisplayClusterColorGradingObjectListColumnNames::ItemEnabled)
+				.DefaultLabel(FText::GetEmpty())
+				.FixedWidth(24.0f)
+
+				+ SHeaderRow::Column(DisplayClusterColorGradingObjectListColumnNames::ItemLabel)
+				.DefaultLabel(LOCTEXT("ColorGradingDrawerList_ItemLabelHeader", "Item Label"))
+				.FillWidth(1.0f)
+			)
+		]
+	];
+}
+
+void SDisplayClusterColorGradingObjectList::RefreshList()
+{
+	if (ListView.IsValid())
+	{
+		ListView->RebuildList();
+	}
+}
+
+TArray<FDisplayClusterColorGradingListItemRef> SDisplayClusterColorGradingObjectList::GetSelectedItems()
+{
+	return ListView->GetSelectedItems();
+}
+
+void SDisplayClusterColorGradingObjectList::SetSelectedItems(const TArray<FDisplayClusterColorGradingListItemRef>& InSelectedItems)
+{
+	ListView->ClearSelection();
+	ListView->SetItemSelection(InSelectedItems, true);
+}
+
+TSharedRef<ITableRow> SDisplayClusterColorGradingObjectList::GenerateListItemRow(FDisplayClusterColorGradingListItemRef Item, const TSharedRef<STableViewBase>& OwnerTable)
+{
+	return SNew(SColorGradingListItemRow, OwnerTable, Item);
+}
+
+void SDisplayClusterColorGradingObjectList::OnSelectionChanged(FDisplayClusterColorGradingListItemRef SelectedItem, ESelectInfo::Type SelectInfo)
+{
+	OnSelectionChangedDelegate.ExecuteIfBound(SharedThis<SDisplayClusterColorGradingObjectList>(this), SelectedItem, SelectInfo);
+}
+
+#undef LOCTEXT_NAMESPACE
\ No newline at end of file
diff --git a/Engine/Plugins/Runtime/nDisplay/Source/DisplayClusterColorGrading/Private/SDisplayClusterColorGradingObjectList.h b/Engine/Plugins/Runtime/nDisplay/Source/DisplayClusterColorGrading/Private/SDisplayClusterColorGradingObjectList.h
new file mode 100644
--- /dev/null
+++ b/Engine/Plugins/Runtime/nDisplay/Source/DisplayClusterColorGrading/Private/SDisplayClusterColorGradingObjectList.h
@@ -0,0 +1,80 @@
+// Copyright Epic Games, Inc. All Rights Reserved.
+
+#pragma once
+
+#include "CoreMinimal.h"
+#include "Components/ActorComponent.h"
+#include "Framework/SlateDelegates.h"
+#include "GameFramework/Actor.h"
+#include "Widgets/DeclarativeSyntaxSupport.h"
+#include "Widgets/SCompoundWidget.h"
+
+class ITableRow;
+class STableViewBase;
+
+template<class T>
+class SListView;
+
+DECLARE_DELEGATE_TwoParams(FOnColorGradingItemEnabledChanged, TSharedPtr<struct FDisplayClusterColorGradingListItem>, bool);
+
+/** A structure to store references to color gradable actors and components */
+struct FDisplayClusterColorGradingListItem
+{
+	/** The actor that is color gradable */
+	TWeakObjectPtr<AActor> Actor;
+
+	/** The component that is color gradable */
+	TWeakObjectPtr<UActorComponent> Component;
+
+	/** Attribute that retrieves whether color grading is enabled on the color gradable item */
+	TAttribute<bool> IsItemEnabled;
+
+	/** Delegate raised when the enabled state of the color gradable item has been changed */
+	FOnColorGradingItemEnabledChanged OnItemEnabledChanged;
+
+	FDisplayClusterColorGradingListItem(AActor* InActor, UActorComponent* InComponent = nullptr)
+		: Actor(InActor)
+		, Component(InComponent)
+		, IsItemEnabled(false)
+	{ }
+};
+
+typedef TSharedPtr<FDisplayClusterColorGradingListItem> FDisplayClusterColorGradingListItemRef;
+
+/** Displays a list of color gradable items */
+class SDisplayClusterColorGradingObjectList : public SCompoundWidget
+{
+public:
+	DECLARE_DELEGATE_ThreeParams(FOnSelectionChanged, TSharedRef<SDisplayClusterColorGradingObjectList>, FDisplayClusterColorGradingListItemRef, ESelectInfo::Type);
+
+public:
+	SLATE_BEGIN_ARGS(SDisplayClusterColorGradingObjectList) {}
+		SLATE_ARGUMENT(const TArray<FDisplayClusterColorGradingListItemRef>*, ColorGradingItemsSource)
+		SLATE_EVENT(FOnSelectionChanged, OnSelectionChanged)
+	SLATE_END_ARGS()
+
+	void Construct(const FArguments& InArgs);
+
+	/** Refreshes the list, updating the UI to reflect the current state of the source items list*/
+	void RefreshList();
+
+	/** Gets a list of currently selected items */
+	TArray<FDisplayClusterColorGradingListItemRef> GetSelectedItems();
+
+	/** Selects the specified list of items */
+	void SetSelectedItems(const TArray<FDisplayClusterColorGradingListItemRef>& InSelectedItems);
+
+private:
+	/** Generates the table row widget for the specified list item */
+	TSharedRef<ITableRow> GenerateListItemRow(FDisplayClusterColorGradingListItemRef Item, const TSharedRef<STableViewBase>& OwnerTable);
+
+	/** Raised when the internal list view's selection has changed */
+	void OnSelectionChanged(FDisplayClusterColorGradingListItemRef SelectedItem, ESelectInfo::Type SelectInfo);
+
+private:
+	/** Internal list view used to display the list of color gradable items */
+	TSharedPtr<SListView<FDisplayClusterColorGradingListItemRef>> ListView;
+
+	/** A delegate that is raised when the list of selected items is changed */
+	FOnSelectionChanged OnSelectionChangedDelegate;
+};
\ No newline at end of file
diff --git a/Engine/Plugins/Runtime/nDisplay/Source/DisplayClusterColorGrading/Public/IDisplayClusterColorGrading.h b/Engine/Plugins/Runtime/nDisplay/Source/DisplayClusterColorGrading/Public/IDisplayClusterColorGrading.h
new file mode 100644
--- /dev/null
+++ b/Engine/Plugins/Runtime/nDisplay/Source/DisplayClusterColorGrading/Public/IDisplayClusterColorGrading.h
@@ -0,0 +1,45 @@
+// Copyright Epic Games, Inc. All Rights Reserved.
+
+#pragma once
+
+#include "CoreMinimal.h"
+#include "Modules/ModuleInterface.h"
+#include "Modules/ModuleManager.h"
+
+class IDisplayClusterColorGradingDrawerSingleton;
+
+/**
+ * Display Cluster Color Grading module interface
+ */
+class IDisplayClusterColorGrading : public IModuleInterface
+{
+public:
+	static constexpr const TCHAR* ModuleName = TEXT("DisplayClusterColorGrading");
+
+public:
+	virtual ~IDisplayClusterColorGrading() = default;
+
+	/**
+	* Singleton-like access to this module's interface.  This is just for convenience!
+	* Beware of calling this during the shutdown phase, though.  Your module might have been unloaded already.
+	*
+	* @return Returns singleton instance, loading the module on demand if needed
+	*/
+	static inline IDisplayClusterColorGrading& Get()
+	{
+		return FModuleManager::GetModuleChecked<IDisplayClusterColorGrading>(ModuleName);
+	}
+
+	/**
+	* Checks to see if this module is loaded and ready.  It is only valid to call Get() if IsAvailable() returns true.
+	*
+	* @return True if the module is loaded and ready to use
+	*/
+	static inline bool IsAvailable()
+	{
+		return FModuleManager::Get().IsModuleLoaded(ModuleName);
+	}
+
+	/** Gets the singleton used to manage the color grading drawer */
+	virtual IDisplayClusterColorGradingDrawerSingleton& GetColorGradingDrawerSingleton() const = 0;
+};
\ No newline at end of file
diff --git a/Engine/Plugins/Runtime/nDisplay/Source/DisplayClusterColorGrading/Public/IDisplayClusterColorGradingDrawerSingleton.h b/Engine/Plugins/Runtime/nDisplay/Source/DisplayClusterColorGrading/Public/IDisplayClusterColorGradingDrawerSingleton.h
new file mode 100644
--- /dev/null
+++ b/Engine/Plugins/Runtime/nDisplay/Source/DisplayClusterColorGrading/Public/IDisplayClusterColorGradingDrawerSingleton.h
@@ -0,0 +1,16 @@
+// Copyright Epic Games, Inc. All Rights Reserved.
+
+#pragma once
+
+class IPropertyRowGenerator;
+
+/** A singleton used to manage and store persistent state for the color grading drawer */
+class IDisplayClusterColorGradingDrawerSingleton
+{
+public:
+	/** Docks the color grading drawer in the nDisplay operator window */
+	virtual void DockColorGradingDrawer() = 0;
+
+	/** Refreshes the UI of any open color grading drawers */
+	virtual void RefreshColorGradingDrawers(bool bPreserveDrawerState) = 0;
+};
\ No newline at end of file
diff --git a/Engine/Plugins/Runtime/nDisplay/Source/DisplayClusterConfiguration/Public/DisplayClusterConfigurationTypes_Postprocess.h b/Engine/Plugins/Runtime/nDisplay/Source/DisplayClusterConfiguration/Public/DisplayClusterConfigurationTypes_Postprocess.h
--- a/Engine/Plugins/Runtime/nDisplay/Source/DisplayClusterConfiguration/Public/DisplayClusterConfigurationTypes_Postprocess.h
+++ b/Engine/Plugins/Runtime/nDisplay/Source/DisplayClusterConfiguration/Public/DisplayClusterConfigurationTypes_Postprocess.h
@@ -251,6 +251,11 @@ struct FDisplayClusterConfigurationViewport_PerViewportColorGrading
 {
 	GENERATED_BODY()
 
+#if WITH_EDITORONLY_DATA
+	UPROPERTY(EditAnywhere, Category = "Viewport Settings")
+	FText Name = FText::GetEmpty();
+#endif
+
 	/** Enable the color grading settings for the viewport(s) specified and add them to nDisplay's color grading stack.  This will not affect the inner frustum. */
 	UPROPERTY(EditAnywhere, BlueprintReadWrite, Category = "Viewport Settings", meta = (DisplayName = "Enable Per-Viewport Color Grading"))
 	bool bIsEnabled = true;
@@ -291,6 +296,11 @@ struct FDisplayClusterConfigurationViewport_PerNodeColorGrading
 {
 	GENERATED_BODY()
 
+#if WITH_EDITORONLY_DATA
+	UPROPERTY(EditAnywhere, Category = "Viewport Settings")
+	FText Name = FText::GetEmpty();
+#endif
+
 	/** Enable the color grading settings for the node(s) specified and add them to nDisplay's color grading stack. */
 	UPROPERTY(EditAnywhere, BlueprintReadWrite, Category = "Viewport Settings", meta = (DisplayName = "Enable Per-Node Color Grading"))
 	bool bIsEnabled = true;
diff --git a/Engine/Plugins/Runtime/nDisplay/Source/DisplayClusterConfigurator/Private/Views/Details/Cluster/DisplayClusterConfiguratorNodeSelectionCustomization.cpp b/Engine/Plugins/Runtime/nDisplay/Source/DisplayClusterConfigurator/Private/Views/Details/Cluster/DisplayClusterConfiguratorNodeSelectionCustomization.cpp
--- a/Engine/Plugins/Runtime/nDisplay/Source/DisplayClusterConfigurator/Private/Views/Details/Cluster/DisplayClusterConfiguratorNodeSelectionCustomization.cpp
+++ b/Engine/Plugins/Runtime/nDisplay/Source/DisplayClusterConfigurator/Private/Views/Details/Cluster/DisplayClusterConfiguratorNodeSelectionCustomization.cpp
@@ -320,6 +320,7 @@ void FDisplayClusterConfiguratorPerViewportColorGradingCustomization::SetChildre
 {
 	TSharedPtr<IPropertyHandle> IsEnabledHandle = GET_CHILD_HANDLE(FDisplayClusterConfigurationViewport_PerViewportColorGrading, bIsEnabled);
 	TSharedPtr<IPropertyHandle> IsEntireClusterPostProcessHandle = GET_CHILD_HANDLE(FDisplayClusterConfigurationViewport_PerViewportColorGrading, bIsEntireClusterEnabled);
+	TSharedPtr<IPropertyHandle> NameHandle = GET_CHILD_HANDLE(FDisplayClusterConfigurationViewport_PerViewportColorGrading, Name);
 	TSharedPtr<IPropertyHandle> PostProcessSettingsHandle = GET_CHILD_HANDLE(FDisplayClusterConfigurationViewport_PerViewportColorGrading, ColorGradingSettings);
 	TSharedPtr<IPropertyHandle> ArrayHandle = GET_CHILD_HANDLE(FDisplayClusterConfigurationViewport_PerViewportColorGrading, ApplyPostProcessToObjects);
 
@@ -332,6 +333,7 @@ void FDisplayClusterConfiguratorPerViewportColorGradingCustomization::SetChildre
 
 	InChildBuilder.AddProperty(IsEnabledHandle.ToSharedRef());
 	InChildBuilder.AddProperty(IsEntireClusterPostProcessHandle.ToSharedRef());
+	InChildBuilder.AddProperty(NameHandle.ToSharedRef());
 	InChildBuilder.AddProperty(PostProcessSettingsHandle.ToSharedRef()).EditCondition(IsEnabledEditCondition, nullptr);
 
 	ArrayHandle->SetInstanceMetaData(FDisplayClusterConfiguratorNodeSelection::NAME_ElementToolTip, ArrayHandle->GetPropertyDisplayName().ToString());
diff --git a/Engine/Plugins/Runtime/nDisplay/Source/DisplayClusterLightCardEditor/Private/SDisplayClusterLightCardEditor.cpp b/Engine/Plugins/Runtime/nDisplay/Source/DisplayClusterLightCardEditor/Private/SDisplayClusterLightCardEditor.cpp
--- a/Engine/Plugins/Runtime/nDisplay/Source/DisplayClusterLightCardEditor/Private/SDisplayClusterLightCardEditor.cpp
+++ b/Engine/Plugins/Runtime/nDisplay/Source/DisplayClusterLightCardEditor/Private/SDisplayClusterLightCardEditor.cpp
@@ -15,6 +15,7 @@
 #include "Viewport/DisplayClusterLightCardEditorViewportClient.h"
 
 #include "IDisplayClusterOperator.h"
+#include "IDisplayClusterOperatorViewModel.h"
 
 #include "DisplayClusterRootActor.h"
 #include "DisplayClusterConfigurationTypes.h"
@@ -62,7 +63,7 @@ void SDisplayClusterLightCardEditor::UnregisterTabSpawner()
 void SDisplayClusterLightCardEditor::RegisterLayoutExtension(FLayoutExtender& InExtender)
 {
 	FTabManager::FTab NewTab(FTabId(TabName, ETabIdFlags::SaveLayout), ETabState::OpenedTab);
-	InExtender.ExtendStack(IDisplayClusterOperator::Get().GetOperatorExtensionId(), ELayoutExtensionPosition::After, NewTab);
+	InExtender.ExtendStack(IDisplayClusterOperator::Get().GetPrimaryOperatorExtensionId(), ELayoutExtensionPosition::After, NewTab);
 }
 
 TSharedRef<SDockTab> SDisplayClusterLightCardEditor::SpawnInTab(const FSpawnTabArgs& SpawnTabArgs)
@@ -79,7 +80,7 @@ TSharedRef<SDockTab> SDisplayClusterLightCardEditor::SpawnInTab(const FSpawnTabA
 
 SDisplayClusterLightCardEditor::~SDisplayClusterLightCardEditor()
 {
-	IDisplayClusterOperator::Get().OnActiveRootActorChanged().Remove(ActiveRootActorChangedHandle);
+	IDisplayClusterOperator::Get().GetOperatorViewModel()->OnActiveRootActorChanged().Remove(ActiveRootActorChangedHandle);
 	FCoreUObjectDelegates::OnObjectPropertyChanged.RemoveAll(this);
 
 	if (GEngine != nullptr)
@@ -109,7 +110,7 @@ void SDisplayClusterLightCardEditor::PostRedo(bool bSuccess)
 
 void SDisplayClusterLightCardEditor::Construct(const FArguments& InArgs, const TSharedRef<SDockTab>& MajorTabOwner, const TSharedPtr<SWindow>& WindowOwner)
 {
-	ActiveRootActorChangedHandle = IDisplayClusterOperator::Get().OnActiveRootActorChanged().AddSP(this, &SDisplayClusterLightCardEditor::OnActiveRootActorChanged);
+	ActiveRootActorChangedHandle = IDisplayClusterOperator::Get().GetOperatorViewModel()->OnActiveRootActorChanged().AddSP(this, &SDisplayClusterLightCardEditor::OnActiveRootActorChanged);
 	if (GEngine != nullptr)
 	{
 		GEngine->OnLevelActorDeleted().AddSP(this, &SDisplayClusterLightCardEditor::OnLevelActorDeleted);
diff --git a/Engine/Plugins/Runtime/nDisplay/Source/DisplayClusterOperator/DisplayClusterOperator.Build.cs b/Engine/Plugins/Runtime/nDisplay/Source/DisplayClusterOperator/DisplayClusterOperator.Build.cs
--- a/Engine/Plugins/Runtime/nDisplay/Source/DisplayClusterOperator/DisplayClusterOperator.Build.cs
+++ b/Engine/Plugins/Runtime/nDisplay/Source/DisplayClusterOperator/DisplayClusterOperator.Build.cs
@@ -13,14 +13,15 @@ public class DisplayClusterOperator : ModuleRules
 
 				"Core",
 				"CoreUObject",
+				"EditorStyle",
 				"Engine",
 				"InputCore",
+				"Kismet",
 				"Slate",
 				"SlateCore",
+				"ToolWidgets",
 				"UnrealEd",
-				"EditorStyle",
 				"WorkspaceMenuStructure",
-				"Kismet"
 			});
 	}
 }
diff --git a/Engine/Plugins/Runtime/nDisplay/Source/DisplayClusterOperator/Private/DisplayClusterOperatorModule.cpp b/Engine/Plugins/Runtime/nDisplay/Source/DisplayClusterOperator/Private/DisplayClusterOperatorModule.cpp
--- a/Engine/Plugins/Runtime/nDisplay/Source/DisplayClusterOperator/Private/DisplayClusterOperatorModule.cpp
+++ b/Engine/Plugins/Runtime/nDisplay/Source/DisplayClusterOperator/Private/DisplayClusterOperatorModule.cpp
@@ -2,6 +2,7 @@
 
 #include "DisplayClusterOperatorModule.h"
 
+#include "DisplayClusterOperatorViewModel.h"
 #include "SDisplayClusterOperatorPanel.h"
 #include "DisplayClusterRootActor.h"
 
@@ -10,11 +11,17 @@
 #include "Engine/World.h"
 #include "LevelEditorViewport.h"
 #include "Toolkits/AssetEditorToolkit.h"
+#include "Widgets/Docking/SDockTab.h"
+#include "WorkspaceMenuStructure.h"
+#include "WorkspaceMenuStructureModule.h"
 
 #define LOCTEXT_NAMESPACE "DisplayClusterOperator"
 
+const FName FDisplayClusterOperatorModule::OperatorPanelTabName = TEXT("DisplayClusterOperatorTab");
+
 void FDisplayClusterOperatorModule::StartupModule()
 {
+	OperatorViewModel = MakeShared<FDisplayClusterOperatorViewModel>();
 	OperatorToolBarExtensibilityManager = MakeShared<FExtensibilityManager>();
 
 	RegisterTabSpawners();
@@ -25,9 +32,19 @@ void FDisplayClusterOperatorModule::ShutdownModule()
 	UnregisterTabSpawners();
 }
 
-FName FDisplayClusterOperatorModule::GetOperatorExtensionId()
+TSharedRef<IDisplayClusterOperatorViewModel> FDisplayClusterOperatorModule::GetOperatorViewModel()
+{
+	return OperatorViewModel.ToSharedRef();
+}
+
+FName FDisplayClusterOperatorModule::GetPrimaryOperatorExtensionId()
 {
-	return SDisplayClusterOperatorPanel::TabExtensionId;
+	return SDisplayClusterOperatorPanel::PrimaryTabExtensionId;
+}
+
+FName FDisplayClusterOperatorModule::GetAuxilliaryOperatorExtensionId()
+{
+	return SDisplayClusterOperatorPanel::AuxilliaryTabExtensionId;
 }
 
 void FDisplayClusterOperatorModule::GetRootActorLevelInstances(TArray<ADisplayClusterRootActor*>& OutRootActorInstances)
@@ -58,14 +75,45 @@ void FDisplayClusterOperatorModule::ShowDetailsForObjects(const TArray<UObject*>
 	DetailObjectsChanged.Broadcast(Objects);
 }
 
+void FDisplayClusterOperatorModule::ForceDismissDrawers()
+{
+	if (ActiveOperatorPanel.IsValid())
+	{
+		ActiveOperatorPanel.Pin()->ForceDismissDrawers();
+	}
+}
+
 void FDisplayClusterOperatorModule::RegisterTabSpawners()
 {
-	SDisplayClusterOperatorPanel::RegisterTabSpawner();
+	// Register the nDisplay operator panel tab spawner
+	FGlobalTabmanager::Get()->RegisterNomadTabSpawner(OperatorPanelTabName, FOnSpawnTab::CreateRaw(this, &FDisplayClusterOperatorModule::SpawnOperatorPanelTab))
+		.SetDisplayName(LOCTEXT("TabDisplayName", "nDisplay Operator"))
+		.SetTooltipText(LOCTEXT("TabTooltip", "Open the nDisplay Operator tab."))
+		.SetGroup(WorkspaceMenu::GetMenuStructure().GetLevelEditorVirtualProductionCategory());
 }
 
 void FDisplayClusterOperatorModule::UnregisterTabSpawners()
 {
-	SDisplayClusterOperatorPanel::UnregisterTabSpawner();
+	FGlobalTabmanager::Get()->UnregisterNomadTabSpawner(OperatorPanelTabName);
+}
+
+TSharedRef<SDockTab> FDisplayClusterOperatorModule::SpawnOperatorPanelTab(const FSpawnTabArgs& SpawnTabArgs)
+{
+	const TSharedRef<SDockTab> MajorTab = SNew(SDockTab)
+		.TabRole(ETabRole::NomadTab)
+		.OnTabClosed_Raw(this, &FDisplayClusterOperatorModule::OnOperatorPanelTabClosed);
+
+	MajorTab->SetContent(SAssignNew(ActiveOperatorPanel, SDisplayClusterOperatorPanel, OperatorViewModel->CreateTabManager(MajorTab), SpawnTabArgs.GetOwnerWindow()));
+
+	return MajorTab;
+}
+
+void FDisplayClusterOperatorModule::OnOperatorPanelTabClosed(TSharedRef<SDockTab> Tab)
+{
+	OperatorViewModel->ResetTabManager();
+	OperatorViewModel->SetRootActor(nullptr);
+
+	ActiveOperatorPanel.Reset();
 }
 
 IMPLEMENT_MODULE(FDisplayClusterOperatorModule, DisplayClusterOperator);
diff --git a/Engine/Plugins/Runtime/nDisplay/Source/DisplayClusterOperator/Private/DisplayClusterOperatorModule.h b/Engine/Plugins/Runtime/nDisplay/Source/DisplayClusterOperator/Private/DisplayClusterOperatorModule.h
--- a/Engine/Plugins/Runtime/nDisplay/Source/DisplayClusterOperator/Private/DisplayClusterOperatorModule.h
+++ b/Engine/Plugins/Runtime/nDisplay/Source/DisplayClusterOperator/Private/DisplayClusterOperatorModule.h
@@ -4,12 +4,21 @@
 
 #include "IDisplayClusterOperator.h"
 
+class FDisplayClusterOperatorViewModel;
+class FSpawnTabArgs;
+class SDockTab;
+class SDisplayClusterOperatorPanel;
+
 /**
  * Display Cluster editor module
  */
 class FDisplayClusterOperatorModule :
 	public IDisplayClusterOperator
 {
+public:
+	/** The name of the tab that the operator panel lives in */
+	static const FName OperatorPanelTabName;
+
 public:
 	//~ IModuleInterface interface
 	virtual void StartupModule() override;
@@ -17,26 +26,33 @@ public:
 	//~ End IModuleInterface interface
 
 	//~ IDisplayClusterOperator interface
+	virtual TSharedRef<IDisplayClusterOperatorViewModel> GetOperatorViewModel() override;
 	virtual FOnRegisterLayoutExtensions& OnRegisterLayoutExtensions() override { return RegisterLayoutExtensions; }
-	virtual FOnActiveRootActorChanged& OnActiveRootActorChanged() override { return RootActorChanged; }
+	virtual FOnRegisterStatusBarExtensions& OnRegisterStatusBarExtensions() override { return RegisterStatusBarExtensions; }
 	virtual FOnDetailObjectsChanged& OnDetailObjectsChanged() override { return DetailObjectsChanged; }
 
-	virtual FName GetOperatorExtensionId() override;
+	virtual FName GetPrimaryOperatorExtensionId() override;
+	virtual FName GetAuxilliaryOperatorExtensionId() override;
 	virtual TSharedPtr<FExtensibilityManager> GetOperatorToolBarExtensibilityManager() override { return OperatorToolBarExtensibilityManager; }
 	virtual void GetRootActorLevelInstances(TArray<ADisplayClusterRootActor*>& OutRootActorInstances) override;
 	virtual void ShowDetailsForObject(UObject* Object) override;
 	virtual void ShowDetailsForObjects(const TArray<UObject*>& Objects) override;
+	virtual void ForceDismissDrawers() override;
 	//~ End IDisplayClusterOperator interface
 
 private:
 	void RegisterTabSpawners();
 	void UnregisterTabSpawners();
 
+	TSharedRef<SDockTab> SpawnOperatorPanelTab(const FSpawnTabArgs& SpawnTabArgs);
+	void OnOperatorPanelTabClosed(TSharedRef<SDockTab> Tab);
 
 private:
 	FOnRegisterLayoutExtensions RegisterLayoutExtensions;
-	FOnActiveRootActorChanged RootActorChanged;
+	FOnRegisterStatusBarExtensions RegisterStatusBarExtensions;
 	FOnDetailObjectsChanged DetailObjectsChanged;
 
+	TWeakPtr<SDisplayClusterOperatorPanel> ActiveOperatorPanel;
+	TSharedPtr<FDisplayClusterOperatorViewModel> OperatorViewModel;
 	TSharedPtr<FExtensibilityManager> OperatorToolBarExtensibilityManager;
 };
diff --git a/Engine/Plugins/Runtime/nDisplay/Source/DisplayClusterOperator/Private/DisplayClusterOperatorStatusBarExtender.cpp b/Engine/Plugins/Runtime/nDisplay/Source/DisplayClusterOperator/Private/DisplayClusterOperatorStatusBarExtender.cpp
new file mode 100644
--- /dev/null
+++ b/Engine/Plugins/Runtime/nDisplay/Source/DisplayClusterOperator/Private/DisplayClusterOperatorStatusBarExtender.cpp
@@ -0,0 +1,18 @@
+// Copyright Epic Games, Inc. All Rights Reserved.
+
+#include "DisplayClusterOperatorStatusBarExtender.h"
+
+#include "SDisplayClusterOperatorStatusBar.h"
+
+void FDisplayClusterOperatorStatusBarExtender::AddWidgetDrawer(const FWidgetDrawerConfig& InWidgetDrawerConfig, int32 SlotIndex)
+{
+	WidgetDrawers.Add(TPair<int32, FWidgetDrawerConfig>(SlotIndex, InWidgetDrawerConfig));
+}
+
+void FDisplayClusterOperatorStatusBarExtender::RegisterExtensions(TSharedRef<SDisplayClusterOperatorStatusBar> StatusBar)
+{
+	for (TPair<int32, FWidgetDrawerConfig>& DrawerConfigPair : WidgetDrawers)
+	{
+		StatusBar->RegisterDrawer(MoveTemp(DrawerConfigPair.Value), DrawerConfigPair.Key);
+	}
+}
\ No newline at end of file
diff --git a/Engine/Plugins/Runtime/nDisplay/Source/DisplayClusterOperator/Private/DisplayClusterOperatorViewModel.cpp b/Engine/Plugins/Runtime/nDisplay/Source/DisplayClusterOperator/Private/DisplayClusterOperatorViewModel.cpp
new file mode 100644
--- /dev/null
+++ b/Engine/Plugins/Runtime/nDisplay/Source/DisplayClusterOperator/Private/DisplayClusterOperatorViewModel.cpp
@@ -0,0 +1,24 @@
+// Copyright Epic Games, Inc. All Rights Reserved.
+
+#include "DisplayClusterOperatorViewModel.h"
+
+#include "DisplayClusterRootActor.h"
+
+#include "Framework/Docking/TabManager.h"
+
+void FDisplayClusterOperatorViewModel::SetRootActor(ADisplayClusterRootActor* InRootActor)
+{
+	RootActor = InRootActor;
+	RootActorChanged.Broadcast(RootActor.Get());
+}
+
+TSharedRef<FTabManager> FDisplayClusterOperatorViewModel::CreateTabManager(const TSharedRef<SDockTab>& MajorTabOwner)
+{
+	TabManager = FGlobalTabmanager::Get()->NewTabManager(MajorTabOwner);
+	return TabManager.ToSharedRef();
+}
+
+void FDisplayClusterOperatorViewModel::ResetTabManager()
+{
+	TabManager.Reset();
+}
\ No newline at end of file
diff --git a/Engine/Plugins/Runtime/nDisplay/Source/DisplayClusterOperator/Private/DisplayClusterOperatorViewModel.h b/Engine/Plugins/Runtime/nDisplay/Source/DisplayClusterOperator/Private/DisplayClusterOperatorViewModel.h
new file mode 100644
--- /dev/null
+++ b/Engine/Plugins/Runtime/nDisplay/Source/DisplayClusterOperator/Private/DisplayClusterOperatorViewModel.h
@@ -0,0 +1,35 @@
+// Copyright Epic Games, Inc. All Rights Reserved.
+
+#pragma once
+
+#include "CoreMinimal.h"
+#include "IDisplayClusterOperatorViewModel.h"
+#include "UObject/WeakObjectPtr.h"
+
+class ADisplayClusterRootActor;
+class SDockTab;
+
+/** A view model object that stores any state from the operator panel that should be exposed externally */
+class FDisplayClusterOperatorViewModel : public TSharedFromThis<FDisplayClusterOperatorViewModel>, public IDisplayClusterOperatorViewModel
+{
+public:
+	virtual ~FDisplayClusterOperatorViewModel() = default;
+
+	//~ IDisplayClusterOperatorViewModel interface
+	virtual bool HasRootActor() const override { return RootActor.IsValid(); }
+	virtual ADisplayClusterRootActor* GetRootActor() const override { return RootActor.IsValid() ? RootActor.Get() : nullptr; }
+	virtual void SetRootActor(ADisplayClusterRootActor* InRootActor) override;
+	virtual FOnActiveRootActorChanged& OnActiveRootActorChanged() override { return RootActorChanged; }
+
+	virtual TSharedPtr<FTabManager> GetTabManager() const override { return TabManager; }
+	//~ End IDisplayClusterOperatorViewModel interface
+
+	TSharedRef<FTabManager> CreateTabManager(const TSharedRef<SDockTab>& MajorTabOwner);
+	void ResetTabManager();
+
+private:
+	TWeakObjectPtr<ADisplayClusterRootActor> RootActor;
+	TSharedPtr<FTabManager> TabManager;
+
+	FOnActiveRootActorChanged RootActorChanged;
+};
\ No newline at end of file
diff --git a/Engine/Plugins/Runtime/nDisplay/Source/DisplayClusterOperator/Private/SDisplayClusterOperatorPanel.cpp b/Engine/Plugins/Runtime/nDisplay/Source/DisplayClusterOperator/Private/SDisplayClusterOperatorPanel.cpp
--- a/Engine/Plugins/Runtime/nDisplay/Source/DisplayClusterOperator/Private/SDisplayClusterOperatorPanel.cpp
+++ b/Engine/Plugins/Runtime/nDisplay/Source/DisplayClusterOperator/Private/SDisplayClusterOperatorPanel.cpp
@@ -3,7 +3,10 @@
 #include "SDisplayClusterOperatorPanel.h"
 
 #include "IDisplayClusterOperator.h"
+#include "IDisplayClusterOperatorViewModel.h"
 #include "SDisplayClusterOperatorToolbar.h"
+#include "SDisplayClusterOperatorStatusBar.h"
+#include "DisplayClusterOperatorStatusBarExtender.h"
 
 #include "Styling/AppStyle.h"
 #include "Framework/Docking/TabManager.h"
@@ -11,49 +14,24 @@
 #include "Framework/MultiBox/MultiBoxBuilder.h"
 #include "SKismetInspector.h"
 #include "Widgets/Docking/SDockTab.h"
-#include "WorkspaceMenuStructure.h"
-#include "WorkspaceMenuStructureModule.h"
 
 #include "Components/ActorComponent.h"
 
 #define LOCTEXT_NAMESPACE "SDisplayClusterOperatorPanel"
 
-const FName SDisplayClusterOperatorPanel::TabName = TEXT("DisplayClusterOperatorTab");
 const FName SDisplayClusterOperatorPanel::ToolbarTabId = TEXT("OperatorToolbar");
 const FName SDisplayClusterOperatorPanel::DetailsTabId = TEXT("OperatorDetails");
-const FName SDisplayClusterOperatorPanel::TabExtensionId = TEXT("OperatorTabStack");
-
-void SDisplayClusterOperatorPanel::RegisterTabSpawner()
-{
-	FGlobalTabmanager::Get()->RegisterNomadTabSpawner(TabName, FOnSpawnTab::CreateStatic(&SDisplayClusterOperatorPanel::SpawnInTab))
-			.SetDisplayName(LOCTEXT("TabDisplayName", "nDisplay Operator"))
-			.SetTooltipText(LOCTEXT("TabTooltip", "Open the nDisplay Operator tab."))
-			.SetGroup(WorkspaceMenu::GetMenuStructure().GetLevelEditorVirtualProductionCategory());
-}
-
-void SDisplayClusterOperatorPanel::UnregisterTabSpawner()
-{
-	FGlobalTabmanager::Get()->UnregisterNomadTabSpawner(TabName);
-}
-
-TSharedRef<SDockTab> SDisplayClusterOperatorPanel::SpawnInTab(const FSpawnTabArgs& SpawnTabArgs)
-{
-	const TSharedRef<SDockTab> MajorTab = SNew(SDockTab)
-		.TabRole(ETabRole::NomadTab);
-
-	MajorTab->SetContent(SNew(SDisplayClusterOperatorPanel, MajorTab, SpawnTabArgs.GetOwnerWindow()));
-
-	return MajorTab;
-}
+const FName SDisplayClusterOperatorPanel::PrimaryTabExtensionId = TEXT("PrimaryOperatorTabStack");
+const FName SDisplayClusterOperatorPanel::AuxilliaryTabExtensionId = TEXT("AuxilliaryOperatorTabStack");
 
 SDisplayClusterOperatorPanel::~SDisplayClusterOperatorPanel()
 {
 	IDisplayClusterOperator::Get().OnDetailObjectsChanged().Remove(DetailObjectsChangedHandle);
 }
 
-void SDisplayClusterOperatorPanel::Construct(const FArguments& InArgs, const TSharedRef<SDockTab>& MajorTabOwner, const TSharedPtr<SWindow>& WindowOwner)
+void SDisplayClusterOperatorPanel::Construct(const FArguments& InArgs, const TSharedRef<FTabManager>& InTabManager, const TSharedPtr<SWindow>& WindowOwner)
 {
-	TabManager = FGlobalTabmanager::Get()->NewTabManager(MajorTabOwner);
+	TabManager = InTabManager;
 	TSharedRef<FWorkspaceItem> AppMenuGroup = TabManager->AddLocalWorkspaceMenuCategory(LOCTEXT("OperatorMenuGroupName", "nDisplay Operator"));
 	TabManager->SetAllowWindowMenuBar(true);
 
@@ -84,9 +62,19 @@ void SDisplayClusterOperatorPanel::Construct(const FArguments& InArgs, const TSh
 					->SetOrientation(Orient_Horizontal)
 					->Split
 					(
-						FTabManager::NewStack()
-							->SetExtensionId(TabExtensionId)
-							->SetSizeCoefficient(0.67f)
+						FTabManager::NewSplitter()
+						->SetOrientation(Orient_Vertical)
+						->SetSizeCoefficient(0.67f)
+						->Split
+						(
+							FTabManager::NewStack()
+								->SetExtensionId(PrimaryTabExtensionId)
+						)
+						->Split
+						(
+							FTabManager::NewStack()
+								->SetExtensionId(AuxilliaryTabExtensionId)
+						)
 					)
 					->Split
 					(
@@ -106,22 +94,42 @@ void SDisplayClusterOperatorPanel::Construct(const FArguments& InArgs, const TSh
 
 	ChildSlot
 	[
-		TabManager->RestoreFrom(Layout, WindowOwner).ToSharedRef()
+		SNew(SVerticalBox)
+
+		+SVerticalBox::Slot()
+		.Padding(4.0f, 2.f, 4.f, 2.f)
+		.FillHeight(1.0f)
+		[
+			TabManager->RestoreFrom(Layout, WindowOwner).ToSharedRef()
+		]
+
+		+SVerticalBox::Slot()
+		.Padding(0.0f, 2.0f, 0.0f, 0.0f)
+		.AutoHeight()
+		[
+			SAssignNew(StatusBar, SDisplayClusterOperatorStatusBar)
+		]
 	];
 
-	// TODO: Move ownership of active root controller to a view model object, instead of letting the toolbar control which root actor is active
+	// Allow any external modules to register extensions to the operator panel's status bar
+	FDisplayClusterOperatorStatusBarExtender StatusBarExtender;
+	IDisplayClusterOperator::Get().OnRegisterStatusBarExtensions().Broadcast(StatusBarExtender);
+	StatusBarExtender.RegisterExtensions(StatusBar.ToSharedRef());
+
 	if (ToolbarContainer.IsValid())
 	{
 		// Create toolbar after tab has been restored so child windows can register their toolbar extensions
 		ToolbarContainer->SetContent(
 		SAssignNew(Toolbar, SDisplayClusterOperatorToolbar)
 		.CommandList(TSharedPtr<FUICommandList>()));
+	}
+}
 
-		TWeakObjectPtr<ADisplayClusterRootActor> ActiveRootActor = Toolbar->GetActiveRootActor();
-		if (ActiveRootActor.IsValid())
-		{
-			IDisplayClusterOperator::Get().OnActiveRootActorChanged().Broadcast(ActiveRootActor.Get());
-		}
+void SDisplayClusterOperatorPanel::ForceDismissDrawers()
+{
+	if (StatusBar.IsValid())
+	{
+		StatusBar->DismissDrawer(nullptr);
 	}
 }
 
diff --git a/Engine/Plugins/Runtime/nDisplay/Source/DisplayClusterOperator/Private/SDisplayClusterOperatorPanel.h b/Engine/Plugins/Runtime/nDisplay/Source/DisplayClusterOperator/Private/SDisplayClusterOperatorPanel.h
--- a/Engine/Plugins/Runtime/nDisplay/Source/DisplayClusterOperator/Private/SDisplayClusterOperatorPanel.h
+++ b/Engine/Plugins/Runtime/nDisplay/Source/DisplayClusterOperator/Private/SDisplayClusterOperatorPanel.h
@@ -13,40 +13,34 @@ class SBox;
 class SDockTab;
 class SWindow;
 class SDisplayClusterOperatorToolbar;
+class SDisplayClusterOperatorStatusBar;
 
 /** The nDisplay operator panel that allows users to edit root actor instances with a variety of tools */
 class SDisplayClusterOperatorPanel : public SCompoundWidget
 {
 public:
-	/** The name of the tab that the operator panel lives in */
-	static const FName TabName;
-
 	/** The tab ID that the operator toolbar lives in */
 	static const FName ToolbarTabId;
 
 	/** The tab ID that the details panel lives in */
 	static const FName DetailsTabId;
 
-	/** The ID of the tab stack that can be extended by external tabs */
-	static const FName TabExtensionId;
-
-
-	/** Registers the operator panel with the global tab manager and adds it to the Virtual Projections window menu */
-	static void RegisterTabSpawner();
-
-	/** Unregisters the operator panel from the global tab manager */
-	static void UnregisterTabSpawner();
-
-	/** Creates a tab with the operator panel inside */
-	static TSharedRef<SDockTab> SpawnInTab(const FSpawnTabArgs& SpawnTabArgs);
+	/** The ID of the main tab stack that can be extended by external tabs */
+	static const FName PrimaryTabExtensionId;
 
+	/** The ID of the auxilliary tab stack that can be extended by external tabs; intended for lower-thirds windows like log outputs */
+	static const FName AuxilliaryTabExtensionId;
 
+public:
 	SLATE_BEGIN_ARGS(SDisplayClusterOperatorPanel) {}
 	SLATE_END_ARGS()
 
 	~SDisplayClusterOperatorPanel();
 
-	void Construct(const FArguments& InArgs, const TSharedRef<SDockTab>& MajorTabOwner, const TSharedPtr<SWindow>& WindowOwner);
+	void Construct(const FArguments& InArgs, const TSharedRef<FTabManager>& InTabManager, const TSharedPtr<SWindow>& WindowOwner);
+
+	/** Forces the operator panel to dismiss any open drawers */
+	void ForceDismissDrawers();
 
 private:
 	/** Creates a tab with the operator toolbar in it */
@@ -74,6 +68,9 @@ private:
 	/** A reference to the operator panel's details view */
 	TSharedPtr<class SKismetInspector> DetailsView;
 
+	/** A reference to the operator panel's status bar widget */
+	TSharedPtr<SDisplayClusterOperatorStatusBar> StatusBar;
+
 	/** The delegate handle for the operator module's OnDetailObjectsChanged event */
 	FDelegateHandle DetailObjectsChangedHandle;
 };
\ No newline at end of file
diff --git a/Engine/Plugins/Runtime/nDisplay/Source/DisplayClusterOperator/Private/SDisplayClusterOperatorStatusBar.cpp b/Engine/Plugins/Runtime/nDisplay/Source/DisplayClusterOperator/Private/SDisplayClusterOperatorStatusBar.cpp
new file mode 100644
--- /dev/null
+++ b/Engine/Plugins/Runtime/nDisplay/Source/DisplayClusterOperator/Private/SDisplayClusterOperatorStatusBar.cpp
@@ -0,0 +1,31 @@
+// Copyright Epic Games, Inc. All Rights Reserved.
+
+#include "SDisplayClusterOperatorStatusBar.h"
+
+#include "Framework/Application/SlateApplication.h"
+#include "SWidgetDrawer.h"
+#include "Widgets/Layout/SBox.h"
+
+const FName SDisplayClusterOperatorStatusBar::StatusBarId = FName("DisplayClusterOperatorStatusBar");
+
+void SDisplayClusterOperatorStatusBar::Construct(const FArguments& InArgs)
+{
+    ChildSlot
+	[
+		SNew(SBox)
+		.HeightOverride(FAppStyle::Get().GetFloat("StatusBar.Height"))
+		[
+			SAssignNew(WidgetDrawer, SWidgetDrawer, StatusBarId)
+		]
+	];
+}
+
+void SDisplayClusterOperatorStatusBar::RegisterDrawer(FWidgetDrawerConfig&& Drawer, int32 SlotIndex)
+{
+	WidgetDrawer->RegisterDrawer(MoveTemp(Drawer), SlotIndex);
+}
+
+bool SDisplayClusterOperatorStatusBar::DismissDrawer(const TSharedPtr<SWidget>& NewlyFocusedWidget)
+{
+	return WidgetDrawer->DismissDrawer(NewlyFocusedWidget);
+}
\ No newline at end of file
diff --git a/Engine/Plugins/Runtime/nDisplay/Source/DisplayClusterOperator/Private/SDisplayClusterOperatorStatusBar.h b/Engine/Plugins/Runtime/nDisplay/Source/DisplayClusterOperator/Private/SDisplayClusterOperatorStatusBar.h
new file mode 100644
--- /dev/null
+++ b/Engine/Plugins/Runtime/nDisplay/Source/DisplayClusterOperator/Private/SDisplayClusterOperatorStatusBar.h
@@ -0,0 +1,41 @@
+// Copyright Epic Games, Inc. All Rights Reserved.
+
+#pragma once
+
+#include "CoreMinimal.h"
+#include "Widgets/DeclarativeSyntaxSupport.h"
+#include "Widgets/SCompoundWidget.h"
+
+class SWidgetDrawer;
+struct FWidgetDrawerConfig;
+
+class SDisplayClusterOperatorStatusBar : public SCompoundWidget
+{
+public:
+	static const FName StatusBarId;
+
+public:
+	SLATE_BEGIN_ARGS(SDisplayClusterOperatorStatusBar)
+	{}
+	SLATE_END_ARGS()
+
+	void Construct(const FArguments& InArgs);
+
+	//~ Begin SWidget Interface
+	virtual bool SupportsKeyboardFocus() const override { return false; }
+	//~ End SWidget Interface
+
+	/** Registers a new drawer with this status bar. Registering will add a button to open and close the drawer */
+	void RegisterDrawer(FWidgetDrawerConfig&& Drawer, int32 SlotIndex = INDEX_NONE);
+
+	/**
+	 * Dismisses an open drawer with an animation.  The drawer contents are removed once the animation is complete
+	 *
+	 * @param NewlyFocusedWidget Optional widget to focus
+	 * @return true if any open drawer was dismissed
+	 */
+	bool DismissDrawer(const TSharedPtr<SWidget>& NewlyFocusedWidget);
+
+private:
+	TSharedPtr<SWidgetDrawer> WidgetDrawer;
+};
\ No newline at end of file
diff --git a/Engine/Plugins/Runtime/nDisplay/Source/DisplayClusterOperator/Private/SDisplayClusterOperatorToolbar.cpp b/Engine/Plugins/Runtime/nDisplay/Source/DisplayClusterOperator/Private/SDisplayClusterOperatorToolbar.cpp
--- a/Engine/Plugins/Runtime/nDisplay/Source/DisplayClusterOperator/Private/SDisplayClusterOperatorToolbar.cpp
+++ b/Engine/Plugins/Runtime/nDisplay/Source/DisplayClusterOperator/Private/SDisplayClusterOperatorToolbar.cpp
@@ -3,6 +3,7 @@
 #include "SDisplayClusterOperatorToolbar.h"
 
 #include "IDisplayClusterOperator.h"
+#include "IDisplayClusterOperatorViewModel.h"
 #include "DisplayClusterRootActor.h"
 
 #include "Styling/AppStyle.h"
@@ -27,7 +28,7 @@ SDisplayClusterOperatorToolbar::~SDisplayClusterOperatorToolbar()
 		LevelEditor.OnMapChanged().Remove(MapChangedHandle);
 	}
 	
-	if (ActiveRootActor.IsValid())
+	if (ADisplayClusterRootActor* ActiveRootActor = ViewModel->GetRootActor())
 	{
 		if (UBlueprint* Blueprint = UBlueprint::GetBlueprintFromClass(ActiveRootActor->GetClass()))
 		{
@@ -38,6 +39,7 @@ SDisplayClusterOperatorToolbar::~SDisplayClusterOperatorToolbar()
 
 void SDisplayClusterOperatorToolbar::Construct(const FArguments& InArgs)
 {
+	ViewModel = IDisplayClusterOperator::Get().GetOperatorViewModel();
 	CommandList = InArgs._CommandList;
 
 	FillRootActorList();
@@ -125,7 +127,7 @@ TSharedPtr<FString> SDisplayClusterOperatorToolbar::FillRootActorList(const FStr
 
 void SDisplayClusterOperatorToolbar::ClearSelectedRootActor()
 {
-	if (ActiveRootActor.IsValid())
+	if (ADisplayClusterRootActor* ActiveRootActor = ViewModel->GetRootActor())
 	{
 		if (UBlueprint* Blueprint = UBlueprint::GetBlueprintFromClass(ActiveRootActor->GetClass()))
 		{
@@ -134,8 +136,7 @@ void SDisplayClusterOperatorToolbar::ClearSelectedRootActor()
 	}
 
 	ActiveRootActorName.Reset();
-	ActiveRootActor.Reset();
-	IDisplayClusterOperator::Get().OnActiveRootActorChanged().Broadcast(nullptr);
+	ViewModel->SetRootActor(nullptr);
 	RootActorComboBox->SetSelectedItem(nullptr);
 }
 
@@ -160,7 +161,7 @@ void SDisplayClusterOperatorToolbar::OnRootActorChanged(TSharedPtr<FString> Item
 		}
 	}
 
-	if (ActiveRootActor.IsValid())
+	if (ADisplayClusterRootActor* ActiveRootActor = ViewModel->GetRootActor())
 	{
 		if (UBlueprint* Blueprint = UBlueprint::GetBlueprintFromClass(ActiveRootActor->GetClass()))
 		{
@@ -177,8 +178,7 @@ void SDisplayClusterOperatorToolbar::OnRootActorChanged(TSharedPtr<FString> Item
 		}
 	}
 
-	ActiveRootActor = SelectedRootActor;
-	IDisplayClusterOperator::Get().OnActiveRootActorChanged().Broadcast(SelectedRootActor);
+	ViewModel->SetRootActor(SelectedRootActor);
 }
 
 void SDisplayClusterOperatorToolbar::OnBlueprintCompiled(UBlueprint* Blueprint)
@@ -224,7 +224,7 @@ FText SDisplayClusterOperatorToolbar::GetRootActorComboBoxText() const
 
 void SDisplayClusterOperatorToolbar::OnLevelActorDeleted(AActor* Actor)
 {
-	if (Actor == ActiveRootActor)
+	if (Actor == ViewModel->GetRootActor())
 	{
 		if (Actor && Actor->GetClass()->HasAnyClassFlags(CLASS_NewerVersionExists))
 		{
@@ -240,7 +240,7 @@ void SDisplayClusterOperatorToolbar::OnLevelActorDeleted(AActor* Actor)
 void SDisplayClusterOperatorToolbar::HandleMapChanged(UWorld* InWorld, EMapChangeType InMapChangeType)
 {
 	if (InMapChangeType == EMapChangeType::TearDownWorld &&
-		(!ActiveRootActor.IsValid() || ActiveRootActor->GetWorld() == InWorld))
+		(!ViewModel->HasRootActor() || ViewModel->GetRootActor()->GetWorld() == InWorld))
 	{
 		ClearSelectedRootActor();
 	}
diff --git a/Engine/Plugins/Runtime/nDisplay/Source/DisplayClusterOperator/Private/SDisplayClusterOperatorToolbar.h b/Engine/Plugins/Runtime/nDisplay/Source/DisplayClusterOperator/Private/SDisplayClusterOperatorToolbar.h
--- a/Engine/Plugins/Runtime/nDisplay/Source/DisplayClusterOperator/Private/SDisplayClusterOperatorToolbar.h
+++ b/Engine/Plugins/Runtime/nDisplay/Source/DisplayClusterOperator/Private/SDisplayClusterOperatorToolbar.h
@@ -9,6 +9,7 @@
 
 enum class EMapChangeType : uint8;
 class FUICommandList;
+class IDisplayClusterOperatorViewModel;
 class ADisplayClusterRootActor;
 
 template<class T>
@@ -27,9 +28,6 @@ public:
 
 	void Construct(const FArguments& InArgs);
 
-	/** Gets the root actor that is currently selected in the toolbar */
-	TWeakObjectPtr<ADisplayClusterRootActor> GetActiveRootActor() const { return ActiveRootActor; }
-
 private:
 	/** 
 	 * Constructs the list of root actors that exist on the current level to use for the root actor picker dropdown.
@@ -64,6 +62,9 @@ private:
 	void HandleMapChanged(UWorld* InWorld, EMapChangeType InMapChangeType);
 
 private:
+	/** A reference the the operator panel's view model, which stores the operator panel state */
+	TSharedPtr<IDisplayClusterOperatorViewModel> ViewModel;
+
 	/** The command list used by the toolbar */
 	TSharedPtr<FUICommandList> CommandList;
 	
@@ -76,9 +77,6 @@ private:
 	/** The combo box widget that allows the uesr to pick the active root actor from */
 	TSharedPtr<SComboBox<TSharedPtr<FString>>> RootActorComboBox;
 
-	/** The active root actor that the operator panel is currently editing */
-	TWeakObjectPtr<ADisplayClusterRootActor> ActiveRootActor;
-
 	/** The OnLevelActorDeleted delegate handle */
 	FDelegateHandle LevelActorDeletedHandle;
 
diff --git a/Engine/Plugins/Runtime/nDisplay/Source/DisplayClusterOperator/Public/DisplayClusterOperatorStatusBarExtender.h b/Engine/Plugins/Runtime/nDisplay/Source/DisplayClusterOperator/Public/DisplayClusterOperatorStatusBarExtender.h
new file mode 100644
--- /dev/null
+++ b/Engine/Plugins/Runtime/nDisplay/Source/DisplayClusterOperator/Public/DisplayClusterOperatorStatusBarExtender.h
@@ -0,0 +1,27 @@
+// Copyright Epic Games, Inc. All Rights Reserved.
+
+#pragma once
+
+#include "CoreMinimal.h"
+#include "WidgetDrawerConfig.h"
+
+class SDisplayClusterOperatorStatusBar;
+
+/** An extender that can be passed to external modules to allow them to extend the operator panel's status bar */
+class FDisplayClusterOperatorStatusBarExtender
+{
+public:
+	/**
+	 * Adds a new widget drawer to the status bar
+	 * 
+	 * @param InWidgetDrawerConfig - The widget drawer config to create the widget drawer from
+	 * @param SlotIndex - An optional slot index of the widget drawer to add, which determines where the button for the widget drawer is placed on the status bar
+	 */
+	void DISPLAYCLUSTEROPERATOR_API AddWidgetDrawer(const FWidgetDrawerConfig& InWidgetDrawerConfig, int32 SlotIndex = INDEX_NONE);
+
+	/** Registers any status bar extensions with the specified status bar */
+	void RegisterExtensions(TSharedRef<SDisplayClusterOperatorStatusBar> StatusBar);
+
+private:
+	TArray<TPair<int32, FWidgetDrawerConfig>> WidgetDrawers;
+};
\ No newline at end of file
diff --git a/Engine/Plugins/Runtime/nDisplay/Source/DisplayClusterOperator/Public/IDisplayClusterOperator.h b/Engine/Plugins/Runtime/nDisplay/Source/DisplayClusterOperator/Public/IDisplayClusterOperator.h
--- a/Engine/Plugins/Runtime/nDisplay/Source/DisplayClusterOperator/Public/IDisplayClusterOperator.h
+++ b/Engine/Plugins/Runtime/nDisplay/Source/DisplayClusterOperator/Public/IDisplayClusterOperator.h
@@ -6,12 +6,14 @@
 #include "Modules/ModuleInterface.h"
 #include "Modules/ModuleManager.h"
 
+class IDisplayClusterOperatorViewModel;
 class FLayoutExtender;
 class FExtensibilityManager;
+class FDisplayClusterOperatorStatusBarExtender;
 class ADisplayClusterRootActor;
 
 DECLARE_EVENT_OneParam(IDisplayClusterOperator, FOnRegisterLayoutExtensions, FLayoutExtender&);
-DECLARE_EVENT_OneParam(IDisplayClusterOperator, FOnActiveRootActorChanged, ADisplayClusterRootActor*);
+DECLARE_EVENT_OneParam(IDisplayClusterOperator, FOnRegisterStatusBarExtensions, FDisplayClusterOperatorStatusBarExtender&);
 DECLARE_EVENT_OneParam(IDisplayClusterOperator, FOnDetailObjectsChanged, const TArray<UObject*>&);
 
 /**
@@ -46,17 +48,23 @@ public:
 		return FModuleManager::Get().IsModuleLoaded(ModuleName);
 	}
 
+	/** Gets the operator panel's view model, which stores any state used by the operator panel */
+	virtual TSharedRef<IDisplayClusterOperatorViewModel> GetOperatorViewModel() = 0;
+
 	/** Gets the event handler that is raised when the operator panel processes extensions to its layout */
 	virtual FOnRegisterLayoutExtensions& OnRegisterLayoutExtensions() = 0;
 
-	/** Gets the event handler that is raised when the operator panel changes the root actor being operated on */
-	virtual FOnActiveRootActorChanged& OnActiveRootActorChanged() = 0;
+	/** Gets the event handler that is raised when the operator panel processes extensions to its status bar */
+	virtual FOnRegisterStatusBarExtensions& OnRegisterStatusBarExtensions() = 0;
 
 	/** Gets the event handler that is raised when the objects being displayed in the operator's details panel are changed */
 	virtual FOnDetailObjectsChanged& OnDetailObjectsChanged() = 0;
 
-	/** Gets the extension ID that can be used to add tabs to the operator panel */
-	virtual FName GetOperatorExtensionId() = 0;
+	/** Gets the extension ID for the main window region that can be used to add tabs to the operator panel */
+	virtual FName GetPrimaryOperatorExtensionId() = 0;
+
+	/** Gets the extension ID for the auxilliary window region intended for lower-thirds windows (e.g log output) that can be used to add tabs to the operator panel */
+	virtual FName GetAuxilliaryOperatorExtensionId() = 0;
 
 	/** Gets the extensibility manager for the operator panel's toolbar */
 	virtual TSharedPtr<FExtensibilityManager> GetOperatorToolBarExtensibilityManager() = 0;
@@ -69,4 +77,7 @@ public:
 
 	/** Displays the properties of the specified object in the operator's details panel */
 	virtual void ShowDetailsForObjects(const TArray<UObject*>& Objects) = 0;
+
+	/** Forces the operator panel to dismiss any open drawers */
+	virtual void ForceDismissDrawers() = 0;
 };
\ No newline at end of file
diff --git a/Engine/Plugins/Runtime/nDisplay/Source/DisplayClusterOperator/Public/IDisplayClusterOperatorViewModel.h b/Engine/Plugins/Runtime/nDisplay/Source/DisplayClusterOperator/Public/IDisplayClusterOperatorViewModel.h
new file mode 100644
--- /dev/null
+++ b/Engine/Plugins/Runtime/nDisplay/Source/DisplayClusterOperator/Public/IDisplayClusterOperatorViewModel.h
@@ -0,0 +1,30 @@
+// Copyright Epic Games, Inc. All Rights Reserved.
+
+#pragma once
+
+#include "CoreMinimal.h"
+
+class ADisplayClusterRootActor;
+class FTabManager;
+
+DECLARE_EVENT_OneParam(IDisplayClusterOperatorViewModel, FOnActiveRootActorChanged, ADisplayClusterRootActor*);
+
+/** Interface for a view model object that stores any state from the operator panel that should be exposed externally */
+class DISPLAYCLUSTEROPERATOR_API IDisplayClusterOperatorViewModel
+{
+public:
+	/** Gets whether the view model has been populated with a valid root actor */
+	virtual bool HasRootActor() const = 0;
+
+	/** Gets the root actor that is actively being edited by the operator panel */
+	virtual ADisplayClusterRootActor* GetRootActor() const = 0;
+
+	/** Sets the root actor that is actively being edited by the operator panel */
+	virtual void SetRootActor(ADisplayClusterRootActor* InRootActor) = 0;
+
+	/** Gets the event handler that is raised when the operator panel changes the root actor being operated on */
+	virtual FOnActiveRootActorChanged& OnActiveRootActorChanged() = 0;
+
+	/** Gets the tab manager of the active operator panel, if there is an open operator panel */
+	virtual TSharedPtr<FTabManager> GetTabManager() const = 0;
+};
\ No newline at end of file
diff --git a/Engine/Plugins/Runtime/nDisplay/nDisplay.uplugin b/Engine/Plugins/Runtime/nDisplay/nDisplay.uplugin
--- a/Engine/Plugins/Runtime/nDisplay/nDisplay.uplugin
+++ b/Engine/Plugins/Runtime/nDisplay/nDisplay.uplugin
@@ -117,6 +117,15 @@
 				"Linux"
 			]
 		},
+		{
+			"Name": "DisplayClusterColorGrading",
+			"Type": "Editor",
+			"LoadingPhase": "Default",
+			"PlatformAllowList": [
+				"Win64",
+				"Linux"
+			]
+		},
 		{
 			"Name": "DisplayClusterStageMonitoring",
 			"Type": "Runtime",
@@ -171,6 +180,10 @@
 		}
 	],
 	"Plugins": [
+		{
+			"Name": "ColorCorrectRegions",
+			"Enabled": true
+		},
 		{
 			"Name": "Composure",
 			"Enabled": true
