nDisplay: Exposure limited to -100 to 100. Also fixed the order of operations when scaling the colors.

This avoids the emissive color that the material gets from going over the float16 limits.

diff --git a/Engine/Plugins/Runtime/nDisplay/Source/DisplayCluster/Public/DisplayClusterLightCardActor.h b/Engine/Plugins/Runtime/nDisplay/Source/DisplayCluster/Public/DisplayClusterLightCardActor.h
--- a/Engine/Plugins/Runtime/nDisplay/Source/DisplayCluster/Public/DisplayClusterLightCardActor.h
+++ b/Engine/Plugins/Runtime/nDisplay/Source/DisplayCluster/Public/DisplayClusterLightCardActor.h
@@ -181,7 +181,7 @@ public:
 	float Tint;
 
 	/** 2^Exposure color value multiplier */
-	UPROPERTY(EditAnywhere, BlueprintReadWrite, Category = "Appearance")
+	UPROPERTY(EditAnywhere, BlueprintReadWrite, Category = "Appearance", meta = (UIMin = -100, ClampMin = -100, UIMax = 100, ClampMax = 100))
 	float Exposure;
 
 	/** Linear color value multiplier */
