nDisplay: It now indicates the shape of the polygonal lightcard as you draw it.

diff --git a/Engine/Plugins/Runtime/nDisplay/Source/DisplayClusterLightCardEditor/Private/Viewport/DisplayClusterLightCardEditorViewportClient.cpp b/Engine/Plugins/Runtime/nDisplay/Source/DisplayClusterLightCardEditor/Private/Viewport/DisplayClusterLightCardEditorViewportClient.cpp
--- a/Engine/Plugins/Runtime/nDisplay/Source/DisplayClusterLightCardEditor/Private/Viewport/DisplayClusterLightCardEditorViewportClient.cpp
+++ b/Engine/Plugins/Runtime/nDisplay/Source/DisplayClusterLightCardEditor/Private/Viewport/DisplayClusterLightCardEditorViewportClient.cpp
@@ -52,15 +52,6 @@
 #include "UnrealWidget.h"
 #include "Widgets/Docking/SDockTab.h"
 
-#if WITH_OPENCV
-
-#include "OpenCVHelper.h"
-
-#include "PreOpenCVHeaders.h"
-#include "opencv2/imgproc.hpp"
-#include "PostOpenCVHeaders.h"
-
-#endif //WITH_OPENCV
 
 #define LOCTEXT_NAMESPACE "DisplayClusterLightCardEditorViewportClient"
 
@@ -437,6 +428,52 @@ void FDisplayClusterLightCardEditorViewportClient::Draw(const FSceneView* View,
 		EditorWidget->SetProjectionTransform(FDisplayClusterMeshProjectionTransform(ProjectionMode, View->ViewMatrices.GetViewMatrix()));
 		EditorWidget->Draw(View, this, PDI);
 	}
+
+	// Draw polygonal lightcard in progress
+	if (DrawnMousePositions.Num())
+	{
+		const FDisplayClusterMeshProjectionTransform ProjectionTransform(ProjectionMode, View->ViewMatrices.GetViewMatrix());
+
+		const FLinearColor LineColor = FLinearColor::Green;
+		constexpr float LineThickness = 5;
+
+		// Convenience function to draw a line between two mouse points
+		auto DrawLine = [&](const FIntPoint& MousePosA, const FIntPoint& MousePosB, const FLinearColor& LineColor)
+		{
+			FVector4 ScreenPosA = View->PixelToScreen(MousePosA.X, MousePosA.Y, 0);
+			FVector4 ScreenPosB = View->PixelToScreen(MousePosB.X, MousePosB.Y, 0);
+
+			// ScreenPos will have Z=0 here.
+
+			FVector PointA;
+			FVector PointB;
+
+			if (RenderViewportType == LVT_Perspective)
+			{
+				ScreenPosA.Z = 1;
+				ScreenPosB.Z = 1;
+
+				PointA = GNearClippingPlane * View->ScreenToWorld(ScreenPosA);
+				PointB = GNearClippingPlane * View->ScreenToWorld(ScreenPosB);
+			}
+			else
+			{
+				PointA = View->ScreenToWorld(ScreenPosA);
+				PointB = View->ScreenToWorld(ScreenPosB);
+			}
+
+			PDI->DrawLine(PointA, PointB, LineColor, ESceneDepthPriorityGroup::SDPG_Foreground, LineThickness, 0.0, true /* bScreenSpace */);
+		};
+
+		// Draw committed lines
+		for (int32 PosIdx = 0; PosIdx < DrawnMousePositions.Num() - 1; PosIdx++)
+		{
+			DrawLine(DrawnMousePositions[PosIdx], DrawnMousePositions[PosIdx + 1], LineColor);
+		}
+
+		// Draw line to current mouse position, slightly darker as a visual queue that it is not committed yet.
+		DrawLine(DrawnMousePositions[DrawnMousePositions.Num() - 1], FIntPoint(GetCachedMouseX(), GetCachedMouseY()), LineColor/2);
+	}
 }
 
 FSceneView* FDisplayClusterLightCardEditorViewportClient::CalcSceneView(FSceneViewFamily* ViewFamily, const int32 StereoViewIndex)
@@ -484,6 +521,7 @@ bool FDisplayClusterLightCardEditorViewportClient::InputKey(const FInputKeyEvent
 
 			if (!DrawnMousePositions.Num() || (DrawnMousePositions.Last() != MousePos))
 			{
+				// Append to array of points
 				DrawnMousePositions.Add(MousePos);
 			}
 		}
@@ -644,8 +682,7 @@ void FDisplayClusterLightCardEditorViewportClient::TrackingStopped()
 
 void FDisplayClusterLightCardEditorViewportClient::CreateDrawnLightCard(const TArray<FIntPoint>& MousePositions)
 {
-#if WITH_OPENCV
-	if (DrawnMousePositions.Num() < 3 || !Viewport)
+	if (MousePositions.Num() < 3 || !Viewport)
 	{
 		return;
 	}
@@ -898,8 +935,6 @@ void FDisplayClusterLightCardEditorViewportClient::CreateDrawnLightCard(const TA
 		const FDisplayClusterLightCardEditorHelper::FSphericalCoordinates LightCardCoords(LightCardLocation - ViewOrigin);
 		MoveLightCardTo(*LightCard, LightCardLocation - ViewOrigin);
 	}
-
-#endif // WITH_OPENCV
 }
 
 double FDisplayClusterLightCardEditorViewportClient::CalculateFinalLightCardDistance(double FlushDistance, double DesiredOffsetFromFlush) const
