nDisplay: Added a margin so that the LC polygon lines are always beyond the near plane.

diff --git a/Engine/Plugins/Runtime/nDisplay/Source/DisplayClusterLightCardEditor/Private/Viewport/DisplayClusterLightCardEditorViewportClient.cpp b/Engine/Plugins/Runtime/nDisplay/Source/DisplayClusterLightCardEditor/Private/Viewport/DisplayClusterLightCardEditorViewportClient.cpp
--- a/Engine/Plugins/Runtime/nDisplay/Source/DisplayClusterLightCardEditor/Private/Viewport/DisplayClusterLightCardEditorViewportClient.cpp
+++ b/Engine/Plugins/Runtime/nDisplay/Source/DisplayClusterLightCardEditor/Private/Viewport/DisplayClusterLightCardEditorViewportClient.cpp
@@ -443,25 +443,30 @@ void FDisplayClusterLightCardEditorViewportClient::Draw(const FSceneView* View,
 			FVector4 ScreenPosA = View->PixelToScreen(MousePosA.X, MousePosA.Y, 0);
 			FVector4 ScreenPosB = View->PixelToScreen(MousePosB.X, MousePosB.Y, 0);
 
-			// ScreenPos will have Z=0 here.
+			// ScreenPos will have Z=0 here, but we need Z=1 to be in the near plane.
+			ScreenPosA.Z = 1;
+			ScreenPosB.Z = 1;
 
-			FVector PointA;
-			FVector PointB;
+			FVector4 ViewPosA = View->ViewMatrices.GetInvProjectionMatrix().TransformFVector4(ScreenPosA);
+			FVector4 ViewPosB = View->ViewMatrices.GetInvProjectionMatrix().TransformFVector4(ScreenPosB);
+
+			constexpr float NearPlaneSafetyFactor = 1.01;
 
 			if (RenderViewportType == LVT_Perspective)
 			{
-				ScreenPosA.Z = 1;
-				ScreenPosB.Z = 1;
-
-				PointA = GNearClippingPlane * View->ScreenToWorld(ScreenPosA);
-				PointB = GNearClippingPlane * View->ScreenToWorld(ScreenPosB);
+				ViewPosA.W /= NearPlaneSafetyFactor;
 			}
 			else
 			{
-				PointA = View->ScreenToWorld(ScreenPosA);
-				PointB = View->ScreenToWorld(ScreenPosB);
+				ViewPosA.Z = GNearClippingPlane * NearPlaneSafetyFactor;
 			}
 
+			const FVector4 WorldPosA = View->ViewMatrices.GetInvViewMatrix().TransformFVector4(ViewPosA);
+			const FVector4 WorldPosB = View->ViewMatrices.GetInvViewMatrix().TransformFVector4(ViewPosB);
+
+			const FVector PointA = WorldPosA / WorldPosA.W;
+			const FVector PointB = WorldPosB / WorldPosB.W;
+
 			PDI->DrawLine(PointA, PointB, LineColor, ESceneDepthPriorityGroup::SDPG_Foreground, LineThickness, 0.0, true /* bScreenSpace */);
 		};
 
