nDisplay: Fix Pixel to World for orthographic projections.

PixelToWorld was always applying the root actor transform to the view matrix, but this is not correct for UV map projection.

Also fixed the RC websocket handler using hte origin component transform instead of using the rotation of its owner actor and the location of the origin component. This introduces unwanted rotations when the origin component has local rotations.

Also noticed that LC proxies don't show up in the right place when the root actor has rotations in x or y (which is unusual). This is not fixed yet.

Also noticed that Epic Stage is not placing the LCs correctly when using orthographic projection (before and after this CL). This is not fixed yet.

Note: Preflights were down.

diff --git a/Engine/Plugins/Runtime/nDisplay/Source/DisplayClusterScenePreview/Private/DisplayClusterLightCardEditorHelper.cpp b/Engine/Plugins/Runtime/nDisplay/Source/DisplayClusterScenePreview/Private/DisplayClusterLightCardEditorHelper.cpp
--- a/Engine/Plugins/Runtime/nDisplay/Source/DisplayClusterScenePreview/Private/DisplayClusterLightCardEditorHelper.cpp
+++ b/Engine/Plugins/Runtime/nDisplay/Source/DisplayClusterScenePreview/Private/DisplayClusterLightCardEditorHelper.cpp
@@ -653,18 +653,34 @@ void FDisplayClusterLightCardEditorHelper::PixelToWorld(const FSceneView& View,
 	const FMatrix& InvProjMatrix = View.ViewMatrices.GetInvProjectionMatrix();
 	FMatrix InvViewMatrix = View.ViewMatrices.GetInvViewMatrix();
 
+	// Cancel out the root actor's orientation
 	if (USceneComponent* OriginComponent = ProjectionOriginComponent.Get())
 	{
-		// Cancel out the root actor's orientation
-		InvViewMatrix *= OriginComponent->GetComponentTransform().ToInverseMatrixWithScale();
+		if (const AActor* Actor = OriginComponent->GetOwner())
+		{
+			if (!bIsOrthographic)
+			{
+				FTransform Transform;
+
+				Transform.SetRotation(Actor->GetActorRotation().Quaternion());
+				Transform.SetTranslation(OriginComponent->GetComponentLocation());
+
+				InvViewMatrix *= Transform.ToInverseMatrixWithScale();
+			}
+		}
 	}
 
-	const FVector4 ScreenPos = View.PixelToScreen(PixelPos.X, PixelPos.Y, 0);
-	const FVector4 HomogeneousPos = bIsOrthographic
-		? FVector4(ScreenPos.X, ScreenPos.Y, 1.0f, 1.0f)
-		: FVector4(ScreenPos.X * GNearClippingPlane, ScreenPos.Y * GNearClippingPlane, 0.0f, GNearClippingPlane);
+	FVector4 ScreenPos = View.PixelToScreen(PixelPos.X, PixelPos.Y, 0);
+	ScreenPos.Z = 1; // Force near clip plane
+
+	FVector4 ViewPos = FVector(InvProjMatrix.TransformFVector4(ScreenPos));
+	ViewPos /= ViewPos.W;
+
+	if (bIsOrthographic)
+	{
+		ViewPos.Z = 0;
+	}
 
-	const FVector ViewPos = FVector(InvProjMatrix.TransformFVector4(HomogeneousPos));
 	const FVector UnprojectedViewPos = FDisplayClusterMeshProjectionRenderer::UnprojectViewPosition(ViewPos, ProjectionMode);
 
 	if (bIsOrthographic)
diff --git a/Engine/Plugins/VirtualProduction/EpicStageApp/Source/EpicStageApp/Private/StageAppRouteHandler.cpp b/Engine/Plugins/VirtualProduction/EpicStageApp/Source/EpicStageApp/Private/StageAppRouteHandler.cpp
--- a/Engine/Plugins/VirtualProduction/EpicStageApp/Source/EpicStageApp/Private/StageAppRouteHandler.cpp
+++ b/Engine/Plugins/VirtualProduction/EpicStageApp/Source/EpicStageApp/Private/StageAppRouteHandler.cpp
@@ -133,7 +133,7 @@ bool FStageAppRouteHandler::FPerRendererData::GetSceneViewInitOptions(FSceneView
 		PreviewSettings.FOV,
 		PreviewSettings.Resolution,
 		ViewOriginComponent->GetComponentLocation(),
-		FRotator(ViewOriginComponent->GetComponentRotation().Quaternion() * PreviewSettings.Rotation.Quaternion())
+		FRotator(RootActor->GetActorRotation().Quaternion() * PreviewSettings.Rotation.Quaternion())
 	);
 
 	return true;
