nDisplay: Fixed LC proxies being in the wrong location when NDC is rotated in x or y.

This was also causing jumps when the LC proxy was dragged, or when moved from Epic Stage.

The issue was that the transform was being propagated back and forth incorrectly.

Note: Preflights were down.

diff --git a/Engine/Plugins/Runtime/nDisplay/Source/DisplayClusterLightCardEditor/Private/Viewport/DisplayClusterLightCardEditorViewportClient.cpp b/Engine/Plugins/Runtime/nDisplay/Source/DisplayClusterLightCardEditor/Private/Viewport/DisplayClusterLightCardEditorViewportClient.cpp
--- a/Engine/Plugins/Runtime/nDisplay/Source/DisplayClusterLightCardEditor/Private/Viewport/DisplayClusterLightCardEditorViewportClient.cpp
+++ b/Engine/Plugins/Runtime/nDisplay/Source/DisplayClusterLightCardEditor/Private/Viewport/DisplayClusterLightCardEditorViewportClient.cpp
@@ -1276,6 +1276,8 @@ void FDisplayClusterLightCardEditorViewportClient::UpdatePreviewActor(ADisplayCl
 				
 				SelectLightCard(nullptr);
 				
+				const FTransform RALevelTransformNoScale(RootActorLevelInstance->GetActorRotation(), RootActorLevelInstance->GetActorLocation(), FVector::OneVector);
+
 				for (ADisplayClusterLightCardActor* LightCard : LightCards)
 				{
 					FObjectDuplicationParameters DupeActorParameters(LightCard, PreviewWorld->GetCurrentLevel());
@@ -1284,9 +1286,10 @@ void FDisplayClusterLightCardEditorViewportClient::UpdatePreviewActor(ADisplayCl
 				
 					ADisplayClusterLightCardActor* LightCardProxy = CastChecked<ADisplayClusterLightCardActor>(StaticDuplicateObjectEx(DupeActorParameters));
 					PreviewWorld->GetCurrentLevel()->AddLoadedActor(LightCardProxy);
-				
-					LightCardProxy->SetActorLocation(LightCard->GetActorLocation() - RootActorPtr->GetActorLocation());
-					LightCardProxy->SetActorRotation(LightCard->GetActorRotation() - RootActorPtr->GetActorRotation());
+
+					const FTransform LCLevelRelativeToRALevel = LightCard->GetTransform().GetRelativeTransform(RALevelTransformNoScale);
+
+					LightCardProxy->SetActorTransform(LCLevelRelativeToRALevel);
 					LightCardProxy->PolygonMask = LightCard->PolygonMask;
 					LightCardProxy->bIsProxy = true;
 
@@ -1356,9 +1359,10 @@ void FDisplayClusterLightCardEditorViewportClient::UpdateProxyTransforms()
 		{
 			if (LightCardProxy.LevelInstance.IsValid() && LightCardProxy.Proxy.IsValid())
 			{
-				LightCardProxy.Proxy->SetActorLocation(LightCardProxy.LevelInstance->GetActorLocation() - RootActorLevelInstance->GetActorLocation());
-				LightCardProxy.Proxy->SetActorRotation(LightCardProxy.LevelInstance->GetActorRotation() - RootActorLevelInstance->GetActorRotation());
-				LightCardProxy.Proxy->SetActorScale3D(LightCardProxy.LevelInstance->GetActorScale3D());
+				const FTransform RALevelTransformNoScale(RootActorLevelInstance->GetActorRotation(), RootActorLevelInstance->GetActorLocation(), FVector::OneVector);
+				const FTransform LCLevelRelativeToRALevel = LightCardProxy.LevelInstance->GetTransform().GetRelativeTransform(RALevelTransformNoScale);
+
+				LightCardProxy.Proxy->SetActorTransform(LCLevelRelativeToRALevel);
 
 				// Need to update these manually or the proxy's position will be out of sync next update
 				LightCardProxy.Proxy->Longitude = LightCardProxy.LevelInstance->Longitude;
@@ -1369,6 +1373,7 @@ void FDisplayClusterLightCardEditorViewportClient::UpdateProxyTransforms()
 				LightCardProxy.Proxy->Yaw = LightCardProxy.LevelInstance->Yaw;
 				LightCardProxy.Proxy->Scale = LightCardProxy.LevelInstance->Scale;
 				LightCardProxy.Proxy->UVCoordinates = LightCardProxy.LevelInstance->UVCoordinates;
+				LightCardProxy.Proxy->RadialOffset = LightCardProxy.LevelInstance->RadialOffset;
 			}
 		}
 	}
@@ -1819,11 +1824,12 @@ void FDisplayClusterLightCardEditorViewportClient::PropagateLightCardTransform(A
 		// If that ever changes then the math below will need to be updated to use the 
 		// relative loc/rot of the LC proxies wrt the root actor proxy.
 
-		const FVector RootActorLevelInstanceLocation = RootActorLevelInstance.IsValid() ? RootActorLevelInstance->GetActorLocation() : FVector::ZeroVector;
-		LevelInstance->SetActorLocation(RootActorLevelInstanceLocation + LightCardProxy->GetActorLocation());
+		const FRotator RALevelRotation = RootActorLevelInstance.IsValid() ? RootActorLevelInstance->GetActorRotation() : FRotator::ZeroRotator;
+		const FVector  RALevelLocation = RootActorLevelInstance.IsValid() ? RootActorLevelInstance->GetActorLocation() : FVector::ZeroVector;
+
+		const FTransform RALevelTransformNoScale(RALevelRotation, RALevelLocation, FVector::OneVector);
 
-		const FRotator RootActorLevelInstanceRotation = RootActorLevelInstance.IsValid() ? RootActorLevelInstance->GetActorRotation() : FRotator::ZeroRotator;
-		LevelInstance->SetActorRotation(RootActorLevelInstanceRotation.Quaternion() * LightCardProxy->GetActorRotation().Quaternion());
+		LevelInstance->SetActorTransform(LightCardProxy->GetTransform() * RALevelTransformNoScale);
 
 		TryChangeProperty(GET_MEMBER_NAME_CHECKED(ADisplayClusterLightCardActor, Longitude));
 		TryChangeProperty(GET_MEMBER_NAME_CHECKED(ADisplayClusterLightCardActor, Latitude));
@@ -1833,6 +1839,7 @@ void FDisplayClusterLightCardEditorViewportClient::PropagateLightCardTransform(A
 		TryChangeProperty(GET_MEMBER_NAME_CHECKED(ADisplayClusterLightCardActor, Yaw));
 		TryChangeProperty(GET_MEMBER_NAME_CHECKED(ADisplayClusterLightCardActor, Scale));
 		TryChangeProperty(GET_MEMBER_NAME_CHECKED(ADisplayClusterLightCardActor, UVCoordinates));
+		TryChangeProperty(GET_MEMBER_NAME_CHECKED(ADisplayClusterLightCardActor, RadialOffset));
 		
 		// Snapshot the changed properties so multi-user can update while dragging.
 		if (ChangedProperties.Num() > 0)
