nDisplay: Turned off volumetric fog light cards and chromakey.

diff --git a/Engine/Plugins/Runtime/nDisplay/Source/DisplayCluster/Private/Render/Viewport/DisplayClusterViewportManager.cpp b/Engine/Plugins/Runtime/nDisplay/Source/DisplayCluster/Private/Render/Viewport/DisplayClusterViewportManager.cpp
--- a/Engine/Plugins/Runtime/nDisplay/Source/DisplayCluster/Private/Render/Viewport/DisplayClusterViewportManager.cpp
+++ b/Engine/Plugins/Runtime/nDisplay/Source/DisplayCluster/Private/Render/Viewport/DisplayClusterViewportManager.cpp
@@ -622,6 +622,7 @@ FSceneViewFamily::ConstructionValues FDisplayClusterViewportManager::CreateViewF
 
 			InEngineShowFlags.SetAtmosphere(0);
 			InEngineShowFlags.SetFog(0);
+			InEngineShowFlags.SetVolumetricFog(0);
 			InEngineShowFlags.SetMotionBlur(0); // motion blur doesn't work correctly with scene captures.
 			InEngineShowFlags.SetSeparateTranslucency(0);
 			InEngineShowFlags.SetHMDDistortion(0);
