nDisplay: Add scene outliner to light card editor replacing the old light card list. "Add Flag" is now available under place new light card actors.

diff --git a/Engine/Plugins/Runtime/nDisplay/Source/DisplayCluster/Private/Game/EngineClasses/Scene/DisplayClusterLightCardActor.cpp b/Engine/Plugins/Runtime/nDisplay/Source/DisplayCluster/Private/Game/EngineClasses/Scene/DisplayClusterLightCardActor.cpp
--- a/Engine/Plugins/Runtime/nDisplay/Source/DisplayCluster/Private/Game/EngineClasses/Scene/DisplayClusterLightCardActor.cpp
+++ b/Engine/Plugins/Runtime/nDisplay/Source/DisplayCluster/Private/Game/EngineClasses/Scene/DisplayClusterLightCardActor.cpp
@@ -174,6 +174,16 @@ void ADisplayClusterLightCardActor::PostEditChangeProperty(FPropertyChangedEvent
 	}
 }
 
+FName ADisplayClusterLightCardActor::GetCustomIconName() const
+{
+	if (bIsUVLightCard)
+	{
+		return TEXT("UVLightCard");
+	}
+
+	return Super::GetCustomIconName();
+}
+
 #endif
 
 FTransform ADisplayClusterLightCardActor::GetLightCardTransform(bool bIgnoreSpinYawPitch) const
diff --git a/Engine/Plugins/Runtime/nDisplay/Source/DisplayCluster/Public/DisplayClusterLightCardActor.h b/Engine/Plugins/Runtime/nDisplay/Source/DisplayCluster/Public/DisplayClusterLightCardActor.h
--- a/Engine/Plugins/Runtime/nDisplay/Source/DisplayCluster/Public/DisplayClusterLightCardActor.h
+++ b/Engine/Plugins/Runtime/nDisplay/Source/DisplayCluster/Public/DisplayClusterLightCardActor.h
@@ -45,7 +45,7 @@ struct FLightCardAlphaGradientSettings
 	float Angle = 0;
 };
 
-UCLASS(Blueprintable)
+UCLASS(Blueprintable, DisplayName = "Light Card")
 class DISPLAYCLUSTER_API ADisplayClusterLightCardActor : public AActor
 {
 	GENERATED_BODY()
@@ -82,6 +82,7 @@ public:
 #if WITH_EDITOR
 	virtual void PostEditChangeProperty(FPropertyChangedEvent& PropertyChangedEvent) override;
 	virtual bool ShouldTickIfViewportsOnly() const override { return true; }
+	virtual FName GetCustomIconName() const override;
 #endif
 
 	/**
diff --git a/Engine/Plugins/Runtime/nDisplay/Source/DisplayClusterLightCardEditor/DisplayClusterLightCardEditor.Build.cs b/Engine/Plugins/Runtime/nDisplay/Source/DisplayClusterLightCardEditor/DisplayClusterLightCardEditor.Build.cs
--- a/Engine/Plugins/Runtime/nDisplay/Source/DisplayClusterLightCardEditor/DisplayClusterLightCardEditor.Build.cs
+++ b/Engine/Plugins/Runtime/nDisplay/Source/DisplayClusterLightCardEditor/DisplayClusterLightCardEditor.Build.cs
@@ -35,6 +35,8 @@ public class DisplayClusterLightCardEditor : ModuleRules
 				"UnrealEd",
 				"WorkspaceMenuStructure",
 				"ToolWidgets",
+				"ToolMenus",
+				"SceneOutliner"
 			});
 	}
 }
diff --git a/Engine/Plugins/Runtime/nDisplay/Source/DisplayClusterLightCardEditor/Private/DisplayClusterLightCardEditorCommands.cpp b/Engine/Plugins/Runtime/nDisplay/Source/DisplayClusterLightCardEditor/Private/DisplayClusterLightCardEditorCommands.cpp
--- a/Engine/Plugins/Runtime/nDisplay/Source/DisplayClusterLightCardEditor/Private/DisplayClusterLightCardEditorCommands.cpp
+++ b/Engine/Plugins/Runtime/nDisplay/Source/DisplayClusterLightCardEditor/Private/DisplayClusterLightCardEditorCommands.cpp
@@ -22,6 +22,7 @@ void FDisplayClusterLightCardEditorCommands::RegisterCommands()
 	UI_COMMAND(ViewOrientationBack, "Back", "Orient the view to look at the back of the stage", EUserInterfaceActionType::Button, FInputChord());
 
 	UI_COMMAND(AddNewLightCard, "Add New Light Card", "Add and assign a new Light Card to the actor", EUserInterfaceActionType::Button, FInputChord());
+	UI_COMMAND(AddNewFlag, "Add Flag", "Add and assign a new light control flag", EUserInterfaceActionType::Button, FInputChord());
 	UI_COMMAND(AddExistingLightCard, "Add Existing Light Card", "Add an existing Light Card to the actor", EUserInterfaceActionType::Button, FInputChord());
 	UI_COMMAND(RemoveLightCard, "Remove from Actor", "Remove the Light Card from the actor but do not delete it", EUserInterfaceActionType::Button, FInputChord());
 	UI_COMMAND(PasteHere, "Paste Here", "Paste clipboard contents at the click location", EUserInterfaceActionType::Button, FInputChord());
diff --git a/Engine/Plugins/Runtime/nDisplay/Source/DisplayClusterLightCardEditor/Private/DisplayClusterLightCardEditorCommands.h b/Engine/Plugins/Runtime/nDisplay/Source/DisplayClusterLightCardEditor/Private/DisplayClusterLightCardEditorCommands.h
--- a/Engine/Plugins/Runtime/nDisplay/Source/DisplayClusterLightCardEditor/Private/DisplayClusterLightCardEditorCommands.h
+++ b/Engine/Plugins/Runtime/nDisplay/Source/DisplayClusterLightCardEditor/Private/DisplayClusterLightCardEditorCommands.h
@@ -38,6 +38,7 @@ public:
 	TSharedPtr<FUICommandInfo> ViewOrientationBack;
 
 	TSharedPtr<FUICommandInfo> AddNewLightCard;
+	TSharedPtr<FUICommandInfo> AddNewFlag;
 	TSharedPtr<FUICommandInfo> AddExistingLightCard;
 	TSharedPtr<FUICommandInfo> RemoveLightCard;
 	TSharedPtr<FUICommandInfo> PasteHere;
diff --git a/Engine/Plugins/Runtime/nDisplay/Source/DisplayClusterLightCardEditor/Private/Outliner/DisplayClusterLightCardOutlinerMode.cpp b/Engine/Plugins/Runtime/nDisplay/Source/DisplayClusterLightCardEditor/Private/Outliner/DisplayClusterLightCardOutlinerMode.cpp
new file mode 100644
--- /dev/null
+++ b/Engine/Plugins/Runtime/nDisplay/Source/DisplayClusterLightCardEditor/Private/Outliner/DisplayClusterLightCardOutlinerMode.cpp
@@ -0,0 +1,1147 @@
+// Copyright Epic Games, Inc. All Rights Reserved.
+
+#include "DisplayClusterLightCardOutlinerMode.h"
+
+#include "DisplayClusterLightCardEditorLog.h"
+#include "SDisplayClusterLightCardOutliner.h"
+
+#include "ActorEditorUtils.h"
+#include "ActorFolderPickingMode.h"
+#include "ActorTreeItem.h"
+#include "ComponentTreeItem.h"
+#include "Editor.h"
+#include "EditorActorFolders.h"
+#include "EditorLevelUtils.h"
+#include "EditorViewportCommands.h"
+#include "FolderTreeItem.h"
+#include "LevelTreeItem.h"
+#include "SceneOutlinerMenuContext.h"
+#include "ScopedTransaction.h"
+#include "Selection.h"
+#include "ToolMenus.h"
+#include "WorldTreeItem.h"
+#include "DragAndDrop/ActorDragDropGraphEdOp.h"
+#include "DragAndDrop/FolderDragDropOp.h"
+#include "Editor/SceneOutliner/Private/SSocketChooser.h"
+#include "HAL/PlatformApplicationMisc.h"
+#include "LevelInstance/LevelInstanceInterface.h"
+#include "LevelInstance/LevelInstanceSubsystem.h"
+
+#define LOCTEXT_NAMESPACE "DisplayClusterLightCardOutlinerMode"
+
+namespace UE::LightCardOutlinerMode
+{
+	static const FName DefaultContextBaseMenuName("LightCardOutliner.DefaultContextMenuBase");
+	static const FName DefaultContextMenuName("LightCardOutliner.DefaultContextMenu");
+}
+
+struct FFolderPathSelector
+{
+	bool operator()(TWeakPtr<ISceneOutlinerTreeItem> Item, FFolder& DataOut) const
+	{
+		if (const FFolderTreeItem* FolderItem = Item.Pin()->CastTo<FFolderTreeItem>())
+		{
+			if (FolderItem->IsValid())
+			{
+				DataOut = FolderItem->GetFolder();
+				return true;
+			}
+		}
+		return false;
+	}
+};
+
+struct FWeakActorSelector
+{
+	bool operator()(const TWeakPtr<ISceneOutlinerTreeItem>& Item, TWeakObjectPtr<AActor>& DataOut) const
+	{
+		if (TSharedPtr<ISceneOutlinerTreeItem> ItemPtr = Item.Pin())
+		{
+			if (FActorTreeItem* ActorItem = ItemPtr->CastTo<FActorTreeItem>())
+			{
+				if (ActorItem->IsValid())
+				{
+					DataOut = ActorItem->Actor;
+					return true;
+				}
+			}
+		}
+		return false;
+	}
+};
+
+struct FActorSelector
+{
+	bool operator()(const TWeakPtr<ISceneOutlinerTreeItem>& Item, AActor*& ActorPtrOut) const
+	{
+		if (const TSharedPtr<ISceneOutlinerTreeItem> ItemPtr = Item.Pin())
+		{
+			if (const FActorTreeItem* ActorItem = ItemPtr->CastTo<FActorTreeItem>())
+			{
+				if (ActorItem->IsValid())
+				{
+					if (AActor* Actor = ActorItem->Actor.Get())
+					{
+						ActorPtrOut = Actor;
+						return true;
+					}
+				}
+			}
+			else if (const FComponentTreeItem* ComponentItem = ItemPtr->CastTo<FComponentTreeItem>())
+			{
+				if (ComponentItem->IsValid())
+				{
+					if (AActor* Actor = ComponentItem->Component->GetOwner())
+					{
+						ActorPtrOut = Actor;
+						return true;
+					}
+				}
+			}
+		}
+
+		return false;
+	}
+};
+
+FDisplayClusterLightCardOutlinerMode::FDisplayClusterLightCardOutlinerMode(SSceneOutliner* InSceneOutliner,
+	TWeakPtr<SDisplayClusterLightCardOutliner> InLightCardOutliner, TWeakObjectPtr<UWorld> InSpecifiedWorldToDisplay)
+	: FActorMode(FActorModeParams(InSceneOutliner, InSpecifiedWorldToDisplay,
+	/* bHideComponents */ true, /* bHideLevelInstanceHierarchy */ true, /* bHideUnloadedActors */ true, /* bHideEmptyFolders */ false)),
+	LightCardOutliner(InLightCardOutliner)
+{
+}
+
+void FDisplayClusterLightCardOutlinerMode::SynchronizeSelection()
+{
+	if (LightCardOutliner.IsValid())
+	{
+		// This can be called from undo/redo and the internal selection will be cleared. The cached selection
+		// of the lightcard outliner will contain the last valid selection.
+		LightCardOutliner.Pin()->RestoreCachedSelection();
+	}
+}
+
+FCreateSceneOutlinerMode FDisplayClusterLightCardOutlinerMode::CreateFolderPickerMode(
+	const FFolder::FRootObject& InRootObject) const
+{
+	auto MoveSelectionTo = [this, InRootObject](const FSceneOutlinerTreeItemRef& NewParent)
+	{
+		if (FWorldTreeItem* WorldItem = NewParent->CastTo<FWorldTreeItem>())
+		{
+			SceneOutliner->MoveSelectionTo(GetWorldDefaultRootFolder());
+		}
+		else if (FFolderTreeItem* FolderItem = NewParent->CastTo<FFolderTreeItem>())
+		{
+			SceneOutliner->MoveSelectionTo(FolderItem->GetFolder());
+		}
+		else if (FActorTreeItem* ActorItem = NewParent->CastTo<FActorTreeItem>())
+		{
+			if (FFolder::IsRootObjectValid(InRootObject))
+			{
+				SceneOutliner->MoveSelectionTo(FFolder(InRootObject));
+			}
+		}
+		else if (FLevelTreeItem* LevelItem = NewParent->CastTo<FLevelTreeItem>())
+		{
+			if (FFolder::IsRootObjectValid(InRootObject))
+			{
+				SceneOutliner->MoveSelectionTo(FFolder(InRootObject));
+			}
+		}
+	};
+
+	return FCreateSceneOutlinerMode::CreateLambda([this, MoveSelectionTo, InRootObject](SSceneOutliner* Outliner)
+	{
+		return new FActorFolderPickingMode(Outliner, FOnSceneOutlinerItemPicked::CreateLambda(MoveSelectionTo), nullptr, InRootObject);
+	});
+}
+
+TSharedPtr<SWidget> FDisplayClusterLightCardOutlinerMode::CreateContextMenu()
+{
+	RegisterContextMenu();
+
+	const FSceneOutlinerItemSelection ItemSelection(SceneOutliner->GetSelection());
+
+	USceneOutlinerMenuContext* ContextObject = NewObject<USceneOutlinerMenuContext>();
+	ContextObject->SceneOutliner = StaticCastSharedRef<SSceneOutliner>(SceneOutliner->AsShared());
+	ContextObject->bShowParentTree = SceneOutliner->GetSharedData().bShowParentTree;
+	ContextObject->NumSelectedItems = ItemSelection.Num();
+	ContextObject->NumSelectedFolders = ItemSelection.Num<FFolderTreeItem>();
+	ContextObject->NumWorldsSelected = ItemSelection.Num<FWorldTreeItem>();
+	ContextObject->bRepresentingGameWorld = false;
+	ContextObject->bRepresentingPartitionedWorld = false;
+
+	FToolMenuContext Context(ContextObject);
+
+	FName MenuName = UE::LightCardOutlinerMode::DefaultContextMenuName;
+	SceneOutliner->GetSharedData().ModifyContextMenu.ExecuteIfBound(MenuName, Context);
+
+	// Build up the menu for a selection
+	UToolMenus* ToolMenus = UToolMenus::Get();
+	UToolMenu* Menu = ToolMenus->GenerateMenu(MenuName, Context);
+
+	for (const FToolMenuSection& Section : Menu->Sections)
+	{
+		if (Section.Blocks.Num() > 0)
+		{
+			return ToolMenus->GenerateWidget(Menu);
+		}
+	}
+
+	return nullptr;
+}
+
+void FDisplayClusterLightCardOutlinerMode::OnItemDoubleClick(FSceneOutlinerTreeItemPtr Item)
+{
+	if (const FActorTreeItem* ActorItem = Item->CastTo<FActorTreeItem>())
+	{
+		AActor* Actor = ActorItem->Actor.Get();
+		check(Actor);
+
+		// Assume the user wants this selected in the level viewport
+		UpdateGlobalActorSelection();
+		
+		if (Item->CanInteract())
+		{
+			FSceneOutlinerItemSelection Selection(SceneOutliner->GetSelection());
+			if (Selection.Has<FActorTreeItem>())
+			{
+				const bool bActiveViewportOnly = false;
+				GEditor->MoveViewportCamerasToActor(Selection.GetData<AActor*>(FActorSelector()), bActiveViewportOnly);
+			}
+		}
+		else
+		{
+			const bool bActiveViewportOnly = false;
+			GEditor->MoveViewportCamerasToActor(*Actor, bActiveViewportOnly);
+		}
+	}
+}
+
+FReply FDisplayClusterLightCardOutlinerMode::OnKeyDown(const FKeyEvent& InKeyEvent)
+{
+	const FSceneOutlinerItemSelection& Selection = SceneOutliner->GetSelection();
+
+	// Rename key: Rename selected actors (not rebindable, because it doesn't make much sense to bind.)
+	if (InKeyEvent.GetKey() == EKeys::F2)
+	{
+		if (Selection.Num() == 1)
+		{
+			FSceneOutlinerTreeItemPtr ItemToRename = Selection.SelectedItems[0].Pin();
+
+			if (ItemToRename.IsValid() && CanRenameItem(*ItemToRename) && ItemToRename->CanInteract())
+			{
+				SceneOutliner->SetPendingRenameItem(ItemToRename);
+				SceneOutliner->ScrollItemIntoView(ItemToRename);
+			}
+
+			return FReply::Handled();
+		}
+	}
+
+	// F5 forces a full refresh
+	else if (InKeyEvent.GetKey() == EKeys::F5)
+	{
+		SceneOutliner->FullRefresh();
+		return FReply::Handled();
+	}
+
+	// Delete key: Delete selected actors (not rebindable, because it doesn't make much sense to bind.)
+	// Use Delete and Backspace instead of Platform_Delete because the LevelEditor default Edit Delete is bound to both
+	else if (InKeyEvent.GetKey() == EKeys::Delete || InKeyEvent.GetKey() == EKeys::BackSpace)
+	{
+		if (SceneOutliner->GetSharedData().CustomDelete.IsBound())
+		{
+			SceneOutliner->GetSharedData().CustomDelete.Execute(Selection.SelectedItems);
+		}
+		else
+		{
+			if (RepresentingWorld.IsValid())
+			{
+				// Only handle deleting folders here.
+				SceneOutliner->DeleteFoldersBegin();
+				SceneOutliner->DeleteFoldersEnd();
+			}
+		}
+
+		// Let the delete key get handled by the light card editor.
+		return FReply::Unhandled();
+	}
+
+	/* Allow the user to scroll to the current selection (and expand if needed) by pressing the key bound to
+	 * FEditorViewportCommands::Get().FocusViewportToSelection (Default: 'F')
+	 */
+	else
+	{
+		const FModifierKeysState ModifierKeys = FSlateApplication::Get().GetModifierKeys();
+		const FInputChord CheckChord( InKeyEvent.GetKey(), EModifierKey::FromBools(ModifierKeys.IsControlDown(), ModifierKeys.IsAltDown(), ModifierKeys.IsShiftDown(), ModifierKeys.IsCommandDown()) );
+
+		// Use the keyboard shortcut bound to 'Focus Viewport To Selection'
+		if (FEditorViewportCommands::Get().FocusViewportToSelection->HasActiveChord(CheckChord))
+		{
+			if (Selection.Num() == 1)
+			{
+				FSceneOutlinerTreeItemPtr ItemToFocus = Selection.SelectedItems[0].Pin();
+
+				if (ItemToFocus.IsValid())
+				{
+					SceneOutliner->ScrollItemIntoView(ItemToFocus);
+				}
+
+				// Return Unhandled here so that the level editor viewport can handle this event and focus the selected item
+				return FReply::Unhandled();
+			}
+		}
+	}
+
+	return FReply::Unhandled();
+}
+
+bool FDisplayClusterLightCardOutlinerMode::CanDelete() const
+{
+	// Handled by command in light card editor or in OnKeyDown for folders
+	return false;
+}
+
+bool FDisplayClusterLightCardOutlinerMode::CanRename() const
+{
+	const FSceneOutlinerItemSelection ItemSelection = SceneOutliner->GetSelection();
+	const uint32 NumberOfFolders = ItemSelection.Num<FFolderTreeItem>();
+	return (NumberOfFolders == 1 && NumberOfFolders == ItemSelection.Num());
+}
+
+bool FDisplayClusterLightCardOutlinerMode::CanRenameItem(const ISceneOutlinerTreeItem& Item) const
+{
+	// Can only rename actor and folder items when in actor browsing mode
+	return (Item.IsValid() && (Item.IsA<FActorTreeItem>() || Item.IsA<FFolderTreeItem>()));
+}
+
+bool FDisplayClusterLightCardOutlinerMode::CanCut() const
+{
+	// Handled by command in light card editor
+	return false;
+}
+
+bool FDisplayClusterLightCardOutlinerMode::CanCopy() const
+{
+	// Handled by command in light card editor
+	return false;
+}
+
+bool FDisplayClusterLightCardOutlinerMode::CanPaste() const
+{
+	// Handled by command in light card editor
+	return false;
+}
+
+FFolder FDisplayClusterLightCardOutlinerMode::CreateNewFolder()
+{
+	const FScopedTransaction Transaction(LOCTEXT("UndoAction_CreateFolder", "Create Folder"));
+
+	// FActorFolders calls will rely on GEditor selection for folder creation.
+	UpdateGlobalActorSelection();
+	
+	TArray<FFolder> SelectedFolders = SceneOutliner->GetSelection().GetData<FFolder>(FFolderPathSelector());
+	const FFolder NewFolderName = FActorFolders::Get().GetDefaultFolderForSelection(*RepresentingWorld, &SelectedFolders);
+	FActorFolders::Get().CreateFolderContainingSelection(*RepresentingWorld, NewFolderName);
+
+	return NewFolderName;
+}
+
+FFolder FDisplayClusterLightCardOutlinerMode::GetFolder(const FFolder& ParentPath, const FName& LeafName)
+{
+	return FActorFolders::Get().GetFolderName(*RepresentingWorld, ParentPath, LeafName);
+}
+
+bool FDisplayClusterLightCardOutlinerMode::CreateFolder(const FFolder& NewPath)
+{
+	return FActorFolders::Get().CreateFolder(*RepresentingWorld, NewPath);
+}
+
+bool FDisplayClusterLightCardOutlinerMode::ReparentItemToFolder(const FFolder& FolderPath, const FSceneOutlinerTreeItemPtr& Item)
+{
+	if (const FActorTreeItem* ActorItem = Item->CastTo<FActorTreeItem>())
+	{
+		// Make sure actor has the same root object before updating path
+		if (ActorItem->Actor->GetFolderRootObject() == FolderPath.GetRootObject())
+		{
+			ActorItem->Actor->SetFolderPath_Recursively(FolderPath.GetPath());
+			return true;
+		}
+	}
+	return false;
+}
+
+TSharedPtr<FDragDropOperation> FDisplayClusterLightCardOutlinerMode::CreateDragDropOperation(
+	const FPointerEvent& MouseEvent, const TArray<FSceneOutlinerTreeItemPtr>& InTreeItems) const
+{
+	FSceneOutlinerDragDropPayload DraggedObjects(InTreeItems);
+
+	if (DraggedObjects.Has<FActorTreeItem>() && !DraggedObjects.Has<FFolderTreeItem>())
+	{
+		return FActorDragDropGraphEdOp::New(DraggedObjects.GetData<TWeakObjectPtr<AActor>>(FWeakActorSelector()));
+	}
+
+	TSharedPtr<FSceneOutlinerDragDropOp> OutlinerOp = MakeShareable(new FSceneOutlinerDragDropOp());
+
+	if (DraggedObjects.Has<FActorTreeItem>())
+	{
+		TSharedPtr<FActorDragDropOp> ActorOperation = MakeShareable(new FActorDragDropGraphEdOp);
+		ActorOperation->Init(DraggedObjects.GetData<TWeakObjectPtr<AActor>>(FWeakActorSelector()));
+		OutlinerOp->AddSubOp(ActorOperation);
+	}
+
+	if (DraggedObjects.Has<FFolderTreeItem>())
+	{
+		FFolder::FRootObject CommonRootObject;
+		TArray<FName> DraggedFolders;
+		if (GetFolderNamesFromPayload(DraggedObjects, DraggedFolders, CommonRootObject))
+		{
+			TSharedPtr<FFolderDragDropOp> FolderOperation = MakeShareable(new FFolderDragDropOp);
+			FolderOperation->Init(DraggedFolders, RepresentingWorld.Get(), CommonRootObject);
+			OutlinerOp->AddSubOp(FolderOperation);
+		}
+	}
+	OutlinerOp->Construct();
+	return OutlinerOp;
+}
+
+bool FDisplayClusterLightCardOutlinerMode::ParseDragDrop(FSceneOutlinerDragDropPayload& OutPayload,
+	const FDragDropOperation& Operation) const
+{
+	if (Operation.IsOfType<FSceneOutlinerDragDropOp>())
+	{
+		const FSceneOutlinerDragDropOp& OutlinerOp = static_cast<const FSceneOutlinerDragDropOp&>(Operation);
+		if (const TSharedPtr<const FFolderDragDropOp>& FolderOp = OutlinerOp.GetSubOp<FFolderDragDropOp>())
+		{
+			for (const FName& Folder : FolderOp->Folders)
+			{
+				OutPayload.DraggedItems.Add(SceneOutliner->GetTreeItem(FFolder(FolderOp->RootObject, Folder)));
+			}
+		}
+		if (const TSharedPtr<const FActorDragDropOp>& ActorOp = OutlinerOp.GetSubOp<FActorDragDropOp>())
+		{
+			for (const TWeakObjectPtr<AActor>& Actor : ActorOp->Actors)
+			{
+				OutPayload.DraggedItems.Add(SceneOutliner->GetTreeItem(Actor.Get()));
+			}
+		}
+		return true;
+	}
+	else if (Operation.IsOfType<FActorDragDropOp>())
+	{
+		for (const TWeakObjectPtr<AActor>& Actor : static_cast<const FActorDragDropOp&>(Operation).Actors)
+		{
+			OutPayload.DraggedItems.Add(SceneOutliner->GetTreeItem(Actor.Get()));
+		}
+		return true;
+	}
+
+	return false;
+}
+
+FSceneOutlinerDragValidationInfo FDisplayClusterLightCardOutlinerMode::ValidateDrop(
+	const ISceneOutlinerTreeItem& DropTarget, const FSceneOutlinerDragDropPayload& Payload) const
+{
+	if (Payload.Has<FFolderTreeItem>())
+	{
+		FFolder::FRootObject TargetRootObject = DropTarget.GetRootObject();
+		FFolder::FRootObject CommonPayloadFoldersRootObject;
+		TArray<FName> PayloadFolders;
+		const bool bHasCommonRootObject = GetFolderNamesFromPayload(Payload, PayloadFolders, CommonPayloadFoldersRootObject);
+		if (!bHasCommonRootObject)
+		{
+			return FSceneOutlinerDragValidationInfo(ESceneOutlinerDropCompatibility::IncompatibleGeneric, LOCTEXT("CantMoveFoldersWithMultipleRoots", "Cannot move folders with multiple roots"));
+		}
+		else if (CommonPayloadFoldersRootObject != TargetRootObject)
+		{
+			return FSceneOutlinerDragValidationInfo(ESceneOutlinerDropCompatibility::IncompatibleGeneric, LOCTEXT("CantChangeFoldersRoot", "Cannot change folders root"));
+		}
+	}
+
+	if (const FActorTreeItem* ActorItem = DropTarget.CastTo<FActorTreeItem>())
+	{
+		const AActor* ActorTarget = ActorItem->Actor.Get();
+		if (!ActorTarget)
+		{
+			return FSceneOutlinerDragValidationInfo(ESceneOutlinerDropCompatibility::IncompatibleGeneric, FText());
+		}
+
+		const ILevelInstanceInterface* LevelInstanceTarget = Cast<ILevelInstanceInterface>(ActorTarget);
+		const ULevelInstanceSubsystem* LevelInstanceSubsystem = RepresentingWorld->GetSubsystem<ULevelInstanceSubsystem>();
+		if (LevelInstanceTarget)
+		{
+			check(LevelInstanceSubsystem);
+			if (!LevelInstanceTarget->IsEditing())
+			{
+				return FSceneOutlinerDragValidationInfo(ESceneOutlinerDropCompatibility::IncompatibleGeneric, LOCTEXT("Error_AttachToClosedLevelInstance", "Cannot attach to LevelInstance which is not being edited"));
+			}
+		}
+		else
+		{
+			if (Payload.Has<FFolderTreeItem>())
+			{
+				return FSceneOutlinerDragValidationInfo(ESceneOutlinerDropCompatibility::IncompatibleGeneric, LOCTEXT("FoldersOnActorError", "Cannot attach folders to actors"));
+			}
+
+			if (!Payload.Has<FActorTreeItem>())
+			{
+				return FSceneOutlinerDragValidationInfo(ESceneOutlinerDropCompatibility::IncompatibleGeneric, FText());
+			}
+		}
+
+		FText AttachErrorMsg;
+		bool bCanAttach = true;
+		bool bDraggedOntoAttachmentParent = true;
+		const TArray<TWeakObjectPtr<AActor>>& DragActors = Payload.GetData<TWeakObjectPtr<
+			AActor>>(FWeakActorSelector());
+		for (const TWeakObjectPtr<AActor>& DragActorPtr : DragActors)
+		{
+			AActor* DragActor = DragActorPtr.Get();
+			if (DragActor)
+			{
+				if (bCanAttach)
+				{
+					if (LevelInstanceSubsystem)
+					{
+						// Either all actors must be in a LevelInstance or none of them
+						if (const ILevelInstanceInterface* ParentLevelInstance = LevelInstanceSubsystem->GetParentLevelInstance(DragActor))
+						{
+							if (!ParentLevelInstance->IsEditing())
+							{
+								return FSceneOutlinerDragValidationInfo(ESceneOutlinerDropCompatibility::IncompatibleGeneric, LOCTEXT("Error_RemoveEditingLevelInstance", "Cannot detach from a LevelInstance which is not being edited"));
+							}
+						}
+
+						if (!LevelInstanceSubsystem->CanMoveActorToLevel(DragActor, &AttachErrorMsg))
+						{
+							bCanAttach = bDraggedOntoAttachmentParent = false;
+							break;
+						}
+					}
+
+					if (DragActor->IsChildActor())
+					{
+						AttachErrorMsg = FText::Format(LOCTEXT("Error_AttachChildActor", "Cannot move {0} as it is a child actor."), FText::FromString(DragActor->GetActorLabel()));
+						bCanAttach = bDraggedOntoAttachmentParent = false;
+						break;
+					}
+					if (!LevelInstanceTarget && !GEditor->CanParentActors(ActorTarget, DragActor, &AttachErrorMsg))
+					{
+						bCanAttach = false;
+					}
+				}
+
+				if (DragActor->GetSceneOutlinerParent() != ActorTarget)
+				{
+					bDraggedOntoAttachmentParent = false;
+				}
+			}
+		}
+
+		const FText ActorLabel = FText::FromString(ActorTarget->GetActorLabel());
+		if (bDraggedOntoAttachmentParent)
+		{
+			if (DragActors.Num() == 1)
+			{
+				return FSceneOutlinerDragValidationInfo(ESceneOutlinerDropCompatibility::CompatibleDetach, ActorLabel);
+			}
+			else
+			{
+				return FSceneOutlinerDragValidationInfo(ESceneOutlinerDropCompatibility::CompatibleMultipleDetach, ActorLabel);
+			}
+		}
+		else if (bCanAttach)
+		{
+			if (DragActors.Num() == 1)
+			{
+				return FSceneOutlinerDragValidationInfo(ESceneOutlinerDropCompatibility::CompatibleAttach, ActorLabel);
+			}
+			else
+			{
+				return FSceneOutlinerDragValidationInfo(ESceneOutlinerDropCompatibility::CompatibleMultipleAttach, ActorLabel);
+			}
+		}
+		else
+		{
+			if (DragActors.Num() == 1)
+			{
+				return FSceneOutlinerDragValidationInfo(ESceneOutlinerDropCompatibility::IncompatibleGeneric, AttachErrorMsg);
+			}
+			else
+			{
+				const FText ReasonText = FText::Format(LOCTEXT("DropOntoText", "{0}. {1}"), ActorLabel, AttachErrorMsg);
+				return FSceneOutlinerDragValidationInfo(ESceneOutlinerDropCompatibility::IncompatibleMultipleAttach, ReasonText);
+			}
+		}
+	}
+	else if (DropTarget.IsA<FFolderTreeItem>() || DropTarget.IsA<FWorldTreeItem>() || DropTarget.IsA<FLevelTreeItem>())
+	{
+		const FFolderTreeItem* FolderItem = DropTarget.CastTo<FFolderTreeItem>();
+		const FWorldTreeItem* WorldItem = DropTarget.CastTo<FWorldTreeItem>();
+		const FLevelTreeItem* LevelItem = DropTarget.CastTo<FLevelTreeItem>();
+		// WorldTreeItem and LevelTreeItem are treated as root folders (path = none), with the difference that LevelTreeItem has a RootObject.
+		const FFolder DestinationPath = FolderItem ? FolderItem->GetFolder() : (LevelItem ? FFolder(FFolder::GetOptionalFolderRootObject(LevelItem->Level.Get()).Get(FFolder::GetInvalidRootObject())) : GetWorldDefaultRootFolder());
+		const FFolder::FRootObject& DestinationRootObject = DestinationPath.GetRootObject();
+		ILevelInstanceInterface* LevelInstanceTarget = Cast<ILevelInstanceInterface>(DestinationPath.GetRootObjectPtr());
+		if (LevelInstanceTarget && !LevelInstanceTarget->IsEditing())
+		{
+			return FSceneOutlinerDragValidationInfo(ESceneOutlinerDropCompatibility::IncompatibleGeneric, LOCTEXT("Error_DragInNonEditingLevelInstance", "Cannot drag into a LevelInstance which is not being edited"));
+		}
+
+		if (Payload.Has<FFolderTreeItem>())
+		{
+			FFolder::FRootObject CommonFolderRootObject;
+			TArray<FName> DraggedFolders;
+			if (GetFolderNamesFromPayload(Payload, DraggedFolders, CommonFolderRootObject))
+			{
+				// Iterate over all the folders that have been dragged
+				for (const FName& DraggedFolder : DraggedFolders)
+				{
+					const FName Leaf = FEditorFolderUtils::GetLeafName(DraggedFolder);
+					const FName Parent = FEditorFolderUtils::GetParentPath(DraggedFolder);
+
+					if ((CommonFolderRootObject != DestinationRootObject) && FFolder::IsRootObjectValid(CommonFolderRootObject) && FFolder::IsRootObjectValid(DestinationRootObject))
+					{
+						FFormatNamedArguments Args;
+						Args.Add(TEXT("SourceName"), FText::FromName(Leaf));
+						FText Text = FText::Format(LOCTEXT("CantChangeFolderRoot", "Cannot change {SourceName} folder root"), Args);
+						return FSceneOutlinerDragValidationInfo(ESceneOutlinerDropCompatibility::IncompatibleGeneric, Text);
+					}
+
+					if (Parent == DestinationPath.GetPath())
+					{
+						FFormatNamedArguments Args;
+						Args.Add(TEXT("SourceName"), FText::FromName(Leaf));
+
+						FText Text;
+						if (DestinationPath.IsNone())
+						{
+							Text = FText::Format(LOCTEXT("FolderAlreadyAssignedRoot", "{SourceName} is already assigned to root"), Args);
+						}
+						else
+						{
+							Args.Add(TEXT("DestPath"), FText::FromName(DestinationPath.GetPath()));
+							Text = FText::Format(LOCTEXT("FolderAlreadyAssigned", "{SourceName} is already assigned to {DestPath}"), Args);
+						}
+
+						return FSceneOutlinerDragValidationInfo(ESceneOutlinerDropCompatibility::IncompatibleGeneric, Text);
+					}
+
+					const FString DragFolderPath = DraggedFolder.ToString();
+					const FString LeafName = Leaf.ToString();
+					const FString DstFolderPath = DestinationPath.IsNone() ? FString() : DestinationPath.ToString();
+					const FString NewPath = DstFolderPath / LeafName;
+
+					if (FActorFolders::Get().ContainsFolder(*RepresentingWorld, FFolder(DestinationRootObject, FName(*NewPath))))
+					{
+						// The folder already exists
+						FFormatNamedArguments Args;
+						Args.Add(TEXT("DragName"), FText::FromString(LeafName));
+						return FSceneOutlinerDragValidationInfo(ESceneOutlinerDropCompatibility::IncompatibleGeneric,
+							FText::Format(LOCTEXT("FolderAlreadyExistsRoot", "A folder called \"{DragName}\" already exists at this level"), Args));
+					}
+					else if (DragFolderPath == DstFolderPath || DstFolderPath.StartsWith(DragFolderPath + "/"))
+					{
+						// Cannot drag as a child of itself
+						FFormatNamedArguments Args;
+						Args.Add(TEXT("FolderPath"), FText::FromName(DraggedFolder));
+						return FSceneOutlinerDragValidationInfo(
+							ESceneOutlinerDropCompatibility::IncompatibleGeneric,
+							FText::Format(LOCTEXT("ChildOfItself", "Cannot move \"{FolderPath}\" to be a child of itself"), Args));
+					}
+				}
+			}
+		}
+
+		if (Payload.Has<FActorTreeItem>())
+		{
+			const ULevelInstanceSubsystem* LevelInstanceSubsystem = RepresentingWorld->GetSubsystem<ULevelInstanceSubsystem>();
+			// Iterate over all the actors that have been dragged
+			for (const TWeakObjectPtr<AActor>& WeakActor : Payload.GetData<TWeakObjectPtr<AActor>>(FWeakActorSelector()))
+			{
+				const AActor* Actor = WeakActor.Get();
+
+				bool bActorContainedInLevelInstance = false;
+				if (LevelInstanceSubsystem)
+				{
+					if (const ILevelInstanceInterface* ParentLevelInstance = LevelInstanceSubsystem->GetParentLevelInstance(Actor))
+					{
+						if (!ParentLevelInstance->IsEditing())
+						{
+							return FSceneOutlinerDragValidationInfo(ESceneOutlinerDropCompatibility::IncompatibleGeneric, LOCTEXT("Error_RemoveEditingLevelInstance", "Cannot detach from a LevelInstance which is not being edited"));
+						}
+						bActorContainedInLevelInstance = true;
+					}
+
+					if (const ILevelInstanceInterface* LevelInstance = Cast<ILevelInstanceInterface>(Actor))
+					{
+						FText Reason;
+						if (!LevelInstanceSubsystem->CanMoveActorToLevel(Actor, &Reason))
+						{
+							return FSceneOutlinerDragValidationInfo(ESceneOutlinerDropCompatibility::IncompatibleGeneric, Reason);
+						}
+					}
+				}
+
+				if (Actor->IsChildActor())
+				{
+					return FSceneOutlinerDragValidationInfo(ESceneOutlinerDropCompatibility::IncompatibleGeneric, FText::Format(LOCTEXT("Error_AttachChildActor", "Cannot move {0} as it is a child actor."), FText::FromString(Actor->GetActorLabel())));
+				}
+				else if ((Actor->GetFolderRootObject() != DestinationRootObject) && FFolder::IsRootObjectValid(Actor->GetFolderRootObject()) && FFolder::IsRootObjectValid(DestinationRootObject))
+				{
+					FFormatNamedArguments Args;
+					Args.Add(TEXT("SourceName"), FText::FromString(Actor->GetActorLabel()));
+					FText Text = FText::Format(LOCTEXT("CantChangeActorRoot", "Cannot change {SourceName} folder root"), Args);
+					return FSceneOutlinerDragValidationInfo(ESceneOutlinerDropCompatibility::IncompatibleGeneric, Text);
+				}
+				else if (Actor->GetFolder() == DestinationPath && !Actor->GetSceneOutlinerParent() && !bActorContainedInLevelInstance)
+				{
+					FFormatNamedArguments Args;
+					Args.Add(TEXT("SourceName"), FText::FromString(Actor->GetActorLabel()));
+
+					FText Text;
+					if (DestinationPath.IsNone())
+					{
+						Text = FText::Format(LOCTEXT("FolderAlreadyAssignedRoot", "{SourceName} is already assigned to root"), Args);
+					}
+					else
+					{
+						Args.Add(TEXT("DestPath"), FText::FromName(DestinationPath.GetPath()));
+						Text = FText::Format(LOCTEXT("FolderAlreadyAssigned", "{SourceName} is already assigned to {DestPath}"), Args);
+					}
+
+					return FSceneOutlinerDragValidationInfo(ESceneOutlinerDropCompatibility::IncompatibleGeneric, Text);
+				}
+			}
+		}
+
+		// Everything else is a valid operation
+		if (DestinationPath.IsNone())
+		{
+			return FSceneOutlinerDragValidationInfo(ESceneOutlinerDropCompatibility::CompatibleGeneric, LOCTEXT("MoveToRoot", "Move to root"));
+		}
+		else
+		{
+			FFormatNamedArguments Args;
+			Args.Add(TEXT("DestPath"), FText::FromName(DestinationPath.GetPath()));
+			return FSceneOutlinerDragValidationInfo(ESceneOutlinerDropCompatibility::CompatibleGeneric, FText::Format(LOCTEXT("MoveInto", "Move into \"{DestPath}\""), Args));
+		}
+	}
+	else if (DropTarget.IsA<FComponentTreeItem>())
+	{
+		// we don't allow drag and drop on components for now
+		return FSceneOutlinerDragValidationInfo(ESceneOutlinerDropCompatibility::IncompatibleGeneric, FText());
+	}
+	return FSceneOutlinerDragValidationInfo::Invalid();
+}
+
+void FDisplayClusterLightCardOutlinerMode::OnDrop(ISceneOutlinerTreeItem& DropTarget,
+	const FSceneOutlinerDragDropPayload& Payload, const FSceneOutlinerDragValidationInfo& ValidationInfo) const
+{
+	if (const FActorTreeItem* ActorItem = DropTarget.CastTo<FActorTreeItem>())
+	{
+		AActor* DropActor = ActorItem->Actor.Get();
+		if (!DropActor)
+		{
+			return;
+		}
+
+		FMessageLog EditorErrors("EditorErrors");
+		EditorErrors.NewPage(LOCTEXT("ActorAttachmentsPageLabel", "Actor attachment"));
+
+		if (ValidationInfo.CompatibilityType == ESceneOutlinerDropCompatibility::CompatibleMultipleDetach || ValidationInfo.CompatibilityType == ESceneOutlinerDropCompatibility::CompatibleDetach)
+		{
+			const FScopedTransaction Transaction(LOCTEXT("UndoAction_DetachActors", "Detach actors"));
+
+			TArray<TWeakObjectPtr<AActor>> DraggedActors = Payload.GetData<TWeakObjectPtr<AActor>>(FWeakActorSelector());
+			for (const TWeakObjectPtr<AActor>& WeakActor : DraggedActors)
+			{
+				if (AActor* DragActor = WeakActor.Get())
+				{
+					// Detach from parent
+					USceneComponent* RootComp = DragActor->GetRootComponent();
+					if (RootComp && RootComp->GetAttachParent())
+					{
+						AActor* OldParent = RootComp->GetAttachParent()->GetOwner();
+						// Attachment is persisted on the child so modify both actors for Undo/Redo but do not mark the Parent package dirty
+						OldParent->Modify(/*bAlwaysMarkDirty=*/false);
+						RootComp->DetachFromComponent(FDetachmentTransformRules::KeepWorldTransform);
+
+						DragActor->SetFolderPath_Recursively(OldParent->GetFolderPath());
+					}
+				}
+			}
+		}
+		else if (ValidationInfo.CompatibilityType == ESceneOutlinerDropCompatibility::CompatibleMultipleAttach || ValidationInfo.CompatibilityType == ESceneOutlinerDropCompatibility::CompatibleAttach)
+		{
+			// Show socket chooser if we have sockets to select
+
+			if (ILevelInstanceInterface* TargetLevelInstance = Cast<ILevelInstanceInterface>(DropActor))
+			{
+				check(TargetLevelInstance->IsEditing());
+				const FScopedTransaction Transaction(LOCTEXT("UndoAction_MoveActorsToLevelInstance", "Move actors to LevelInstance"));
+
+				const FFolder DestinationPath = FFolder(FFolder::FRootObject(DropActor));
+				auto MoveToDestination = [&DestinationPath](FFolderTreeItem& Item)
+				{
+					Item.MoveTo(DestinationPath);
+				};
+				Payload.ForEachItem<FFolderTreeItem>(MoveToDestination);
+
+				// Since target root is directly the Level Instance, clear folder path
+				TArray<AActor*> DraggedActors = Payload.GetData<AActor*>(FActorSelector());
+				for (AActor*& Actor : DraggedActors)
+				{
+					Actor->SetFolderPath_Recursively(FName());
+				}
+
+				ULevelInstanceSubsystem* LevelInstanceSubsystem = RepresentingWorld->GetSubsystem<ULevelInstanceSubsystem>();
+				check(LevelInstanceSubsystem);
+				LevelInstanceSubsystem->MoveActorsTo(TargetLevelInstance, DraggedActors);
+			}
+			else
+			{
+				auto PerformAttachment = [](FName SocketName, TWeakObjectPtr<AActor> Parent, const TArray<TWeakObjectPtr<AActor>> NewAttachments)
+				{
+					AActor* ParentActor = Parent.Get();
+					if (ParentActor)
+					{
+						// modify parent and child
+						const FScopedTransaction Transaction(LOCTEXT("UndoAction_PerformAttachment", "Attach actors"));
+
+						// Attach each child
+						bool bAttached = false;
+						for (const TWeakObjectPtr<AActor>& Child : NewAttachments)
+						{
+							AActor* ChildActor = Child.Get();
+							if (GEditor->CanParentActors(ParentActor, ChildActor))
+							{
+								GEditor->ParentActors(ParentActor, ChildActor, SocketName);
+
+								ChildActor->SetFolderPath_Recursively(ParentActor->GetFolderPath());
+							}
+						}
+					}
+				};
+
+				TArray<TWeakObjectPtr<AActor>> DraggedActors = Payload.GetData<TWeakObjectPtr<AActor>>(FWeakActorSelector());
+				USceneComponent* Component = DropActor->GetRootComponent();
+				if ((Component != NULL) && (Component->HasAnySockets()))
+				{
+					// Create the popup
+					FSlateApplication::Get().PushMenu(
+						SceneOutliner->AsShared(),
+						FWidgetPath(),
+						SNew(SSocketChooserPopup)
+						.SceneComponent(Component)
+						.OnSocketChosen_Lambda(PerformAttachment, DropActor, MoveTemp(DraggedActors)),
+						FSlateApplication::Get().GetCursorPos(),
+						FPopupTransitionEffect(FPopupTransitionEffect::TypeInPopup)
+					);
+				}
+				else
+				{
+					PerformAttachment(NAME_None, DropActor, MoveTemp(DraggedActors));
+				}
+			}
+
+		}
+		// Report errors
+		EditorErrors.Notify(NSLOCTEXT("ActorAttachmentError", "AttachmentsFailed", "Attachments Failed!"));
+	}
+	else if (DropTarget.IsA<FFolderTreeItem>() || DropTarget.IsA<FWorldTreeItem>() || DropTarget.IsA<FLevelTreeItem>())
+	{
+		const FFolderTreeItem* FolderItem = DropTarget.CastTo<FFolderTreeItem>();
+		const FWorldTreeItem* WorldItem = DropTarget.CastTo<FWorldTreeItem>();
+		const FLevelTreeItem* LevelItem = DropTarget.CastTo<FLevelTreeItem>();
+		// WorldTreeItem and LevelTreeItem are treated as root folders (path = none), with the difference that LevelTreeItem has a RootObject.
+		const FFolder DestinationPath = FolderItem ? FolderItem->GetFolder() : (LevelItem ? FFolder(FFolder::GetOptionalFolderRootObject(LevelItem->Level.Get()).Get(FFolder::GetInvalidRootObject())) : GetWorldDefaultRootFolder());
+
+		const FScopedTransaction Transaction(LOCTEXT("MoveOutlinerItems", "Move World Outliner Items"));
+
+		auto MoveToDestination = [&DestinationPath](FFolderTreeItem& Item)
+		{
+			Item.MoveTo(DestinationPath);
+		};
+		Payload.ForEachItem<FFolderTreeItem>(MoveToDestination);
+
+		// Set the folder path on all the dragged actors, and detach any that need to be moved
+		if (Payload.Has<FActorTreeItem>())
+		{
+			TSet<const AActor*> ParentActors;
+			TSet<const AActor*> ChildActors;
+
+			TArray<AActor*> MovingActorsToValidRootObject;
+			Payload.ForEachItem<FActorTreeItem>([&DestinationPath, &ParentActors, &ChildActors, &MovingActorsToValidRootObject](const FActorTreeItem& ActorItem)
+			{
+				AActor* Actor = ActorItem.Actor.Get();
+				if (Actor)
+				{
+					// First mark this object as a parent, then set its children's path
+					ParentActors.Add(Actor);
+
+					const FFolder SrcFolder = Actor->GetFolder();
+
+					// If the folder root object changes, 1st pass will put actors at root. 2nd pass will set the destination path.
+					FName NewPath = (SrcFolder.GetRootObject() == DestinationPath.GetRootObject()) ? DestinationPath.GetPath() : NAME_None;
+
+					Actor->SetFolderPath(NewPath);
+					FActorEditorUtils::TraverseActorTree_ParentFirst(Actor, [&](AActor* InActor) {
+						ChildActors.Add(InActor);
+						InActor->SetFolderPath(NewPath);
+						return true;
+						}, false);
+
+					if ((Actor->GetFolderRootObject() != DestinationPath.GetRootObject()) && SrcFolder.IsRootObjectPersistentLevel() && (DestinationPath.IsRootObjectValid() && !DestinationPath.IsRootObjectPersistentLevel()))
+					{
+						MovingActorsToValidRootObject.Add(Actor);
+					}
+				}
+			});
+
+			// Detach parent actors
+			for (const AActor* Parent : ParentActors)
+			{
+				USceneComponent* RootComp = Parent->GetRootComponent();
+
+				// We don't detach if it's a child of another that's been dragged
+				if (RootComp && RootComp->GetAttachParent() && !ChildActors.Contains(Parent))
+				{
+					if (AActor* OldParentActor = RootComp->GetAttachParent()->GetOwner())
+					{
+						// Attachment is persisted on the child so modify both actors for Undo/Redo but do not mark the Parent package dirty
+						OldParentActor->Modify(/*bAlwaysMarkDirty=*/false);
+					}
+					RootComp->DetachFromComponent(FDetachmentTransformRules::KeepWorldTransform);
+				}
+			}
+
+			auto MoveActorsToLevel = [](const TArray<AActor*>& InActorsToMove, ULevel* InDestLevel, const FName& InDestinationPath)
+			{
+				// We are moving actors to another level
+				const bool bWarnAboutReferences = true;
+				const bool bWarnAboutRenaming = true;
+				const bool bMoveAllOrFail = true;
+				TArray<AActor*> MovedActors;
+				if (!EditorLevelUtils::MoveActorsToLevel(InActorsToMove, InDestLevel, bWarnAboutReferences, bWarnAboutRenaming, bMoveAllOrFail, &MovedActors))
+				{
+					UE_LOG(DisplayClusterLightCardEditorLog, Warning, TEXT("Failed to move actors because not all actors could be moved"));
+				}
+				// Once moved, update actors folder path
+				for (AActor* Actor : MovedActors)
+				{
+					Actor->SetFolderPath_Recursively(InDestinationPath);
+				}
+			};
+
+			if (DestinationPath.IsRootObjectPersistentLevel())
+			{
+				const ULevelInstanceSubsystem* LevelInstanceSubsystem = RepresentingWorld->GetSubsystem<ULevelInstanceSubsystem>();
+				check(LevelInstanceSubsystem);
+				ULevel* DestinationLevel = RepresentingWorld->PersistentLevel;
+				check(DestinationLevel);
+
+				TArray<AActor*> LevelInstanceActorsToMove;
+				TArray<AActor*> ActorsToMoveToPersistentLevel;
+				Payload.ForEachItem<FActorTreeItem>([LevelInstanceSubsystem, &LevelInstanceActorsToMove, &ActorsToMoveToPersistentLevel](const FActorTreeItem& ActorItem)
+				{
+					AActor* Actor = ActorItem.Actor.Get();
+					if (const ILevelInstanceInterface* ParentLevelInstance = LevelInstanceSubsystem->GetParentLevelInstance(Actor))
+					{
+						check(ParentLevelInstance->IsEditing());
+						LevelInstanceActorsToMove.Add(Actor);
+					}
+					else
+					{
+						const FFolder ActorSrcFolder = Actor->GetFolder();
+						if (ActorSrcFolder.IsRootObjectValid() && !ActorSrcFolder.IsRootObjectPersistentLevel())
+						{
+							ActorsToMoveToPersistentLevel.Add(Actor);
+						}
+					}
+				});
+
+				// We are moving actors outside of an editing level instance to a folder (or root) into the persistent level.
+				if (LevelInstanceActorsToMove.Num() > 0)
+				{
+					TArray<AActor*> MovedActors;
+					LevelInstanceSubsystem->MoveActorsToLevel(LevelInstanceActorsToMove, DestinationLevel, &MovedActors);
+					// Once moved, update actors folder path
+					for (AActor* Actor : MovedActors)
+					{
+						Actor->SetFolderPath_Recursively(DestinationPath.GetPath());
+					}
+				}
+				if (ActorsToMoveToPersistentLevel.Num() > 0)
+				{
+					MoveActorsToLevel(ActorsToMoveToPersistentLevel, DestinationLevel, DestinationPath.GetPath());
+				}
+			}
+			else if (MovingActorsToValidRootObject.Num())
+			{
+				if (ILevelInstanceInterface* TargetLevelInstance = Cast<ILevelInstanceInterface>(DestinationPath.GetRootObjectPtr()))
+				{
+					// We are moving actors inside an editing level instance
+					check(TargetLevelInstance->IsEditing());
+
+					ULevelInstanceSubsystem* LevelInstanceSubsystem = RepresentingWorld->GetSubsystem<ULevelInstanceSubsystem>();
+					check(LevelInstanceSubsystem);
+					TArray<AActor*> MovedActors;
+					LevelInstanceSubsystem->MoveActorsTo(TargetLevelInstance, MovingActorsToValidRootObject, &MovedActors);
+					// Once moved, update actors folder path
+					for (AActor* Actor : MovedActors)
+					{
+						Actor->SetFolderPath_Recursively(DestinationPath.GetPath());
+					}
+				}
+				else if (ULevel* DestinationLevel = Cast<ULevel>(DestinationPath.GetRootObjectPtr()))
+				{
+					MoveActorsToLevel(MovingActorsToValidRootObject, DestinationLevel, DestinationPath.GetPath());
+				}
+			}
+		}
+	}
+}
+
+void FDisplayClusterLightCardOutlinerMode::RegisterContextMenu()
+{
+	UToolMenus* ToolMenus = UToolMenus::Get();
+
+	if (!ToolMenus->IsMenuRegistered(UE::LightCardOutlinerMode::DefaultContextBaseMenuName))
+	{
+		UToolMenu* Menu = ToolMenus->RegisterMenu(UE::LightCardOutlinerMode::DefaultContextBaseMenuName);
+
+		Menu->AddDynamicSection("DynamicHierarchySection", FNewToolMenuDelegate::CreateLambda([](UToolMenu* InMenu)
+			{
+				USceneOutlinerMenuContext* Context = InMenu->FindContext<USceneOutlinerMenuContext>();
+				if (!Context || !Context->SceneOutliner.IsValid())
+				{
+					return;
+				}
+
+				// NOTE: the name "Section" is used in many other places
+				FToolMenuSection& Section = InMenu->FindOrAddSection("Section");
+				Section.Label = LOCTEXT("HierarchySectionName", "Hierarchy");
+
+				SSceneOutliner* SceneOutliner = Context->SceneOutliner.Pin().Get();
+				if (Context->bShowParentTree)
+				{
+					// Only show create folder if no items are selected or a folder is selected.
+					// Actors will have the MoveTo option available to support folder creation.
+					// This lightcard outliner implementation doesn't currently support the normal folder
+					// hierarchy options like duplication.
+					if (Context->NumSelectedItems == 0 || Context->NumSelectedFolders == 1)
+					{
+						FSceneOutlinerMenuHelper::AddMenuEntryCreateFolder(Section, *SceneOutliner);
+					}
+				}
+			}));
+
+		Menu->AddDynamicSection("DynamicMainSection", FNewToolMenuDelegate::CreateLambda([](UToolMenu* InMenu)
+			{
+				// We always create a section here, even if there is no parent so that clients can still extend the menu
+				FToolMenuSection& Section = InMenu->AddSection("MainSection", LOCTEXT("OutlinerSectionName", "Outliner"));
+
+				if (USceneOutlinerMenuContext* Context = InMenu->FindContext<USceneOutlinerMenuContext>())
+				{
+					// Don't add any of these menu items if we're not showing the parent tree
+					// Can't move worlds or level blueprints
+					if (Context->bShowParentTree && Context->NumSelectedItems > 0 && Context->NumWorldsSelected == 0 && Context->SceneOutliner.IsValid())
+					{
+						Section.AddSubMenu(
+							"MoveActorsTo",
+							LOCTEXT("MoveActorsTo", "Move To"),
+							LOCTEXT("MoveActorsTo_Tooltip", "Move selection to another folder"),
+							FNewToolMenuDelegate::CreateSP(Context->SceneOutliner.Pin().Get(), &SSceneOutliner::FillFoldersSubMenu));
+					}
+
+					if (Context->bShowParentTree && Context->NumSelectedItems > 0 && Context->SceneOutliner.IsValid())
+					{
+						// Only add the menu option to wp levels
+						if (!Context->bRepresentingGameWorld && Context->bRepresentingPartitionedWorld)
+						{
+							// If selection contains some unpinned items, show the pin option
+							// If the selection contains folders, always show the pin option
+							if (Context->NumPinnedItems != Context->NumSelectedItems || Context->NumSelectedFolders > 0)
+							{
+								Section.AddMenuEntry(
+									"PinItems",
+									LOCTEXT("Pin", "Pin"),
+									FText(),
+									FSlateIcon(),
+									FUIAction(FExecuteAction::CreateSP(Context->SceneOutliner.Pin().Get(), &SSceneOutliner::PinSelectedItems)));
+							}
+
+							// If the selection contains some pinned items, show the unpin option
+							// If the selection contains folders, always show the unpin option
+							if (Context->NumPinnedItems != 0 || Context->NumSelectedFolders > 0)
+							{
+								Section.AddMenuEntry(
+									"UnpinItems",
+									LOCTEXT("Unpin", "Unpin"),
+									FText(),
+									FSlateIcon(),
+									FUIAction(FExecuteAction::CreateSP(Context->SceneOutliner.Pin().Get(), &SSceneOutliner::UnpinSelectedItems)));
+							}
+						}
+					}
+
+					if (Context->NumSelectedItems > 0 && Context->SceneOutliner.IsValid())
+					{
+						SSceneOutliner* SceneOutliner = Context->SceneOutliner.Pin().Get();
+						SceneOutliner->AddSourceControlMenuOptions(InMenu);
+					}
+				}
+			}));
+	}
+
+	if (!ToolMenus->IsMenuRegistered(UE::LightCardOutlinerMode::DefaultContextMenuName))
+	{
+		ToolMenus->RegisterMenu(UE::LightCardOutlinerMode::DefaultContextMenuName, UE::LightCardOutlinerMode::DefaultContextBaseMenuName);
+	}
+}
+
+bool FDisplayClusterLightCardOutlinerMode::CanPasteFoldersOnlyFromClipboard() const
+{
+	FString PasteString;
+	FPlatformApplicationMisc::ClipboardPaste(PasteString);
+	return PasteString.StartsWith("BEGIN FOLDERLIST");
+}
+
+bool FDisplayClusterLightCardOutlinerMode::GetFolderNamesFromPayload(const FSceneOutlinerDragDropPayload& InPayload,
+	TArray<FName>& OutFolders, FFolder::FRootObject& OutCommonRootObject) const
+{
+	return FFolder::GetFolderPathsAndCommonRootObject(InPayload.GetData<FFolder>(FFolderPathSelector()), OutFolders, OutCommonRootObject);
+}
+
+FFolder FDisplayClusterLightCardOutlinerMode::GetWorldDefaultRootFolder() const
+{
+	return FFolder::GetWorldRootFolder(RepresentingWorld.Get());
+}
+
+void FDisplayClusterLightCardOutlinerMode::UpdateGlobalActorSelection()
+{
+	TArray<AActor*> SelectedActors = SceneOutliner->GetSelection().GetData<AActor*>(FActorSelector());
+	
+	const FScopedTransaction Transaction(NSLOCTEXT("UnrealEd", "SelectActors", "Select Actors"));
+	GEditor->GetSelectedActors()->Modify();
+
+	GEditor->GetSelectedActors()->BeginBatchSelectOperation();
+	
+	GEditor->SelectNone(false, true, true);
+	
+	const bool bShouldSelect = true;
+	const bool bNotifyAfterSelect = false;
+	const bool bSelectEvenIfHidden = true;
+	for (AActor* Actor : SelectedActors)
+	{
+		GEditor->SelectActor(Actor, bShouldSelect, bNotifyAfterSelect, bSelectEvenIfHidden);
+	}
+
+	GEditor->GetSelectedActors()->EndBatchSelectOperation(/*bNotify*/false);
+	GEditor->NoteSelectionChange();
+}
+
+#undef LOCTEXT_NAMESPACE
diff --git a/Engine/Plugins/Runtime/nDisplay/Source/DisplayClusterLightCardEditor/Private/Outliner/DisplayClusterLightCardOutlinerMode.h b/Engine/Plugins/Runtime/nDisplay/Source/DisplayClusterLightCardEditor/Private/Outliner/DisplayClusterLightCardOutlinerMode.h
new file mode 100644
--- /dev/null
+++ b/Engine/Plugins/Runtime/nDisplay/Source/DisplayClusterLightCardEditor/Private/Outliner/DisplayClusterLightCardOutlinerMode.h
@@ -0,0 +1,61 @@
+// Copyright Epic Games, Inc. All Rights Reserved.
+
+#pragma once
+
+#include "ActorMode.h"
+
+class SDisplayClusterLightCardOutliner;
+
+/**
+ * Custom outliner mode for the light card editor based on FActorBrowsingMode. The primary goal is to allow basic actor
+ * functionality such as selecting actors locally, creating folders, and drag & drop. Selection should occur locally
+ * when possible and not update the GEditor selection. Additionally the context menu is cleaned up.
+ */
+class FDisplayClusterLightCardOutlinerMode final : public FActorMode
+{
+public:
+	
+	FDisplayClusterLightCardOutlinerMode(SSceneOutliner* InSceneOutliner, TWeakPtr<SDisplayClusterLightCardOutliner> InLightCardOutliner,
+		TWeakObjectPtr<UWorld> InSpecifiedWorldToDisplay = nullptr);
+
+	virtual void SynchronizeSelection() override;
+	virtual bool IsInteractive() const override { return true; }
+	virtual FCreateSceneOutlinerMode CreateFolderPickerMode(const FFolder::FRootObject& InRootObject = FFolder::GetInvalidRootObject()) const override;
+	virtual TSharedPtr<SWidget> CreateContextMenu() override;
+	virtual ESelectionMode::Type GetSelectionMode() const override { return ESelectionMode::Multi; }
+	virtual bool SupportsKeyboardFocus() const override { return true; }
+	virtual bool ShouldShowFolders() const override { return true; }
+	virtual bool SupportsCreateNewFolder() const override { return true; }
+	virtual bool ShowStatusBar() const override { return true; }
+	virtual bool ShowViewButton() const override { return true; }
+	virtual bool ShowFilterOptions() const override { return true; }
+	virtual void OnItemDoubleClick(FSceneOutlinerTreeItemPtr Item) override;
+	virtual FReply OnKeyDown(const FKeyEvent& InKeyEvent) override;
+	virtual bool CanDelete() const override;
+	virtual bool CanRename() const override;
+	virtual bool CanRenameItem(const ISceneOutlinerTreeItem& Item) const override;
+	virtual bool CanCut() const override;
+	virtual bool CanCopy() const override;
+	virtual bool CanPaste() const override;
+	virtual FFolder CreateNewFolder() override;
+	virtual FFolder GetFolder(const FFolder& ParentPath, const FName& LeafName) override;
+	virtual bool CreateFolder(const FFolder& NewFolder) override;
+	virtual bool ReparentItemToFolder(const FFolder& FolderPath, const FSceneOutlinerTreeItemPtr& Item) override;
+	virtual bool CanSupportDragAndDrop() const override { return true; }
+	virtual TSharedPtr<FDragDropOperation> CreateDragDropOperation(const FPointerEvent& MouseEvent, const TArray<FSceneOutlinerTreeItemPtr>& InTreeItems) const override;
+	virtual bool ParseDragDrop(FSceneOutlinerDragDropPayload& OutPayload, const FDragDropOperation& Operation) const override;
+	virtual FSceneOutlinerDragValidationInfo ValidateDrop(const ISceneOutlinerTreeItem& DropTarget, const FSceneOutlinerDragDropPayload& Payload) const override;
+	virtual void OnDrop(ISceneOutlinerTreeItem& DropTarget, const FSceneOutlinerDragDropPayload& Payload, const FSceneOutlinerDragValidationInfo& ValidationInfo) const override;
+
+private:
+	static void RegisterContextMenu();
+	bool CanPasteFoldersOnlyFromClipboard() const;
+	bool GetFolderNamesFromPayload(const FSceneOutlinerDragDropPayload& InPayload, TArray<FName>& OutFolders, FFolder::FRootObject& OutCommonRootObject) const;
+	FFolder GetWorldDefaultRootFolder() const;
+
+	/** Updates GEditor selection */
+	void UpdateGlobalActorSelection();
+	
+private:
+	TWeakPtr<SDisplayClusterLightCardOutliner> LightCardOutliner;
+};
diff --git a/Engine/Plugins/Runtime/nDisplay/Source/DisplayClusterLightCardEditor/Private/Outliner/SDisplayClusterLightCardOutliner.cpp b/Engine/Plugins/Runtime/nDisplay/Source/DisplayClusterLightCardEditor/Private/Outliner/SDisplayClusterLightCardOutliner.cpp
new file mode 100644
--- /dev/null
+++ b/Engine/Plugins/Runtime/nDisplay/Source/DisplayClusterLightCardEditor/Private/Outliner/SDisplayClusterLightCardOutliner.cpp
@@ -0,0 +1,407 @@
+// Copyright Epic Games, Inc. All Rights Reserved.
+
+#include "SDisplayClusterLightCardOutliner.h"
+
+#include "DisplayClusterLightCardEditorCommands.h"
+#include "DisplayClusterLightCardOutlinerMode.h"
+#include "SDisplayClusterLightCardEditor.h"
+
+#include "IDisplayClusterOperator.h"
+
+#include "DisplayClusterConfigurationTypes.h"
+
+#include "DisplayClusterLightCardActor.h"
+#include "DisplayClusterRootActor.h"
+
+#include "ActorTreeItem.h"
+#include "ClassIconFinder.h"
+#include "FolderTreeItem.h"
+#include "SceneOutlinerMenuContext.h"
+#include "SceneOutlinerModule.h"
+#include "SceneOutlinerTextInfoColumn.h"
+#include "Selection.h"
+#include "SSceneOutliner.h"
+#include "ToolMenus.h"
+#include "WorldTreeItem.h"
+#include "Framework/Commands/GenericCommands.h"
+
+#define LOCTEXT_NAMESPACE "SDisplayClusterLightCardOutliner"
+
+SDisplayClusterLightCardOutliner::~SDisplayClusterLightCardOutliner()
+{
+	if (SceneOutliner.IsValid() && SceneOutlinerSelectionChanged.IsValid())
+	{
+		SceneOutliner->GetOnItemSelectionChanged().Remove(SceneOutlinerSelectionChanged);
+	}
+}
+
+void SDisplayClusterLightCardOutliner::Construct(const FArguments& InArgs, TSharedPtr<SDisplayClusterLightCardEditor> InLightCardEditor, TSharedPtr<FUICommandList> InCommandList)
+{
+	LightCardEditorPtr = InLightCardEditor;
+	CommandList = InCommandList;
+	
+	CreateWorldOutliner();
+}
+
+void SDisplayClusterLightCardOutliner::SetRootActor(ADisplayClusterRootActor* NewRootActor)
+{
+	TArray<ADisplayClusterLightCardActor*> PreviouslySelectedLightCards;
+	GetSelectedLightCards(PreviouslySelectedLightCards);
+	
+	RootActor = NewRootActor;
+	FillLightCardList();
+
+	CreateWorldOutliner();
+	
+	for (const TSharedPtr<FLightCardTreeItem>& LightCardTreeItem : LightCardTree)
+	{
+		if (LightCardTreeItem->IsActorLayer())
+		{
+			LightCardTreeView->SetItemExpansion(LightCardTreeItem, true);
+		}
+	}
+	
+	// Select previously selected light cards. This fixes an issue where the details panel may clear in some cases
+	// and also supports maintaining the selection of shared light cards across different root actors.
+	SelectLightCards(MoveTemp(PreviouslySelectedLightCards));
+}
+
+void SDisplayClusterLightCardOutliner::GetSelectedLightCards(TArray<ADisplayClusterLightCardActor*>& OutSelectedLightCards) const
+{
+	TArray<FSceneOutlinerTreeItemPtr> SelectedOutlinerLightCards = SceneOutliner->GetSelectedItems();
+	for (const FSceneOutlinerTreeItemPtr& SelectedLightCard : SelectedOutlinerLightCards)
+	{
+		TSharedPtr<FLightCardTreeItem> MatchingLightCard = GetLightCardTreeItemFromOutliner(SelectedLightCard);
+		if (MatchingLightCard.IsValid() && MatchingLightCard->LightCardActor.IsValid())
+		{
+			OutSelectedLightCards.Add(MatchingLightCard->LightCardActor.Get());
+		}
+	}
+}
+
+void SDisplayClusterLightCardOutliner::SelectLightCards(const TArray<ADisplayClusterLightCardActor*>& LightCardsToSelect)
+{
+	TArray<ADisplayClusterLightCardActor*> ActualSelectedLightCards;
+	TArray<FSceneOutlinerTreeItemPtr> SelectedTreeItems;
+	CachedSelectedActors.Reset();
+	
+	for (const TSharedPtr<FLightCardTreeItem>& TreeItem : LightCardActors)
+	{
+		if (LightCardsToSelect.Contains(TreeItem->LightCardActor))
+		{
+			if (FSceneOutlinerTreeItemPtr OutlinerItem = SceneOutliner->GetTreeItem(TreeItem->LightCardActor.Get()))
+			{
+				SelectedTreeItems.Add(OutlinerItem);
+				ActualSelectedLightCards.Add(TreeItem->LightCardActor.Get());
+				CachedSelectedActors.Add(TreeItem->LightCardActor);
+			}
+		}
+	}
+	
+	SceneOutliner->SetItemSelection(SelectedTreeItems, true);
+	
+	IDisplayClusterOperator::Get().ShowDetailsForObjects(*reinterpret_cast<TArray<UObject*>*>(&ActualSelectedLightCards));
+}
+
+void SDisplayClusterLightCardOutliner::RestoreCachedSelection()
+{
+	TArray<ADisplayClusterLightCardActor*> ActorsToSelect;
+	for (const TWeakObjectPtr<ADisplayClusterLightCardActor>& CachedActor : CachedSelectedActors)
+	{
+		if (CachedActor.IsValid())
+		{
+			ActorsToSelect.Add(CachedActor.Get());
+		}
+	}
+	
+	SelectLightCards(MoveTemp(ActorsToSelect));
+}
+
+void SDisplayClusterLightCardOutliner::CreateWorldOutliner()
+{
+	auto OutlinerFilterPredicate = [this](const AActor* InActor)
+	{
+		return TrackedActors.Contains(InActor);
+	};
+
+	const FSceneOutlinerModule& SceneOutlinerModule = FModuleManager::Get().LoadModuleChecked<FSceneOutlinerModule>("SceneOutliner");
+
+	if (SceneOutliner.IsValid() && SceneOutlinerSelectionChanged.IsValid())
+	{
+		SceneOutliner->GetOnItemSelectionChanged().Remove(SceneOutlinerSelectionChanged);
+	}
+	
+	FSceneOutlinerInitializationOptions SceneOutlinerOptions;
+
+	// Explicitly add in the columns we want. Otherwise all column visibility can be toggled and this outliner only needs several.
+	
+	SceneOutlinerOptions.ColumnMap.Add(FSceneOutlinerBuiltInColumnTypes::Gutter(),
+		FSceneOutlinerColumnInfo(ESceneOutlinerColumnVisibility::Visible, 0, FCreateSceneOutlinerColumn(), false, TOptional<float>(),
+				FSceneOutlinerBuiltInColumnTypes::Gutter_Localized()));
+	
+	SceneOutlinerOptions.ColumnMap.Add(FSceneOutlinerBuiltInColumnTypes::Label(),
+		FSceneOutlinerColumnInfo(ESceneOutlinerColumnVisibility::Visible, 10, FCreateSceneOutlinerColumn(), false, TOptional<float>(),
+				FSceneOutlinerBuiltInColumnTypes::Label_Localized()));
+
+	// Add a custom class/type column. The default ActorInfo column returns the class FName only, where as we want to display a friendly
+	// display for the class.
+	{
+		// Based on struct FGetInfo from SceneOutlinerActorInfoColumn
+		const FGetTextForItem InternalNameInfoText = FGetTextForItem::CreateLambda([](const ISceneOutlinerTreeItem& Item) -> FString
+		{
+			if (const FActorTreeItem* ActorItem = Item.CastTo<FActorTreeItem>())
+			{
+				if (const AActor* Actor = ActorItem->Actor.Get())
+				{
+					return Actor->GetClass()->GetDisplayNameText().ToString();
+				}
+			
+				return FString();
+			}
+			if (Item.IsA<FFolderTreeItem>())
+			{
+				return LOCTEXT("FolderTypeName", "Folder").ToString();
+			}
+			if (Item.IsA<FWorldTreeItem>())
+			{
+				return LOCTEXT("WorldTypeName", "World").ToString();
+			}
+
+			return FString();
+		});
+		
+		SceneOutlinerOptions.ColumnMap.Add(FSceneOutlinerBuiltInColumnTypes::ActorInfo(), FSceneOutlinerColumnInfo(ESceneOutlinerColumnVisibility::Visible, 20,
+			FCreateSceneOutlinerColumn::CreateStatic(&FTextInfoColumn::CreateTextInfoColumn, FSceneOutlinerBuiltInColumnTypes::ActorInfo(), InternalNameInfoText, FText::GetEmpty()),
+			true, TOptional<float>(),
+				FSceneOutlinerBuiltInColumnTypes::ActorInfo_Localized()));
+	}
+	
+	SceneOutlinerOptions.ModifyContextMenu = FSceneOutlinerModifyContextMenu::CreateSP(this, &SDisplayClusterLightCardOutliner::RegisterContextMenu);
+	SceneOutlinerOptions.OutlinerIdentifier = TEXT("LightCardEditorOutliner");
+	SceneOutlinerOptions.Filters->AddFilterPredicate<FActorTreeItem>(FActorTreeItem::FFilterPredicate::CreateLambda(OutlinerFilterPredicate));
+
+	TWeakPtr<SDisplayClusterLightCardOutliner> WeakPtrThis = SharedThis(this);
+	UWorld* OutlinerWorld = RootActor.IsValid() ? RootActor->GetWorld() : nullptr;
+	SceneOutlinerOptions.ModeFactory = FCreateSceneOutlinerMode::CreateLambda([WeakPtrThis, SpecifiedWorld = MakeWeakObjectPtr(OutlinerWorld)](SSceneOutliner* Outliner)
+		{
+			return new FDisplayClusterLightCardOutlinerMode(Outliner, WeakPtrThis, SpecifiedWorld);
+		});
+	
+	SceneOutliner = StaticCastSharedRef<SSceneOutliner>(SceneOutlinerModule.CreateSceneOutliner(SceneOutlinerOptions));
+	SceneOutlinerSelectionChanged = SceneOutliner->GetOnItemSelectionChanged().AddRaw(this, &SDisplayClusterLightCardOutliner::OnOutlinerSelectionChanged);
+
+	ChildSlot
+	[
+		SceneOutliner.ToSharedRef()
+	];
+}
+
+bool SDisplayClusterLightCardOutliner::FillLightCardList()
+{
+	TArray<TSharedPtr<FLightCardTreeItem>> OriginalTree = MoveTemp(LightCardTree);
+	LightCardActors.Empty();
+
+	if (RootActor.IsValid())
+	{
+		FDisplayClusterConfigurationICVFX_VisibilityList& RootActorLightCards = RootActor->GetConfigData()->StageSettings.Lightcard.ShowOnlyList;
+
+		for (const FActorLayer& ActorLayer : RootActorLightCards.ActorLayers)
+		{
+			TSharedPtr<FLightCardTreeItem> LightCardTreeItem = MakeShared<FLightCardTreeItem>();
+			LightCardTreeItem->ActorLayer = ActorLayer.Name;
+			LightCardTree.Add(LightCardTreeItem);
+		}
+
+		for (const TSoftObjectPtr<AActor>& LightCardActor : RootActorLightCards.Actors)
+		{
+			if (LightCardActor.IsValid() && LightCardActor->IsA<ADisplayClusterLightCardActor>())
+			{
+				TSharedPtr<FLightCardTreeItem> LightCardTreeItem = MakeShared<FLightCardTreeItem>();
+				LightCardTreeItem->LightCardActor = CastChecked<ADisplayClusterLightCardActor>(LightCardActor.Get());
+				LightCardTree.Add(LightCardTreeItem);
+				LightCardActors.Add(LightCardTreeItem);
+			}
+		}
+
+		// If there are any layers that are specified as light card layers, iterate over all actors in the world and 
+		// add any that are members of any of the light card layers to the list. Only add an actor once, even if it is
+		// in multiple layers
+		if (RootActorLightCards.ActorLayers.Num())
+		{
+			if (UWorld* World = RootActor->GetWorld())
+			{
+				for (const TWeakObjectPtr<ADisplayClusterLightCardActor> WeakActor : TActorRange<ADisplayClusterLightCardActor>(World))
+				{
+					if (WeakActor.IsValid())
+					{
+						for (const FActorLayer& ActorLayer : RootActorLightCards.ActorLayers)
+						{
+							if (WeakActor->Layers.Contains(ActorLayer.Name))
+							{
+								TSharedPtr<FLightCardTreeItem> LightCardReference = MakeShared<FLightCardTreeItem>();
+								LightCardReference->LightCardActor = WeakActor.Get();
+								LightCardReference->ActorLayer = ActorLayer.Name;
+								LightCardActors.Add(LightCardReference);
+								break;
+							}
+						}
+					}
+				}
+			}
+		}
+	}
+
+	TrackedActors.Empty(LightCardActors.Num());
+
+	for (const TSharedPtr<FLightCardTreeItem>& LightCard : LightCardActors)
+	{
+		TrackedActors.Add(LightCard->LightCardActor.Get(), LightCard);
+	}
+	
+	// Check if the tree items have changed.
+	{
+		if (OriginalTree.Num() != LightCardTree.Num())
+		{
+			return true;
+		}
+	
+		for (const TSharedPtr<FLightCardTreeItem>& OriginalTreeItem : OriginalTree)
+		{
+			if (!LightCardTree.ContainsByPredicate([OriginalTreeItem](const TSharedPtr<FLightCardTreeItem>& LightCardTreeItem)
+			{
+				return OriginalTreeItem.IsValid() && LightCardTreeItem.IsValid() &&
+					OriginalTreeItem->LightCardActor.Get() == LightCardTreeItem->LightCardActor.Get() &&
+						OriginalTreeItem->ActorLayer == LightCardTreeItem->ActorLayer;
+			}))
+			{
+				return true;
+			}
+		}
+	}
+
+	return false;
+}
+
+void SDisplayClusterLightCardOutliner::OnOutlinerSelectionChanged(FSceneOutlinerTreeItemPtr TreeItem,
+	ESelectInfo::Type Type)
+{
+	MostRecentSelectedItem = TreeItem;
+
+	TArray<ADisplayClusterLightCardActor*> SelectedLightCards;
+	GetSelectedLightCards(SelectedLightCards);
+
+	IDisplayClusterOperator::Get().ShowDetailsForObjects(*reinterpret_cast<TArray<UObject*>*>(&SelectedLightCards));
+
+	if (LightCardEditorPtr.IsValid())
+	{
+		LightCardEditorPtr.Pin()->SelectLightCardProxies(SelectedLightCards);
+	}
+}
+
+TSharedPtr<SDisplayClusterLightCardOutliner::FLightCardTreeItem> SDisplayClusterLightCardOutliner::
+GetLightCardTreeItemFromOutliner(FSceneOutlinerTreeItemPtr InOutlinerTreeItem) const
+{
+	if (const FActorTreeItem* ActorTreeItem = InOutlinerTreeItem->CastTo<FActorTreeItem>())
+	{
+		if (const TSharedPtr<FLightCardTreeItem>* MatchingLightCard = TrackedActors.Find(
+			ActorTreeItem->Actor.Get()))
+		{
+			return *MatchingLightCard;
+		}
+	}
+
+	return nullptr;
+}
+
+AActor* SDisplayClusterLightCardOutliner::GetActorFromTreeItem(FSceneOutlinerTreeItemPtr InOutlinerTreeItem) const
+{
+	if (InOutlinerTreeItem != nullptr)
+	{
+		if (const FActorTreeItem* ActorTreeItem = InOutlinerTreeItem->CastTo<FActorTreeItem>())
+		{
+			return ActorTreeItem->Actor.Get();
+		}
+	}
+
+	return nullptr;
+}
+
+FReply SDisplayClusterLightCardOutliner::OnKeyDown(const FGeometry& MyGeometry, const FKeyEvent& InKeyEvent)
+{
+	if (CommandList->ProcessCommandBindings(InKeyEvent))
+	{
+		return FReply::Handled();
+	}
+	return FReply::Unhandled();
+}
+
+void SDisplayClusterLightCardOutliner::RegisterContextMenu(FName& InName, FToolMenuContext& InContext)
+{
+	static const FName LightCardEditorSectionName = "DynamicLightCardSection";
+	
+	UToolMenus* ToolMenus = UToolMenus::Get();
+	if (ensure(ToolMenus->IsMenuRegistered(InName)))
+	{
+		UToolMenu* Menu = ToolMenus->FindMenu(InName);
+
+		InContext.AddCleanup([Menu]()
+		{
+			Menu->RemoveSection(LightCardEditorSectionName);
+		});
+
+		check(LightCardEditorPtr.IsValid());
+		InContext.AppendCommandList(LightCardEditorPtr.Pin()->GetCommandList());
+		
+		Menu->AddDynamicSection(LightCardEditorSectionName, FNewToolMenuDelegate::CreateLambda([this](UToolMenu* InMenu)
+		{
+			USceneOutlinerMenuContext* Context = InMenu->FindContext<USceneOutlinerMenuContext>();
+			if (!Context || !Context->SceneOutliner.IsValid())
+			{
+				return;
+			}
+			
+			AActor* Actor = GetActorFromTreeItem(MostRecentSelectedItem.Pin());
+
+			if (Actor == nullptr && Context->SceneOutliner.Pin()->GetSelection().Num() == 1)
+			{
+				// If a delete operation was undone the most recent selection will be invalid.
+				// Assume a single entry light card is correct.
+				TArray<ADisplayClusterLightCardActor*> SelectedLightCardActors;
+				GetSelectedLightCards(SelectedLightCardActors);
+				if (SelectedLightCardActors.Num() == 1)
+				{
+					Actor = SelectedLightCardActors[0];
+				}
+			}
+			
+			bool bAddFullEditMenu = false;
+			
+			if (ADisplayClusterLightCardActor* LightCardActor = Cast<ADisplayClusterLightCardActor>(Actor))
+			{
+				FToolMenuSection& Section = InMenu->AddSection("LightCardDefaultSection", LOCTEXT("LightCardSectionName", "Light Cards"));
+
+				Section.AddMenuEntry(FDisplayClusterLightCardEditorCommands::Get().RemoveLightCard);
+				Section.AddMenuEntry(FDisplayClusterLightCardEditorCommands::Get().SaveLightCardTemplate);
+				bAddFullEditMenu = true;
+			}
+			
+			FToolMenuSection& Section = InMenu->AddSection("LightCardEditSection", LOCTEXT("LightCardEditSectionName", "Edit"));
+
+			if (bAddFullEditMenu)
+			{
+				Section.AddMenuEntry(FGenericCommands::Get().Cut);
+				Section.AddMenuEntry(FGenericCommands::Get().Copy);
+			}
+			
+			Section.AddMenuEntry(FGenericCommands::Get().Paste);
+			
+			if (bAddFullEditMenu)
+			{	
+				Section.AddMenuEntry(FGenericCommands::Get().Duplicate);
+				Section.AddMenuEntry(FGenericCommands::Get().Delete);
+			}
+		}));
+	}
+}
+
+#undef LOCTEXT_NAMESPACE
diff --git a/Engine/Plugins/Runtime/nDisplay/Source/DisplayClusterLightCardEditor/Private/Outliner/SDisplayClusterLightCardOutliner.h b/Engine/Plugins/Runtime/nDisplay/Source/DisplayClusterLightCardEditor/Private/Outliner/SDisplayClusterLightCardOutliner.h
new file mode 100644
--- /dev/null
+++ b/Engine/Plugins/Runtime/nDisplay/Source/DisplayClusterLightCardEditor/Private/Outliner/SDisplayClusterLightCardOutliner.h
@@ -0,0 +1,122 @@
+// Copyright Epic Games, Inc. All Rights Reserved.
+
+#pragma once
+
+#include "CoreMinimal.h"
+#include "SceneOutlinerFwd.h"
+#include "Widgets/DeclarativeSyntaxSupport.h"
+#include "Widgets/SCompoundWidget.h"
+
+class FExtender;
+class FUICommandList;
+class ADisplayClusterRootActor;
+class ADisplayClusterLightCardActor;
+class SDisplayClusterLightCardEditor;
+class ITableRow;
+class STableViewBase;
+
+class SSceneOutliner;
+template<class T>
+class STreeView;
+
+/** Displays all of the light cards associated with a particular display cluster root actor in a list view */
+class SDisplayClusterLightCardOutliner : public SCompoundWidget
+{
+public:
+	struct FLightCardTreeItem
+	{
+		TWeakObjectPtr<ADisplayClusterLightCardActor> LightCardActor;
+		FName ActorLayer;
+
+		FLightCardTreeItem() :
+			LightCardActor(nullptr),
+			ActorLayer(NAME_None)
+		{ }
+
+		bool IsInActorLayer() { return !ActorLayer.IsNone() && LightCardActor.IsValid(); }
+		bool IsActorLayer() { return !ActorLayer.IsNone() && !LightCardActor.IsValid(); }
+	};
+
+public:
+	SLATE_BEGIN_ARGS(SDisplayClusterLightCardOutliner)
+	{}
+	SLATE_END_ARGS()
+
+	virtual ~SDisplayClusterLightCardOutliner() override;
+
+	void Construct(const FArguments& InArgs, TSharedPtr<SDisplayClusterLightCardEditor> InLightCardEditor, TSharedPtr<FUICommandList> InCommandList);
+
+	void SetRootActor(ADisplayClusterRootActor* NewRootActor);
+
+	const TArray<TSharedPtr<FLightCardTreeItem>>& GetLightCardActors() const { return LightCardActors; }
+
+	/** Gets a list of light card actors that are currently selected in the list */
+	void GetSelectedLightCards(TArray<ADisplayClusterLightCardActor*>& OutSelectedLightCards) const;
+
+	/** Select light cards in the outliner and display their details */
+	void SelectLightCards(const TArray<ADisplayClusterLightCardActor*>& LightCardsToSelect);
+
+	/** Restores the last valid cached selection to the outliner */
+	void RestoreCachedSelection();
+
+private:
+	/** Creates a world outliner based on the root actor world */
+	void CreateWorldOutliner();
+	
+	/**
+	 * Fill the LightCard list with available LightCards
+	 * @return True if the list has been modified
+	 */
+	bool FillLightCardList();
+
+	/** Called when the user clicks on a selection in the outliner */
+	void OnOutlinerSelectionChanged(FSceneOutlinerTreeItemPtr TreeItem, ESelectInfo::Type Type);
+
+	/** Return our light card tree item from the outliner tree item */
+	TSharedPtr<FLightCardTreeItem> GetLightCardTreeItemFromOutliner(FSceneOutlinerTreeItemPtr InOutlinerTreeItem) const;
+
+	/** Return an actor if the tree item is an actor tree item */
+	AActor* GetActorFromTreeItem(FSceneOutlinerTreeItemPtr InOutlinerTreeItem) const;
+	
+	virtual FReply OnKeyDown(const FGeometry& MyGeometry, const FKeyEvent& InKeyEvent) override;
+
+	/** Registers a context menu for use with tool menus */
+	void RegisterContextMenu(FName& InName, struct FToolMenuContext& InContext);
+
+private:
+	/** Pointer to the light card editor that owns this widget */
+	TWeakPtr<SDisplayClusterLightCardEditor> LightCardEditorPtr;
+
+	/** A hierarchical list of light card tree items to be displayed in the tree view */
+	TArray<TSharedPtr<FLightCardTreeItem>> LightCardTree;
+
+	/** All actors currently tracked by the scene outliner */
+	TMap<TObjectPtr<AActor>, TSharedPtr<FLightCardTreeItem>> TrackedActors;
+
+	/** Cached actors currently selected. Used for synchronizing with the outliner mode */
+	TArray<TWeakObjectPtr<ADisplayClusterLightCardActor>> CachedSelectedActors;
+	
+	/** A flattened list of all the light card actors being displayed in the tree view */
+	TArray<TSharedPtr<FLightCardTreeItem>> LightCardActors;
+
+	/** The active root actor whose light cards are being displayed */
+	TWeakObjectPtr<ADisplayClusterRootActor> RootActor;
+
+	/** The scene outliner managed by this widget */
+	TSharedPtr<SSceneOutliner> SceneOutliner;
+
+	/** The most recently selected item from the outliner */
+	TWeakPtr<ISceneOutlinerTreeItem> MostRecentSelectedItem;
+
+	/** Called when the outliner makes a selection change */
+	FDelegateHandle SceneOutlinerSelectionChanged;
+	
+	/** The tree view widget displaying the light cards */
+	TSharedPtr<STreeView<TSharedPtr<FLightCardTreeItem>>> LightCardTreeView;
+	
+	/** Extenders for menus */
+	TSharedPtr<FExtender> Extenders;
+
+	/** Mapped commands for this list */
+	TSharedPtr<FUICommandList> CommandList;
+};
\ No newline at end of file
diff --git a/Engine/Plugins/Runtime/nDisplay/Source/DisplayClusterLightCardEditor/Private/SDisplayClusterLightCardEditor.cpp b/Engine/Plugins/Runtime/nDisplay/Source/DisplayClusterLightCardEditor/Private/SDisplayClusterLightCardEditor.cpp
--- a/Engine/Plugins/Runtime/nDisplay/Source/DisplayClusterLightCardEditor/Private/SDisplayClusterLightCardEditor.cpp
+++ b/Engine/Plugins/Runtime/nDisplay/Source/DisplayClusterLightCardEditor/Private/SDisplayClusterLightCardEditor.cpp
@@ -5,10 +5,10 @@
 #include "DisplayClusterLightCardEditorCommands.h"
 #include "DisplayClusterLightCardEditorStyle.h"
 #include "IDisplayClusterLightCardEditor.h"
-#include "SDisplayClusterLightCardList.h"
-#include "LightCardTemplates/SDisplayClusterLightCardTemplateList.h"
 #include "LightCardTemplates/DisplayClusterLightCardTemplate.h"
 #include "LightCardTemplates/DisplayClusterLightCardTemplateHelpers.h"
+#include "LightCardTemplates/SDisplayClusterLightCardTemplateList.h"
+#include "Outliner/SDisplayClusterLightCardOutliner.h"
 #include "Settings/DisplayClusterLightCardEditorSettings.h"
 
 #include "Viewport/DisplayClusterLightcardEditorViewport.h"
@@ -17,24 +17,24 @@
 #include "IDisplayClusterOperator.h"
 #include "IDisplayClusterOperatorViewModel.h"
 
-#include "DisplayClusterRootActor.h"
 #include "DisplayClusterConfigurationTypes.h"
+#include "DisplayClusterRootActor.h"
 #include "Components/DisplayClusterCameraComponent.h"
 
-#include "AssetRegistry/AssetRegistryModule.h"
 #include "ContentBrowserModule.h"
 #include "FileHelpers.h"
+#include "IContentBrowserSingleton.h"
+#include "PropertyCustomizationHelpers.h"
+#include "ScopedTransaction.h"
+#include "Selection.h"
+#include "AssetRegistry/AssetRegistryModule.h"
 #include "Framework/Commands/GenericCommands.h"
-#include "Framework/Docking/TabManager.h"
 #include "Framework/Docking/LayoutExtender.h"
+#include "Framework/Docking/TabManager.h"
 #include "Framework/MultiBox/MultiBoxBuilder.h"
 #include "Framework/MultiBox/MultiBoxExtender.h"
-#include "IContentBrowserSingleton.h"
 #include "Misc/FileHelper.h"
-#include "PropertyCustomizationHelpers.h"
-#include "ScopedTransaction.h"
 #include "Misc/TransactionObjectEvent.h"
-#include "Selection.h"
 #include "Styling/SlateIconFinder.h"
 #include "Toolkits/AssetEditorToolkit.h"
 #include "Widgets/Docking/SDockTab.h"
@@ -44,6 +44,7 @@
 #define LOCTEXT_NAMESPACE "SDisplayClusterLightCardEditor"
 
 const FName SDisplayClusterLightCardEditor::TabName = TEXT("DisplayClusterLightCardEditorTab");
+const FName SDisplayClusterLightCardEditor::OutlinerTabName = TEXT("DisplayClusterLightCardOutlinerTab");
 
 void SDisplayClusterLightCardEditor::RegisterTabSpawner()
 {
@@ -53,17 +54,28 @@ void SDisplayClusterLightCardEditor::RegisterTabSpawner()
 		.SetDisplayName(LOCTEXT("TabDisplayName", "Light Cards Editor"))
 		.SetTooltipText(LOCTEXT("TabTooltip", "Editing tools for nDisplay light cards."))
 		.SetMenuType(ETabSpawnerMenuType::Hidden);
+
+	// Outliner needs to be dockable within the owning operator panel. Register it as a nomad tab that the operator panel
+	// will invoke.
+	FGlobalTabmanager::Get()->RegisterNomadTabSpawner(OutlinerTabName, FOnSpawnTab::CreateStatic(&SDisplayClusterLightCardEditor::SpawnOutlinerTab))
+	.SetDisplayName(LOCTEXT("OutlinerTab_DisplayName", "Outliner"))
+	.SetTooltipText(LOCTEXT("OutlinerTab_Tooltip", "Light Card Editor Outliner"))
+	.SetMenuType(ETabSpawnerMenuType::Hidden);
 }
 
 void SDisplayClusterLightCardEditor::UnregisterTabSpawner()
 {
 	FGlobalTabmanager::Get()->UnregisterNomadTabSpawner(TabName);
+	FGlobalTabmanager::Get()->UnregisterNomadTabSpawner(OutlinerTabName);
 }
 
 void SDisplayClusterLightCardEditor::RegisterLayoutExtension(FLayoutExtender& InExtender)
 {
 	FTabManager::FTab NewTab(FTabId(TabName, ETabIdFlags::SaveLayout), ETabState::OpenedTab);
 	InExtender.ExtendStack(IDisplayClusterOperator::Get().GetPrimaryOperatorExtensionId(), ELayoutExtensionPosition::After, NewTab);
+
+	FTabManager::FTab OutlinerTab(FTabId(OutlinerTabName, ETabIdFlags::SaveLayout), ETabState::OpenedTab);
+	InExtender.ExtendLayout(IDisplayClusterOperator::Get().GetDetailsTabId(), ELayoutExtensionPosition::Above, OutlinerTab);
 }
 
 TSharedRef<SDockTab> SDisplayClusterLightCardEditor::SpawnInTab(const FSpawnTabArgs& SpawnTabArgs)
@@ -78,6 +90,33 @@ TSharedRef<SDockTab> SDisplayClusterLightCardEditor::SpawnInTab(const FSpawnTabA
 	return MajorTab;
 }
 
+TSharedRef<SDockTab> SDisplayClusterLightCardEditor::SpawnOutlinerTab(const FSpawnTabArgs& SpawnTabArgs)
+{
+	TSharedPtr<SWidget> OutlinerWidget = SNullWidget::NullWidget;
+	if (const TSharedPtr<FTabManager> OperatorPanelTabManager = IDisplayClusterOperator::Get().GetOperatorViewModel()->GetTabManager())
+	{
+		TSharedPtr<SDockTab> ExistingTab = OperatorPanelTabManager->FindExistingLiveTab(TabName);
+		if (!ExistingTab.IsValid())
+		{
+			ExistingTab = OperatorPanelTabManager->TryInvokeTab(TabName);
+		}
+
+		if (ExistingTab.IsValid())
+		{
+			const TSharedRef<SDisplayClusterLightCardEditor> EditorPtr = StaticCastSharedRef<SDisplayClusterLightCardEditor>(ExistingTab->GetContent());
+			OutlinerWidget = EditorPtr->CreateLightCardOutlinerWidget();
+		}
+	}
+		
+	return SNew(SDockTab)
+	.CanEverClose(false)
+	.OnCanCloseTab(SDockTab::FCanCloseTab::CreateLambda([](){ return false; }))
+	.Content()
+	[
+		OutlinerWidget.ToSharedRef()
+	];
+}
+
 SDisplayClusterLightCardEditor::~SDisplayClusterLightCardEditor()
 {
 	IDisplayClusterOperator::Get().GetOperatorViewModel()->OnActiveRootActorChanged().Remove(ActiveRootActorChangedHandle);
@@ -121,66 +160,9 @@ void SDisplayClusterLightCardEditor::Construct(const FArguments& InArgs, const T
 	BindCommands();
 	RegisterToolbarExtensions();
 	
-	TabManager = FGlobalTabmanager::Get()->NewTabManager(MajorTabOwner);
-
-	const FName LightCardTab = TEXT("LightCards");
-	const FName ViewportTab = TEXT("Viewport"); 
-	
-	const TSharedRef<FTabManager::FLayout> Layout = FTabManager::NewLayout("LightCardEditor")
-	->AddArea
-	(
-		FTabManager::NewPrimaryArea()
-		
-		->Split
-		(
-			FTabManager::NewSplitter()
-			->SetSizeCoefficient(0.25f)
-			->SetOrientation(Orient_Vertical)
-			->Split
-			(
-			FTabManager::NewStack()
-				->AddTab(LightCardTab, ETabState::OpenedTab)
-				->SetHideTabWell(true)
-				->SetForegroundTab(LightCardTab)
-			)
-		)
-		->Split
-		(
-			FTabManager::NewStack()
-			->AddTab(ViewportTab, ETabState::OpenedTab)
-			->SetSizeCoefficient(0.75f)
-			->SetHideTabWell(true)
-			->SetForegroundTab(ViewportTab)
-		)
-	);
-
-	TabManager->RegisterTabSpawner(LightCardTab, FOnSpawnTab::CreateLambda([this](const FSpawnTabArgs&) -> TSharedRef<SDockTab>
-	{
-		return SNew(SDockTab)
-		.TabRole(ETabRole::PanelTab)
-		.CanEverClose(false)
-		.OnCanCloseTab(SDockTab::FCanCloseTab::CreateLambda([](){ return false; }))
-		.ContentPadding(0)
-		[
-			CreateLightCardListWidget()
-		];
-	}));
-	
-	TabManager->RegisterTabSpawner(ViewportTab, FOnSpawnTab::CreateLambda([this](const FSpawnTabArgs&) -> TSharedRef<SDockTab>
-	{
-		return SNew(SDockTab)
-		.TabRole(ETabRole::PanelTab)
-		.CanEverClose(false)
-		.OnCanCloseTab(SDockTab::FCanCloseTab::CreateLambda([](){ return false; }))
-		.ContentPadding(0)
-		[
-			CreateViewportWidget()
-		];
-	}));
-	
 	ChildSlot                     
 	[
-		TabManager->RestoreFrom(Layout, WindowOwner).ToSharedRef()
+		CreateViewportWidget()
 	];
 
 	BindCompileDelegates();
@@ -192,14 +174,14 @@ void SDisplayClusterLightCardEditor::Construct(const FArguments& InArgs, const T
 
 void SDisplayClusterLightCardEditor::SelectLightCards(const TArray<ADisplayClusterLightCardActor*>& LightCardsToSelect)
 {
-	check(LightCardList);
-	LightCardList->SelectLightCards(LightCardsToSelect);
+	check(LightCardOutliner);
+	LightCardOutliner->SelectLightCards(LightCardsToSelect);
 }
 
 void SDisplayClusterLightCardEditor::GetSelectedLightCards(TArray<ADisplayClusterLightCardActor*>& OutSelectedLightCards)
 {
-	check(LightCardList);
-	LightCardList->GetSelectedLightCards(OutSelectedLightCards);
+	check(LightCardOutliner);
+	LightCardOutliner->GetSelectedLightCards(OutSelectedLightCards);
 }
 
 void SDisplayClusterLightCardEditor::SelectLightCardProxies(const TArray<ADisplayClusterLightCardActor*>& LightCardsToSelect)
@@ -214,7 +196,7 @@ void SDisplayClusterLightCardEditor::CenterLightCardInView(ADisplayClusterLightC
 	ViewportView->GetLightCardEditorViewportClient()->CenterLightCardInView(LightCard);
 }
 
-ADisplayClusterLightCardActor* SDisplayClusterLightCardEditor::SpawnLightCard()
+ADisplayClusterLightCardActor* SDisplayClusterLightCardEditor::SpawnLightCard(const FName& LightCardName)
 {
 	if (!ActiveRootActor.IsValid())
 	{
@@ -228,7 +210,7 @@ ADisplayClusterLightCardActor* SDisplayClusterLightCardEditor::SpawnLightCard()
 	FActorSpawnParameters SpawnParameters;
 	SpawnParameters.bNoFail = true;
 	SpawnParameters.SpawnCollisionHandlingOverride = ESpawnActorCollisionHandlingMethod::AdjustIfPossibleButAlwaysSpawn;
-	SpawnParameters.Name = TEXT("LightCard");
+	SpawnParameters.Name = LightCardName;
 	SpawnParameters.NameMode = FActorSpawnParameters::ESpawnActorNameMode::Requested;
 	SpawnParameters.OverrideLevel = ActiveRootActor->GetWorld()->GetCurrentLevel();
 
@@ -291,6 +273,7 @@ ADisplayClusterLightCardActor* SDisplayClusterLightCardEditor::SpawnLightCardFro
 		RecentlyPlacedItem.ItemType = FDisplayClusterLightCardEditorRecentItem::Type_LightCardTemplate;
 		
 		AddRecentlyPlacedItem(MoveTemp(RecentlyPlacedItem));
+		SelectLightCards({NewLightCard});
 	}
 	
 	return NewLightCard;
@@ -314,6 +297,7 @@ void SDisplayClusterLightCardEditor::AddNewLightCard()
 		RecentlyPlacedItem.ItemType = FDisplayClusterLightCardEditorRecentItem::Type_LightCard;
 		
 		AddRecentlyPlacedItem(MoveTemp(RecentlyPlacedItem));
+		SelectLightCards({NewLightCard});
 	}
 }
 
@@ -425,6 +409,28 @@ void SDisplayClusterLightCardEditor::AddExistingLightCard()
 	SelectedActorPtr.Reset();
 }
 
+void SDisplayClusterLightCardEditor::AddNewFlag()
+{
+	check(ActiveRootActor.IsValid());
+
+	FScopedTransaction Transaction(LOCTEXT("AddNewFlagTransactionMessage", "Add New Flag"));
+
+	// When adding a new lightcard, usually the desired location is in the middle of the viewport
+	if (ADisplayClusterLightCardActor* NewLightCard = SpawnLightCard(TEXT("Flag")))
+	{
+		NewLightCard->Color = FLinearColor(0.f, 0.f, 0.f);
+		
+		CenterLightCardInView(*NewLightCard);
+
+		FDisplayClusterLightCardEditorRecentItem RecentlyPlacedItem;
+		RecentlyPlacedItem.ObjectPath = NewLightCard->GetClass();
+		RecentlyPlacedItem.ItemType = FDisplayClusterLightCardEditorRecentItem::Type_Flag;
+		
+		AddRecentlyPlacedItem(MoveTemp(RecentlyPlacedItem));
+		SelectLightCards({NewLightCard});
+	}
+}
+
 void SDisplayClusterLightCardEditor::AddLightCardsToActor(TArray<ADisplayClusterLightCardActor*> LightCards)
 {
 	if (ActiveRootActor.IsValid())
@@ -469,8 +475,12 @@ bool SDisplayClusterLightCardEditor::CanAddLightCard() const
 
 void SDisplayClusterLightCardEditor::CutLightCards()
 {
+	TArray<ADisplayClusterLightCardActor*> SelectedLightCards;
+	LightCardOutliner->GetSelectedLightCards(SelectedLightCards);
+	
 	CopyLightCards();
-	RemoveLightCards(true);
+
+	RemoveLightCards(SelectedLightCards, /*bDeleteLightCards*/ true);
 }
 
 bool SDisplayClusterLightCardEditor::CanCutLightCards()
@@ -481,7 +491,7 @@ bool SDisplayClusterLightCardEditor::CanCutLightCards()
 void SDisplayClusterLightCardEditor::CopyLightCards()
 {
 	TArray<ADisplayClusterLightCardActor*> SelectedLightCards;
-	LightCardList->GetSelectedLightCards(SelectedLightCards);
+	LightCardOutliner->GetSelectedLightCards(SelectedLightCards);
 
 	UWorld* EditorWorld = GEditor->GetEditorWorldContext().World();
 
@@ -507,7 +517,7 @@ void SDisplayClusterLightCardEditor::CopyLightCards()
 bool SDisplayClusterLightCardEditor::CanCopyLightCards() const
 {
 	TArray<ADisplayClusterLightCardActor*> SelectedLightCards;
-	LightCardList->GetSelectedLightCards(SelectedLightCards);
+	LightCardOutliner->GetSelectedLightCards(SelectedLightCards);
 
 	return SelectedLightCards.Num() > 0;
 }
@@ -565,14 +575,20 @@ bool SDisplayClusterLightCardEditor::CanDuplicateLightCards() const
 void SDisplayClusterLightCardEditor::RemoveLightCards(bool bDeleteLightCardActor)
 {
 	TArray<ADisplayClusterLightCardActor*> SelectedLightCards;
-	LightCardList->GetSelectedLightCards(SelectedLightCards);
+	LightCardOutliner->GetSelectedLightCards(SelectedLightCards);
+
+	RemoveLightCards(SelectedLightCards, bDeleteLightCardActor);
+}
 
+void SDisplayClusterLightCardEditor::RemoveLightCards(
+	const TArray<ADisplayClusterLightCardActor*>& InLightCardsToRemove, bool bDeleteLightCards)
+{
 	FScopedTransaction Transaction(LOCTEXT("RemoveLightCardTransactionMessage", "Remove Light Card(s)"));
 
 	USelection* EdSelectionManager = GEditor->GetSelectedActors();
 	UWorld* WorldToUse = nullptr;
 
-	if (bDeleteLightCardActor)
+	if (bDeleteLightCards)
 	{
 		EdSelectionManager->BeginBatchSelectOperation();
 		EdSelectionManager->Modify();
@@ -586,14 +602,14 @@ void SDisplayClusterLightCardEditor::RemoveLightCards(bool bDeleteLightCardActor
 
 		FDisplayClusterConfigurationICVFX_VisibilityList& RootActorLightCards = ConfigData->StageSettings.Lightcard.ShowOnlyList;
 
-		for (ADisplayClusterLightCardActor* LightCard : SelectedLightCards)
+		for (ADisplayClusterLightCardActor* LightCard : InLightCardsToRemove)
 		{
 			RootActorLightCards.Actors.RemoveAll([&](const TSoftObjectPtr<AActor>& Actor)
 				{
 					return Actor.Get() == LightCard;
 				});
 
-			if (bDeleteLightCardActor)
+			if (bDeleteLightCards)
 			{
 				WorldToUse = LightCard->GetWorld();
 				GEditor->SelectActor(LightCard, /*bSelect =*/true, /*bNotifyForActor =*/false, /*bSelectEvenIfHidden =*/true);
@@ -605,7 +621,7 @@ void SDisplayClusterLightCardEditor::RemoveLightCards(bool bDeleteLightCardActor
 		}
 	}
 
-	if (bDeleteLightCardActor)
+	if (bDeleteLightCards)
 	{
 		EdSelectionManager->EndBatchSelectOperation();
 
@@ -621,7 +637,7 @@ void SDisplayClusterLightCardEditor::RemoveLightCards(bool bDeleteLightCardActor
 bool SDisplayClusterLightCardEditor::CanRemoveLightCards() const
 {
 	TArray<ADisplayClusterLightCardActor*> SelectedLightCards;
-	LightCardList->GetSelectedLightCards(SelectedLightCards);
+	LightCardOutliner->GetSelectedLightCards(SelectedLightCards);
 
 	return SelectedLightCards.Num() > 0;
 }
@@ -629,7 +645,7 @@ bool SDisplayClusterLightCardEditor::CanRemoveLightCards() const
 void SDisplayClusterLightCardEditor::CreateLightCardTemplate()
 {
 	TArray<ADisplayClusterLightCardActor*> SelectedLightCards;
-	LightCardList->GetSelectedLightCards(SelectedLightCards);
+	LightCardOutliner->GetSelectedLightCards(SelectedLightCards);
 
 	check(SelectedLightCards.Num() == 1);
 	const ADisplayClusterLightCardActor* LightCardActor = SelectedLightCards[0];
@@ -692,7 +708,7 @@ void SDisplayClusterLightCardEditor::CreateLightCardTemplate()
 bool SDisplayClusterLightCardEditor::CanCreateLightCardTemplate() const
 {
 	TArray<ADisplayClusterLightCardActor*> SelectedLightCards;
-	LightCardList->GetSelectedLightCards(SelectedLightCards);
+	LightCardOutliner->GetSelectedLightCards(SelectedLightCards);
 
 	return SelectedLightCards.Num() == 1;
 }
@@ -763,7 +779,7 @@ void SDisplayClusterLightCardEditor::OnActiveRootActorChanged(ADisplayClusterRoo
 	
 	// The new root actor pointer could be null, indicating that it was deleted or the user didn't select a valid root actor
 	ActiveRootActor = NewRootActor;
-	LightCardList->SetRootActor(NewRootActor);
+	LightCardOutliner->SetRootActor(NewRootActor);
 	ViewportView->SetRootActor(NewRootActor);
 	
 	BindCompileDelegates();
@@ -774,9 +790,9 @@ void SDisplayClusterLightCardEditor::OnActiveRootActorChanged(ADisplayClusterRoo
 	RefreshLabels();
 }
 
-TSharedRef<SWidget> SDisplayClusterLightCardEditor::CreateLightCardListWidget()
+TSharedRef<SWidget> SDisplayClusterLightCardEditor::CreateLightCardOutlinerWidget()
 {
-	return SAssignNew(LightCardList, SDisplayClusterLightCardList, SharedThis(this), CommandList);
+	return SAssignNew(LightCardOutliner, SDisplayClusterLightCardOutliner, SharedThis(this), CommandList);
 }
 
 TSharedRef<SWidget> SDisplayClusterLightCardEditor::CreateViewportWidget()
@@ -792,7 +808,9 @@ TSharedRef<SWidget> SDisplayClusterLightCardEditor::GeneratePlaceActorsMenu()
 	MenuBuilder.BeginSection("PlaceActors", LOCTEXT("PlaceActorsMenuHeader", "Place Actors"));
 	{
 		FSlateIcon LightCardIcon = FSlateIconFinder::FindIconForClass(ADisplayClusterLightCardActor::StaticClass());
-		
+
+		MenuBuilder.AddMenuEntry(FDisplayClusterLightCardEditorCommands::Get().AddNewFlag,
+			NAME_None, TAttribute<FText>(), TAttribute<FText>(), LightCardIcon);
 		MenuBuilder.AddMenuEntry(FDisplayClusterLightCardEditorCommands::Get().AddNewLightCard,
 			NAME_None, TAttribute<FText>(), TAttribute<FText>(), LightCardIcon);
 		MenuBuilder.AddMenuEntry(FDisplayClusterLightCardEditorCommands::Get().AddExistingLightCard,
@@ -912,6 +930,10 @@ TSharedRef<SWidget> SDisplayClusterLightCardEditor::GeneratePlaceActorsMenu()
 					{
 						AddNewLightCard();
 					}
+					else if (RecentlyPlacedItem.ItemType == FDisplayClusterLightCardEditorRecentItem::Type_Flag)
+					{
+						AddNewFlag();
+					}
 					else if (RecentlyPlacedItem.ItemType == FDisplayClusterLightCardEditorRecentItem::Type_LightCardTemplate)
 					{
 						if (const UDisplayClusterLightCardTemplate* Template = Cast<UDisplayClusterLightCardTemplate>(ObjectWeakPtr.Get()))
@@ -1014,6 +1036,11 @@ void SDisplayClusterLightCardEditor::BindCommands()
 		FExecuteAction::CreateSP(this, &SDisplayClusterLightCardEditor::AddNewLightCard),
 		FCanExecuteAction::CreateSP(this, &SDisplayClusterLightCardEditor::CanAddLightCard));
 
+	CommandList->MapAction(
+		FDisplayClusterLightCardEditorCommands::Get().AddNewFlag,
+		FExecuteAction::CreateSP(this, &SDisplayClusterLightCardEditor::AddNewFlag),
+		FCanExecuteAction::CreateSP(this, &SDisplayClusterLightCardEditor::CanAddLightCard));
+	
 	CommandList->MapAction(
 		FDisplayClusterLightCardEditorCommands::Get().AddExistingLightCard,
 		FExecuteAction::CreateSP(this, &SDisplayClusterLightCardEditor::AddExistingLightCard),
@@ -1103,9 +1130,9 @@ void SDisplayClusterLightCardEditor::RefreshPreviewActors(EDisplayClusterLightCa
 	
 	if (ADisplayClusterRootActor* RootActor = GetActiveRootActor().Get())
 	{
-		if (LightCardList.IsValid())
+		if (LightCardOutliner.IsValid())
 		{
-			LightCardList->SetRootActor(RootActor);
+			LightCardOutliner->SetRootActor(RootActor);
 		}
 			
 		if (ViewportView.IsValid())
@@ -1147,9 +1174,9 @@ bool SDisplayClusterLightCardEditor::IsOurObject(UObject* InObject,
 	EDisplayClusterLightCardEditorProxyType ProxyType = EDisplayClusterLightCardEditorProxyType::All;
 	
 	bool bIsOurActor = IsOurActor(GetActiveRootActor().Get());
-	if (!bIsOurActor && LightCardList.IsValid())
+	if (!bIsOurActor && LightCardOutliner.IsValid())
 	{
-		for (const TSharedPtr<SDisplayClusterLightCardList::FLightCardTreeItem>& LightCard : LightCardList->GetLightCardActors())
+		for (const TSharedPtr<SDisplayClusterLightCardOutliner::FLightCardTreeItem>& LightCard : LightCardOutliner->GetLightCardActors())
 		{
 			bIsOurActor = IsOurActor(LightCard->LightCardActor.Get());
 			if (bIsOurActor)
@@ -1166,9 +1193,9 @@ bool SDisplayClusterLightCardEditor::IsOurObject(UObject* InObject,
 
 void SDisplayClusterLightCardEditor::BindCompileDelegates()
 {
-	if (LightCardList.IsValid())
+	if (LightCardOutliner.IsValid())
 	{
-		for (const TSharedPtr<SDisplayClusterLightCardList::FLightCardTreeItem>& LightCardActor : LightCardList->GetLightCardActors())
+		for (const TSharedPtr<SDisplayClusterLightCardOutliner::FLightCardTreeItem>& LightCardActor : LightCardOutliner->GetLightCardActors())
 		{
 			if (LightCardActor.IsValid() && LightCardActor->LightCardActor.IsValid())
 			{
@@ -1183,9 +1210,9 @@ void SDisplayClusterLightCardEditor::BindCompileDelegates()
 
 void SDisplayClusterLightCardEditor::RemoveCompileDelegates()
 {
-	if (LightCardList.IsValid())
+	if (LightCardOutliner.IsValid())
 	{
-		for (const TSharedPtr<SDisplayClusterLightCardList::FLightCardTreeItem>& LightCardActor : LightCardList->GetLightCardActors())
+		for (const TSharedPtr<SDisplayClusterLightCardOutliner::FLightCardTreeItem>& LightCardActor : LightCardOutliner->GetLightCardActors())
 		{
 			if (LightCardActor.IsValid() && LightCardActor->LightCardActor.IsValid())
 			{
@@ -1222,8 +1249,8 @@ void SDisplayClusterLightCardEditor::OnActorPropertyChanged(UObject* ObjectBeing
 
 void SDisplayClusterLightCardEditor::OnLevelActorDeleted(AActor* Actor)
 {
-	if (LightCardList.IsValid() &&
-		LightCardList->GetLightCardActors().ContainsByPredicate([Actor](const TSharedPtr<SDisplayClusterLightCardList::FLightCardTreeItem>& LightCardTreeItem)
+	if (LightCardOutliner.IsValid() &&
+		LightCardOutliner->GetLightCardActors().ContainsByPredicate([Actor](const TSharedPtr<SDisplayClusterLightCardOutliner::FLightCardTreeItem>& LightCardTreeItem)
 	{
 		return LightCardTreeItem.IsValid() && Actor == LightCardTreeItem->LightCardActor.Get();
 	}))
@@ -1265,11 +1292,6 @@ void SDisplayClusterLightCardEditor::OnObjectTransacted(UObject* Object,
 	}
 }
 
-void SDisplayClusterLightCardEditor::OnLightCardListChanged()
-{
-	RefreshPreviewActors(EDisplayClusterLightCardEditorProxyType::LightCards);
-}
-
 void SDisplayClusterLightCardEditor::AddRecentlyPlacedItem(const FDisplayClusterLightCardEditorRecentItem& InItem)
 {
 	UDisplayClusterLightCardEditorSettings* Settings = GetMutableDefault<UDisplayClusterLightCardEditorSettings>();
diff --git a/Engine/Plugins/Runtime/nDisplay/Source/DisplayClusterLightCardEditor/Private/SDisplayClusterLightCardEditor.h b/Engine/Plugins/Runtime/nDisplay/Source/DisplayClusterLightCardEditor/Private/SDisplayClusterLightCardEditor.h
--- a/Engine/Plugins/Runtime/nDisplay/Source/DisplayClusterLightCardEditor/Private/SDisplayClusterLightCardEditor.h
+++ b/Engine/Plugins/Runtime/nDisplay/Source/DisplayClusterLightCardEditor/Private/SDisplayClusterLightCardEditor.h
@@ -17,7 +17,7 @@ class FTabManager;
 class FToolBarBuilder;
 class FUICommandList;
 class SDockTab;
-class SDisplayClusterLightCardList;
+class SDisplayClusterLightCardOutliner;
 class SDisplayClusterLightCardTemplateList;
 class SDisplayClusterLightCardEditorViewport;
 class ADisplayClusterRootActor;
@@ -33,6 +33,9 @@ public:
 	/** The name of the tab that the light card editor lives in */
 	static const FName TabName;
 
+	/** The name of the tab that the outliner lives in */
+	static const FName OutlinerTabName;
+
 	/** Registers the light card editor with the global tab manager and adds it to the operator panel's extension tab stack */
 	static void RegisterTabSpawner();
 
@@ -45,6 +48,9 @@ public:
 	/** Creates a tab with the light card editor inside */
 	static TSharedRef<SDockTab> SpawnInTab(const FSpawnTabArgs& SpawnTabArgs);
 
+	/** Spawns in the outliner tab */
+	static TSharedRef<SDockTab> SpawnOutlinerTab(const FSpawnTabArgs& SpawnTabArgs);
+	
 	SLATE_BEGIN_ARGS(SDisplayClusterLightCardEditor)
 	{}
 	SLATE_END_ARGS()
@@ -56,8 +62,11 @@ public:
 	virtual void PostRedo(bool bSuccess) override;
 	// ~FEditorUndoClient
 	
-	void Construct(const FArguments& args, const TSharedRef<SDockTab>& MajorTabOwner, const TSharedPtr<SWindow>& WindowOwner);
+	void Construct(const FArguments& InArgs, const TSharedRef<SDockTab>& MajorTabOwner, const TSharedPtr<SWindow>& WindowOwner);
 
+	/** Return the command list for the light card editor */
+	TSharedPtr<FUICommandList> GetCommandList() const { return CommandList; }
+	
 	/** The current active root actor for this light card editor */
 	TWeakObjectPtr<ADisplayClusterRootActor> GetActiveRootActor() const { return ActiveRootActor; }
 
@@ -74,7 +83,7 @@ public:
 	void CenterLightCardInView(ADisplayClusterLightCardActor& LightCard);
 
 	/** Spawns a new light card and adds it to the root actor */
-	ADisplayClusterLightCardActor* SpawnLightCard();
+	ADisplayClusterLightCardActor* SpawnLightCard(const FName& LightCardName = TEXT("LightCard"));
 
 	/** Spawns a new light card from a light card template */
 	ADisplayClusterLightCardActor* SpawnLightCardFromTemplate(const UDisplayClusterLightCardTemplate* InTemplate, ULevel* InLevel = nullptr, bool bIsPreview = false);
@@ -85,6 +94,9 @@ public:
 	/** Select an existing Light Card from a menu */
 	void AddExistingLightCard();
 
+	/** Adds a new light card configured as a flag */
+	void AddNewFlag();
+
 	/** Adds the given Light Card to the root actor */
 	void AddLightCardsToActor(TArray<ADisplayClusterLightCardActor*> LightCards);
 
@@ -121,6 +133,13 @@ public:
 	 */
 	void RemoveLightCards(bool bDeleteLightCardActor);
 
+	/**
+	 * Remove the given light cards from the actor
+	 * @param InLightCardsToRemove Light cards which should be removed
+	 * @param bDeleteLightCards If the light cards should be deleted from the level
+	 */
+	void RemoveLightCards(const TArray<ADisplayClusterLightCardActor*>& InLightCardsToRemove, bool bDeleteLightCards);
+	
 	/** If the selected Light Card can be removed */
 	bool CanRemoveLightCards() const;
 
@@ -150,7 +169,7 @@ private:
 	void OnActiveRootActorChanged(ADisplayClusterRootActor* NewRootActor);
 
 	/** Creates the widget used to show the list of light cards associated with the active root actor */
-	TSharedRef<SWidget> CreateLightCardListWidget();
+	TSharedRef<SWidget> CreateLightCardOutlinerWidget();
 
 	/** Create the widget for selecting light card templates */
 	TSharedRef<SWidget> CreateLightCardTemplateWidget();
@@ -211,9 +230,6 @@ private:
 	/** Raised when any object is transacted */
 	void OnObjectTransacted(UObject* Object, const FTransactionObjectEvent& TransactionObjectEvent);
 
-	/** Raised when the light card list has added or removed a card */
-	void OnLightCardListChanged();
-
 	/** Add an item to recently placed list. Reorders matching items to the top of the list and trims the list */
 	void AddRecentlyPlacedItem(const FDisplayClusterLightCardEditorRecentItem& InItem);
 
@@ -221,10 +237,8 @@ private:
 	void CleanupRecentlyPlacedItems();
 	
 private:
-	TSharedPtr<FTabManager> TabManager;
-	
-	/** The light card list widget */
-	TSharedPtr<SDisplayClusterLightCardList> LightCardList;
+	/** The light card outliner widget */
+	TSharedPtr<SDisplayClusterLightCardOutliner> LightCardOutliner;
 
 	/** Templates to their icon brushes */
 	TMap<TWeakObjectPtr<UDisplayClusterLightCardTemplate>, TSharedPtr<FSlateBrush>> TemplateBrushes;
diff --git a/Engine/Plugins/Runtime/nDisplay/Source/DisplayClusterLightCardEditor/Private/SDisplayClusterLightCardList.cpp b/Engine/Plugins/Runtime/nDisplay/Source/DisplayClusterLightCardEditor/Private/SDisplayClusterLightCardList.cpp
deleted file mode 100644
diff --git a/Engine/Plugins/Runtime/nDisplay/Source/DisplayClusterLightCardEditor/Private/SDisplayClusterLightCardList.h b/Engine/Plugins/Runtime/nDisplay/Source/DisplayClusterLightCardEditor/Private/SDisplayClusterLightCardList.h
deleted file mode 100644
diff --git a/Engine/Plugins/Runtime/nDisplay/Source/DisplayClusterLightCardEditor/Private/Settings/DisplayClusterLightCardEditorSettings.cpp b/Engine/Plugins/Runtime/nDisplay/Source/DisplayClusterLightCardEditor/Private/Settings/DisplayClusterLightCardEditorSettings.cpp
--- a/Engine/Plugins/Runtime/nDisplay/Source/DisplayClusterLightCardEditor/Private/Settings/DisplayClusterLightCardEditorSettings.cpp
+++ b/Engine/Plugins/Runtime/nDisplay/Source/DisplayClusterLightCardEditor/Private/Settings/DisplayClusterLightCardEditorSettings.cpp
@@ -23,6 +23,7 @@ UDisplayClusterLightCardEditorProjectSettings::UDisplayClusterLightCardEditorPro
 }
 
 const FName FDisplayClusterLightCardEditorRecentItem::Type_LightCard = "LightCard";
+const FName FDisplayClusterLightCardEditorRecentItem::Type_Flag = "Flag";
 const FName FDisplayClusterLightCardEditorRecentItem::Type_LightCardTemplate = "LightCardTemplate";
 
 FText FDisplayClusterLightCardEditorRecentItem::GetItemDisplayName() const
@@ -31,6 +32,10 @@ FText FDisplayClusterLightCardEditorRecentItem::GetItemDisplayName() const
 	{
 		return LOCTEXT("LightCardRecentItemName", "Light Card");
 	}
+	if (ItemType == Type_Flag)
+	{
+		return LOCTEXT("FlagRecentItemName", "Flag");
+	}
 	if (const UObject* Object = ObjectPath.LoadSynchronous())
 	{
 		return FText::FromString(Object->GetName());
@@ -41,7 +46,7 @@ FText FDisplayClusterLightCardEditorRecentItem::GetItemDisplayName() const
 
 const FSlateBrush* FDisplayClusterLightCardEditorRecentItem::GetSlateBrush() const
 {
-	if (ItemType == Type_LightCard)
+	if (ItemType == Type_LightCard || ItemType == Type_Flag)
 	{
 		return FSlateIconFinder::FindIconBrushForClass(ADisplayClusterLightCardActor::StaticClass());
 	}
diff --git a/Engine/Plugins/Runtime/nDisplay/Source/DisplayClusterLightCardEditor/Private/Settings/DisplayClusterLightCardEditorSettings.h b/Engine/Plugins/Runtime/nDisplay/Source/DisplayClusterLightCardEditor/Private/Settings/DisplayClusterLightCardEditorSettings.h
--- a/Engine/Plugins/Runtime/nDisplay/Source/DisplayClusterLightCardEditor/Private/Settings/DisplayClusterLightCardEditorSettings.h
+++ b/Engine/Plugins/Runtime/nDisplay/Source/DisplayClusterLightCardEditor/Private/Settings/DisplayClusterLightCardEditorSettings.h
@@ -38,6 +38,7 @@ struct FDisplayClusterLightCardEditorRecentItem
 	GENERATED_BODY()
 	
 	static const FName Type_LightCard;
+	static const FName Type_Flag;
 	static const FName Type_LightCardTemplate;
 
 	/** Path of object to instantiate */
diff --git a/Engine/Plugins/Runtime/nDisplay/Source/DisplayClusterOperator/Private/DisplayClusterOperatorModule.cpp b/Engine/Plugins/Runtime/nDisplay/Source/DisplayClusterOperator/Private/DisplayClusterOperatorModule.cpp
--- a/Engine/Plugins/Runtime/nDisplay/Source/DisplayClusterOperator/Private/DisplayClusterOperatorModule.cpp
+++ b/Engine/Plugins/Runtime/nDisplay/Source/DisplayClusterOperator/Private/DisplayClusterOperatorModule.cpp
@@ -47,6 +47,11 @@ FName FDisplayClusterOperatorModule::GetAuxilliaryOperatorExtensionId()
 	return SDisplayClusterOperatorPanel::AuxilliaryTabExtensionId;
 }
 
+FName FDisplayClusterOperatorModule::GetDetailsTabId()
+{
+	return SDisplayClusterOperatorPanel::DetailsTabId;
+}
+
 void FDisplayClusterOperatorModule::GetRootActorLevelInstances(TArray<ADisplayClusterRootActor*>& OutRootActorInstances)
 {
 	UWorld* World = GCurrentLevelEditingViewportClient ? GCurrentLevelEditingViewportClient->GetWorld() : nullptr;
diff --git a/Engine/Plugins/Runtime/nDisplay/Source/DisplayClusterOperator/Private/DisplayClusterOperatorModule.h b/Engine/Plugins/Runtime/nDisplay/Source/DisplayClusterOperator/Private/DisplayClusterOperatorModule.h
--- a/Engine/Plugins/Runtime/nDisplay/Source/DisplayClusterOperator/Private/DisplayClusterOperatorModule.h
+++ b/Engine/Plugins/Runtime/nDisplay/Source/DisplayClusterOperator/Private/DisplayClusterOperatorModule.h
@@ -33,6 +33,7 @@ public:
 
 	virtual FName GetPrimaryOperatorExtensionId() override;
 	virtual FName GetAuxilliaryOperatorExtensionId() override;
+	virtual FName GetDetailsTabId() override;
 	virtual TSharedPtr<FExtensibilityManager> GetOperatorToolBarExtensibilityManager() override { return OperatorToolBarExtensibilityManager; }
 	virtual void GetRootActorLevelInstances(TArray<ADisplayClusterRootActor*>& OutRootActorInstances) override;
 	virtual void ShowDetailsForObject(UObject* Object) override;
diff --git a/Engine/Plugins/Runtime/nDisplay/Source/DisplayClusterOperator/Private/SDisplayClusterOperatorPanel.cpp b/Engine/Plugins/Runtime/nDisplay/Source/DisplayClusterOperator/Private/SDisplayClusterOperatorPanel.cpp
--- a/Engine/Plugins/Runtime/nDisplay/Source/DisplayClusterOperator/Private/SDisplayClusterOperatorPanel.cpp
+++ b/Engine/Plugins/Runtime/nDisplay/Source/DisplayClusterOperator/Private/SDisplayClusterOperatorPanel.cpp
@@ -64,7 +64,7 @@ void SDisplayClusterOperatorPanel::Construct(const FArguments& InArgs, const TSh
 					(
 						FTabManager::NewSplitter()
 						->SetOrientation(Orient_Vertical)
-						->SetSizeCoefficient(0.67f)
+						->SetSizeCoefficient(0.75f)
 						->Split
 						(
 							FTabManager::NewStack()
@@ -78,10 +78,16 @@ void SDisplayClusterOperatorPanel::Construct(const FArguments& InArgs, const TSh
 					)
 					->Split
 					(
-						FTabManager::NewStack()
+						FTabManager::NewSplitter()
+						->SetOrientation(Orient_Vertical)
+						->SetSizeCoefficient(0.25f)
+						->Split
+						(
+							FTabManager::NewStack()
 							->AddTab(DetailsTabId, ETabState::OpenedTab)
 							->SetHideTabWell(true)
-							->SetSizeCoefficient(0.33f)
+						)
+						
 					)
 			)
 		);
diff --git a/Engine/Plugins/Runtime/nDisplay/Source/DisplayClusterOperator/Public/IDisplayClusterOperator.h b/Engine/Plugins/Runtime/nDisplay/Source/DisplayClusterOperator/Public/IDisplayClusterOperator.h
--- a/Engine/Plugins/Runtime/nDisplay/Source/DisplayClusterOperator/Public/IDisplayClusterOperator.h
+++ b/Engine/Plugins/Runtime/nDisplay/Source/DisplayClusterOperator/Public/IDisplayClusterOperator.h
@@ -66,6 +66,9 @@ public:
 	/** Gets the extension ID for the auxilliary window region intended for lower-thirds windows (e.g log output) that can be used to add tabs to the operator panel */
 	virtual FName GetAuxilliaryOperatorExtensionId() = 0;
 
+	/** Gets the tab ID used for the details panel */
+	virtual FName GetDetailsTabId() = 0;
+	
 	/** Gets the extensibility manager for the operator panel's toolbar */
 	virtual TSharedPtr<FExtensibilityManager> GetOperatorToolBarExtensibilityManager() = 0;
 
