Show a warning in the lightcard editor and app when real-time rendering is disabled.

diff --git a/Engine/Plugins/Runtime/nDisplay/Source/DisplayCluster/Private/Game/EngineClasses/Scene/DisplayClusterPreviewComponent.cpp b/Engine/Plugins/Runtime/nDisplay/Source/DisplayCluster/Private/Game/EngineClasses/Scene/DisplayClusterPreviewComponent.cpp
--- a/Engine/Plugins/Runtime/nDisplay/Source/DisplayCluster/Private/Game/EngineClasses/Scene/DisplayClusterPreviewComponent.cpp
+++ b/Engine/Plugins/Runtime/nDisplay/Source/DisplayCluster/Private/Game/EngineClasses/Scene/DisplayClusterPreviewComponent.cpp
@@ -101,7 +101,7 @@ bool UDisplayClusterPreviewComponent::InitializePreviewComponent(ADisplayCluster
 
 bool UDisplayClusterPreviewComponent::IsPreviewEnabled() const
 {
-	return (ViewportConfig && RootActor && RootActor->bPreviewEnable);
+	return (ViewportConfig && RootActor && RootActor->IsPreviewEnabled());
 }
 
 void UDisplayClusterPreviewComponent::RestorePreviewMeshMaterial()
diff --git a/Engine/Plugins/Runtime/nDisplay/Source/DisplayCluster/Private/Game/EngineClasses/Scene/DisplayClusterRootActorEditor.cpp b/Engine/Plugins/Runtime/nDisplay/Source/DisplayCluster/Private/Game/EngineClasses/Scene/DisplayClusterRootActorEditor.cpp
--- a/Engine/Plugins/Runtime/nDisplay/Source/DisplayCluster/Private/Game/EngineClasses/Scene/DisplayClusterRootActorEditor.cpp
+++ b/Engine/Plugins/Runtime/nDisplay/Source/DisplayCluster/Private/Game/EngineClasses/Scene/DisplayClusterRootActorEditor.cpp
@@ -208,6 +208,11 @@ bool ADisplayClusterRootActor::IsPreviewEnabled() const
 	}
 #endif
 
+	if (!PreviewEnableOverriders.IsEmpty())
+	{
+		return true;
+	}
+
 	return bPreviewEnable;
 }
 
@@ -325,7 +330,7 @@ bool ADisplayClusterRootActor::ImplUpdatePreviewConfiguration_Editor(const FStri
 {
 	PreviewRenderFrame.Reset();
 
-	if (bPreviewEnable && ViewportManager.IsValid())
+	if (IsPreviewEnabled() && ViewportManager.IsValid())
 	{
 		FDisplayClusterPreviewSettings PreviewSettings;
 		PreviewSettings.PreviewRenderTargetRatioMult = PreviewRenderTargetRatioMult;
@@ -424,7 +429,7 @@ bool ADisplayClusterRootActor::ImplRenderPassPreviewClusterNode_Editor()
 
 void ADisplayClusterRootActor::ImplRenderPreview_Editor()
 {
-	if (CurrentConfigData == nullptr || !bPreviewEnable || !ViewportManager.IsValid())
+	if (CurrentConfigData == nullptr || !IsPreviewEnabled() || !ViewportManager.IsValid())
 	{
 		// no preview
 		return;
@@ -817,7 +822,7 @@ void ADisplayClusterRootActor::UpdatePreviewComponents()
 
 	TArray<UDisplayClusterPreviewComponent*> IteratedPreviewComponents;
 	
-	if (CurrentConfigData != nullptr && bPreviewEnable)
+	if (CurrentConfigData != nullptr && IsPreviewEnabled())
 	{
 		for (const UE_TRANSITIONAL_OBJECT_PTR_TEMPLATE2_ARG2(TPair, FString, UDisplayClusterConfigurationClusterNode)& Node : CurrentConfigData->Cluster->Nodes)
 		{
@@ -914,6 +919,22 @@ bool ADisplayClusterRootActor::ShouldThisFrameOutputPreviewToPostProcessRenderTa
 	return bOutputFrameToPostProcessRenderTarget;
 }
 
+void ADisplayClusterRootActor::AddPreviewEnableOverride(const uint8* Object)
+{
+	check(Object);
+	PreviewEnableOverriders.Add(Object);
+
+	UpdatePreviewComponents();
+}
+
+void ADisplayClusterRootActor::RemovePreviewEnableOverride(const uint8* Object)
+{
+	check(Object);
+	PreviewEnableOverriders.Remove(Object);
+
+	UpdatePreviewComponents();
+}
+
 FString ADisplayClusterRootActor::GeneratePreviewComponentName_Editor(const FString& NodeId, const FString& ViewportId) const
 {
 	return FString::Printf(TEXT("%s_%s"), *NodeId, *ViewportId);
diff --git a/Engine/Plugins/Runtime/nDisplay/Source/DisplayCluster/Public/DisplayClusterRootActor.h b/Engine/Plugins/Runtime/nDisplay/Source/DisplayCluster/Public/DisplayClusterRootActor.h
--- a/Engine/Plugins/Runtime/nDisplay/Source/DisplayCluster/Public/DisplayClusterRootActor.h
+++ b/Engine/Plugins/Runtime/nDisplay/Source/DisplayCluster/Public/DisplayClusterRootActor.h
@@ -419,6 +419,15 @@ public:
 	/** When rendering the preview determine which render target should be used for the current frame. */
 	bool ShouldThisFrameOutputPreviewToPostProcessRenderTarget() const;
 
+	/** Force preview rendering to be enabled regardless of the user's setting until a matching RemovePreviewEnableOverride call is made. */
+	void AddPreviewEnableOverride(const uint8* Object);
+
+	/**
+	 * Stop forcing preview rendering to be enabled for this caller. If other objects have called AddPreviewEnableOverride, it will remain
+	 * forced until they have also removed their overrides.
+	 */
+	void RemovePreviewEnableOverride(const uint8* Object);
+
 	float GetPreviewRenderTargetRatioMult() const
 	{
 		return PreviewRenderTargetRatioMult;
@@ -468,6 +477,9 @@ private:
 
 	/** Enables preview components to output to the post process render target when bPreviewEnablePostProcess is disabled. Contains all subscribed objects. */
 	TSet<const uint8*> PostProcessRenderTargetObservers;
+
+	/* Addresses of callers to AddPreviewEnableOverride that haven't removed their overrides yet. */
+	TSet<const uint8*> PreviewEnableOverriders;
 	
 	TWeakPtr<IDisplayClusterConfiguratorBlueprintEditor> ToolkitPtr;
 
diff --git a/Engine/Plugins/Runtime/nDisplay/Source/DisplayClusterLightCardEditor/Private/Viewport/DisplayClusterLightCardEditorViewportClient.cpp b/Engine/Plugins/Runtime/nDisplay/Source/DisplayClusterLightCardEditor/Private/Viewport/DisplayClusterLightCardEditorViewportClient.cpp
--- a/Engine/Plugins/Runtime/nDisplay/Source/DisplayClusterLightCardEditor/Private/Viewport/DisplayClusterLightCardEditorViewportClient.cpp
+++ b/Engine/Plugins/Runtime/nDisplay/Source/DisplayClusterLightCardEditor/Private/Viewport/DisplayClusterLightCardEditorViewportClient.cpp
@@ -405,6 +405,38 @@ void FDisplayClusterLightCardEditorViewportClient::Draw(FViewport* InViewport, F
 		DrawStatsHUD( World, Viewport, DebugCanvas, NULL, EmptyPropertyArray, GetViewLocation(), GetViewRotation() );
 	}
 
+	// Show warning if nDisplay real-time rendering is disabled
+	if (!IDisplayClusterScenePreview::Get().IsRealTimePreviewEnabled())
+	{
+		const FString WarningLines[] = {
+			LOCTEXT("RealTimeWarningLine1", "WARNING: Real-time rendering is disabled.").ToString(),
+			LOCTEXT("RealTimeWarningLine2", "nDisplay previews and light card positions may not update correctly.").ToString(),
+		};
+		const int32 TextMargin = 8;
+
+		const float RightMarginX = Viewport->GetSizeXY().X - TextMargin;
+		FCanvasTextItem RealTimeWarningText(
+			FVector2D(RightMarginX, Viewport->GetSizeXY().Y - TextMargin),
+			FText::GetEmpty(),
+			GEngine->GetSmallFont(),
+			FLinearColor::Yellow
+		);
+		RealTimeWarningText.EnableShadow(FLinearColor::Black);
+
+		// Draw each line, starting from the bottom edge of the viewport
+		for (int32 LineIndex = sizeof(WarningLines) / sizeof(FString) - 1; LineIndex >= 0; --LineIndex)
+		{
+			const FString& LineString = WarningLines[LineIndex];
+			const int32 TextWidth = RealTimeWarningText.Font->GetStringSize(*LineString);
+
+			RealTimeWarningText.Text = FText::FromString(LineString);
+			RealTimeWarningText.Position.X = RightMarginX - TextWidth;
+			RealTimeWarningText.Position.Y -= RealTimeWarningText.Font->GetStringHeightSize(*LineString);
+
+			Canvas->DrawItem(RealTimeWarningText);
+		}
+	}
+
 	if(!IsRealtime())
 	{
 		// Wait for the rendering thread to finish drawing the view before returning.
@@ -1226,8 +1258,10 @@ void FDisplayClusterLightCardEditorViewportClient::UpdatePreviewActor(ADisplayCl
 				Finalize();
 				return;
 			}
-			
-			RootActorPtr->SubscribeToPostProcessRenderTarget(reinterpret_cast<uint8*>(this));
+
+			const uint8* GenericThis = reinterpret_cast<uint8*>(this);
+			RootActorPtr->SubscribeToPostProcessRenderTarget(GenericThis);
+			RootActorPtr->AddPreviewEnableOverride(GenericThis);
 			RootActorLevelInstance = RootActorPtr;
 			
 			if (ProxyType == EDisplayClusterLightCardEditorProxyType::All ||
@@ -1422,7 +1456,9 @@ void FDisplayClusterLightCardEditorViewportClient::DestroyProxies(
 
 		if (RootActorLevelInstance.IsValid())
 		{
-			RootActorLevelInstance->UnsubscribeFromPostProcessRenderTarget(reinterpret_cast<uint8*>(this));
+			const uint8* GenericThis = reinterpret_cast<uint8*>(this);
+			RootActorLevelInstance->UnsubscribeFromPostProcessRenderTarget(GenericThis);
+			RootActorLevelInstance->RemovePreviewEnableOverride(GenericThis);
 			RootActorLevelInstance.Reset();
 		}
 	}
diff --git a/Engine/Plugins/Runtime/nDisplay/Source/DisplayClusterScenePreview/Private/DisplayClusterScenePreviewModule.cpp b/Engine/Plugins/Runtime/nDisplay/Source/DisplayClusterScenePreview/Private/DisplayClusterScenePreviewModule.cpp
--- a/Engine/Plugins/Runtime/nDisplay/Source/DisplayClusterScenePreview/Private/DisplayClusterScenePreviewModule.cpp
+++ b/Engine/Plugins/Runtime/nDisplay/Source/DisplayClusterScenePreview/Private/DisplayClusterScenePreviewModule.cpp
@@ -11,6 +11,10 @@
 #include "Engine/TextureRenderTarget2D.h"
 #include "Misc/TransactionObjectEvent.h"
 
+#if WITH_EDITOR
+#include "LevelEditorViewport.h"
+#endif
+
 #define LOCTEXT_NAMESPACE "DisplayClusterScenePreview"
 
 static TAutoConsoleVariable<float> CVarDisplayClusterScenePreviewRenderTickDelay(
@@ -232,6 +236,11 @@ bool FDisplayClusterScenePreviewModule::RenderQueued(int32 RendererId, FDisplayC
 	return InternalRenderQueued(RendererId, RenderSettings, Canvas, RenderTarget->GetSizeXY(), ResultDelegate);
 }
 
+bool FDisplayClusterScenePreviewModule::IsRealTimePreviewEnabled() const
+{
+	return bIsRealTimePreviewEnabled;
+}
+
 ADisplayClusterRootActor* FDisplayClusterScenePreviewModule::InternalGetRendererRootActor(FRendererConfig& RendererConfig)
 {
 	if (!RendererConfig.RootActor.IsValid() && !RendererConfig.RootActorPath.IsEmpty())
@@ -292,6 +301,9 @@ void FDisplayClusterScenePreviewModule::InternalSetRendererRootActor(FRendererCo
 
 bool FDisplayClusterScenePreviewModule::InternalRenderImmediate(FRendererConfig& RendererConfig, FDisplayClusterMeshProjectionRenderSettings& RenderSettings, FCanvas& Canvas)
 {
+	// Update this so that whoever gets the callback can immediately check whether the nDisplay preview may be out of date
+	UpdateIsRealTimePreviewEnabled();
+
 	UWorld* World = RendererConfig.RootActor.IsValid() ? RendererConfig.RootActor->GetWorld() : nullptr;
 	if (!World)
 	{
@@ -393,12 +405,17 @@ void FDisplayClusterScenePreviewModule::RegisterRootActorEvents(ADisplayClusterR
 		return;
 	}
 
-	// Register/unregister actors to use post-processing for preview renders so their lighting looks correct
-	Actor->UnsubscribeFromPostProcessRenderTarget(reinterpret_cast<uint8*>(this));
+	const uint8* GenericThis = reinterpret_cast<uint8*>(this);
+
+	// Register/unregister actors to use post-processing for preview renders so their lighting looks correct,
+	// and add/remove preview override so our preview renders have the latest nDisplay preview texture.
+	Actor->UnsubscribeFromPostProcessRenderTarget(GenericThis);
+	Actor->RemovePreviewEnableOverride(GenericThis);
 
 	if (bShouldRegister)
 	{
-		Actor->SubscribeToPostProcessRenderTarget(reinterpret_cast<uint8*>(this));
+		Actor->SubscribeToPostProcessRenderTarget(GenericThis);
+		Actor->AddPreviewEnableOverride(GenericThis);
 	}
 
 	// Register/unregister for Blueprint events
@@ -443,6 +460,31 @@ void FDisplayClusterScenePreviewModule::AutoPopulateScene(FRendererConfig& Rende
 	RendererConfig.bIsSceneDirty = false;
 }
 
+bool FDisplayClusterScenePreviewModule::UpdateIsRealTimePreviewEnabled()
+{
+#if WITH_EDITOR
+	bIsRealTimePreviewEnabled = false;
+
+	if (!GEditor)
+	{
+		return false;
+	}
+
+	for (const FLevelEditorViewportClient* LevelViewport : GEditor->GetLevelViewportClients())
+	{
+		if (LevelViewport && LevelViewport->IsRealtime())
+		{
+			bIsRealTimePreviewEnabled = true;
+			break;
+		}
+	}
+
+	return bIsRealTimePreviewEnabled;
+#else
+	return false;
+#endif
+}
+
 bool FDisplayClusterScenePreviewModule::OnTick(float DeltaTime)
 {
 	// This loop should break when we either run out of jobs or complete a single job
diff --git a/Engine/Plugins/Runtime/nDisplay/Source/DisplayClusterScenePreview/Private/DisplayClusterScenePreviewModule.h b/Engine/Plugins/Runtime/nDisplay/Source/DisplayClusterScenePreview/Private/DisplayClusterScenePreviewModule.h
--- a/Engine/Plugins/Runtime/nDisplay/Source/DisplayClusterScenePreview/Private/DisplayClusterScenePreviewModule.h
+++ b/Engine/Plugins/Runtime/nDisplay/Source/DisplayClusterScenePreview/Private/DisplayClusterScenePreviewModule.h
@@ -40,6 +40,7 @@ public:
 	virtual bool Render(int32 RendererId, FDisplayClusterMeshProjectionRenderSettings& RenderSettings, FCanvas& Canvas) override;
 	virtual bool RenderQueued(int32 RendererId, FDisplayClusterMeshProjectionRenderSettings& RenderSettings, const FIntPoint& Size, FRenderResultDelegate ResultDelegate) override;
 	virtual bool RenderQueued(int32 RendererId, FDisplayClusterMeshProjectionRenderSettings& RenderSettings, const TWeakPtr<FCanvas> Canvas, FRenderResultDelegate ResultDelegate) override;
+	virtual bool IsRealTimePreviewEnabled() const override;
 	//~ End IDisplayClusterScenePreview Interface
 
 private:
@@ -126,6 +127,9 @@ private:
 	/** Clear and re-populate a renderer's scene with the root actor and lightcards if applicable. */
 	void AutoPopulateScene(FRendererConfig& RendererConfig);
 
+	/** Check whether nDisplay preview textures are being updated in real time. */
+	bool UpdateIsRealTimePreviewEnabled();
+
 	/** Called on tick to process the queued renders. */
 	bool OnTick(float DeltaTime);
 
@@ -160,4 +164,7 @@ private:
 
 	/** Whether this is currently registered for actor update events. */
 	bool bIsRegisteredForActorEvents = false;
+
+	/** Whether nDisplay preview textures are being updated in real time. */
+	bool bIsRealTimePreviewEnabled = false;
 };
\ No newline at end of file
diff --git a/Engine/Plugins/Runtime/nDisplay/Source/DisplayClusterScenePreview/Public/IDisplayClusterScenePreview.h b/Engine/Plugins/Runtime/nDisplay/Source/DisplayClusterScenePreview/Public/IDisplayClusterScenePreview.h
--- a/Engine/Plugins/Runtime/nDisplay/Source/DisplayClusterScenePreview/Public/IDisplayClusterScenePreview.h
+++ b/Engine/Plugins/Runtime/nDisplay/Source/DisplayClusterScenePreview/Public/IDisplayClusterScenePreview.h
@@ -181,4 +181,9 @@ public:
 	 * @param ResultDelegate The delegate to call when the render is complete. It will be passed a FRenderTarget containing the rendered preview.
 	 */
 	virtual bool RenderQueued(int32 RendererId, FDisplayClusterMeshProjectionRenderSettings& RenderSettings, TWeakPtr<FCanvas> Canvas, FRenderResultDelegate ResultDelegate) = 0;
+
+	/**
+	 * Check whether nDisplay preview textures are being updated in real time.
+	 */
+	virtual bool IsRealTimePreviewEnabled() const = 0;
 };
\ No newline at end of file
diff --git a/Engine/Plugins/VirtualProduction/EpicStageApp/Source/EpicStageApp/Private/StageAppResponse.h b/Engine/Plugins/VirtualProduction/EpicStageApp/Source/EpicStageApp/Private/StageAppResponse.h
--- a/Engine/Plugins/VirtualProduction/EpicStageApp/Source/EpicStageApp/Private/StageAppResponse.h
+++ b/Engine/Plugins/VirtualProduction/EpicStageApp/Source/EpicStageApp/Private/StageAppResponse.h
@@ -87,6 +87,12 @@ struct FRCPreviewRenderCompletedEvent
 	 */
 	UPROPERTY()
 	int64 SequenceNumber = -1;
+
+	/**
+	 * If true, real-time preview rendering is disabled in the engine, which may prevent nDisplay previews and light card positions from updating correctly.
+	 */
+	UPROPERTY()
+	bool IsRealTimeDisabled = false;
 };
 
 /**
diff --git a/Engine/Plugins/VirtualProduction/EpicStageApp/Source/EpicStageApp/Private/StageAppRouteHandler.cpp b/Engine/Plugins/VirtualProduction/EpicStageApp/Source/EpicStageApp/Private/StageAppRouteHandler.cpp
--- a/Engine/Plugins/VirtualProduction/EpicStageApp/Source/EpicStageApp/Private/StageAppRouteHandler.cpp
+++ b/Engine/Plugins/VirtualProduction/EpicStageApp/Source/EpicStageApp/Private/StageAppRouteHandler.cpp
@@ -442,6 +442,7 @@ void FStageAppRouteHandler::HandleWebSocketNDisplayPreviewRender(const FRemoteCo
 			FRCPreviewRenderCompletedEvent Event;
 			Event.RendererId = RendererId;
 			Event.ImageBase64 = FBase64::Encode(JpegData.GetData(), JpegData.Num());
+			Event.IsRealTimeDisabled = !IDisplayClusterScenePreview::Get().IsRealTimePreviewEnabled();
 
 			if (PerRendererData)
 			{
