nDisplay: Added new drawer mode to the color grading drawer that displays a list of additional, customized properties from the selected items

diff --git a/Engine/Plugins/Runtime/nDisplay/Source/DisplayClusterColorGrading/DisplayClusterColorGrading.Build.cs b/Engine/Plugins/Runtime/nDisplay/Source/DisplayClusterColorGrading/DisplayClusterColorGrading.Build.cs
--- a/Engine/Plugins/Runtime/nDisplay/Source/DisplayClusterColorGrading/DisplayClusterColorGrading.Build.cs
+++ b/Engine/Plugins/Runtime/nDisplay/Source/DisplayClusterColorGrading/DisplayClusterColorGrading.Build.cs
@@ -23,6 +23,7 @@ public class DisplayClusterColorGrading : ModuleRules
 				"InputCore",
 				"Kismet",
 				"PropertyEditor",
+				"PropertyPath",
 				"Slate",
 				"SlateCore",
 				"ToolWidgets",
diff --git a/Engine/Plugins/Runtime/nDisplay/Source/DisplayClusterColorGrading/Private/DataModelGenerators/DisplayClusterColorGradingGenerator_RootActor.cpp b/Engine/Plugins/Runtime/nDisplay/Source/DisplayClusterColorGrading/Private/DataModelGenerators/DisplayClusterColorGradingGenerator_RootActor.cpp
--- a/Engine/Plugins/Runtime/nDisplay/Source/DisplayClusterColorGrading/Private/DataModelGenerators/DisplayClusterColorGradingGenerator_RootActor.cpp
+++ b/Engine/Plugins/Runtime/nDisplay/Source/DisplayClusterColorGrading/Private/DataModelGenerators/DisplayClusterColorGradingGenerator_RootActor.cpp
@@ -13,10 +13,14 @@
 #include "DisplayClusterConfigurationTypes_ICVFX.h"
 #include "DisplayClusterConfigurationTypes_Postprocess.h"
 
+#include "DetailCategoryBuilder.h"
+#include "DetailLayoutBuilder.h"
 #include "Framework/MultiBox/MultiBoxBuilder.h"
+#include "IDetailCustomization.h"
 #include "IDetailTreeNode.h"
 #include "IPropertyRowGenerator.h"
 #include "PropertyHandle.h"
+#include "PropertyPathHelpers.h"
 #include "ScopedTransaction.h"
 #include "Widgets/Input/SButton.h"
 #include "Widgets/Text/SInlineEditableTextBlock.h"
@@ -84,19 +88,23 @@ void FDisplayClusterColorGradingGenerator_ColorGradingRenderingSettings::AddDeta
 	}
 }
 
-TSharedPtr<IDetailTreeNode> FDisplayClusterColorGradingGenerator_ColorGradingRenderingSettings::FindPropertyTreeNode(const TSharedRef<IDetailTreeNode>& Node, const FName& PropertyName)
+TSharedPtr<IDetailTreeNode> FDisplayClusterColorGradingGenerator_ColorGradingRenderingSettings::FindPropertyTreeNode(const TSharedRef<IDetailTreeNode>& Node, const FCachedPropertyPath& PropertyPath)
 {
 	if (Node->GetNodeType() == EDetailNodeType::Item)
 	{
 		FName NodeName = Node->GetNodeName();
-		if (Node->GetNodeName() == PropertyName)
+		if (Node->GetNodeName() == PropertyPath.GetLastSegment().GetName())
 		{
-			return Node;
-		}
-		else
-		{
-			return nullptr;
+			TSharedPtr<IPropertyHandle> FoundPropertyHandle = Node->CreatePropertyHandle();
+			FString FoundPropertyPath = FoundPropertyHandle->GeneratePathToProperty();
+
+			if (PropertyPath == FoundPropertyPath)
+			{
+				return Node;
+			}
 		}
+		
+		return nullptr;
 	}
 	else
 	{
@@ -104,7 +112,7 @@ TSharedPtr<IDetailTreeNode> FDisplayClusterColorGradingGenerator_ColorGradingRen
 		Node->GetChildren(Children);
 		for (const TSharedRef<IDetailTreeNode>& Child : Children)
 		{
-			if (TSharedPtr<IDetailTreeNode> PropertyTreeNode = FindPropertyTreeNode(Child, PropertyName))
+			if (TSharedPtr<IDetailTreeNode> PropertyTreeNode = FindPropertyTreeNode(Child, PropertyPath))
 			{
 				return PropertyTreeNode;
 			}
@@ -114,13 +122,13 @@ TSharedPtr<IDetailTreeNode> FDisplayClusterColorGradingGenerator_ColorGradingRen
 	}
 }
 
-TSharedPtr<IPropertyHandle> FDisplayClusterColorGradingGenerator_ColorGradingRenderingSettings::FindPropertyHandle(IPropertyRowGenerator& PropertyRowGenerator, const FName& PropertyNodeName)
+TSharedPtr<IPropertyHandle> FDisplayClusterColorGradingGenerator_ColorGradingRenderingSettings::FindPropertyHandle(IPropertyRowGenerator& PropertyRowGenerator, const FCachedPropertyPath& PropertyPath)
 {
 	const TArray<TSharedRef<IDetailTreeNode>>& RootNodes = PropertyRowGenerator.GetRootTreeNodes();
 
 	for (const TSharedRef<IDetailTreeNode>& RootNode : RootNodes)
 	{
-		if (TSharedPtr<IDetailTreeNode> PropertyTreeNode = FindPropertyTreeNode(RootNode, PropertyNodeName))
+		if (TSharedPtr<IDetailTreeNode> PropertyTreeNode = FindPropertyTreeNode(RootNode, PropertyPath))
 		{
 			return PropertyTreeNode->CreatePropertyHandle();
 		}
@@ -129,14 +137,87 @@ TSharedPtr<IPropertyHandle> FDisplayClusterColorGradingGenerator_ColorGradingRen
 	return nullptr;
 }
 
+#define CREATE_PROPERTY_PATH(RootObjectClass, PropertyPath) FCachedPropertyPath(GET_MEMBER_NAME_STRING_CHECKED(RootObjectClass, PropertyPath))
 
 TSharedRef<IDisplayClusterColorGradingDataModelGenerator> FDisplayClusterColorGradingGenerator_RootActor::MakeInstance()
 {
 	return MakeShareable(new FDisplayClusterColorGradingGenerator_RootActor());
 }
 
+class FRootActorDetailsSectionCustomization : public IDetailCustomization
+{
+public:
+	enum EDetailsSectionType
+	{
+		Viewports,
+		InnerFrustum,
+		OCIO_AllViewports,
+		OCIO_PerViewport
+	};
+
+public:
+	FRootActorDetailsSectionCustomization(EDetailsSectionType InSectionType)
+		: SectionType(InSectionType)
+	{ }
+
+	virtual void CustomizeDetails(IDetailLayoutBuilder& DetailBuilder) override
+	{
+		TArray<FName> Categories;
+		DetailBuilder.GetCategoryNames(Categories);
+
+		for (const FName& Category : Categories)
+		{
+			DetailBuilder.HideCategory(Category);
+		}
+
+		// TransformCommon is a custom category that doesn't get returned by GetCategoryNames that also needs to be hidden
+		DetailBuilder.HideCategory(TEXT("TransformCommon"));
+
+		switch (SectionType)
+		{
+		case EDetailsSectionType::Viewports:
+			{
+				IDetailCategoryBuilder& ViewportsCategoryBuilder = DetailBuilder.EditCategory(TEXT("CustomViewportsCategory"), LOCTEXT("CustomViewportsCategoryLabel", "Viewports"));
+
+				ViewportsCategoryBuilder.AddProperty(TEXT("ViewportScreenPercentageMultiplierRef"), ADisplayClusterRootActor::StaticClass());
+				ViewportsCategoryBuilder.AddProperty(TEXT("FreezeRenderOuterViewportsRef"), ADisplayClusterRootActor::StaticClass());
+			}
+			break;
+
+		case EDetailsSectionType::InnerFrustum:
+			{
+				IDetailCategoryBuilder& InnerFrustumCategoryBuilder = DetailBuilder.EditCategory(TEXT("CustomICVFXCategory"), LOCTEXT("CustomICVFXCategoryLabel", "In-Camera VFX"));
+
+				InnerFrustumCategoryBuilder.AddProperty(GET_MEMBER_NAME_CHECKED(ADisplayClusterRootActor, InnerFrustumPriority), ADisplayClusterRootActor::StaticClass());
+			}
+			break;
+
+		case EDetailsSectionType::OCIO_AllViewports:
+			{
+				IDetailCategoryBuilder& AllViewportsOCIOCategoryBuilder = DetailBuilder.EditCategory(TEXT("CustomAllViewportsOCIOCategory"), LOCTEXT("CustomAllViewportsOCIOCategoryLabel", "All Viewports"));
+
+				AllViewportsOCIOCategoryBuilder.AddProperty(TEXT("ClusterOCIOColorConfigurationRef"), ADisplayClusterRootActor::StaticClass());
+			}
+			break;
+
+		case EDetailsSectionType::OCIO_PerViewport:
+			{
+				IDetailCategoryBuilder& PerViewportOCIOCategoryBuilder = DetailBuilder.EditCategory(TEXT("CustomPerViewportOCIOCategory"), LOCTEXT("CustomPerViewportOCIOCategoryLabel", "Per-Viewport"));
+
+				PerViewportOCIOCategoryBuilder.AddProperty(TEXT("PerViewportOCIOProfilesRef"), ADisplayClusterRootActor::StaticClass());
+			}
+			break;
+		}
+	}
+
+private:
+	EDetailsSectionType SectionType;
+};
+
 void FDisplayClusterColorGradingGenerator_RootActor::GenerateDataModel(IPropertyRowGenerator& PropertyRowGenerator, FDisplayClusterColorGradingDataModel& OutColorGradingDataModel)
 {
+	RootActors.Empty();
+
 	const TArray<TWeakObjectPtr<UObject>>& SelectedObjects = PropertyRowGenerator.GetSelectedObjects();
 	for (const TWeakObjectPtr<UObject>& SelectedObject : SelectedObjects)
 	{
@@ -148,7 +229,7 @@ void FDisplayClusterColorGradingGenerator_RootActor::GenerateDataModel(IProperty
 	}
 
 	// Add a color grading group for the root actor's "EntireClusterColorGrading" property
-	if (TSharedPtr<IPropertyHandle> EntireClusterColorGradingHandle = FindPropertyHandle(PropertyRowGenerator, GET_MEMBER_NAME_CHECKED(FDisplayClusterConfigurationICVFX_StageSettings, EntireClusterColorGrading)))
+	if (TSharedPtr<IPropertyHandle> EntireClusterColorGradingHandle = FindPropertyHandle(PropertyRowGenerator, CREATE_PROPERTY_PATH(UDisplayClusterConfigurationData, StageSettings.EntireClusterColorGrading)))
 	{
 		FDisplayClusterColorGradingDataModel::FColorGradingGroup EntireClusterGroup = CreateColorGradingGroup(EntireClusterColorGradingHandle);
 		EntireClusterGroup.EditConditionPropertyHandle = EntireClusterColorGradingHandle->GetChildHandle(GET_MEMBER_NAME_CHECKED(FDisplayClusterConfigurationViewport_EntireClusterColorGrading, bEnableEntireClusterColorGrading));
@@ -175,7 +256,7 @@ void FDisplayClusterColorGradingGenerator_RootActor::GenerateDataModel(IProperty
 	}
 
 	// Add a color grading group for each element in the root actor's "PerViewportColorGrading" array
-	if (TSharedPtr<IPropertyHandle> PerViewportColorGradingHandle = FindPropertyHandle(PropertyRowGenerator, GET_MEMBER_NAME_CHECKED(FDisplayClusterConfigurationICVFX_StageSettings, PerViewportColorGrading)))
+	if (TSharedPtr<IPropertyHandle> PerViewportColorGradingHandle = FindPropertyHandle(PropertyRowGenerator, CREATE_PROPERTY_PATH(UDisplayClusterConfigurationData, StageSettings.PerViewportColorGrading)))
 	{
 		check(PerViewportColorGradingHandle->AsArray().IsValid());
 
@@ -247,6 +328,50 @@ void FDisplayClusterColorGradingGenerator_RootActor::GenerateDataModel(IProperty
 			.Image(FAppStyle::GetBrush("Icons.PlusCircle"))
 			.ColorAndOpacity(FSlateColor::UseForeground())
 		];
+
+
+	FDisplayClusterColorGradingDataModel::FDetailsSection ViewportsDetailsSection;
+	ViewportsDetailsSection.DisplayName = LOCTEXT("ViewportsDetailsSectionLabel", "Viewports");
+	ViewportsDetailsSection.DetailCustomizationDelegate = FOnGetDetailCustomizationInstance::CreateLambda([]
+	{
+		return MakeShared<FRootActorDetailsSectionCustomization>(FRootActorDetailsSectionCustomization::EDetailsSectionType::Viewports);
+	});
+
+	OutColorGradingDataModel.DetailsSections.Add(ViewportsDetailsSection);
+
+	FDisplayClusterColorGradingDataModel::FDetailsSection InnerFrustumDetailsSection;
+	InnerFrustumDetailsSection.DisplayName = LOCTEXT("InnerFrustumDetailsSectionLabel", "Inner Frustum");
+	InnerFrustumDetailsSection.EditConditionPropertyHandle = FindPropertyHandle(PropertyRowGenerator, CREATE_PROPERTY_PATH(UDisplayClusterConfigurationData, StageSettings.bEnableInnerFrustums));
+	InnerFrustumDetailsSection.DetailCustomizationDelegate = FOnGetDetailCustomizationInstance::CreateLambda([]
+	{
+		return MakeShared<FRootActorDetailsSectionCustomization>(FRootActorDetailsSectionCustomization::EDetailsSectionType::InnerFrustum);
+	});
+
+	OutColorGradingDataModel.DetailsSections.Add(InnerFrustumDetailsSection);
+
+	FDisplayClusterColorGradingDataModel::FDetailsSection OCIODetailsSection;
+	OCIODetailsSection.DisplayName = LOCTEXT("OCIODetailsSectionLabel", "OCIO");
+	OCIODetailsSection.EditConditionPropertyHandle = FindPropertyHandle(PropertyRowGenerator, CREATE_PROPERTY_PATH(UDisplayClusterConfigurationData, StageSettings.bUseOverallClusterOCIOConfiguration));
+
+	FDisplayClusterColorGradingDataModel::FDetailsSubsection AllViewportsOCIODetailsSubsection;
+	AllViewportsOCIODetailsSubsection.DisplayName = LOCTEXT("AllViewportsOCIOSubsectionLabel", "All Viewports");
+	AllViewportsOCIODetailsSubsection.DetailCustomizationDelegate = FOnGetDetailCustomizationInstance::CreateLambda([]
+	{
+		return MakeShared<FRootActorDetailsSectionCustomization>(FRootActorDetailsSectionCustomization::EDetailsSectionType::OCIO_AllViewports);
+	});
+
+	OCIODetailsSection.Subsections.Add(AllViewportsOCIODetailsSubsection);
+
+	FDisplayClusterColorGradingDataModel::FDetailsSubsection PerViewportOCIODetailsSubsection;
+	PerViewportOCIODetailsSubsection.DisplayName = LOCTEXT("PerViewportOCIOSubsectionLabel", "Per-Viewport");
+	PerViewportOCIODetailsSubsection.DetailCustomizationDelegate = FOnGetDetailCustomizationInstance::CreateLambda([]
+	{
+		return MakeShared<FRootActorDetailsSectionCustomization>(FRootActorDetailsSectionCustomization::EDetailsSectionType::OCIO_PerViewport);
+	});
+
+	OCIODetailsSection.Subsections.Add(PerViewportOCIODetailsSubsection);
+
+	OutColorGradingDataModel.DetailsSections.Add(OCIODetailsSection);
 }
 
 FReply FDisplayClusterColorGradingGenerator_RootActor::AddColorGradingGroup()
@@ -438,8 +563,118 @@ TSharedRef<IDisplayClusterColorGradingDataModelGenerator> FDisplayClusterColorGr
 	return MakeShareable(new FDisplayClusterColorGradingGenerator_ICVFXCamera());
 }
 
+class FICVFXCameraDetailsSectionCustomization : public IDetailCustomization
+{
+public:
+	enum EDetailsSectionType
+	{
+		ICVFX,
+		Overscan,
+		Chromakey_Markers,
+		Chromakey_Custom,
+		OCIO_AllNodes,
+		OCIO_PerNode
+	};
+
+public:
+	FICVFXCameraDetailsSectionCustomization(EDetailsSectionType InSectionType)
+		: SectionType(InSectionType)
+	{ }
+
+	virtual void CustomizeDetails(IDetailLayoutBuilder& DetailBuilder) override
+	{
+		TArray<FName> Categories;
+		DetailBuilder.GetCategoryNames(Categories);
+
+		for (const FName& Category : Categories)
+		{
+			DetailBuilder.HideCategory(Category);
+		}
+
+		// TransformCommon is a custom category that doesn't get returned by GetCategoryNames that also needs to be hidden
+		DetailBuilder.HideCategory(TEXT("TransformCommon"));
+
+		switch (SectionType)
+		{
+		case EDetailsSectionType::ICVFX:
+			{
+				IDetailCategoryBuilder& ICVFXCategoryBuilder = DetailBuilder.EditCategory(TEXT("CustomICVFXCategory"), LOCTEXT("CustomICVFXCategoryLabel", "In-Camera VFX"));
+				AddProperty(DetailBuilder, ICVFXCategoryBuilder, TEXT("BufferRatioRef"));
+				AddProperty(DetailBuilder, ICVFXCategoryBuilder, TEXT("ExternalCameraActorRef"));
+				AddProperty(DetailBuilder, ICVFXCategoryBuilder, TEXT("HiddenICVFXViewportsRef"));
+
+				IDetailCategoryBuilder& SoftEdgeCategoryBuilder = DetailBuilder.EditCategory(TEXT("CustomSoftEdgeCategory"), LOCTEXT("CustomSoftEdgeCategoryLabel", "Soft Edge"));
+				AddProperty(DetailBuilder, SoftEdgeCategoryBuilder, TEXT("SoftEdgeRef"), true);
+
+				IDetailCategoryBuilder& BorderCategoryBuilder = DetailBuilder.EditCategory(TEXT("CustomBorderCategory"), LOCTEXT("CustomBorderCategoryLabel", "Border"));
+				AddProperty(DetailBuilder, BorderCategoryBuilder, TEXT("BorderRef"), true);
+			}
+			break;
+
+		case EDetailsSectionType::Overscan:
+			{
+				IDetailCategoryBuilder& OverscanCategoryBuilder = DetailBuilder.EditCategory(TEXT("CustomOverscanCategory"), LOCTEXT("CustomOverscanCategoryLabel", "Inner Frustum Overscan"));
+				AddProperty(DetailBuilder, OverscanCategoryBuilder, TEXT("CustomFrustumRef"), true);
+			}
+			break;
+
+		case EDetailsSectionType::Chromakey_Markers:
+			{
+				IDetailCategoryBuilder& ChromakeyCategoryBuilder = DetailBuilder.EditCategory(TEXT("CustomChromakeyCategory"), LOCTEXT("CustomChromakeyCategoryLabel", "Chromakey"));
+				AddProperty(DetailBuilder, ChromakeyCategoryBuilder, TEXT("ChromakeyColorRef"));
+
+				IDetailCategoryBuilder& ChromakeyMarkersCategoryBuilder = DetailBuilder.EditCategory(TEXT("CustomChromakeyMarkersCategory"), LOCTEXT("CustomChromakeyMarkersCategoryLabel", "ChromakeyMarkers"));
+				AddProperty(DetailBuilder, ChromakeyMarkersCategoryBuilder, TEXT("ChromakeyMarkersRef"), true);
+			}
+			break;
+
+		case EDetailsSectionType::Chromakey_Custom:
+			{
+				IDetailCategoryBuilder& ChromakeyCategoryBuilder = DetailBuilder.EditCategory(TEXT("CustomChromakeyCategory"), LOCTEXT("CustomChromakeyCategoryLabel", "Chromakey"));
+				AddProperty(DetailBuilder, ChromakeyCategoryBuilder, TEXT("ChromakeyColorRef"));
+
+				IDetailCategoryBuilder& ChromakeyCustomCategoryBuilder = DetailBuilder.EditCategory(TEXT("CustomChromakeyCustomCategory"), LOCTEXT("CustomChromakeyCustomCategoryLabel", "Custom Chromakey"));
+				AddProperty(DetailBuilder, ChromakeyCustomCategoryBuilder, TEXT("ChromakeyRenderTextureRef"), true);
+			}
+			break;
+
+		case EDetailsSectionType::OCIO_AllNodes:
+			{
+				IDetailCategoryBuilder& OCIOAllNodesCategoryBuilder = DetailBuilder.EditCategory(TEXT("CustomAllNodesOCIOCategory"), LOCTEXT("CustomAllNodesOCIOCategoryLabel", "All Nodes"));
+				AddProperty(DetailBuilder, OCIOAllNodesCategoryBuilder, TEXT("OCIOColorConfiguratonRef"));
+			}
+			break;
+
+		case EDetailsSectionType::OCIO_PerNode:
+			{
+				IDetailCategoryBuilder& OCIOPerNodeCategoryBuilder = DetailBuilder.EditCategory(TEXT("CustomPerNodeOCIOCategory"), LOCTEXT("CustomPerNodeOCIOCategoryLabel", "Per-Node"));
+				AddProperty(DetailBuilder, OCIOPerNodeCategoryBuilder, TEXT("PerNodeOCIOProfilesRef"));
+			}
+			break;
+		}
+	}
+
+private:
+	void AddProperty(IDetailLayoutBuilder& DetailBuilder, IDetailCategoryBuilder& Category, FName PropertyName, bool bExpandChildProperties = false)
+	{
+		TSharedRef<IPropertyHandle> PropertyHandle = DetailBuilder.GetProperty(PropertyName, UDisplayClusterICVFXCameraComponent::StaticClass());
+
+		if (bExpandChildProperties)
+		{
+			PropertyHandle->SetInstanceMetaData(TEXT("ShowOnlyInnerProperties"), TEXT("1"));
+		}
+
+		Category.AddProperty(PropertyHandle);
+	}
+
+private:
+	EDetailsSectionType SectionType;
+};
+
 void FDisplayClusterColorGradingGenerator_ICVFXCamera::GenerateDataModel(IPropertyRowGenerator& PropertyRowGenerator, FDisplayClusterColorGradingDataModel& OutColorGradingDataModel)
 {
+	CameraComponents.Empty();
+
 	const TArray<TWeakObjectPtr<UObject>>& SelectedObjects = PropertyRowGenerator.GetSelectedObjects();
 	for (const TWeakObjectPtr<UObject>& SelectedObject : SelectedObjects)
 	{
@@ -451,7 +686,7 @@ void FDisplayClusterColorGradingGenerator_ICVFXCamera::GenerateDataModel(IProper
 	}
 
 	// Add a color grading group for the camera's "AllNodesColorGrading" property
-	if (TSharedPtr<IPropertyHandle> AllNodesColorGradingHandle = FindPropertyHandle(PropertyRowGenerator, GET_MEMBER_NAME_CHECKED(FDisplayClusterConfigurationICVFX_CameraSettings, AllNodesColorGrading)))
+	if (TSharedPtr<IPropertyHandle> AllNodesColorGradingHandle = FindPropertyHandle(PropertyRowGenerator, CREATE_PROPERTY_PATH(UDisplayClusterICVFXCameraComponent, CameraSettings.AllNodesColorGrading)))
 	{
 		FDisplayClusterColorGradingDataModel::FColorGradingGroup AllNodesGroup = CreateColorGradingGroup(AllNodesColorGradingHandle);
 		AllNodesGroup.EditConditionPropertyHandle = AllNodesColorGradingHandle->GetChildHandle(GET_MEMBER_NAME_CHECKED(FDisplayClusterConfigurationViewport_AllNodesColorGrading, bEnableInnerFrustumAllNodesColorGrading));
@@ -478,7 +713,7 @@ void FDisplayClusterColorGradingGenerator_ICVFXCamera::GenerateDataModel(IProper
 	}
 
 	// Add a color grading group for each element in the camera's "PerNodeColorGrading" array
-	if (TSharedPtr<IPropertyHandle> PerNodeColorGradingHandle = FindPropertyHandle(PropertyRowGenerator, GET_MEMBER_NAME_CHECKED(FDisplayClusterConfigurationICVFX_CameraSettings, PerNodeColorGrading)))
+	if (TSharedPtr<IPropertyHandle> PerNodeColorGradingHandle = FindPropertyHandle(PropertyRowGenerator, CREATE_PROPERTY_PATH(UDisplayClusterICVFXCameraComponent, CameraSettings.PerNodeColorGrading)))
 	{
 		check(PerNodeColorGradingHandle->AsArray().IsValid());
 
@@ -550,6 +785,84 @@ void FDisplayClusterColorGradingGenerator_ICVFXCamera::GenerateDataModel(IProper
 			.Image(FAppStyle::GetBrush("Icons.PlusCircle"))
 			.ColorAndOpacity(FSlateColor::UseForeground())
 		];
+
+	{
+		FDisplayClusterColorGradingDataModel::FDetailsSection InnerFrustumDetailsSection;
+		InnerFrustumDetailsSection.DisplayName = LOCTEXT("InnerFrustumDetailsSectionLabel", "Inner Frustum");
+		InnerFrustumDetailsSection.EditConditionPropertyHandle = FindPropertyHandle(PropertyRowGenerator, CREATE_PROPERTY_PATH(UDisplayClusterICVFXCameraComponent, CameraSettings.bEnable));
+
+		FDisplayClusterColorGradingDataModel::FDetailsSubsection ICVFXDetailsSubsection;
+		ICVFXDetailsSubsection.DisplayName = LOCTEXT("ICVFXSubsectionLabel", "ICVFX");
+		ICVFXDetailsSubsection.DetailCustomizationDelegate = FOnGetDetailCustomizationInstance::CreateLambda([]
+		{
+			return MakeShared<FICVFXCameraDetailsSectionCustomization>(FICVFXCameraDetailsSectionCustomization::EDetailsSectionType::ICVFX);
+		});
+
+		InnerFrustumDetailsSection.Subsections.Add(ICVFXDetailsSubsection);
+
+		FDisplayClusterColorGradingDataModel::FDetailsSubsection OverscanDetailsSubsection;
+		OverscanDetailsSubsection.DisplayName = LOCTEXT("OverscanDetailsSubsectionLabel", "Overscan");
+		OverscanDetailsSubsection.DetailCustomizationDelegate = FOnGetDetailCustomizationInstance::CreateLambda([]
+		{
+			return MakeShared<FICVFXCameraDetailsSectionCustomization>(FICVFXCameraDetailsSectionCustomization::EDetailsSectionType::Overscan);
+		});
+
+		InnerFrustumDetailsSection.Subsections.Add(OverscanDetailsSubsection);
+
+		OutColorGradingDataModel.DetailsSections.Add(InnerFrustumDetailsSection);
+	}
+
+	{
+		FDisplayClusterColorGradingDataModel::FDetailsSection ChromakeyDetailsSection;
+		ChromakeyDetailsSection.DisplayName = LOCTEXT("ChromakeyDetailsSectionLabel", "Chromakey");
+		ChromakeyDetailsSection.EditConditionPropertyHandle = FindPropertyHandle(PropertyRowGenerator, CREATE_PROPERTY_PATH(UDisplayClusterICVFXCameraComponent, CameraSettings.Chromakey.bEnable));
+
+		FDisplayClusterColorGradingDataModel::FDetailsSubsection ChromakeyMarkersDetailsSubsection;
+		ChromakeyMarkersDetailsSubsection.DisplayName = LOCTEXT("ChromakeyMarkersDetailsSubsectionLabel", "Markers");
+		ChromakeyMarkersDetailsSubsection.DetailCustomizationDelegate = FOnGetDetailCustomizationInstance::CreateLambda([]
+		{
+			return MakeShared<FICVFXCameraDetailsSectionCustomization>(FICVFXCameraDetailsSectionCustomization::EDetailsSectionType::Chromakey_Markers);
+		});
+
+		ChromakeyDetailsSection.Subsections.Add(ChromakeyMarkersDetailsSubsection);
+
+		FDisplayClusterColorGradingDataModel::FDetailsSubsection ChromakeyCustomDetailsSubsection;
+		ChromakeyCustomDetailsSubsection.DisplayName = LOCTEXT("ChromakeyCustomDetailsSubsectionLabel", "Custom");
+		ChromakeyCustomDetailsSubsection.DetailCustomizationDelegate = FOnGetDetailCustomizationInstance::CreateLambda([]
+		{
+			return MakeShared<FICVFXCameraDetailsSectionCustomization>(FICVFXCameraDetailsSectionCustomization::EDetailsSectionType::Chromakey_Custom);
+		});
+
+		ChromakeyDetailsSection.Subsections.Add(ChromakeyCustomDetailsSubsection);
+
+		OutColorGradingDataModel.DetailsSections.Add(ChromakeyDetailsSection);
+	}
+
+	{
+		FDisplayClusterColorGradingDataModel::FDetailsSection OCIODetailsSection;
+		OCIODetailsSection.DisplayName = LOCTEXT("OCIODetailsSectionLabel", "OCIO");
+		OCIODetailsSection.EditConditionPropertyHandle = FindPropertyHandle(PropertyRowGenerator, CREATE_PROPERTY_PATH(UDisplayClusterICVFXCameraComponent, CameraSettings.AllNodesOCIOConfiguration.bIsEnabled));
+
+		FDisplayClusterColorGradingDataModel::FDetailsSubsection AllNodesOCIODetailsSubsection;
+		AllNodesOCIODetailsSubsection.DisplayName = LOCTEXT("AllNodesOCIODetailsSubsectionLabel", "All Nodes");
+		AllNodesOCIODetailsSubsection.DetailCustomizationDelegate = FOnGetDetailCustomizationInstance::CreateLambda([]
+		{
+			return MakeShared<FICVFXCameraDetailsSectionCustomization>(FICVFXCameraDetailsSectionCustomization::EDetailsSectionType::OCIO_AllNodes);
+		});
+
+		OCIODetailsSection.Subsections.Add(AllNodesOCIODetailsSubsection);
+
+		FDisplayClusterColorGradingDataModel::FDetailsSubsection PerNodeOCIODetailsSubsection;
+		PerNodeOCIODetailsSubsection.DisplayName = LOCTEXT("PerNodeOCIODetailsSubsectionLabel", "Per-Node");
+		PerNodeOCIODetailsSubsection.DetailCustomizationDelegate = FOnGetDetailCustomizationInstance::CreateLambda([]
+		{
+			return MakeShared<FICVFXCameraDetailsSectionCustomization>(FICVFXCameraDetailsSectionCustomization::EDetailsSectionType::OCIO_PerNode);
+		});
+
+		OCIODetailsSection.Subsections.Add(PerNodeOCIODetailsSubsection);
+
+		OutColorGradingDataModel.DetailsSections.Add(OCIODetailsSection);
+	}
 }
 
 FReply FDisplayClusterColorGradingGenerator_ICVFXCamera::AddColorGradingGroup()
diff --git a/Engine/Plugins/Runtime/nDisplay/Source/DisplayClusterColorGrading/Private/DataModelGenerators/DisplayClusterColorGradingGenerator_RootActor.h b/Engine/Plugins/Runtime/nDisplay/Source/DisplayClusterColorGrading/Private/DataModelGenerators/DisplayClusterColorGradingGenerator_RootActor.h
--- a/Engine/Plugins/Runtime/nDisplay/Source/DisplayClusterColorGrading/Private/DataModelGenerators/DisplayClusterColorGradingGenerator_RootActor.h
+++ b/Engine/Plugins/Runtime/nDisplay/Source/DisplayClusterColorGrading/Private/DataModelGenerators/DisplayClusterColorGradingGenerator_RootActor.h
@@ -13,6 +13,7 @@ class IPropertyHandle;
 class IPropertyRowGenerator;
 class FDisplayClusterColorGradingDataModel;
 class UDisplayClusterICVFXCameraComponent;
+struct FCachedPropertyPath;
 
 /** Base generator for any object that needs a color grading data model generated from an FColorGradingRenderingSettings struct */
 class FDisplayClusterColorGradingGenerator_ColorGradingRenderingSettings : public IDisplayClusterColorGradingDataModelGenerator
@@ -28,10 +29,10 @@ protected:
 	void AddDetailsViewPropertyToGroup(const TSharedPtr<IPropertyHandle>& GroupPropertyHandle, FDisplayClusterColorGradingDataModel::FColorGradingGroup& Group, FName PropertyName, const FString& CategoryOverride);
 
 	/** Recursively searches the detail tree hierarchy for a property detail tree node whose name matches the specified name */
-	TSharedPtr<IDetailTreeNode> FindPropertyTreeNode(const TSharedRef<IDetailTreeNode>& Node, const FName& PropertyName);
+	TSharedPtr<IDetailTreeNode> FindPropertyTreeNode(const TSharedRef<IDetailTreeNode>& Node, const FCachedPropertyPath& PropertyPath);
 
 	/** Finds a property handle in the specified property row generator whose name matches the specified name */
-	TSharedPtr<IPropertyHandle> FindPropertyHandle(IPropertyRowGenerator& PropertyRowGenerator, const FName& PropertyNodeName);
+	TSharedPtr<IPropertyHandle> FindPropertyHandle(IPropertyRowGenerator& PropertyRowGenerator, const FCachedPropertyPath& PropertyPath);
 };
 
 /** Color grading data model generator for an nDisplay root actor */
diff --git a/Engine/Plugins/Runtime/nDisplay/Source/DisplayClusterColorGrading/Private/DisplayClusterColorGradingDataModel.cpp b/Engine/Plugins/Runtime/nDisplay/Source/DisplayClusterColorGrading/Private/DisplayClusterColorGradingDataModel.cpp
--- a/Engine/Plugins/Runtime/nDisplay/Source/DisplayClusterColorGrading/Private/DisplayClusterColorGradingDataModel.cpp
+++ b/Engine/Plugins/Runtime/nDisplay/Source/DisplayClusterColorGrading/Private/DisplayClusterColorGradingDataModel.cpp
@@ -21,6 +21,7 @@ FDisplayClusterColorGradingDataModel::FDisplayClusterColorGradingDataModel()
 
 	FPropertyRowGeneratorArgs Args;
 	PropertyRowGenerator = PropertyEditorModule.CreatePropertyRowGenerator(Args);
+	PropertyRowGenerator->OnRowsRefreshed().AddRaw(this, &FDisplayClusterColorGradingDataModel::OnPropertyRowGeneratorRefreshed);
 }
 
 TArray<TWeakObjectPtr<UObject>> FDisplayClusterColorGradingDataModel::GetObjects() const
@@ -40,23 +41,6 @@ void FDisplayClusterColorGradingDataModel::SetObjects(const TArray<UObject*>& In
 	if (PropertyRowGenerator.IsValid())
 	{
 		PropertyRowGenerator->SetObjects(InObjects);
-
-		const TArray<TWeakObjectPtr<UObject>>& SelectedObjects = PropertyRowGenerator->GetSelectedObjects();
-
-		if (SelectedObjects.Num() == 1)
-		{
-			TWeakObjectPtr<UObject> SelectedObject = SelectedObjects[0];
-
-			if (SelectedObject.IsValid())
-			{
-				if (TSharedPtr<IDisplayClusterColorGradingDataModelGenerator> Generator = GetDataModelGenerator(SelectedObject->GetClass()))
-				{
-					Generator->GenerateDataModel(*PropertyRowGenerator, *this);
-				}
-			}
-		}
-
-		// TODO: Figure out what needs to be done to support multiple disparate types of objects being color graded at the same time
 	}
 
 	SelectedColorGradingGroupIndex = ColorGradingGroups.Num() ? 0 : INDEX_NONE;
@@ -85,6 +69,7 @@ void FDisplayClusterColorGradingDataModel::Reset()
 {
 	DataModelGeneratorInstances.Empty();
 	ColorGradingGroups.Empty();
+	DetailsSections.Empty();
 	SelectedColorGradingGroupIndex = INDEX_NONE;
 	SelectedColorGradingElementIndex = INDEX_NONE;
 	ColorGradingGroupToolBarWidget = nullptr;
@@ -160,4 +145,29 @@ TSharedPtr<IDisplayClusterColorGradingDataModelGenerator> FDisplayClusterColorGr
 	return nullptr;
 }
 
+void FDisplayClusterColorGradingDataModel::OnPropertyRowGeneratorRefreshed()
+{
+	ColorGradingGroups.Empty();
+	DetailsSections.Empty();
+
+	const TArray<TWeakObjectPtr<UObject>>& SelectedObjects = PropertyRowGenerator->GetSelectedObjects();
+
+	if (SelectedObjects.Num() == 1)
+	{
+		TWeakObjectPtr<UObject> SelectedObject = SelectedObjects[0];
+
+		if (SelectedObject.IsValid())
+		{
+			if (TSharedPtr<IDisplayClusterColorGradingDataModelGenerator> Generator = GetDataModelGenerator(SelectedObject->GetClass()))
+			{
+				Generator->GenerateDataModel(*PropertyRowGenerator, *this);
+			}
+		}
+	}
+
+	// TODO: Figure out what needs to be done to support multiple disparate types of objects being color graded at the same time
+
+	OnDataModelGeneratedDelegate.Broadcast();
+}
+
 #undef LOCTEXT_NAMESPACE
\ No newline at end of file
diff --git a/Engine/Plugins/Runtime/nDisplay/Source/DisplayClusterColorGrading/Private/DisplayClusterColorGradingDataModel.h b/Engine/Plugins/Runtime/nDisplay/Source/DisplayClusterColorGrading/Private/DisplayClusterColorGradingDataModel.h
--- a/Engine/Plugins/Runtime/nDisplay/Source/DisplayClusterColorGrading/Private/DisplayClusterColorGradingDataModel.h
+++ b/Engine/Plugins/Runtime/nDisplay/Source/DisplayClusterColorGrading/Private/DisplayClusterColorGradingDataModel.h
@@ -3,6 +3,7 @@
 #pragma once
 
 #include "CoreMinimal.h"
+#include "PropertyEditorDelegates.h"
 
 class IPropertyHandle;
 class IPropertyRowGenerator;
@@ -85,7 +86,34 @@ public:
 		TArray<TSharedPtr<IPropertyHandle>> DetailsViewPropertyHandles;
 	};
 
-	DECLARE_MULTICAST_DELEGATE(FOnColorGradingSelectionChanged);
+	/** Stores data used to generate a details subsection, which is a selectable element in a details section that modifies which properties are displayed in the section */
+	struct FDetailsSubsection
+	{
+		/** The display name of the subsection */
+		FText DisplayName;
+
+		/** A delegate that creates the details customization for the details subsection */
+		FOnGetDetailCustomizationInstance DetailCustomizationDelegate;
+	};
+
+	/** Stores data used to generate a details section in the drawer, which displays a specific subset of properties from the selected objects */
+	struct FDetailsSection
+	{
+		/** The display name of the details section*/
+		FText DisplayName = FText::GetEmpty();
+
+		/** The handle of the property that controls the editing condition of the properties in the color grading group */
+		TSharedPtr<IPropertyHandle> EditConditionPropertyHandle = nullptr;
+
+		/** A list of selectable subsections to display in the details view */
+		TArray<FDetailsSubsection> Subsections;
+
+		/** A delegate that creates the details customization for the details section */
+		FOnGetDetailCustomizationInstance DetailCustomizationDelegate;
+	};
+
+	DECLARE_EVENT(FDisplayClusterColorGradingDataModel, FOnDataModelGenerated);
+	DECLARE_EVENT(FDisplayClusterColorGradingDataModel, FOnColorGradingSelectionChanged);
 
 public:
 	FDisplayClusterColorGradingDataModel();
@@ -114,6 +142,9 @@ public:
 	/** Sets the selected color grading group */
 	void SetSelectedColorGradingGroup(int32 InColorGradingGroupIndex);
 
+	/** Gets the delegate that is raised when the data model is generated */
+	FOnDataModelGenerated& OnDataModelGenerated() { return OnDataModelGeneratedDelegate; }
+
 	/** Gets the delegate that is raised when the selected color grading group has changed */
 	FOnColorGradingSelectionChanged& OnColorGradingGroupSelectionChanged() { return OnColorGradingGroupSelectionChangedDelegate; }
 
@@ -133,10 +164,16 @@ private:
 	/** Attempts to find an instance of a data model generator that can be used for the specified class */
 	TSharedPtr<IDisplayClusterColorGradingDataModelGenerator> GetDataModelGenerator(const UClass* InClass);
 
+	/** Callback that is raised when the property row generator the data model is generated from has been refreshed */
+	void OnPropertyRowGeneratorRefreshed();
+
 public:
 	/** A list of color grading groups in the color grading data model */
 	TArray<FColorGradingGroup> ColorGradingGroups;
 
+	/** A list of details sections in the color grading data model */
+	TArray<FDetailsSection> DetailsSections;
+
 	/** An optional widget to append to the end of the color grading group toolbar */
 	TSharedPtr<SWidget> ColorGradingGroupToolBarWidget = nullptr;
 
@@ -156,6 +193,9 @@ private:
 	/** The currently selected color grading element */
 	int32 SelectedColorGradingElementIndex = INDEX_NONE;
 
+	/** Delegate that is raised when the data model has been generated */
+	FOnDataModelGenerated OnDataModelGeneratedDelegate;
+
 	/** Delegate that is raised when the selected color grading group has been changed */
 	FOnColorGradingSelectionChanged OnColorGradingGroupSelectionChangedDelegate;
 
diff --git a/Engine/Plugins/Runtime/nDisplay/Source/DisplayClusterColorGrading/Private/DisplayClusterColorGradingStyle.h b/Engine/Plugins/Runtime/nDisplay/Source/DisplayClusterColorGrading/Private/DisplayClusterColorGradingStyle.h
--- a/Engine/Plugins/Runtime/nDisplay/Source/DisplayClusterColorGrading/Private/DisplayClusterColorGradingStyle.h
+++ b/Engine/Plugins/Runtime/nDisplay/Source/DisplayClusterColorGrading/Private/DisplayClusterColorGradingStyle.h
@@ -23,6 +23,7 @@ public:
 		SetCoreContentRoot(FPaths::EngineContentDir() / TEXT("Editor/Slate"));
 
 		Set("ColorGradingDrawer.Icon", new IMAGE_BRUSH_SVG("Components/nDisplayCamera", Icon16x16));
+		Set("ColorGradingDrawer.ColorGradingMode", new CORE_IMAGE_BRUSH_SVG("Starship/Common/PostProcessing", Icon16x16));
 		Set("ColorGradingDrawer.Viewports", new CORE_IMAGE_BRUSH_SVG("Starship/Common/Viewports", Icon16x16));
 		Set("ColorGradingDrawer.Nodes", new IMAGE_BRUSH_SVG("Cluster/ClusterNode", Icon16x16));
 
diff --git a/Engine/Plugins/Runtime/nDisplay/Source/DisplayClusterColorGrading/Private/Drawer/SDisplayClusterColorGradingDrawer.cpp b/Engine/Plugins/Runtime/nDisplay/Source/DisplayClusterColorGrading/Private/Drawer/SDisplayClusterColorGradingDrawer.cpp
--- a/Engine/Plugins/Runtime/nDisplay/Source/DisplayClusterColorGrading/Private/Drawer/SDisplayClusterColorGradingDrawer.cpp
+++ b/Engine/Plugins/Runtime/nDisplay/Source/DisplayClusterColorGrading/Private/Drawer/SDisplayClusterColorGradingDrawer.cpp
@@ -2,10 +2,12 @@
 
 #include "SDisplayClusterColorGradingDrawer.h"
 
+#include "DisplayClusterColorGradingStyle.h"
 #include "IDisplayClusterColorGrading.h"
 #include "IDisplayClusterColorGradingDrawerSingleton.h"
 #include "SDisplayClusterColorGradingObjectList.h"
 #include "SDisplayClusterColorGradingColorWheelPanel.h"
+#include "SDisplayClusterColorGradingDetailsPanel.h"
 
 #include "IDisplayClusterOperator.h"
 #include "IDisplayClusterOperatorViewModel.h"
@@ -73,6 +75,7 @@ SDisplayClusterColorGradingDrawer::~SDisplayClusterColorGradingDrawer()
 void SDisplayClusterColorGradingDrawer::Construct(const FArguments& InArgs, bool bInIsInDrawer)
 {
 	ColorGradingDataModel = MakeShared<FDisplayClusterColorGradingDataModel>();
+	ColorGradingDataModel->OnDataModelGenerated().AddSP(this, &SDisplayClusterColorGradingDrawer::OnColorGradingDataModelGenerated);
 
 	bIsInDrawer = bInIsInDrawer;
 	OperatorViewModel = IDisplayClusterOperator::Get().GetOperatorViewModel();
@@ -116,6 +119,43 @@ void SDisplayClusterColorGradingDrawer::Construct(const FArguments& InArgs, bool
 					[
 						SNew(SBox)
 						.HeightOverride(24.0f)
+						[
+							SNew(SHorizontalBox)
+
+							+ SHorizontalBox::Slot()
+							.AutoWidth()
+							.VAlign(VAlign_Center)
+							.Padding(2, 0, 2, 0)
+							[
+								SNew(SCheckBox)
+								.Style(FAppStyle::Get(), "ToggleButtonCheckbox")
+								.ToolTipText(LOCTEXT("ColorGradingDrawerModeToolTip", "Sets the drawer to display the color grading properties of the selected items"))
+								.OnCheckStateChanged(this, &SDisplayClusterColorGradingDrawer::OnDrawerModeSelected, EDisplayClusterColorGradingDrawerMode::ColorGrading)
+								.IsChecked(this, &SDisplayClusterColorGradingDrawer::IsDrawerModeSelected, EDisplayClusterColorGradingDrawerMode::ColorGrading)
+								[
+									SNew(SImage)
+									.ColorAndOpacity(FSlateColor::UseForeground())
+									.Image(FDisplayClusterColorGradingStyle::Get().GetBrush("ColorGradingDrawer.ColorGradingMode"))
+								]
+							]
+
+							+ SHorizontalBox::Slot()
+							.AutoWidth()
+							.VAlign(VAlign_Center)
+							.Padding(2, 0, 2, 0)
+							[
+								SNew(SCheckBox)
+								.Style(FAppStyle::Get(), "ToggleButtonCheckbox")
+								.ToolTipText(LOCTEXT("ColorGradingDrawerModeToolTip", "Sets the drawer to display the additional in-camera VFX properties of the selected items"))
+								.OnCheckStateChanged(this, &SDisplayClusterColorGradingDrawer::OnDrawerModeSelected, EDisplayClusterColorGradingDrawerMode::DetailsView)
+								.IsChecked(this, &SDisplayClusterColorGradingDrawer::IsDrawerModeSelected, EDisplayClusterColorGradingDrawerMode::DetailsView)
+								[
+									SNew(SImage)
+									.ColorAndOpacity(FSlateColor::UseForeground())
+									.Image(FAppStyle::Get().GetBrush("EditorPreferences.TabIcon"))
+								]
+							]
+						]
 					]
 				]
 
@@ -287,11 +327,25 @@ void SDisplayClusterColorGradingDrawer::Construct(const FArguments& InArgs, bool
 					SNew(SBorder)
 					.BorderImage(FAppStyle::Get().GetBrush("Brushes.Panel"))
 					.Padding(FMargin(2.0f, 2.0f, 2.0f, 0.0f))
+					.Visibility(this, &SDisplayClusterColorGradingDrawer::GetDrawerModeVisibility, EDisplayClusterColorGradingDrawerMode::ColorGrading)
 					[
 						SAssignNew(ColorWheelPanel, SDisplayClusterColorGradingColorWheelPanel)
 						.ColorGradingDataModelSource(ColorGradingDataModel)
 					]
 				]
+
+				// Slot for the details views
+				+SVerticalBox::Slot()
+				[
+					SNew(SBorder)
+					.BorderImage(FAppStyle::Get().GetBrush("Brushes.Panel"))
+					.Padding(FMargin(2.0f, 2.0f, 2.0f, 0.0f))
+					.Visibility(this, &SDisplayClusterColorGradingDrawer::GetDrawerModeVisibility, EDisplayClusterColorGradingDrawerMode::DetailsView)
+					[
+						SAssignNew(DetailsPanel, SDisplayClusterColorGradingDetailsPanel)
+						.ColorGradingDataModelSource(ColorGradingDataModel)
+					]
+				]
 			]
 		]
 	];
@@ -311,6 +365,11 @@ void SDisplayClusterColorGradingDrawer::Refresh(bool bPreserveDrawerState)
 		ColorWheelPanel->Refresh();
 	}
 
+	if (DetailsPanel.IsValid())
+	{
+		DetailsPanel->Refresh();
+	}
+
 	if (bPreserveDrawerState)
 	{
 		SetDrawerState(DrawerState);
@@ -354,12 +413,18 @@ FDisplayClusterColorGradingDrawerState SDisplayClusterColorGradingDrawer::GetDra
 
 	DrawerState.SelectedColorGradingGroup = ColorGradingDataModel->GetSelectedColorGradingGroupIndex();
 	DrawerState.SelectedColorGradingElement = ColorGradingDataModel->GetSelectedColorGradingElementIndex();
+	DrawerState.DrawerMode = CurrentDrawerMode;
 
 	if (ColorWheelPanel.IsValid())
 	{
 		ColorWheelPanel->GetDrawerState(DrawerState);
 	}
 
+	if (DetailsPanel.IsValid())
+	{
+		DetailsPanel->GetDrawerState(DrawerState);
+	}
+
 	if (LevelActorsList.IsValid())
 	{
 		TArray<FDisplayClusterColorGradingListItemRef> SelectedLevelItems = LevelActorsList->GetSelectedItems();
@@ -405,6 +470,8 @@ FDisplayClusterColorGradingDrawerState SDisplayClusterColorGradingDrawer::GetDra
 
 void SDisplayClusterColorGradingDrawer::SetDrawerState(const FDisplayClusterColorGradingDrawerState& InDrawerState)
 {
+	CurrentDrawerMode = InDrawerState.DrawerMode;
+
 	TArray<FDisplayClusterColorGradingListItemRef> LevelItemsToSelect;
 	TArray<FDisplayClusterColorGradingListItemRef> RootActorItemsToSelect;
 	for (const TWeakObjectPtr<UObject>& SelectedObject : InDrawerState.SelectedObjects)
@@ -458,6 +525,11 @@ void SDisplayClusterColorGradingDrawer::SetDrawerState(const FDisplayClusterColo
 	{
 		ColorWheelPanel->SetDrawerState(InDrawerState);
 	}
+
+	if (DetailsPanel.IsValid())
+	{
+		DetailsPanel->SetDrawerState(InDrawerState);
+	}
 }
 
 TSharedRef<SWidget> SDisplayClusterColorGradingDrawer::CreateDockInLayoutButton()
@@ -681,9 +753,27 @@ void SDisplayClusterColorGradingDrawer::FillColorGradingGroupToolBar()
 	}
 }
 
+ECheckBoxState SDisplayClusterColorGradingDrawer::IsDrawerModeSelected(EDisplayClusterColorGradingDrawerMode InDrawerMode) const
+{
+	return CurrentDrawerMode == InDrawerMode ? ECheckBoxState::Checked : ECheckBoxState::Unchecked;
+}
+
+EVisibility SDisplayClusterColorGradingDrawer::GetDrawerModeVisibility(EDisplayClusterColorGradingDrawerMode InDrawerMode) const
+{
+	return CurrentDrawerMode == InDrawerMode ? EVisibility::Visible : EVisibility::Collapsed;
+}
+
+void SDisplayClusterColorGradingDrawer::OnDrawerModeSelected(ECheckBoxState State, EDisplayClusterColorGradingDrawerMode InDrawerMode)
+{
+	if (State == ECheckBoxState::Checked)
+	{
+		CurrentDrawerMode = InDrawerMode;
+	}
+}
+
 EVisibility SDisplayClusterColorGradingDrawer::GetColorGradingGroupToolBarVisibility() const
 {
-	if (ColorGradingDataModel->bShowColorGradingGroupToolBar)
+	if (ColorGradingDataModel->bShowColorGradingGroupToolBar && CurrentDrawerMode == EDisplayClusterColorGradingDrawerMode::ColorGrading)
 	{
 		return EVisibility::Visible;
 	}
@@ -810,6 +900,21 @@ void SDisplayClusterColorGradingDrawer::OnActiveRootActorChanged(ADisplayCluster
 	Refresh(bPreserveDrawerState);
 }
 
+void SDisplayClusterColorGradingDrawer::OnColorGradingDataModelGenerated()
+{
+	FillColorGradingGroupToolBar();
+
+	if (ColorWheelPanel.IsValid())
+	{
+		ColorWheelPanel->Refresh();
+	}
+
+	if (DetailsPanel.IsValid())
+	{
+		DetailsPanel->Refresh();
+	}
+}
+
 void SDisplayClusterColorGradingDrawer::OnListSelectionChanged(TSharedRef<SDisplayClusterColorGradingObjectList> SourceList, FDisplayClusterColorGradingListItemRef SelectedItem, ESelectInfo::Type SelectInfo)
 {
 	TArray<FDisplayClusterColorGradingListItemRef> SelectedObjects = SourceList->GetSelectedItems();
@@ -844,13 +949,6 @@ void SDisplayClusterColorGradingDrawer::OnListSelectionChanged(TSharedRef<SDispl
 	}
 
 	ColorGradingDataModel->SetObjects(ObjectsToColorGrade);
-
-	FillColorGradingGroupToolBar();
-
-	if (ColorWheelPanel.IsValid())
-	{
-		ColorWheelPanel->Refresh();
-	}
 }
 
 FReply SDisplayClusterColorGradingDrawer::DockInLayout()
diff --git a/Engine/Plugins/Runtime/nDisplay/Source/DisplayClusterColorGrading/Private/Drawer/SDisplayClusterColorGradingDrawer.h b/Engine/Plugins/Runtime/nDisplay/Source/DisplayClusterColorGrading/Private/Drawer/SDisplayClusterColorGradingDrawer.h
--- a/Engine/Plugins/Runtime/nDisplay/Source/DisplayClusterColorGrading/Private/Drawer/SDisplayClusterColorGradingDrawer.h
+++ b/Engine/Plugins/Runtime/nDisplay/Source/DisplayClusterColorGrading/Private/Drawer/SDisplayClusterColorGradingDrawer.h
@@ -18,6 +18,14 @@ class IDisplayClusterOperatorViewModel;
 class IPropertyRowGenerator;
 class SHorizontalBox;
 class SDisplayClusterColorGradingColorWheelPanel;
+class SDisplayClusterColorGradingDetailsPanel;
+
+/** Drawer modes that the color grading drawer can be in */
+enum EDisplayClusterColorGradingDrawerMode
+{
+	ColorGrading,
+	DetailsView
+};
 
 /** Stores the state of the drawer UI that can be reloaded in cases where the drawer or any of its elements are reloaded (such as when the drawer is reopened or docked) */
 struct FDisplayClusterColorGradingDrawerState
@@ -39,6 +47,12 @@ struct FDisplayClusterColorGradingDrawerState
 
 	/** The color display mode of the color wheels */
 	SDisplayClusterColorGradingColorWheel::EColorDisplayMode ColorDisplayMode;
+
+	/** The mode the drawer is in */
+	EDisplayClusterColorGradingDrawerMode DrawerMode;
+
+	/** Indicates which subsections were selected for each section in the details panel */
+	TArray<int32> SelectedDetailsSubsections;
 };
 
 /** Color grading drawer widget, which displays a list of color gradable items, and the color wheel panel */
@@ -96,6 +110,15 @@ private:
 	/** Fills the color grading group toolbar using the color grading data model */
 	void FillColorGradingGroupToolBar();
 
+	/** Gets whether the specified drawer mode is currently selected */
+	ECheckBoxState IsDrawerModeSelected(EDisplayClusterColorGradingDrawerMode InDrawerMode) const;
+
+	/** Gets the visibility state for the specified drawer mode */
+	EVisibility GetDrawerModeVisibility(EDisplayClusterColorGradingDrawerMode InDrawerMode) const;
+
+	/** Raised when the specified drawer mode checked state is changed */
+	void OnDrawerModeSelected(ECheckBoxState State, EDisplayClusterColorGradingDrawerMode InDrawerMode);
+
 	/** Gets the visibility state of the color grading group toolbar */
 	EVisibility GetColorGradingGroupToolBarVisibility() const;
 
@@ -119,6 +142,9 @@ private:
 
 	/** Raised when the user has changed the active root actor selected in the nDisplay operator panel */
 	void OnActiveRootActorChanged(ADisplayClusterRootActor* NewRootActor);
+	
+	/** Raised when the color grading data model has been generated */
+	void OnColorGradingDataModelGenerated();
 
 	/** Raised when the user has selected a new item in any of the drawer's list views */
 	void OnListSelectionChanged(TSharedRef<SDisplayClusterColorGradingObjectList> SourceList, FDisplayClusterColorGradingListItemRef SelectedItem, ESelectInfo::Type SelectInfo);
@@ -134,6 +160,7 @@ private:
 	TSharedPtr<SDisplayClusterColorGradingObjectList> RootActorList;
 	TSharedPtr<SHorizontalBox> ColorGradingGroupToolBarBox;
 	TSharedPtr<SDisplayClusterColorGradingColorWheelPanel> ColorWheelPanel;
+	TSharedPtr<SDisplayClusterColorGradingDetailsPanel> DetailsPanel;
 
 	/** List of color gradable items found in the current level */
 	TArray<FDisplayClusterColorGradingListItemRef> LevelColorGradingItems;
@@ -144,6 +171,9 @@ private:
 	/** The color grading data model for the currently selected objects */
 	TSharedPtr<FDisplayClusterColorGradingDataModel> ColorGradingDataModel;
 
+	/** The current mode that the drawer is in */
+	EDisplayClusterColorGradingDrawerMode CurrentDrawerMode;
+
 	/** Gets whether this widget is in a drawer or docked in a tab */
 	bool bIsInDrawer = false;
 
diff --git a/Engine/Plugins/Runtime/nDisplay/Source/DisplayClusterColorGrading/Private/SDisplayClusterColorGradingColorWheelPanel.cpp b/Engine/Plugins/Runtime/nDisplay/Source/DisplayClusterColorGrading/Private/SDisplayClusterColorGradingColorWheelPanel.cpp
--- a/Engine/Plugins/Runtime/nDisplay/Source/DisplayClusterColorGrading/Private/SDisplayClusterColorGradingColorWheelPanel.cpp
+++ b/Engine/Plugins/Runtime/nDisplay/Source/DisplayClusterColorGrading/Private/SDisplayClusterColorGradingColorWheelPanel.cpp
@@ -469,12 +469,15 @@ void SDisplayClusterColorGradingColorWheelPanel::FillDetailsView(const FDisplayC
 			ColorGradingGroup.DetailsViewPropertyHandles[0]->GetOuterObjects(OuterObjects);
 		}
 
-		DetailsView->RegisterInstancedCustomPropertyLayout(OuterObjects[0]->GetClass(), FOnGetDetailCustomizationInstance::CreateLambda([&ColorGradingGroup]()
+		if (OuterObjects.Num())
 		{
-			return MakeShareable(new FColorGradingDetailsCustomization(ColorGradingGroup.DetailsViewPropertyHandles));
-		}));
+			DetailsView->RegisterInstancedCustomPropertyLayout(OuterObjects[0]->GetClass(), FOnGetDetailCustomizationInstance::CreateLambda([&ColorGradingGroup]()
+			{
+				return MakeShareable(new FColorGradingDetailsCustomization(ColorGradingGroup.DetailsViewPropertyHandles));
+			}));
 
-		DetailsView->SetObjects(OuterObjects);
+			DetailsView->SetObjects(OuterObjects);
+		}
 	}
 }
 
@@ -547,15 +550,18 @@ void SDisplayClusterColorGradingColorWheelPanel::ClearColorWheels()
 
 void SDisplayClusterColorGradingColorWheelPanel::SetColorWheelOrientation(EOrientation NewOrientation)
 {
-	ColorWheelOrientation = NewOrientation;
-
-	for (const TSharedPtr<SDisplayClusterColorGradingColorWheel>& ColorWheel : ColorWheels)
+	if (ColorWheelOrientation != NewOrientation)
 	{
-		if (ColorWheel.IsValid())
+		ColorWheelOrientation = NewOrientation;
+
+		for (const TSharedPtr<SDisplayClusterColorGradingColorWheel>& ColorWheel : ColorWheels)
 		{
-			ColorWheel->SetOrientation(ColorWheelOrientation);
-		}
-	};
+			if (ColorWheel.IsValid())
+			{
+				ColorWheel->SetOrientation(ColorWheelOrientation);
+			}
+		};
+	}
 }
 
 bool SDisplayClusterColorGradingColorWheelPanel::IsColorWheelOrientationSelected(EOrientation Orientation) const
diff --git a/Engine/Plugins/Runtime/nDisplay/Source/DisplayClusterColorGrading/Private/SDisplayClusterColorGradingDetailsPanel.cpp b/Engine/Plugins/Runtime/nDisplay/Source/DisplayClusterColorGrading/Private/SDisplayClusterColorGradingDetailsPanel.cpp
new file mode 100644
--- /dev/null
+++ b/Engine/Plugins/Runtime/nDisplay/Source/DisplayClusterColorGrading/Private/SDisplayClusterColorGradingDetailsPanel.cpp
@@ -0,0 +1,328 @@
+// Copyright Epic Games, Inc. All Rights Reserved.
+
+#include "SDisplayClusterColorGradingDetailsPanel.h"
+
+#include "DisplayClusterColorGradingDataModel.h"
+#include "Drawer/SDisplayClusterColorGradingDrawer.h"
+
+#include "DetailCategoryBuilder.h"
+#include "IDetailsView.h"
+#include "Modules/ModuleManager.h"
+#include "PropertyEditorModule.h"
+#include "Widgets/SBoxPanel.h"
+#include "Widgets/Input/SCheckBox.h"
+#include "Widgets/Layout/SBox.h"
+#include "Widgets/Layout/SSplitter.h"
+#include "Widgets/Text/STextBlock.h"
+
+/** A widget that wraps a details view used to display the properties specified by a details section in the color grading data model */
+class SDetailsSectionView : public SCompoundWidget
+{
+public:
+	SLATE_BEGIN_ARGS(SDetailsSectionView)
+	{}
+	SLATE_END_ARGS()
+
+	void Construct(const FArguments& InArgs)
+	{
+		FPropertyEditorModule& EditModule = FModuleManager::Get().GetModuleChecked<FPropertyEditorModule>("PropertyEditor");
+		FDetailsViewArgs DetailsViewArgs;
+		DetailsViewArgs.bAllowSearch = false;
+		DetailsViewArgs.NameAreaSettings = FDetailsViewArgs::HideNameArea;
+		DetailsViewArgs.bHideSelectionTip = true;
+		DetailsView = EditModule.CreateDetailView(DetailsViewArgs);
+
+		ChildSlot
+		[
+			SNew(SVerticalBox)
+
+			+ SVerticalBox::Slot()
+			.AutoHeight()
+			.Padding(2, 4)
+			[
+				SAssignNew(HeaderBox, SBox)
+			]
+
+			+ SVerticalBox::Slot()
+			.FillHeight(1.0f)
+			.Padding(2, 0)
+			[
+				DetailsView.ToSharedRef()
+			]
+		];
+	}
+
+	/** Fills the details view with the specified objects, and configures it used the specified details section */
+	void FillDetailsSection(const TArray<TWeakObjectPtr<UObject>>& Objects, const FDisplayClusterColorGradingDataModel::FDetailsSection& InDetailsSection)
+	{
+		DetailsSection = &InDetailsSection;
+		ObjectClass = nullptr;
+
+		if (Objects.Num())
+		{
+			ObjectClass = UObject::StaticClass();
+
+			const TWeakObjectPtr<UObject>& Object = Objects[0];
+			if (Object.IsValid())
+			{
+				ObjectClass = Object->GetClass();
+			}
+		}
+
+		// If a details subsection isn't already selected and the details section has subsections, set the selected subsection to the first subsection
+		if (SelectedSubsectionIndex == INDEX_NONE && DetailsSection->Subsections.Num())
+		{
+			SelectedSubsectionIndex = 0;
+		}
+
+		if (ObjectClass)
+		{
+			// If a subsection is selected, use the details customization for that subsection. Otherwise, use the details customization for the entire section
+			if (SelectedSubsectionIndex > INDEX_NONE && SelectedSubsectionIndex < DetailsSection->Subsections.Num())
+			{
+				DetailsView->RegisterInstancedCustomPropertyLayout(ObjectClass, DetailsSection->Subsections[SelectedSubsectionIndex].DetailCustomizationDelegate);
+			}
+			else
+			{
+				DetailsView->RegisterInstancedCustomPropertyLayout(ObjectClass, DetailsSection->DetailCustomizationDelegate);
+			}
+		}
+
+		DetailsView->SetObjects(Objects);
+		HeaderBox->SetContent(CreateHeader());
+	}
+
+	/** Clears the details view and header */
+	void EmptyDetailsSection()
+	{
+		DetailsSection = nullptr;
+		ObjectClass = nullptr;
+		SelectedSubsectionIndex = INDEX_NONE;
+
+		if (HeaderBox.IsValid())
+		{
+			HeaderBox->SetContent(SNullWidget::NullWidget);
+		}
+
+		if (DetailsView.IsValid())
+		{
+			TArray<UObject*> EmptyList;
+			DetailsView->SetObjects(EmptyList);
+		}
+	}
+
+	/** Gets the selected subsection index */
+	int32 GetSelectedSubsectionIndex() const
+	{
+		return SelectedSubsectionIndex;
+	}
+
+	/** Sets the selected subsection index*/
+	void SetSelectedSubsectionIndex(int32 InSubsectionIndex)
+	{
+		if (DetailsSection)
+		{
+			SelectedSubsectionIndex = InSubsectionIndex;
+
+			if (ObjectClass)
+			{
+				DetailsView->UnregisterInstancedCustomPropertyLayout(ObjectClass);
+
+				// If a subsection is selected, use the details customization for that subsection. Otherwise, use the details customization for the entire section
+				if (SelectedSubsectionIndex > INDEX_NONE && SelectedSubsectionIndex < DetailsSection->Subsections.Num())
+				{
+					DetailsView->RegisterInstancedCustomPropertyLayout(ObjectClass, DetailsSection->Subsections[SelectedSubsectionIndex].DetailCustomizationDelegate);
+				}
+				else
+				{
+					DetailsView->RegisterInstancedCustomPropertyLayout(ObjectClass, DetailsSection->DetailCustomizationDelegate);
+				}
+			}
+
+			DetailsView->ForceRefresh();
+		}
+	}
+
+private:
+	TSharedRef<SWidget> CreateHeader()
+	{
+		if (DetailsSection)
+		{
+			TSharedRef<SHorizontalBox> HorizontalBox = SNew(SHorizontalBox);
+
+			if (DetailsSection->EditConditionPropertyHandle.IsValid() && DetailsSection->EditConditionPropertyHandle->IsValidHandle())
+			{
+				HorizontalBox->AddSlot()
+				.AutoWidth()
+				.VAlign(VAlign_Center)
+				.Padding(2, 0, 4, 0)
+				[
+					DetailsSection->EditConditionPropertyHandle->CreatePropertyValueWidget(false)
+				];
+			}
+
+			HorizontalBox->AddSlot()
+			.FillWidth(1.0f)
+			.VAlign(VAlign_Center)
+			.Padding(2, 0, 0, 0)
+			[
+				SNew(STextBlock)
+				.Text(DetailsSection->DisplayName)
+				.Font(FAppStyle::Get().GetFontStyle("NormalFontBold"))
+			];
+
+			if (DetailsSection->Subsections.Num())
+			{
+				HorizontalBox->AddSlot()
+				.AutoWidth()
+				.VAlign(VAlign_Center)
+				[
+					CreateSubsectionToolBar()
+				];
+			}
+
+			return HorizontalBox;
+		}
+
+		return SNullWidget::NullWidget;
+	}
+
+	TSharedRef<SWidget> CreateSubsectionToolBar()
+	{
+		if (DetailsSection && DetailsSection->Subsections.Num())
+		{
+			TSharedRef<SHorizontalBox> ToolBarBox = SNew(SHorizontalBox);
+
+			for (int32 Index = 0; Index < DetailsSection->Subsections.Num(); ++Index)
+			{
+				ToolBarBox->AddSlot()
+				.AutoWidth()
+				.Padding(2.0f, 0.0f, 2.0f, 0.0f)
+				[
+					SNew(SCheckBox)
+					.Style(FAppStyle::Get(), "DetailsView.SectionButton")
+					.IsChecked(this, &SDetailsSectionView::IsSubsectionSelected, Index)
+					.OnCheckStateChanged(this, &SDetailsSectionView::OnSelectedSubsectionCheckedChanged, Index)
+					[
+						SNew(STextBlock)
+						.TextStyle(FAppStyle::Get(), "SmallText")
+						.Text(DetailsSection->Subsections[Index].DisplayName)
+					]
+				];
+			}
+
+			return ToolBarBox;
+		}
+
+		return SNullWidget::NullWidget;
+	}
+
+	ECheckBoxState IsSubsectionSelected(int32 SubsectionIndex) const
+	{
+		return SelectedSubsectionIndex == SubsectionIndex ? ECheckBoxState::Checked : ECheckBoxState::Unchecked;
+	}
+
+	void OnSelectedSubsectionCheckedChanged(ECheckBoxState CheckBoxState, int32 SubsectionIndex)
+	{
+		if (CheckBoxState == ECheckBoxState::Checked)
+		{
+			SetSelectedSubsectionIndex(SubsectionIndex);
+		}
+	}
+
+private:
+	TSharedPtr<IDetailsView> DetailsView;
+	TSharedPtr<SBox> HeaderBox;
+
+	/** Pointer to the details section this view represents */
+	const FDisplayClusterColorGradingDataModel::FDetailsSection* DetailsSection = nullptr;
+
+	/** The object class of the objects that the details view is currently displaying */
+	UClass* ObjectClass = nullptr;
+
+	/** The index of the currently selected details subsection */
+	int32 SelectedSubsectionIndex = INDEX_NONE;
+};
+
+void SDisplayClusterColorGradingDetailsPanel::Construct(const FArguments& InArgs)
+{
+	ColorGradingDataModel = InArgs._ColorGradingDataModelSource;
+
+    ChildSlot
+	[
+		SAssignNew(Splitter, SSplitter)
+		.PhysicalSplitterHandleSize(2.0f)
+	];
+
+	DetailsSectionViews.AddDefaulted(MaxNumDetailsSections);
+	for (int32 Index = 0; Index < MaxNumDetailsSections; ++Index)
+	{
+		Splitter->AddSlot()
+		[
+			SAssignNew(DetailsSectionViews[Index], SDetailsSectionView)
+			.Visibility(this, &SDisplayClusterColorGradingDetailsPanel::GetDetailsSectionVisibility, Index)
+		];
+	}
+}
+
+void SDisplayClusterColorGradingDetailsPanel::Refresh()
+{
+	FillDetailsSections();
+}
+
+void SDisplayClusterColorGradingDetailsPanel::GetDrawerState(FDisplayClusterColorGradingDrawerState& OutDrawerState)
+{
+	for (int32 Index = 0; Index < DetailsSectionViews.Num(); ++Index)
+	{
+		OutDrawerState.SelectedDetailsSubsections.Add(DetailsSectionViews[Index]->GetSelectedSubsectionIndex());
+	}
+}
+
+void SDisplayClusterColorGradingDetailsPanel::SetDrawerState(const FDisplayClusterColorGradingDrawerState& InDrawerState)
+{
+	if (ColorGradingDataModel.IsValid())
+	{
+		if (InDrawerState.SelectedDetailsSubsections.Num() == ColorGradingDataModel->DetailsSections.Num())
+		{
+			const int32 NumDetailsSections = FMath::Min(ColorGradingDataModel->DetailsSections.Num(), MaxNumDetailsSections);
+			for (int32 Index = 0; Index < NumDetailsSections; ++Index)
+			{
+				DetailsSectionViews[Index]->SetSelectedSubsectionIndex(InDrawerState.SelectedDetailsSubsections[Index]);
+			}
+		}
+	}
+}
+
+void SDisplayClusterColorGradingDetailsPanel::FillDetailsSections()
+{
+	if (ColorGradingDataModel.IsValid())
+	{
+		TArray<TWeakObjectPtr<UObject>> Objects = ColorGradingDataModel->GetObjects();
+
+		const int32 NumDetailsSections = FMath::Min(ColorGradingDataModel->DetailsSections.Num(), MaxNumDetailsSections);
+		for (int32 Index = 0; Index < MaxNumDetailsSections; ++Index)
+		{
+			if (Index < ColorGradingDataModel->DetailsSections.Num())
+			{
+				const FDisplayClusterColorGradingDataModel::FDetailsSection& DetailsSection = ColorGradingDataModel->DetailsSections[Index];
+
+				DetailsSectionViews[Index]->FillDetailsSection(Objects, DetailsSection);
+			}
+			else
+			{
+				DetailsSectionViews[Index]->EmptyDetailsSection();
+			}
+		}
+	}
+}
+
+EVisibility SDisplayClusterColorGradingDetailsPanel::GetDetailsSectionVisibility(int32 SectionIndex) const
+{
+	if (ColorGradingDataModel.IsValid())
+	{
+		const int32 NumSections = FMath::Min(ColorGradingDataModel->DetailsSections.Num(), MaxNumDetailsSections);
+		return SectionIndex < NumSections ? EVisibility::Visible : EVisibility::Collapsed;
+	}
+
+	return EVisibility::Collapsed;
+}
\ No newline at end of file
diff --git a/Engine/Plugins/Runtime/nDisplay/Source/DisplayClusterColorGrading/Private/SDisplayClusterColorGradingDetailsPanel.h b/Engine/Plugins/Runtime/nDisplay/Source/DisplayClusterColorGrading/Private/SDisplayClusterColorGradingDetailsPanel.h
new file mode 100644
--- /dev/null
+++ b/Engine/Plugins/Runtime/nDisplay/Source/DisplayClusterColorGrading/Private/SDisplayClusterColorGradingDetailsPanel.h
@@ -0,0 +1,50 @@
+// Copyright Epic Games, Inc. All Rights Reserved.
+
+#pragma once
+
+#include "CoreMinimal.h"
+#include "Widgets/DeclarativeSyntaxSupport.h"
+#include "Widgets/SCompoundWidget.h"
+
+class FDisplayClusterColorGradingDataModel;
+class SDetailsSectionView;
+class SSplitter;
+struct FDisplayClusterColorGradingDrawerState;
+
+/** A panel that displays several property details views based on the color grading data model */
+class SDisplayClusterColorGradingDetailsPanel : public SCompoundWidget
+{
+public:
+	/** The maximum number of details sections that are allowed to be displayed at the same time */
+	static const int32 MaxNumDetailsSections = 3;
+
+public:
+	SLATE_BEGIN_ARGS(SDisplayClusterColorGradingDetailsPanel) {}
+		SLATE_ARGUMENT(TSharedPtr<FDisplayClusterColorGradingDataModel>, ColorGradingDataModelSource)
+	SLATE_END_ARGS()
+
+	void Construct(const FArguments& InArgs);
+
+	/** Refreshes the details panel to reflect the current state of the color grading data model */
+	void Refresh();
+
+	/** Adds the state of the details panel to the specified drawer state */
+	void GetDrawerState(FDisplayClusterColorGradingDrawerState& OutDrawerState);
+
+	/** Sets the state of the details panel from the specified drawer state */
+	void SetDrawerState(const FDisplayClusterColorGradingDrawerState& InDrawerState);
+
+private:
+	/** Fills the details sections based on the current state of the color gradin data model */
+	void FillDetailsSections();
+
+	/** Gets the visibility state of the specified details section */
+	EVisibility GetDetailsSectionVisibility(int32 SectionIndex) const;
+
+private:
+	/** The color grading data model that the panel is displaying */
+	TSharedPtr<FDisplayClusterColorGradingDataModel> ColorGradingDataModel;
+
+	TArray<TSharedPtr<SDetailsSectionView>> DetailsSectionViews;
+	TSharedPtr<SSplitter> Splitter;
+};
\ No newline at end of file
diff --git a/Engine/Plugins/Runtime/nDisplay/Source/DisplayClusterConfigurator/Private/Views/Details/DisplayClusterEditorPropertyReferenceTypeCustomization.cpp b/Engine/Plugins/Runtime/nDisplay/Source/DisplayClusterConfigurator/Private/Views/Details/DisplayClusterEditorPropertyReferenceTypeCustomization.cpp
--- a/Engine/Plugins/Runtime/nDisplay/Source/DisplayClusterConfigurator/Private/Views/Details/DisplayClusterEditorPropertyReferenceTypeCustomization.cpp
+++ b/Engine/Plugins/Runtime/nDisplay/Source/DisplayClusterConfigurator/Private/Views/Details/DisplayClusterEditorPropertyReferenceTypeCustomization.cpp
@@ -21,7 +21,7 @@ TSharedRef<IPropertyTypeCustomization> FDisplayClusterEditorPropertyReferenceTyp
 
 void FDisplayClusterEditorPropertyReferenceTypeCustomization::CustomizeHeader(TSharedRef<IPropertyHandle> InPropertyHandle, FDetailWidgetRow& InHeaderRow, IPropertyTypeCustomizationUtils& CustomizationUtils)
 {
-	CustomizationUtilsPtr = &CustomizationUtils;
+	PropertyUtilities = CustomizationUtils.GetPropertyUtilities();
 
 	check(InPropertyHandle->IsValidHandle());
 	if (InPropertyHandle->HasMetaData(EditConditionPathMetadataKey))
@@ -141,18 +141,34 @@ void FDisplayClusterEditorPropertyReferenceTypeCustomization::CustomizeChildren(
 					ReferencedPropertyHandle->SetToolTipText(InPropertyHandle->GetToolTipText());
 				}
 
-				ReferencedPropertyHandle->SetOnPropertyValueChanged(FSimpleDelegate::CreateRaw(this, &FDisplayClusterEditorPropertyReferenceTypeCustomization::OnReferencedPropertyValueChanged));
+				ReferencedPropertyHandle->SetOnPropertyValueChanged(FSimpleDelegate::CreateSP(this, &FDisplayClusterEditorPropertyReferenceTypeCustomization::OnReferencedPropertyValueChanged));
+				
+				const bool bHasShowOnlyInners = InPropertyHandle->HasMetaData(TEXT("ShowOnlyInnerProperties")) || InPropertyHandle->GetInstanceMetaData(TEXT("ShowOnlyInnerProperties"));
+				if (bHasShowOnlyInners)
+				{
+					// If the property referencer has the ShowOnlyInnerProperties metadata, expand the reference property and display the child properties directly
+					uint32 NumChildren;
+					ReferencedPropertyHandle->GetNumChildren(NumChildren);
 
-				IDetailPropertyRow& PropertyRow = InChildBuilder.AddProperty(ReferencedPropertyHandle.ToSharedRef());
+					for (uint32 Index = 0; Index < NumChildren; ++Index)
+					{
+						TSharedPtr<IPropertyHandle> ChildHandle = ReferencedPropertyHandle->GetChildHandle(Index);
+						InChildBuilder.AddProperty(ChildHandle.ToSharedRef());
+					}
+				}
+				else
+				{
+					IDetailPropertyRow& PropertyRow = InChildBuilder.AddProperty(ReferencedPropertyHandle.ToSharedRef());
+
+					if (EditConditionPropertyHandles.Num())
+					{
+						PropertyRow.EditCondition(CreateEditConditional(), nullptr);
+					}
+				}
 
 				// Mark the property with the "IsCustomized" flag so that any subsequent layout builders can account for the property
 				// being moved and placed here.
 				ReferencedPropertyHandle->MarkHiddenByCustomization();
-
-				if (EditConditionPropertyHandles.Num())
-				{
-					PropertyRow.EditCondition(CreateEditConditional(), nullptr);
-				}
 			}
 		}
 	}
@@ -161,9 +177,9 @@ void FDisplayClusterEditorPropertyReferenceTypeCustomization::CustomizeChildren(
 void FDisplayClusterEditorPropertyReferenceTypeCustomization::OnReferencedPropertyValueChanged()
 {
 	// When a referenced property is changed, we have to trigger layout refresh
-	if (CustomizationUtilsPtr)
+	if (PropertyUtilities.IsValid())
 	{
-		CustomizationUtilsPtr->GetPropertyUtilities()->ForceRefresh();
+		PropertyUtilities.Pin()->ForceRefresh();
 	}
 }
 
diff --git a/Engine/Plugins/Runtime/nDisplay/Source/DisplayClusterConfigurator/Private/Views/Details/DisplayClusterEditorPropertyReferenceTypeCustomization.h b/Engine/Plugins/Runtime/nDisplay/Source/DisplayClusterConfigurator/Private/Views/Details/DisplayClusterEditorPropertyReferenceTypeCustomization.h
--- a/Engine/Plugins/Runtime/nDisplay/Source/DisplayClusterConfigurator/Private/Views/Details/DisplayClusterEditorPropertyReferenceTypeCustomization.h
+++ b/Engine/Plugins/Runtime/nDisplay/Source/DisplayClusterConfigurator/Private/Views/Details/DisplayClusterEditorPropertyReferenceTypeCustomization.h
@@ -43,7 +43,7 @@ private:
 	bool IsEditable() const;
 
 	// Customization utils instance to trigger GUI update on referenced value changes
-	IPropertyTypeCustomizationUtils* CustomizationUtilsPtr = nullptr;
+	TWeakPtr<IPropertyUtilities> PropertyUtilities = nullptr;
 
 	TArray<FPropertyHandlePair> ReferencedPropertyHandles;
 	TArray<TSharedRef<IPropertyHandle>> EditConditionPropertyHandles;
