Turned on unity files for DisplayClusterTests by default to save compile times.

Before: Total CPU Time: 43.78125 s
After: Total CPU Time: 23.375 s

diff --git a/Engine/Plugins/Runtime/nDisplay/Source/DisplayClusterTests/DisplayClusterTests.Build.cs b/Engine/Plugins/Runtime/nDisplay/Source/DisplayClusterTests/DisplayClusterTests.Build.cs
--- a/Engine/Plugins/Runtime/nDisplay/Source/DisplayClusterTests/DisplayClusterTests.Build.cs
+++ b/Engine/Plugins/Runtime/nDisplay/Source/DisplayClusterTests/DisplayClusterTests.Build.cs
@@ -21,7 +21,6 @@ public class DisplayClusterTests : ModuleRules
 			});
 
 		OptimizeCode = CodeOptimization.Never;
-		bUseUnity = false;
 		PCHUsage = PCHUsageMode.NoPCHs;
 	}
 }
