nDisplay: Don't hide components that have bCastHiddenShadow enabled.

This allows rendering the shadow but not the geo.

diff --git a/Engine/Plugins/Runtime/nDisplay/Source/DisplayCluster/Private/Game/EngineClasses/Scene/DisplayClusterRootActor.cpp b/Engine/Plugins/Runtime/nDisplay/Source/DisplayCluster/Private/Game/EngineClasses/Scene/DisplayClusterRootActor.cpp
--- a/Engine/Plugins/Runtime/nDisplay/Source/DisplayCluster/Private/Game/EngineClasses/Scene/DisplayClusterRootActor.cpp
+++ b/Engine/Plugins/Runtime/nDisplay/Source/DisplayCluster/Private/Game/EngineClasses/Scene/DisplayClusterRootActor.cpp
@@ -537,7 +537,8 @@ bool ADisplayClusterRootActor::GetHiddenInGamePrimitives(TSet<FPrimitiveComponen
 
 				for (UPrimitiveComponent* PrimComp : PrimitiveComponents)
 				{
-					if ((bHideVisualizationComponents && PrimComp->IsVisualizationComponent()) || bActorHideInGame || PrimComp->bHiddenInGame)
+					if ((bHideVisualizationComponents && PrimComp->IsVisualizationComponent()) 
+						|| ((bActorHideInGame || PrimComp->bHiddenInGame) && !PrimComp->bCastHiddenShadow))
 					{
 						OutPrimitives.Add(PrimComp->ComponentId);
 					}
