nDisplay: Changed default size of UV lightcards.

So that they don't appear as big when they are created.

diff --git a/Engine/Plugins/Runtime/nDisplay/Source/DisplayClusterLightCardEditor/Private/SDisplayClusterLightCardEditor.cpp b/Engine/Plugins/Runtime/nDisplay/Source/DisplayClusterLightCardEditor/Private/SDisplayClusterLightCardEditor.cpp
--- a/Engine/Plugins/Runtime/nDisplay/Source/DisplayClusterLightCardEditor/Private/SDisplayClusterLightCardEditor.cpp
+++ b/Engine/Plugins/Runtime/nDisplay/Source/DisplayClusterLightCardEditor/Private/SDisplayClusterLightCardEditor.cpp
@@ -226,6 +226,7 @@ ADisplayClusterLightCardActor* SDisplayClusterLightCardEditor::SpawnLightCard(co
 		if (ProjectionMode == EDisplayClusterMeshProjectionType::UV)
 		{
 			NewLightCard->bIsUVLightCard = true;
+			NewLightCard->Scale /= 4;
 		}
 	}
 
