nDisplay: Removed old lightcard blueprints from place actors menu.

This should avoid users mistakenly using the old light cards which are not supported in the icvfx operator panel.

diff --git a/Engine/Plugins/Runtime/nDisplay/Source/DisplayClusterConfigurator/DisplayClusterConfigurator.Build.cs b/Engine/Plugins/Runtime/nDisplay/Source/DisplayClusterConfigurator/DisplayClusterConfigurator.Build.cs
--- a/Engine/Plugins/Runtime/nDisplay/Source/DisplayClusterConfigurator/DisplayClusterConfigurator.Build.cs
+++ b/Engine/Plugins/Runtime/nDisplay/Source/DisplayClusterConfigurator/DisplayClusterConfigurator.Build.cs
@@ -37,7 +37,6 @@ public class DisplayClusterConfigurator : ModuleRules
 				"MediaAssets",
 				"MessageLog",
 				"PinnedCommandList",
-				"PlacementMode",
 				"Projects",
 				"PropertyEditor",
 				"Serialization",
@@ -49,7 +48,6 @@ public class DisplayClusterConfigurator : ModuleRules
 				"SubobjectEditor",
 				"SubobjectDataInterface",
 				"ToolWidgets",
-				"VPUtilitiesEditor",
 				"ProceduralMeshComponent",
 			});
 	}
diff --git a/Engine/Plugins/Runtime/nDisplay/Source/DisplayClusterConfigurator/Private/DisplayClusterConfiguratorModule.cpp b/Engine/Plugins/Runtime/nDisplay/Source/DisplayClusterConfigurator/Private/DisplayClusterConfiguratorModule.cpp
--- a/Engine/Plugins/Runtime/nDisplay/Source/DisplayClusterConfigurator/Private/DisplayClusterConfiguratorModule.cpp
+++ b/Engine/Plugins/Runtime/nDisplay/Source/DisplayClusterConfigurator/Private/DisplayClusterConfiguratorModule.cpp
@@ -40,8 +40,6 @@
 #include "AssetTypeCategories.h"
 #include "HAL/IConsoleManager.h"
 #include "Modules/ModuleManager.h"
-#include "IPlacementModeModule.h"
-#include "IVPUtilitiesEditorModule.h"
 #include "ActorFactories/ActorFactoryBlueprint.h"
 
 #define REGISTER_PROPERTY_LAYOUT(PropertyType, CustomizationType) { \
@@ -89,7 +87,6 @@ void FDisplayClusterConfiguratorModule::StartupModule()
 
 	RegisterCustomLayouts();
 	RegisterSettings();
-	RegisterPlacementModeItems();
 	RegisterSectionMappings();
 	
 	FDisplayClusterConfiguratorStyle::Get();
@@ -127,7 +124,6 @@ void FDisplayClusterConfiguratorModule::ShutdownModule()
 
 	UnregisterSettings();
 	UnregisterCustomLayouts();
-	UnregisterPlacementModeItems();
 	UnregisterSectionMappings();
 
 	MenuExtensibilityManager.Reset();
@@ -303,74 +299,6 @@ void FDisplayClusterConfiguratorModule::UnregisterSectionMappings()
 	}
 }
 
-void FDisplayClusterConfiguratorModule::RegisterPlacementModeItems()
-{
-	auto RegisterPlaceActors = [&]() -> void
-	{
-		if (!GEditor)
-		{
-			return;
-		}
-		
-		const FPlacementCategoryInfo* Info = IVPUtilitiesEditorModule::Get().GetVirtualProductionPlacementCategoryInfo();
-		if (!Info)
-		{
-			UE_LOG(DisplayClusterConfiguratorLog, Warning, TEXT("Could not find or create VirtualProduction Place Actor Category"));
-			return;
-		}
-
-		FAssetData LightCardAssetData(
-			TEXT("/nDisplay/LightCard/LightCard"),
-			TEXT("/nDisplay/LightCard"),
-			TEXT("LightCard"),
-			FTopLevelAssetPath(TEXT("/Script/Engine"), TEXT("Blueprint"))
-		);
-		check (LightCardAssetData.IsValid());
-	
-		PlaceActors.Add(IPlacementModeModule::Get().RegisterPlaceableItem(Info->UniqueHandle, MakeShared<FPlaceableItem>(
-			*UActorFactoryBlueprint::StaticClass(),
-			LightCardAssetData,
-			NAME_None,
-			NAME_None,
-			TOptional<FLinearColor>(),
-			TOptional<int32>(),
-			NSLOCTEXT("PlacementMode", "LightCard", "LightCard")
-		)));
-	};
-
-	if (GEngine && GEngine->IsInitialized())
-	{
-		RegisterPlaceActors();
-	}
-	else
-	{
-		PostEngineInitHandle = FCoreDelegates::OnPostEngineInit.AddLambda(RegisterPlaceActors);
-	}
-}
-
-void FDisplayClusterConfiguratorModule::UnregisterPlacementModeItems()
-{
-	if (!IsEngineExitRequested() && GEditor && UObjectInitialized())
-	{
-		IPlacementModeModule& PlacementModeModule = IPlacementModeModule::Get();
-
-		for (TOptional<FPlacementModeID>& PlaceActor : PlaceActors)
-		{
-			if (PlaceActor.IsSet())
-			{
-				PlacementModeModule.UnregisterPlaceableItem(*PlaceActor);
-			}
-		}
-	}
-
-	if (PostEngineInitHandle.IsValid())
-	{
-		FCoreDelegates::OnPostEngineInit.Remove(PostEngineInitHandle);
-	}
-	
-	PlaceActors.Empty();
-}
-
 TSharedPtr<FKismetCompilerContext> FDisplayClusterConfiguratorModule::GetCompilerForDisplayClusterBP(UBlueprint* BP,
                                                                                                      FCompilerResultsLog& InMessageLog, const FKismetCompilerOptions& InCompileOptions)
 {
diff --git a/Engine/Plugins/Runtime/nDisplay/Source/DisplayClusterConfigurator/Private/DisplayClusterConfiguratorModule.h b/Engine/Plugins/Runtime/nDisplay/Source/DisplayClusterConfigurator/Private/DisplayClusterConfiguratorModule.h
--- a/Engine/Plugins/Runtime/nDisplay/Source/DisplayClusterConfigurator/Private/DisplayClusterConfiguratorModule.h
+++ b/Engine/Plugins/Runtime/nDisplay/Source/DisplayClusterConfigurator/Private/DisplayClusterConfiguratorModule.h
@@ -40,12 +40,6 @@ private:
 	void RegisterSectionMappings();
 	void UnregisterSectionMappings();
 
-	/** Register items to show up in the Place Actors panel. */
-	void RegisterPlacementModeItems();
-	
-	/** Unregister items in Place Actors panel */
-	void UnregisterPlacementModeItems();
-
 	static TSharedPtr<FKismetCompilerContext> GetCompilerForDisplayClusterBP(UBlueprint* BP, FCompilerResultsLog& InMessageLog, const FKismetCompilerOptions& InCompileOptions);
 
 private:
@@ -55,7 +49,6 @@ private:
 	FDisplayClusterConfiguratorKismetCompiler  BlueprintCompiler;
 	TArray<FName> RegisteredClassLayoutNames;
 	TArray<FName> RegisteredPropertyLayoutNames;
-	TArray<TOptional<FPlacementModeID>> PlaceActors;
 	FDelegateHandle FilesLoadedHandle;
 	FDelegateHandle PostEngineInitHandle;
 
