nDisplay: UV LCs now have default feathering of 0.05.

This is just enough to mitigate the jagged look when at 512x512 resolution.

diff --git a/Engine/Plugins/Runtime/nDisplay/Source/DisplayClusterLightCardEditor/Private/SDisplayClusterLightCardEditor.cpp b/Engine/Plugins/Runtime/nDisplay/Source/DisplayClusterLightCardEditor/Private/SDisplayClusterLightCardEditor.cpp
--- a/Engine/Plugins/Runtime/nDisplay/Source/DisplayClusterLightCardEditor/Private/SDisplayClusterLightCardEditor.cpp
+++ b/Engine/Plugins/Runtime/nDisplay/Source/DisplayClusterLightCardEditor/Private/SDisplayClusterLightCardEditor.cpp
@@ -227,6 +227,7 @@ ADisplayClusterLightCardActor* SDisplayClusterLightCardEditor::SpawnLightCard(co
 		{
 			NewLightCard->bIsUVLightCard = true;
 			NewLightCard->Scale /= 4;
+			NewLightCard->Feathering = 0.05; // Just enough to avoid jagged look on UV lightcards.
 		}
 	}
 
