nDisplay: UV LCs are now placed slightly closer than the DRCA to avoid mouse click misses.

diff --git a/Engine/Plugins/Runtime/nDisplay/Source/DisplayCluster/Private/Game/EngineClasses/Scene/DisplayClusterLightCardActor.cpp b/Engine/Plugins/Runtime/nDisplay/Source/DisplayCluster/Private/Game/EngineClasses/Scene/DisplayClusterLightCardActor.cpp
--- a/Engine/Plugins/Runtime/nDisplay/Source/DisplayCluster/Private/Game/EngineClasses/Scene/DisplayClusterLightCardActor.cpp
+++ b/Engine/Plugins/Runtime/nDisplay/Source/DisplayCluster/Private/Game/EngineClasses/Scene/DisplayClusterLightCardActor.cpp
@@ -235,7 +235,12 @@ void ADisplayClusterLightCardActor::UpdateLightCardTransform()
 
 		// Set world location and rotation such that the light card is always projected onto a YZ plane a distance of UVPlaneDefaultDistance from the world origin, facing in the -X direction
 		// This ensures that when the UV light cards are rendered to the light card map, they are always positioned and oriented correctly regardless of the stages location and rotation.
-		LightCardTransformerComponent->SetWorldLocation(FVector(UVPlaneDefaultDistance, -UVPlaneDefaultSize * (0.5 - UVCoordinates.X), UVPlaneDefaultSize * (0.5 - UVCoordinates.Y)));
+		
+		// We place them slightly closer than UVPlaneDefaultDistance so that mouse clicks hit it first instead of DCRA meshes (avoids re-linetracing).
+		// UV LCs operate in orthographic projection so this should have no visual effect.
+		constexpr float DistanceFactor = 0.99f;
+
+		LightCardTransformerComponent->SetWorldLocation(FVector(DistanceFactor * UVPlaneDefaultDistance, -UVPlaneDefaultSize * (0.5 - UVCoordinates.X), UVPlaneDefaultSize * (0.5 - UVCoordinates.Y)));
 		LightCardTransformerComponent->SetWorldRotation(FVector(-1, 0, 0).Rotation());
 		LightCardTransformerComponent->SetWorldScale3D(FVector::OneVector);
 	}
