nDisplay: Localization conflict fix

diff --git a/Engine/Plugins/Runtime/nDisplay/Source/DisplayClusterColorGrading/Private/Drawer/SDisplayClusterColorGradingDrawer.cpp b/Engine/Plugins/Runtime/nDisplay/Source/DisplayClusterColorGrading/Private/Drawer/SDisplayClusterColorGradingDrawer.cpp
--- a/Engine/Plugins/Runtime/nDisplay/Source/DisplayClusterColorGrading/Private/Drawer/SDisplayClusterColorGradingDrawer.cpp
+++ b/Engine/Plugins/Runtime/nDisplay/Source/DisplayClusterColorGrading/Private/Drawer/SDisplayClusterColorGradingDrawer.cpp
@@ -146,7 +146,7 @@ void SDisplayClusterColorGradingDrawer::Construct(const FArguments& InArgs, bool
 							[
 								SNew(SCheckBox)
 								.Style(FAppStyle::Get(), "ToggleButtonCheckbox")
-								.ToolTipText(LOCTEXT("ColorGradingDrawerModeToolTip", "Sets the drawer to display the additional in-camera VFX properties of the selected items"))
+								.ToolTipText(LOCTEXT("DetailsViewDrawerModeToolTip", "Sets the drawer to display the additional in-camera VFX properties of the selected items"))
 								.OnCheckStateChanged(this, &SDisplayClusterColorGradingDrawer::OnDrawerModeSelected, EDisplayClusterColorGradingDrawerMode::DetailsView)
 								.IsChecked(this, &SDisplayClusterColorGradingDrawer::IsDrawerModeSelected, EDisplayClusterColorGradingDrawerMode::DetailsView)
 								[
