nDisplay: Added global cartesian mode to the light card editor's transform widget, allowing the user to position light cards along cardinal directions. Also changed the transform widget to render a sprite for its origin, which should improve visibility and usability.

diff --git a/Engine/Plugins/Runtime/nDisplay/Source/DisplayClusterLightCardEditor/Private/DisplayClusterLightCardEditorCommands.cpp b/Engine/Plugins/Runtime/nDisplay/Source/DisplayClusterLightCardEditor/Private/DisplayClusterLightCardEditorCommands.cpp
--- a/Engine/Plugins/Runtime/nDisplay/Source/DisplayClusterLightCardEditor/Private/DisplayClusterLightCardEditorCommands.cpp
+++ b/Engine/Plugins/Runtime/nDisplay/Source/DisplayClusterLightCardEditor/Private/DisplayClusterLightCardEditorCommands.cpp
@@ -21,6 +21,15 @@ void FDisplayClusterLightCardEditorCommands::RegisterCommands()
 	UI_COMMAND(ViewOrientationFront, "Front", "Orient the view to look at the front of the stage", EUserInterfaceActionType::Button, FInputChord());
 	UI_COMMAND(ViewOrientationBack, "Back", "Orient the view to look at the back of the stage", EUserInterfaceActionType::Button, FInputChord());
 
+#if PLATFORM_MAC
+	UI_COMMAND(CycleEditorWidgetCoordinateSystem, "Cycle Transform Coordinate System", "Cycles the transform gizmo coordinate systems between cartesian and spherical coordinates", EUserInterfaceActionType::Button, FInputChord(EModifierKey::Command, EKeys::Tilde));
+#else
+	UI_COMMAND(CycleEditorWidgetCoordinateSystem, "Cycle Transform Coordinate System", "Cycles the transform gizmo coordinate systems between cartesian and spherical coordinates", EUserInterfaceActionType::Button, FInputChord(EModifierKey::Control, EKeys::Tilde));
+#endif
+
+	UI_COMMAND(SphericalCoordinateSystem, "Spherical Coordinate System", "Move objects using latitude and longitude", EUserInterfaceActionType::RadioButton, FInputChord());
+	UI_COMMAND(CartesianCoordinateSystem, "Cartesian Coordinate System", "Move objects using X, Y, and Z axes", EUserInterfaceActionType::RadioButton, FInputChord());
+
 	UI_COMMAND(AddNewLightCard, "Add New Light Card", "Add and assign a new Light Card to the actor", EUserInterfaceActionType::Button, FInputChord());
 	UI_COMMAND(AddNewFlag, "Add Flag", "Add and assign a new light control flag", EUserInterfaceActionType::Button, FInputChord());
 	UI_COMMAND(AddExistingLightCard, "Add Existing Light Card", "Add an existing Light Card to the actor", EUserInterfaceActionType::Button, FInputChord());
diff --git a/Engine/Plugins/Runtime/nDisplay/Source/DisplayClusterLightCardEditor/Private/DisplayClusterLightCardEditorCommands.h b/Engine/Plugins/Runtime/nDisplay/Source/DisplayClusterLightCardEditor/Private/DisplayClusterLightCardEditorCommands.h
--- a/Engine/Plugins/Runtime/nDisplay/Source/DisplayClusterLightCardEditor/Private/DisplayClusterLightCardEditorCommands.h
+++ b/Engine/Plugins/Runtime/nDisplay/Source/DisplayClusterLightCardEditor/Private/DisplayClusterLightCardEditorCommands.h
@@ -37,6 +37,10 @@ public:
 	TSharedPtr<FUICommandInfo> ViewOrientationFront;
 	TSharedPtr<FUICommandInfo> ViewOrientationBack;
 
+	TSharedPtr<FUICommandInfo> CycleEditorWidgetCoordinateSystem;
+	TSharedPtr<FUICommandInfo> CartesianCoordinateSystem;
+	TSharedPtr<FUICommandInfo> SphericalCoordinateSystem;
+
 	TSharedPtr<FUICommandInfo> AddNewLightCard;
 	TSharedPtr<FUICommandInfo> AddNewFlag;
 	TSharedPtr<FUICommandInfo> AddExistingLightCard;
diff --git a/Engine/Plugins/Runtime/nDisplay/Source/DisplayClusterLightCardEditor/Private/Viewport/DisplayClusterLightCardEditorViewportClient.cpp b/Engine/Plugins/Runtime/nDisplay/Source/DisplayClusterLightCardEditor/Private/Viewport/DisplayClusterLightCardEditorViewportClient.cpp
--- a/Engine/Plugins/Runtime/nDisplay/Source/DisplayClusterLightCardEditor/Private/Viewport/DisplayClusterLightCardEditorViewportClient.cpp
+++ b/Engine/Plugins/Runtime/nDisplay/Source/DisplayClusterLightCardEditor/Private/Viewport/DisplayClusterLightCardEditorViewportClient.cpp
@@ -455,10 +455,12 @@ void FDisplayClusterLightCardEditorViewportClient::Draw(const FSceneView* View,
 		// Project the editor widget's world position into projection space so that it renders at the appropriate screen location.
 		// This needs to be computed on the render thread using the render thread's scene view, which will be behind the game thread's scene view
 		// by at least one frame
+		const bool bDrawZAxis = EditorWidget->GetWidgetMode() == FDisplayClusterLightCardEditorWidget::EWidgetMode::WM_Translate &&
+			EditorWidgetCoordinateSystem == FDisplayClusterLightCardEditorHelper::ECoordinateSystem::Cartesian;
 
 		EditorWidget->SetTransform(CachedEditorWidgetWorldTransform);
 		EditorWidget->SetProjectionTransform(FDisplayClusterMeshProjectionTransform(ProjectionMode, View->ViewMatrices.GetViewMatrix()));
-		EditorWidget->Draw(View, this, PDI);
+		EditorWidget->Draw(View, this, PDI, bDrawZAxis);
 	}
 
 	// Draw polygonal lightcard in progress
@@ -1493,6 +1495,17 @@ void FDisplayClusterLightCardEditorViewportClient::SelectLightCards(const TArray
 	}
 }
 
+void FDisplayClusterLightCardEditorViewportClient::SetEditorWidgetMode(FDisplayClusterLightCardEditorWidget::EWidgetMode InWidgetMode)
+{
+	// Force the coordinate system back to spherical if the widget mode is set to anything besides translate
+	if (InWidgetMode != FDisplayClusterLightCardEditorWidget::EWidgetMode::WM_Translate)
+	{
+		EditorWidgetCoordinateSystem = FDisplayClusterLightCardEditorHelper::ECoordinateSystem::Spherical;
+	}
+
+	EditorWidget->SetWidgetMode(InWidgetMode);
+}
+
 void FDisplayClusterLightCardEditorViewportClient::SetProjectionMode(EDisplayClusterMeshProjectionType InProjectionMode, ELevelViewportType InViewportType)
 {
 	ProjectionMode = InProjectionMode;
@@ -1579,6 +1592,27 @@ void FDisplayClusterLightCardEditorViewportClient::ResetCamera(bool bLocationOnl
 	ResetFOVs();
 }
 
+void FDisplayClusterLightCardEditorViewportClient::CycleCoordinateSystem()
+{
+	if (EditorWidgetCoordinateSystem == FDisplayClusterLightCardEditorHelper::ECoordinateSystem::Cartesian)
+	{
+		SetCoordinateSystem(FDisplayClusterLightCardEditorHelper::ECoordinateSystem::Spherical);
+	}
+	else
+	{
+		SetCoordinateSystem(FDisplayClusterLightCardEditorHelper::ECoordinateSystem::Cartesian);
+	}
+}
+
+void FDisplayClusterLightCardEditorViewportClient::SetCoordinateSystem(FDisplayClusterLightCardEditorHelper::ECoordinateSystem NewCoordinateSystem)
+{
+	// The only widget mode that supports multiple coordinate systems at the moment is translation
+	if (EditorWidget->GetWidgetMode() == FDisplayClusterLightCardEditorWidget::EWidgetMode::WM_Translate)
+	{
+		EditorWidgetCoordinateSystem = NewCoordinateSystem;
+	}
+}
+
 void FDisplayClusterLightCardEditorViewportClient::MoveLightCardTo(ADisplayClusterLightCardActor& LightCard, const FDisplayClusterLightCardEditorHelper::FSphericalCoordinates& SphericalCoords) const
 {
 	const FVector Origin = GetViewLocation(); // Assumed the same as the LC origin in proxy space. Call VerifyAndFixLightCardOrigin to ensure this.
@@ -1900,7 +1934,7 @@ void FDisplayClusterLightCardEditorViewportClient::MoveSelectedLightCards(FViewp
 	InViewport->GetMousePos(MousePos);
 
 	// Move the light cards
-	ProjectionHelper->DragLightCards(SelectedLightCards, MousePos, *View, DragWidgetOffset, CurrentAxis, LastSelectedLightCard.Get());
+	ProjectionHelper->DragLightCards(SelectedLightCards, MousePos, *View, EditorWidgetCoordinateSystem, DragWidgetOffset, CurrentAxis, LastSelectedLightCard.Get());
 
 	// Update the level instances
 	for (const TWeakObjectPtr<ADisplayClusterLightCardActor>& LightCard : SelectedLightCards)
@@ -2282,13 +2316,20 @@ bool FDisplayClusterLightCardEditorViewportClient::CalcEditorWidgetTransform(FTr
 		}
 		else
 		{
-			// The translation widget should be oriented to show the x axis pointing in the longitudinal direction and the y axis pointing in the latitudinal direction
-			const FVector ProjectionOrigin = ProjectionOriginComponent.IsValid() ? ProjectionOriginComponent->GetComponentLocation() : FVector::ZeroVector;
-			const FVector RadialVector = (LightCardPosition - ProjectionOrigin).GetSafeNormal();
-			const FVector AzimuthalVector = (FVector::ZAxisVector ^ RadialVector).GetSafeNormal();
-			const FVector InclinationVector = RadialVector ^ AzimuthalVector;
+			if (EditorWidgetCoordinateSystem == FDisplayClusterLightCardEditorHelper::ECoordinateSystem::Spherical)
+			{
+				// The translation widget should be oriented to show the x axis pointing in the longitudinal direction and the y axis pointing in the latitudinal direction
+				const FVector ProjectionOrigin = ProjectionOriginComponent.IsValid() ? ProjectionOriginComponent->GetComponentLocation() : FVector::ZeroVector;
+				const FVector RadialVector = (LightCardPosition - ProjectionOrigin).GetSafeNormal();
+				const FVector AzimuthalVector = (FVector::ZAxisVector ^ RadialVector).GetSafeNormal();
+				const FVector InclinationVector = RadialVector ^ AzimuthalVector;
 
-			WidgetOrientation = FMatrix(AzimuthalVector, InclinationVector, RadialVector, FVector::ZeroVector).ToQuat();
+				WidgetOrientation = FMatrix(AzimuthalVector, InclinationVector, RadialVector, FVector::ZeroVector).ToQuat();
+			}
+			else if (EditorWidgetCoordinateSystem == FDisplayClusterLightCardEditorHelper::ECoordinateSystem::Cartesian)
+			{
+				WidgetOrientation = FMatrix(FVector::XAxisVector, FVector::YAxisVector, FVector::ZAxisVector, FVector::ZeroVector).ToQuat();
+			}
 		}
 	}
 	else
diff --git a/Engine/Plugins/Runtime/nDisplay/Source/DisplayClusterLightCardEditor/Private/Viewport/DisplayClusterLightCardEditorViewportClient.h b/Engine/Plugins/Runtime/nDisplay/Source/DisplayClusterLightCardEditor/Private/Viewport/DisplayClusterLightCardEditorViewportClient.h
--- a/Engine/Plugins/Runtime/nDisplay/Source/DisplayClusterLightCardEditor/Private/Viewport/DisplayClusterLightCardEditorViewportClient.h
+++ b/Engine/Plugins/Runtime/nDisplay/Source/DisplayClusterLightCardEditor/Private/Viewport/DisplayClusterLightCardEditorViewportClient.h
@@ -93,7 +93,7 @@ public:
 	void SelectLightCards(const TArray<ADisplayClusterLightCardActor*>& LightCardsToSelect);
 
 	FDisplayClusterLightCardEditorWidget::EWidgetMode GetEditorWidgetMode() const { return EditorWidget->GetWidgetMode(); }
-	void SetEditorWidgetMode(FDisplayClusterLightCardEditorWidget::EWidgetMode InWidgetMode) { EditorWidget->SetWidgetMode(InWidgetMode); }
+	void SetEditorWidgetMode(FDisplayClusterLightCardEditorWidget::EWidgetMode InWidgetMode);
 
 	/** Gets the current projection mode of the editor viewport */
 	EDisplayClusterMeshProjectionType GetProjectionMode() const { return ProjectionMode; }
@@ -113,6 +113,15 @@ public:
 	/** Resets the camera to the initial rotation / position */
 	void ResetCamera(bool bLocationOnly = false);
 
+	/** Cycles the editor widget's coordinate system between cartesian and spherical */
+	void CycleCoordinateSystem();
+
+	/** Gets the editor widget's current coordinate system */
+	FDisplayClusterLightCardEditorHelper::ECoordinateSystem GetCoordinateSystem() const { return EditorWidgetCoordinateSystem; }
+
+	/** Sets the editor widget's coordinate system to the specified coordinate system */
+	void SetCoordinateSystem(FDisplayClusterLightCardEditorHelper::ECoordinateSystem NewCoordinateSystem);
+
 	/** Moves specified card to desired coordinates. Actual radius will be based on flush constraint and LightCard's RadialOffset.
 	 *
 	 * @param LightCard The light card that we are moving
@@ -274,6 +283,9 @@ private:
 	/** The cached editor widget transform in unprojected world space */
 	FTransform CachedEditorWidgetWorldTransform;
 
+	/** The current coordinate system of the editor widget */
+	FDisplayClusterLightCardEditorHelper::ECoordinateSystem EditorWidgetCoordinateSystem = FDisplayClusterLightCardEditorHelper::ECoordinateSystem::Spherical;
+
 	/** The offset between the widget's origin and the place it was clicked when a drag action was started */
 	FVector DragWidgetOffset;
 
diff --git a/Engine/Plugins/Runtime/nDisplay/Source/DisplayClusterLightCardEditor/Private/Viewport/DisplayClusterLightCardEditorWidget.cpp b/Engine/Plugins/Runtime/nDisplay/Source/DisplayClusterLightCardEditor/Private/Viewport/DisplayClusterLightCardEditorWidget.cpp
--- a/Engine/Plugins/Runtime/nDisplay/Source/DisplayClusterLightCardEditor/Private/Viewport/DisplayClusterLightCardEditorWidget.cpp
+++ b/Engine/Plugins/Runtime/nDisplay/Source/DisplayClusterLightCardEditor/Private/Viewport/DisplayClusterLightCardEditorWidget.cpp
@@ -4,10 +4,17 @@
 
 #include "DisplayClusterLightCardEditorViewportClient.h"
 
+#include "Engine/Texture2D.h"
 #include "SceneManagement.h"
 #include "UnrealWidget.h"
 
-void FDisplayClusterLightCardEditorWidget::Draw(const FSceneView* View, const FDisplayClusterLightCardEditorViewportClient* ViewportClient, FPrimitiveDrawInterface* PDI)
+FDisplayClusterLightCardEditorWidget::FDisplayClusterLightCardEditorWidget()
+	: TranslateOriginTexture(LoadObject<UTexture2D>(nullptr, TEXT("/nDisplay/Icons/S_LightCardEditorWidgetTranslate")))
+	, ScaleOriginTexture(LoadObject<UTexture2D>(nullptr, TEXT("/nDisplay/Icons/S_LightCardEditorWidgetScale")))
+{
+}
+
+void FDisplayClusterLightCardEditorWidget::Draw(const FSceneView* View, const FDisplayClusterLightCardEditorViewportClient* ViewportClient, FPrimitiveDrawInterface* PDI, bool bDrawZAxis)
 {
 	const bool bIsOrthographic = ViewportClient->GetRenderViewportType() != LVT_Perspective;
 
@@ -25,7 +32,15 @@ void FDisplayClusterLightCardEditorWidget::Draw(const FSceneView* View, const FD
 	{
 		DrawAxis(PDI, EAxisList::Type::X, SizeScalar, LengthScalar, OrthoScalar);
 		DrawAxis(PDI, EAxisList::Type::Y, SizeScalar, LengthScalar, OrthoScalar);
-		DrawOrigin(PDI, SizeScalar, OrthoScalar);
+
+		if (bDrawZAxis)
+		{
+			DrawAxis(PDI, EAxisList::Type::Z, SizeScalar, LengthScalar, OrthoScalar);
+		}
+
+		// When drawing with sprites, the length scalar is needed instead of the size scalar. Orthographic scaling is already factored into the length scalar,
+		// so the ortho scalar here can be 1.0
+		DrawOrigin(PDI, LengthScalar, 1.0f);
 	}
 }
 
@@ -55,11 +70,23 @@ void FDisplayClusterLightCardEditorWidget::DrawAxis(FPrimitiveDrawInterface* PDI
 void FDisplayClusterLightCardEditorWidget::DrawOrigin(FPrimitiveDrawInterface* PDI, float SizeScalar, float OrthoScalar)
 {
 	const FVector Origin = ProjectionTransform.ProjectPosition(Transform.GetTranslation());
-	const FLinearColor Color = HighlightedAxis == EAxisList::Type::XYZ ? HighlightColor : FLinearColor::Black;
-
-	PDI->SetHitProxy(new HWidgetAxis(EAxisList::Type::XYZ));
-	PDI->DrawPoint(Origin, Color, OriginSize * SizeScalar * OrthoScalar, ESceneDepthPriorityGroup::SDPG_Foreground);
-	PDI->SetHitProxy(nullptr);
+	const FLinearColor Color = HighlightedAxis == EAxisList::Type::XYZ ? HighlightColor : FLinearColor::White;
+	const float SpriteSize = OriginSize * SizeScalar * OrthoScalar;
+	
+	UTexture* OriginTexture = WidgetMode == EWidgetMode::WM_Scale ? ScaleOriginTexture : TranslateOriginTexture;
+	if (FTextureResource* Texture = OriginTexture->GetResource())
+	{
+		PDI->SetHitProxy(new HWidgetAxis(EAxisList::Type::XYZ));
+		PDI->DrawSprite(Origin,
+			SpriteSize, SpriteSize,
+			Texture,
+			Color,
+			ESceneDepthPriorityGroup::SDPG_Foreground,
+			0, Texture->GetSizeX(),
+			0, Texture->GetSizeY(),
+			SE_BLEND_Masked);
+		PDI->SetHitProxy(nullptr);
+	}
 }
 
 void FDisplayClusterLightCardEditorWidget::DrawCircle(FPrimitiveDrawInterface* PDI, EAxisList::Type Axis, float SizeScalar, float LengthScalar)
diff --git a/Engine/Plugins/Runtime/nDisplay/Source/DisplayClusterLightCardEditor/Private/Viewport/DisplayClusterLightCardEditorWidget.h b/Engine/Plugins/Runtime/nDisplay/Source/DisplayClusterLightCardEditor/Private/Viewport/DisplayClusterLightCardEditorWidget.h
--- a/Engine/Plugins/Runtime/nDisplay/Source/DisplayClusterLightCardEditor/Private/Viewport/DisplayClusterLightCardEditorWidget.h
+++ b/Engine/Plugins/Runtime/nDisplay/Source/DisplayClusterLightCardEditor/Private/Viewport/DisplayClusterLightCardEditorWidget.h
@@ -28,7 +28,7 @@ public:
 	constexpr static float AxisCapSize = 15.f;
 
 	/** The size of the axes origin */
-	constexpr static float OriginSize = 35.f;
+	constexpr static float OriginSize = 10.0f;
 
 	/** The size of the rotation circle, in pixels */
 	constexpr static float CirlceRadius = 100.0f;
@@ -50,7 +50,9 @@ public:
 		WM_Max
 	};
 
-	void Draw(const FSceneView* View, const FDisplayClusterLightCardEditorViewportClient* ViewportClient, FPrimitiveDrawInterface* PDI);
+	FDisplayClusterLightCardEditorWidget();
+
+	void Draw(const FSceneView* View, const FDisplayClusterLightCardEditorViewportClient* ViewportClient, FPrimitiveDrawInterface* PDI, bool bDrawZAxis);
 
 	EWidgetMode GetWidgetMode() const { return WidgetMode; }
 	void SetWidgetMode(EWidgetMode NewWidgetMode) { WidgetMode = NewWidgetMode; }
@@ -82,6 +84,12 @@ private:
 	float GetSizeScreenScalar(const FSceneView* View, const FDisplayClusterLightCardEditorViewportClient* ViewportClient) const;
 
 private:
+	/** Texture to use to render the translation widget origin */
+	UTexture2D* TranslateOriginTexture;
+
+	/** Texture to use to render the scale widget origin */
+	UTexture2D* ScaleOriginTexture;
+
 	/** The current mode the widget is in */
 	EWidgetMode WidgetMode = WM_Translate;
 
diff --git a/Engine/Plugins/Runtime/nDisplay/Source/DisplayClusterLightCardEditor/Private/Viewport/DisplayClusterLightcardEditorViewport.cpp b/Engine/Plugins/Runtime/nDisplay/Source/DisplayClusterLightCardEditor/Private/Viewport/DisplayClusterLightcardEditorViewport.cpp
--- a/Engine/Plugins/Runtime/nDisplay/Source/DisplayClusterLightCardEditor/Private/Viewport/DisplayClusterLightcardEditorViewport.cpp
+++ b/Engine/Plugins/Runtime/nDisplay/Source/DisplayClusterLightCardEditor/Private/Viewport/DisplayClusterLightcardEditorViewport.cpp
@@ -239,10 +239,35 @@ public:
 		}
 
 		ToolbarBuilder.EndSection();
+		ToolbarBuilder.AddSeparator();
+
+		ToolbarBuilder.AddToolBarButton(FDisplayClusterLightCardEditorCommands::Get().CycleEditorWidgetCoordinateSystem,
+			NAME_None,
+			TAttribute<FText>(),
+			TAttribute<FText>(),
+			TAttribute<FSlateIcon>(this, &SDisplayClusterLightCardEditorViewportToolBar::GetLocalToWorldIcon),
+			FName(TEXT("CycleTransformGizmoCoordSystem")),
+
+			FNewMenuDelegate::CreateLambda([](FMenuBuilder& InMenuBuilder)
+			{
+				InMenuBuilder.AddMenuEntry(FDisplayClusterLightCardEditorCommands::Get().SphericalCoordinateSystem);
+				InMenuBuilder.AddMenuEntry(FDisplayClusterLightCardEditorCommands::Get().CartesianCoordinateSystem);
+			}
+		));
 
 		return ToolbarBuilder.MakeWidget();
 	}
 
+	FSlateIcon GetLocalToWorldIcon() const
+	{
+		if (EditorViewport.IsValid() && EditorViewport.Pin()->GetLightCardEditorViewportClient()->GetCoordinateSystem() == FDisplayClusterLightCardEditorHelper::ECoordinateSystem::Cartesian)
+		{
+			return FSlateIcon(FAppStyle::GetAppStyleSetName(), TEXT("EditorViewport.RelativeCoordinateSystem_World"));
+		}
+
+		return FSlateIcon(FAppStyle::GetAppStyleSetName(), TEXT("Icons.Transform"));
+	}
+
 	TSharedRef<SWidget> MakeDrawingToolBar()
 	{
 		FSlimHorizontalToolBarBuilder ToolbarBuilder(EditorViewport.Pin()->GetCommandList(), FMultiBoxCustomization::None);
@@ -473,6 +498,23 @@ void SDisplayClusterLightCardEditorViewport::BindCommands()
 			Commands.ResetCamera,
 			FExecuteAction::CreateSP(ViewportClient.Get(), &FDisplayClusterLightCardEditorViewportClient::ResetCamera, false));
 
+		CommandList->MapAction(
+			Commands.CycleEditorWidgetCoordinateSystem,
+			FExecuteAction::CreateSP(ViewportClient.Get(), &FDisplayClusterLightCardEditorViewportClient::CycleCoordinateSystem));
+
+		CommandList->MapAction(
+			Commands.SphericalCoordinateSystem,
+			FExecuteAction::CreateSP(ViewportClient.Get(), &FDisplayClusterLightCardEditorViewportClient::SetCoordinateSystem, FDisplayClusterLightCardEditorHelper::ECoordinateSystem::Spherical),
+			FCanExecuteAction(),
+			FIsActionChecked::CreateLambda([this]() { return ViewportClient->GetCoordinateSystem() == FDisplayClusterLightCardEditorHelper::ECoordinateSystem::Spherical; }));
+
+		CommandList->MapAction(
+			Commands.CartesianCoordinateSystem,
+			FExecuteAction::CreateSP(ViewportClient.Get(), &FDisplayClusterLightCardEditorViewportClient::SetCoordinateSystem, FDisplayClusterLightCardEditorHelper::ECoordinateSystem::Cartesian),
+			FCanExecuteAction(),
+			FIsActionChecked::CreateLambda([this]() { return ViewportClient->GetCoordinateSystem() == FDisplayClusterLightCardEditorHelper::ECoordinateSystem::Cartesian; }));
+
+
 		CommandList->MapAction(
 			FDisplayClusterLightCardEditorCommands::Get().DrawLightCard,
 			FExecuteAction::CreateSP(this, &SDisplayClusterLightCardEditorViewport::DrawLightCard),
diff --git a/Engine/Plugins/Runtime/nDisplay/Source/DisplayClusterScenePreview/Private/DisplayClusterLightCardEditorHelper.cpp b/Engine/Plugins/Runtime/nDisplay/Source/DisplayClusterScenePreview/Private/DisplayClusterLightCardEditorHelper.cpp
--- a/Engine/Plugins/Runtime/nDisplay/Source/DisplayClusterScenePreview/Private/DisplayClusterLightCardEditorHelper.cpp
+++ b/Engine/Plugins/Runtime/nDisplay/Source/DisplayClusterScenePreview/Private/DisplayClusterLightCardEditorHelper.cpp
@@ -509,7 +509,7 @@ void FDisplayClusterLightCardEditorHelper::MoveLightCardsTo(
 }
 
 void FDisplayClusterLightCardEditorHelper::DragLightCards(const TArray<TWeakObjectPtr<ADisplayClusterLightCardActor>>& LightCards, const FIntPoint& PixelPos, const FSceneView& SceneView,
-	const FVector& DragWidgetOffset, EAxisList::Type DragAxis, ADisplayClusterLightCardActor* PrimaryLightCard)
+	ECoordinateSystem CoordinateSystem, const FVector& DragWidgetOffset, EAxisList::Type DragAxis, ADisplayClusterLightCardActor* PrimaryLightCard)
 {
 	if (LightCards.IsEmpty())
 	{
@@ -528,7 +528,7 @@ void FDisplayClusterLightCardEditorHelper::DragLightCards(const TArray<TWeakObje
 	}
 
 	// Maps are already ready, so we can move immediately
-	InternalDragLightCards(LightCards, PixelPos, SceneView, DragWidgetOffset, DragAxis, PrimaryLightCard);
+	InternalDragLightCards(LightCards, PixelPos, SceneView, CoordinateSystem, DragWidgetOffset, DragAxis, PrimaryLightCard);
 }
 
 void FDisplayClusterLightCardEditorHelper::VerifyAndFixLightCardOrigin(ADisplayClusterLightCardActor& LightCard)
@@ -932,9 +932,9 @@ FVector FDisplayClusterLightCardEditorHelper::GetProjectionOrigin() const
 }
 
 FRotator FDisplayClusterLightCardEditorHelper::GetLightCardRotationDelta(const FIntPoint& PixelPos, const FSceneView& View, ADisplayClusterLightCardActor& LightCard,
-	EAxisList::Type DragAxis, const FVector& DragWidgetOffset)
+	ECoordinateSystem CoordinateSystem, EAxisList::Type DragAxis, const FVector& DragWidgetOffset)
 {
-	const FSphericalCoordinates DeltaCoords = GetLightCardTranslationDelta(PixelPos, View, LightCard, DragAxis, DragWidgetOffset);
+	const FSphericalCoordinates DeltaCoords = GetLightCardTranslationDelta(PixelPos, View, LightCard, CoordinateSystem, DragAxis, DragWidgetOffset);
 	const FSphericalCoordinates LightCardCoords = GetLightCardCoordinates(LightCard);
 	const FSphericalCoordinates NewCoords = LightCardCoords + DeltaCoords;
 
@@ -958,6 +958,7 @@ FDisplayClusterLightCardEditorHelper::FSphericalCoordinates FDisplayClusterLight
 	const FIntPoint& PixelPos,
 	const FSceneView& View,
 	ADisplayClusterLightCardActor& LightCard,
+	ECoordinateSystem CoordinateSystem,
 	EAxisList::Type DragAxis,
 	const FVector& DragWidgetOffset)
 {
@@ -972,43 +973,85 @@ FDisplayClusterLightCardEditorHelper::FSphericalCoordinates FDisplayClusterLight
 
 	const FVector LocalDirection = LightCard.GetActorRotation().RotateVector(Direction);
 	const FVector LightCardLocation = LightCard.GetLightCardTransform().GetTranslation() - Origin;
+	const FSphericalCoordinates LightCardCoords = GetLightCardCoordinates(LightCard);
 
-	FVector Normal;
-	float Distance;
+	FSphericalCoordinates DeltaCoords;
 
-	// If the light card is in the southern hemisphere of the view origin, use the southern normal map; otherwise, use the north normal map
-	if (LightCardLocation.Z < 0.0f)
-	{
-		SouthNormalMap.GetNormalAndDistanceAtPosition(LightCard.GetLightCardTransform().GetTranslation(), Normal, Distance);
-	}
-	else
+	// If we are in a cartesian coordinate system and are constraining to an axis, perform the constraint calculations in
+	// cartesian coordinates, then convert to spherical coordinates at the end. Otherwise, perform all calculations in spherical coodinates
+	if (CoordinateSystem == ECoordinateSystem::Cartesian && DragAxis != EAxisList::Type::XYZ)
 	{
-		NorthNormalMap.GetNormalAndDistanceAtPosition(LightCard.GetLightCardTransform().GetTranslation(), Normal, Distance);
-	}
+		// For consistency, project the cursor direction vector onto the sphere the light card is currently on. Gives a good balance of approximating
+		// the true projection plane (one that works with all view projections) and the general stage normal map
+		const FVector RequestedLocation = LocalDirection * LightCardLocation.Size();
 
-	const FSphericalCoordinates LightCardCoords = GetLightCardCoordinates(LightCard);
-	const FSphericalCoordinates RequestedCoords(LocalDirection * Distance);
+		// Compute the axis to constrain the translation to based on the axis that was dragged. Axis must be rotated into the light card's local space
+		FVector Axis = FVector::ZeroVector;
+		if (DragAxis == EAxisList::Type::X)
+		{
+			Axis = FVector::XAxisVector;
+		}
+		else if (DragAxis == EAxisList::Type::Y)
+		{
+			Axis = FVector::YAxisVector;
+		}
+		else if (DragAxis == EAxisList::Type::Z)
+		{
+			Axis = FVector::ZAxisVector;
+		}
 
-	FSphericalCoordinates DeltaCoords = RequestedCoords - LightCardCoords;
+		const FVector LocalAxis = LightCard.GetActorRotation().RotateVector(Axis);
 
-	if (DragAxis == EAxisList::Type::X)
-	{
-		DeltaCoords.Inclination = 0;
+		// Compute the offset between the requested location and the light card's current location, and project that
+		// offset onto the constraint axis
+		const FVector DeltaLocation = ((RequestedLocation - LightCardLocation) | LocalAxis) * LocalAxis;
+
+		const FVector ConstrainedLocation = LightCardLocation + DeltaLocation;
+
+		const FSphericalCoordinates ConstrainedCoords(ConstrainedLocation);
+		DeltaCoords = ConstrainedCoords - LightCardCoords;
 	}
-	else if (DragAxis == EAxisList::Type::Y)
+	else
 	{
-		// Convert the inclination to Cartesian coordinates, project it to the x-z plane, and convert back to spherical coordinates. This ensures that the motion in the inclination
-		// plane always lines up with the mouse's projected location along that plane
-		const double FixedInclination = FMath::Abs(FMath::Atan2(
-			FMath::Cos(DeltaCoords.Azimuth) * FMath::Sin(RequestedCoords.Inclination),
-			FMath::Cos(RequestedCoords.Inclination))
-		);
+		FVector Normal;
+		float Distance;
+
+		// If the light card is in the southern hemisphere of the view origin, use the southern normal map; otherwise, use the north normal map
+		if (LightCardLocation.Z < 0.0f)
+		{
+			SouthNormalMap.GetNormalAndDistanceAtPosition(LightCard.GetLightCardTransform().GetTranslation(), Normal, Distance);
+		}
+		else
+		{
+			NorthNormalMap.GetNormalAndDistanceAtPosition(LightCard.GetLightCardTransform().GetTranslation(), Normal, Distance);
+		}
 
-		// When translating along the inclination axis, the azimuth delta can only be intervals of pi
-		const double FixedAzimuth = FMath::RoundToInt(DeltaCoords.Azimuth / PI) * PI;
+		const FSphericalCoordinates RequestedCoords(LocalDirection * Distance);
 
-		DeltaCoords.Azimuth = FixedAzimuth;
-		DeltaCoords.Inclination = FixedInclination - LightCardCoords.Inclination;
+		DeltaCoords = RequestedCoords - LightCardCoords;
+
+		if (CoordinateSystem == ECoordinateSystem::Spherical)
+		{
+			if (DragAxis == EAxisList::Type::X)
+			{
+				DeltaCoords.Inclination = 0;
+			}
+			else if (DragAxis == EAxisList::Type::Y)
+			{
+				// Convert the inclination to Cartesian coordinates, project it to the x-z plane, and convert back to spherical coordinates. This ensures that the motion in the inclination
+				// plane always lines up with the mouse's projected location along that plane
+				const double FixedInclination = FMath::Abs(FMath::Atan2(
+					FMath::Cos(DeltaCoords.Azimuth) * FMath::Sin(RequestedCoords.Inclination),
+					FMath::Cos(RequestedCoords.Inclination))
+				);
+
+				// When translating along the inclination axis, the azimuth delta can only be intervals of pi
+				const double FixedAzimuth = FMath::RoundToInt(DeltaCoords.Azimuth / PI) * PI;
+
+				DeltaCoords.Azimuth = FixedAzimuth;
+				DeltaCoords.Inclination = FixedInclination - LightCardCoords.Inclination;
+			}
+		}
 	}
 
 	return DeltaCoords;
@@ -1061,7 +1104,7 @@ void FDisplayClusterLightCardEditorHelper::InternalMoveLightCardTo(
 }
 
 void FDisplayClusterLightCardEditorHelper::InternalDragLightCards(const TArray<TWeakObjectPtr<ADisplayClusterLightCardActor>>& LightCards, const FIntPoint& PixelPos, const FSceneView& View,
-	const FVector& DragWidgetOffset, EAxisList::Type DragAxis, ADisplayClusterLightCardActor* PrimaryLightCard)
+	ECoordinateSystem CoordinateSystem, const FVector& DragWidgetOffset, EAxisList::Type DragAxis, ADisplayClusterLightCardActor* PrimaryLightCard)
 {
 	if (LightCards.IsEmpty() || !ensure(!bNormalMapInvalid) || !UpdateRootActor())
 	{
@@ -1075,17 +1118,17 @@ void FDisplayClusterLightCardEditorHelper::InternalDragLightCards(const TArray<T
 
 	if (PrimaryLightCard && !PrimaryLightCard->bIsUVLightCard)
 	{
-		const bool bUseDeltaRotation = (DragAxis == EAxisList::Type::XYZ) || (DragAxis == EAxisList::Type::Y);
+		const bool bUseDeltaRotation = (DragAxis == EAxisList::Type::XYZ) || (DragAxis == EAxisList::Type::Y) || (CoordinateSystem == ECoordinateSystem::Cartesian);
 
 		const FRotator DeltaRotation =
 			bUseDeltaRotation ?
-			GetLightCardRotationDelta(PixelPos, View, *PrimaryLightCard, DragAxis, DragWidgetOffset)
+			GetLightCardRotationDelta(PixelPos, View, *PrimaryLightCard, CoordinateSystem, DragAxis, DragWidgetOffset)
 			: FRotator::ZeroRotator;
 
 		const FSphericalCoordinates DeltaCoords =
 			bUseDeltaRotation ?
 			FSphericalCoordinates()
-			: GetLightCardTranslationDelta(PixelPos, View, *PrimaryLightCard, DragAxis, DragWidgetOffset);
+			: GetLightCardTranslationDelta(PixelPos, View, *PrimaryLightCard, CoordinateSystem, DragAxis, DragWidgetOffset);
 
 		for (const TWeakObjectPtr<ADisplayClusterLightCardActor> LightCard : LightCards)
 		{
@@ -1101,7 +1144,7 @@ void FDisplayClusterLightCardEditorHelper::InternalDragLightCards(const TArray<T
 
 			// We will adjust the spin (to maintain the apparent spin) when using center of gizmo
 			// or dragging longitudinally (this seems to provide an intuitive behavior)
-			if (DragAxis == EAxisList::Type::XYZ || DragAxis == EAxisList::Type::Y) // Dragging center of gizmo
+			if (bUseDeltaRotation) // Dragging center of gizmo
 			{
 				// We might need this to put back the LightCard exactly as it was
 				const ADisplayClusterLightCardActor::PositionalParams OriginalPositionalParams = LightCard->GetPositionalParams();
diff --git a/Engine/Plugins/Runtime/nDisplay/Source/DisplayClusterScenePreview/Public/DisplayClusterLightCardEditorHelper.h b/Engine/Plugins/Runtime/nDisplay/Source/DisplayClusterScenePreview/Public/DisplayClusterLightCardEditorHelper.h
--- a/Engine/Plugins/Runtime/nDisplay/Source/DisplayClusterScenePreview/Public/DisplayClusterLightCardEditorHelper.h
+++ b/Engine/Plugins/Runtime/nDisplay/Source/DisplayClusterScenePreview/Public/DisplayClusterLightCardEditorHelper.h
@@ -29,6 +29,12 @@ class FEditorViewportClient;
 struct FDisplayClusterLightCardEditorHelper
 {
 public:
+	enum class ECoordinateSystem : uint8
+	{
+		Cartesian,
+		Spherical
+	};
+
 	struct FSphericalCoordinates
 	{
 	public:
@@ -190,12 +196,13 @@ public:
 	 * @param LightCards The light cards that we are moving
 	 * @param PixelPos The screen pixel position of the widget
 	 * @param SceneView The scene view used to convert from pixel position to 3D position
+	 * @param CoordinateSystem The coordinate system to use when computing drag constraints
 	 * @param DragWidgetOffset The offset between the actual cursor position and the position of the widget when the drag action started
 	 * @param DragAxis The axis along which the widget is being dragged
 	 * @param PrimaryLightCard The light card used to calculate the translation/rotation delta. If not provided, the last entry in LightCards will be used.
 	 */
 	DISPLAYCLUSTERSCENEPREVIEW_API void DragLightCards(const TArray<TWeakObjectPtr<ADisplayClusterLightCardActor>>& LightCards, const FIntPoint& PixelPos, const FSceneView& SceneView,
-		const FVector& DragWidgetOffset, EAxisList::Type DragAxis, ADisplayClusterLightCardActor* PrimaryLightCard = nullptr);
+		ECoordinateSystem CoordinateSystem, const FVector& DragWidgetOffset, EAxisList::Type DragAxis, ADisplayClusterLightCardActor* PrimaryLightCard = nullptr);
 
 	/** Ensures that the light card root component is at the same location as the projection/view origin */
 	DISPLAYCLUSTERSCENEPREVIEW_API void VerifyAndFixLightCardOrigin(ADisplayClusterLightCardActor& LightCard);
@@ -267,19 +274,19 @@ private:
 	FVector GetProjectionOrigin() const;
 
 	/** Determines the appropriate delta rotation needed to move the specified light card to the given position within the view. */
-	FRotator GetLightCardRotationDelta(const FIntPoint& PixelPos, const FSceneView& View, ADisplayClusterLightCardActor& LightCard, EAxisList::Type CurrentAxis,
-		const FVector& DragWidgetOffset);
+	FRotator GetLightCardRotationDelta(const FIntPoint& PixelPos, const FSceneView& View, ADisplayClusterLightCardActor& LightCard,
+		ECoordinateSystem CoordinateSystem, EAxisList::Type DragAxis, const FVector& DragWidgetOffset);
 
 	/** Determines the appropriate delta in spherical coordinates needed to move the specified light card to the given position within the view. */
 	FSphericalCoordinates GetLightCardTranslationDelta(const FIntPoint& PixelPos, const FSceneView& View, ADisplayClusterLightCardActor& LightCard,
-		EAxisList::Type DragAxis, const FVector& DragWidgetOffset);
+		ECoordinateSystem CoordinateSystem, EAxisList::Type DragAxis, const FVector& DragWidgetOffset);
 
 	/** Moves specified card to desired coordinates immediately using the current normal maps. Requires valid normal maps. */
 	void InternalMoveLightCardTo(ADisplayClusterLightCardActor& LightCard, const FSphericalCoordinates& Position, bool bIsFinalChange) const;
 
 	/** Moves specified cards to a coordinate in viewport space as if dragged by a widget. Requires valid normal maps. */
 	void InternalDragLightCards(const TArray<TWeakObjectPtr<ADisplayClusterLightCardActor>>& LightCards, const FIntPoint& PixelPos, const FSceneView& View,
-		const FVector& DragWidgetOffset, EAxisList::Type DragAxis, ADisplayClusterLightCardActor* PrimaryLightCard);
+		ECoordinateSystem CoordinateSystem, const FVector& DragWidgetOffset, EAxisList::Type DragAxis, ADisplayClusterLightCardActor* PrimaryLightCard);
 
 	/** Gets the scene view init options to use when rendering the normal map cache */
 	void GetNormalMapSceneViewInitOptions(const FVector& ViewDirection, FSceneViewInitOptions& OutViewInitOptions);
diff --git a/Engine/Plugins/VirtualProduction/EpicStageApp/Source/EpicStageApp/Private/StageAppRouteHandler.cpp b/Engine/Plugins/VirtualProduction/EpicStageApp/Source/EpicStageApp/Private/StageAppRouteHandler.cpp
--- a/Engine/Plugins/VirtualProduction/EpicStageApp/Source/EpicStageApp/Private/StageAppRouteHandler.cpp
+++ b/Engine/Plugins/VirtualProduction/EpicStageApp/Source/EpicStageApp/Private/StageAppRouteHandler.cpp
@@ -723,6 +723,7 @@ void FStageAppRouteHandler::DragLightCards(FPerRendererData& PerRendererData, FV
 		PerRendererData.DraggedLightCards,
 		PixelPos,
 		View,
+		FDisplayClusterLightCardEditorHelper::ECoordinateSystem::Spherical,
 		DragWidgetOffset,
 		EAxisList::Type::XYZ,
 		PerRendererData.PrimaryLightCard.Get()
