nDisplay: Add modular feature to extend DisplayClusterLightCardActor

diff --git a/Engine/Plugins/Runtime/nDisplay/Source/DisplayCluster/DisplayCluster.Build.cs b/Engine/Plugins/Runtime/nDisplay/Source/DisplayCluster/DisplayCluster.Build.cs
--- a/Engine/Plugins/Runtime/nDisplay/Source/DisplayCluster/DisplayCluster.Build.cs
+++ b/Engine/Plugins/Runtime/nDisplay/Source/DisplayCluster/DisplayCluster.Build.cs
@@ -20,6 +20,7 @@ public class DisplayCluster : ModuleRules
 				"Core",
 				"CoreUObject",
 				"DisplayClusterConfiguration",
+				"DisplayClusterLightCardExtender",
 				"Engine"
 			});
 
diff --git a/Engine/Plugins/Runtime/nDisplay/Source/DisplayCluster/Private/Game/EngineClasses/Scene/DisplayClusterLightCardActor.cpp b/Engine/Plugins/Runtime/nDisplay/Source/DisplayCluster/Private/Game/EngineClasses/Scene/DisplayClusterLightCardActor.cpp
--- a/Engine/Plugins/Runtime/nDisplay/Source/DisplayCluster/Private/Game/EngineClasses/Scene/DisplayClusterLightCardActor.cpp
+++ b/Engine/Plugins/Runtime/nDisplay/Source/DisplayCluster/Private/Game/EngineClasses/Scene/DisplayClusterLightCardActor.cpp
@@ -2,6 +2,7 @@
 
 #include "DisplayClusterLightCardActor.h"
 #include "DisplayClusterRootActor.h"
+#include "IDisplayClusterLightCardActorExtender.h"
 
 #include "Components/DisplayClusterLabelComponent.h"
 #include "Components/StaticMeshComponent.h"
@@ -76,6 +77,15 @@ ADisplayClusterLightCardActor::ADisplayClusterLightCardActor(const FObjectInitia
 
 	LabelComponent = CreateOptionalDefaultSubobject<UDisplayClusterLabelComponent>(TEXT("Label"), true);
 	LabelComponent->AttachToComponent(LightCardComponent, FAttachmentTransformRules::KeepRelativeTransform);
+
+	CreateComponentsForExtenders();
+}
+
+void ADisplayClusterLightCardActor::PostLoad()
+{
+	Super::PostLoad();
+
+	CleanUpComponentsForExtenders();
 }
 
 void ADisplayClusterLightCardActor::OnConstruction(const FTransform& Transform)
@@ -486,3 +496,51 @@ void ADisplayClusterLightCardActor::SetPositionalParams(const PositionalParams&
 	Yaw = Params.Yaw;
 	RadialOffset = Params.RadialOffset;
 }
+
+void ADisplayClusterLightCardActor::CreateComponentsForExtenders()
+{
+	IModularFeatures& ModularFeatures = IModularFeatures::Get();
+	const TArray<IDisplayClusterLightCardActorExtender*> Extenders = ModularFeatures.GetModularFeatureImplementations<IDisplayClusterLightCardActorExtender>(IDisplayClusterLightCardActorExtender::ModularFeatureName);
+	for (IDisplayClusterLightCardActorExtender* Extender : Extenders)
+	{
+		UClass* ExtenderComponentClass = Extender->GetAdditionalSubobjectClass().Get();
+		const FName ExtenderName = Extender->GetExtenderName();
+		if (!ensureAlwaysMsgf(ExtenderComponentClass, TEXT("Cannot apply display cluster Extender. Invalid component class.")))
+		{
+			continue;
+		}
+		else if (!ensureAlwaysMsgf(!ExtenderName.IsNone(), TEXT("Cannot apply display cluster Extender. Extender name cannot be empty.")))
+		{
+			continue;
+		}
+		else if (!ensureAlwaysMsgf(!ExtenderNameToComponentMap.Contains(ExtenderName), TEXT("Cannot apply display cluster Extender. Extender name already in use.")))
+		{
+			continue;
+		}
+
+		const FName ObjectName = MakeUniqueObjectName(this, ExtenderComponentClass, ExtenderName);
+		constexpr bool bIsRequired = false;
+		constexpr bool bIsTransient = false;
+		UObject* Component = CreateDefaultSubobject(ObjectName, UActorComponent::StaticClass(), ExtenderComponentClass, bIsRequired, bIsTransient);
+		ExtenderNameToComponentMap.Add(ExtenderName, CastChecked<UActorComponent>(Component));
+	}
+}
+
+void ADisplayClusterLightCardActor::CleanUpComponentsForExtenders()
+{
+	TArray<FName> AvailableExtenders;
+	IModularFeatures& ModularFeatures = IModularFeatures::Get();
+	const TArray<IDisplayClusterLightCardActorExtender*> Extenders = ModularFeatures.GetModularFeatureImplementations<IDisplayClusterLightCardActorExtender>(IDisplayClusterLightCardActorExtender::ModularFeatureName);
+	for (IDisplayClusterLightCardActorExtender* Extender : Extenders)
+	{
+		AvailableExtenders.Add(Extender->GetExtenderName());
+	}
+	const TMap<FName, TObjectPtr<UActorComponent>> CachedExtenderNameToComponentMap(ExtenderNameToComponentMap);
+	for (const TTuple<FName, TObjectPtr<UActorComponent>>& ExtenderNameToComponentPair : CachedExtenderNameToComponentMap)
+	{
+		if (!AvailableExtenders.Contains(ExtenderNameToComponentPair.Key))
+		{
+			ExtenderNameToComponentMap.Remove(ExtenderNameToComponentPair.Key);
+		}
+	}
+}
diff --git a/Engine/Plugins/Runtime/nDisplay/Source/DisplayCluster/Public/DisplayClusterLightCardActor.h b/Engine/Plugins/Runtime/nDisplay/Source/DisplayCluster/Public/DisplayClusterLightCardActor.h
--- a/Engine/Plugins/Runtime/nDisplay/Source/DisplayCluster/Public/DisplayClusterLightCardActor.h
+++ b/Engine/Plugins/Runtime/nDisplay/Source/DisplayCluster/Public/DisplayClusterLightCardActor.h
@@ -8,6 +8,7 @@
 #include "DisplayClusterLightCardActor.generated.h"
 
 class ADisplayClusterRootActor;
+class UActorComponent;
 class UDisplayClusterLabelComponent;
 class USceneComponent;
 class USpringArmComponent;
@@ -21,6 +22,7 @@ enum class EDisplayClusterLightCardMask : uint8
 	Square,
 	UseTextureAlpha,
 	Polygon,
+	MAX	UMETA(Hidden)
 };
 
 USTRUCT(Blueprintable)
@@ -76,6 +78,7 @@ public:
 public:
 	ADisplayClusterLightCardActor(const FObjectInitializer& ObjectInitializer);
 
+	virtual void PostLoad() override;
 	virtual void OnConstruction(const FTransform& Transform) override;
 	virtual void Tick(float DeltaSeconds) override;
 
@@ -221,7 +224,18 @@ public:
 	UPROPERTY(Transient)
 	TObjectPtr<UTexture> PolygonMask = nullptr;
 
+	static FName GetExtenderNameToComponentMapMemberName()
+	{
+		return GET_MEMBER_NAME_CHECKED(ADisplayClusterLightCardActor, ExtenderNameToComponentMap);
+	}
+
 protected:
+	/** Creates components for IDisplayClusterLightCardActorExtender */
+	void CreateComponentsForExtenders();
+
+	/** Removes components that were added by IDisplayClusterLightCardActorExtender */
+	void CleanUpComponentsForExtenders();
+
 	UPROPERTY(EditDefaultsOnly, BlueprintReadOnly, Category = "Default")
 	TObjectPtr<USceneComponent> DefaultSceneRootComponent;
 
@@ -233,7 +247,11 @@ protected:
 
 	UPROPERTY(EditDefaultsOnly, BlueprintReadOnly, Category = "Default")
 	TObjectPtr<UStaticMeshComponent> LightCardComponent;
-	
+
+	/** Components added by the IDisplayLightCardActorExtender */
+	UPROPERTY(VisibleAnywhere, Category = "Default")
+	TMap<FName, TObjectPtr<UActorComponent>> ExtenderNameToComponentMap;
+
 	UPROPERTY(VisibleAnywhere, transient, BlueprintReadOnly, Category = "Default")
 	TObjectPtr<UDisplayClusterLabelComponent> LabelComponent;
 
@@ -244,4 +262,4 @@ private:
 	/** Stores the user translucency value when labels are displayed */
 	TOptional<int32> SavedTranslucencySortPriority;
 	
-};
\ No newline at end of file
+};
diff --git a/Engine/Plugins/Runtime/nDisplay/Source/DisplayClusterLightCardEditor/DisplayClusterLightCardEditor.Build.cs b/Engine/Plugins/Runtime/nDisplay/Source/DisplayClusterLightCardEditor/DisplayClusterLightCardEditor.Build.cs
--- a/Engine/Plugins/Runtime/nDisplay/Source/DisplayClusterLightCardEditor/DisplayClusterLightCardEditor.Build.cs
+++ b/Engine/Plugins/Runtime/nDisplay/Source/DisplayClusterLightCardEditor/DisplayClusterLightCardEditor.Build.cs
@@ -18,6 +18,7 @@ public class DisplayClusterLightCardEditor : ModuleRules
 				"DisplayCluster",
 				"DisplayClusterConfiguration",
 				"DisplayClusterOperator",
+				"DisplayClusterLightCardExtender",
 				"DisplayClusterLightCardEditorShaders",
 				"DisplayClusterScenePreview",
 				"EditorStyle",
diff --git a/Engine/Plugins/Runtime/nDisplay/Source/DisplayClusterLightCardEditor/Private/DetailCustomizations/DisplayClusterLightCardActorDetails.cpp b/Engine/Plugins/Runtime/nDisplay/Source/DisplayClusterLightCardEditor/Private/DetailCustomizations/DisplayClusterLightCardActorDetails.cpp
new file mode 100644
--- /dev/null
+++ b/Engine/Plugins/Runtime/nDisplay/Source/DisplayClusterLightCardEditor/Private/DetailCustomizations/DisplayClusterLightCardActorDetails.cpp
@@ -0,0 +1,71 @@
+// Copyright Epic Games, Inc. All Rights Reserved.
+
+#include "DisplayClusterLightCardActorDetails.h"
+
+#include "DisplayClusterLightCardActor.h"
+#include "IDisplayClusterLightCardActorExtender.h"
+
+#include "DetailCategoryBuilder.h"
+#include "DetailLayoutBuilder.h"
+#include "PropertyHandle.h"
+#include "Algo/Find.h"
+#include "Features/IModularFeatures.h"
+
+TSharedRef<IDetailCustomization> FDisplayClusterLightCardActorDetails::MakeInstance()
+{
+	return MakeShared<FDisplayClusterLightCardActorDetails>();
+}
+
+void FDisplayClusterLightCardActorDetails::CustomizeDetails(IDetailLayoutBuilder& InLayoutBuilder)
+{
+	// Add the detail views of components added for light card actor Extenders
+	const TSharedRef<IPropertyHandle> ExtenderNameToComponentMapHandle = InLayoutBuilder.GetProperty(ADisplayClusterLightCardActor::GetExtenderNameToComponentMapMemberName());
+	ExtenderNameToComponentMapHandle->MarkHiddenByCustomization();
+
+	TArray<const void*> RawDatas;
+	ExtenderNameToComponentMapHandle->AccessRawData(RawDatas);
+	if (!RawDatas.IsEmpty())
+	{
+		IModularFeatures& ModularFeatures = IModularFeatures::Get();
+		const TArray<IDisplayClusterLightCardActorExtender*> Extenders = ModularFeatures.GetModularFeatureImplementations<IDisplayClusterLightCardActorExtender>(IDisplayClusterLightCardActorExtender::ModularFeatureName);
+
+		FName Category;
+		bool bShouldShowSubcategories = false;
+		TArray<UObject*> Components;
+		for (const void* RawData : RawDatas)
+		{
+			TMap<FName, UActorComponent*>* ExtenderNameToComponentMapPtr = (TMap<FName, UActorComponent*>*)(RawData);
+			check(ExtenderNameToComponentMapPtr);
+
+			for (const TTuple<FName, UActorComponent*>& ExtenderNameToComponentPair : *ExtenderNameToComponentMapPtr)
+			{
+				if (!ensureMsgf(ExtenderNameToComponentPair.Value, TEXT("Trying to display component for Extender %s, but component is invalid."), *ExtenderNameToComponentPair.Key.ToString()))
+				{
+					continue;
+				}
+
+				IDisplayClusterLightCardActorExtender* const* ExtenderPtr = 
+					Algo::FindByPredicate(Extenders, [ExtenderNameToComponentPair](const IDisplayClusterLightCardActorExtender* Extender)
+						{
+							return Extender->GetExtenderName() == ExtenderNameToComponentPair.Key;
+						});
+				if (!ensureMsgf(ExtenderPtr, TEXT("Cannot find Extender %s expected to exist for component %s."), *ExtenderNameToComponentPair.Key.ToString(), *ExtenderNameToComponentPair.Value->GetName()))
+				{
+					continue;
+				}
+
+				Category = (*ExtenderPtr)->GetCategory();
+				bShouldShowSubcategories = (*ExtenderPtr)->ShouldShowSubcategories();
+				Components.Add(ExtenderNameToComponentPair.Value);
+			}
+		}
+
+		FAddPropertyParams AddPropertyParams;
+		AddPropertyParams.CreateCategoryNodes(bShouldShowSubcategories);
+		AddPropertyParams.HideRootObjectNode(true);
+
+		InLayoutBuilder
+			.EditCategory(Category)
+			.AddExternalObjects(Components, EPropertyLocation::Default, AddPropertyParams);
+	}
+}
diff --git a/Engine/Plugins/Runtime/nDisplay/Source/DisplayClusterLightCardEditor/Private/DetailCustomizations/DisplayClusterLightCardActorDetails.h b/Engine/Plugins/Runtime/nDisplay/Source/DisplayClusterLightCardEditor/Private/DetailCustomizations/DisplayClusterLightCardActorDetails.h
new file mode 100644
--- /dev/null
+++ b/Engine/Plugins/Runtime/nDisplay/Source/DisplayClusterLightCardEditor/Private/DetailCustomizations/DisplayClusterLightCardActorDetails.h
@@ -0,0 +1,16 @@
+// Copyright Epic Games, Inc. All Rights Reserved.
+
+#pragma once
+
+#include "CoreMinimal.h"
+#include "IDetailCustomization.h"
+
+class FDisplayClusterLightCardActorDetails final : public IDetailCustomization
+{
+public:
+	static TSharedRef<IDetailCustomization> MakeInstance();
+
+	// IDetailCustomization interface
+	virtual void CustomizeDetails(IDetailLayoutBuilder& InLayoutBuilder) override;
+	// End IDetailCustomization interface
+};
diff --git a/Engine/Plugins/Runtime/nDisplay/Source/DisplayClusterLightCardEditor/Private/DisplayClusterLightCardEditorModule.cpp b/Engine/Plugins/Runtime/nDisplay/Source/DisplayClusterLightCardEditor/Private/DisplayClusterLightCardEditorModule.cpp
--- a/Engine/Plugins/Runtime/nDisplay/Source/DisplayClusterLightCardEditor/Private/DisplayClusterLightCardEditorModule.cpp
+++ b/Engine/Plugins/Runtime/nDisplay/Source/DisplayClusterLightCardEditor/Private/DisplayClusterLightCardEditorModule.cpp
@@ -7,6 +7,7 @@
 #include "DisplayClusterLightCardEditorCommands.h"
 #include "DisplayClusterRootActor.h"
 #include "SDisplayClusterLightCardEditor.h"
+#include "DetailCustomizations/DisplayClusterLightCardActorDetails.h"
 #include "Settings/DisplayClusterLightCardEditorSettings.h"
 
 #include "ISettingsModule.h"
@@ -18,7 +19,7 @@ void FDisplayClusterLightCardEditorModule::StartupModule()
 {
 	RegisterTabSpawners();
 	RegisterSettings();
-
+	RegisterDetailCustomizations();
 	FDisplayClusterLightCardEditorCommands::Register();
 }
 
@@ -26,6 +27,7 @@ void FDisplayClusterLightCardEditorModule::ShutdownModule()
 {
 	UnregisterTabSpawners();
 	UnregisterSettings();
+	UnregisterDetailCustomizations();
 
 	FDisplayClusterLightCardEditorCommands::Unregister();
 }
@@ -86,6 +88,24 @@ void FDisplayClusterLightCardEditorModule::UnregisterSettings()
 	}
 }
 
+void FDisplayClusterLightCardEditorModule::RegisterDetailCustomizations()
+{
+	FPropertyEditorModule& PropertyModule = FModuleManager::LoadModuleChecked<FPropertyEditorModule>("PropertyEditor");
+	PropertyModule.RegisterCustomClassLayout(
+		ADisplayClusterLightCardActor::StaticClass()->GetFName(), 
+		FOnGetDetailCustomizationInstance::CreateStatic(&FDisplayClusterLightCardActorDetails::MakeInstance)
+	);
+}
+
+void FDisplayClusterLightCardEditorModule::UnregisterDetailCustomizations()
+{
+	FPropertyEditorModule* PropertyModule = FModuleManager::GetModulePtr<FPropertyEditorModule>("PropertyEditor");
+	if (PropertyModule)
+	{
+		PropertyModule->UnregisterCustomClassLayout(ADisplayClusterLightCardActor::StaticClass()->GetFName());
+	}
+}
+
 IMPLEMENT_MODULE(FDisplayClusterLightCardEditorModule, DisplayClusterLightCardEditor);
 
 #undef LOCTEXT_NAMESPACE
diff --git a/Engine/Plugins/Runtime/nDisplay/Source/DisplayClusterLightCardEditor/Private/DisplayClusterLightCardEditorModule.h b/Engine/Plugins/Runtime/nDisplay/Source/DisplayClusterLightCardEditor/Private/DisplayClusterLightCardEditorModule.h
--- a/Engine/Plugins/Runtime/nDisplay/Source/DisplayClusterLightCardEditor/Private/DisplayClusterLightCardEditorModule.h
+++ b/Engine/Plugins/Runtime/nDisplay/Source/DisplayClusterLightCardEditor/Private/DisplayClusterLightCardEditorModule.h
@@ -24,4 +24,7 @@ private:
 
 	void RegisterSettings();
 	void UnregisterSettings();
+
+	void RegisterDetailCustomizations();
+	void UnregisterDetailCustomizations();
 };
diff --git a/Engine/Plugins/Runtime/nDisplay/nDisplay.uplugin b/Engine/Plugins/Runtime/nDisplay/nDisplay.uplugin
--- a/Engine/Plugins/Runtime/nDisplay/nDisplay.uplugin
+++ b/Engine/Plugins/Runtime/nDisplay/nDisplay.uplugin
@@ -196,6 +196,10 @@
 			"Name": "MediaIOFramework",
 			"Enabled": true
 		},
+		{
+			"Name": "nDisplayModularFeatures",
+			"Enabled": true
+		},
 		{
 			"Name": "OpenCV",
 			"Enabled": true
diff --git a/Engine/Plugins/Runtime/nDisplayModularFeatures/Source/DisplayClusterLightCardExtender/DisplayClusterLightCardExtender.build.cs b/Engine/Plugins/Runtime/nDisplayModularFeatures/Source/DisplayClusterLightCardExtender/DisplayClusterLightCardExtender.build.cs
new file mode 100644
--- /dev/null
+++ b/Engine/Plugins/Runtime/nDisplayModularFeatures/Source/DisplayClusterLightCardExtender/DisplayClusterLightCardExtender.build.cs
@@ -0,0 +1,16 @@
+// Copyright Epic Games, Inc. All Rights Reserved.
+
+using UnrealBuildTool;
+using System.IO;
+
+public class DisplayClusterLightCardExtender : ModuleRules
+{
+	public DisplayClusterLightCardExtender(ReadOnlyTargetRules ROTargetRules) : base(ROTargetRules)
+	{
+		PublicDependencyModuleNames.AddRange(
+			new string[] {
+				"Core",
+				"CoreUObject",
+			});
+	}
+}
diff --git a/Engine/Plugins/Runtime/nDisplayModularFeatures/Source/DisplayClusterLightCardExtender/Private/DisplayClusterLightCardExtenderModule.cpp b/Engine/Plugins/Runtime/nDisplayModularFeatures/Source/DisplayClusterLightCardExtender/Private/DisplayClusterLightCardExtenderModule.cpp
new file mode 100644
--- /dev/null
+++ b/Engine/Plugins/Runtime/nDisplayModularFeatures/Source/DisplayClusterLightCardExtender/Private/DisplayClusterLightCardExtenderModule.cpp
@@ -0,0 +1,5 @@
+// Copyright Epic Games, Inc. All Rights Reserved.
+
+#include "Modules/ModuleManager.h"
+
+IMPLEMENT_MODULE(FDefaultModuleImpl, DisplayClusterLightCardExtender)
diff --git a/Engine/Plugins/Runtime/nDisplayModularFeatures/Source/DisplayClusterLightCardExtender/Private/IDisplayClusterLightCardActorExtender.cpp b/Engine/Plugins/Runtime/nDisplayModularFeatures/Source/DisplayClusterLightCardExtender/Private/IDisplayClusterLightCardActorExtender.cpp
new file mode 100644
--- /dev/null
+++ b/Engine/Plugins/Runtime/nDisplayModularFeatures/Source/DisplayClusterLightCardExtender/Private/IDisplayClusterLightCardActorExtender.cpp
@@ -0,0 +1,5 @@
+// Copyright Epic Games, Inc. All Rights Reserved.
+
+#include "IDisplayClusterLightCardActorExtender.h"
+
+const FName IDisplayClusterLightCardActorExtender::ModularFeatureName = "DisplayClusterLightCardActorExtender";
diff --git a/Engine/Plugins/Runtime/nDisplayModularFeatures/Source/DisplayClusterLightCardExtender/Public/IDisplayClusterLightCardActorExtender.h b/Engine/Plugins/Runtime/nDisplayModularFeatures/Source/DisplayClusterLightCardExtender/Public/IDisplayClusterLightCardActorExtender.h
new file mode 100644
--- /dev/null
+++ b/Engine/Plugins/Runtime/nDisplayModularFeatures/Source/DisplayClusterLightCardExtender/Public/IDisplayClusterLightCardActorExtender.h
@@ -0,0 +1,35 @@
+// Copyright Epic Games, Inc. All Rights Reserved.
+
+#pragma once
+
+#include "CoreMinimal.h"
+#include "Features/IModularFeature.h"
+#include "Templates/SubclassOf.h"
+
+class UActorComponent;
+
+/** 
+ * Light card Extender API
+ */
+class DISPLAYCLUSTERLIGHTCARDEXTENDER_API IDisplayClusterLightCardActorExtender : public IModularFeature
+{
+public:
+	static const FName ModularFeatureName;
+
+	virtual ~IDisplayClusterLightCardActorExtender() = default;
+
+	/** Returns the name of the Extender */
+	virtual FName GetExtenderName() const = 0;
+
+	/** Returns the class of the additional subobject that will be created in the actor. */
+	virtual TSubclassOf<UActorComponent> GetAdditionalSubobjectClass() = 0;
+
+#if WITH_EDITOR
+	/** Returns the name of the category in which details of the component are displayed */
+	virtual FName GetCategory() const = 0;
+
+	/** If true, displays properties in their respective categories */
+	virtual bool ShouldShowSubcategories() const = 0;
+#endif // WITH_EDITOR
+
+};
diff --git a/Engine/Plugins/Runtime/nDisplayModularFeatures/nDisplayModularFeatures.uplugin b/Engine/Plugins/Runtime/nDisplayModularFeatures/nDisplayModularFeatures.uplugin
new file mode 100644
--- /dev/null
+++ b/Engine/Plugins/Runtime/nDisplayModularFeatures/nDisplayModularFeatures.uplugin
@@ -0,0 +1,24 @@
+{
+	"FileVersion": 3,
+	"Version": 1,
+	"VersionName": "1.0",
+	"FriendlyName": "nDisplay Modular Features",
+	"Description": "Modular Features for nDisplay",
+	"Category": "Misc",
+	"CreatedBy": "Epic Games, Inc.",
+	"CreatedByURL": "http://epicgames.com",
+	"DocsURL": "",
+	"MarketplaceURL": "",
+	"SupportURL": "",
+	"CanContainContent": true,
+	"IsBetaVersion": true,
+	"Installed": false,
+	"Hidden": true,
+	"Modules": [
+		{
+			"Name": "DisplayClusterLightCardExtender",
+			"Type": "Runtime",
+			"LoadingPhase": "Default"
+		}
+	]
+}
