nDisplay: fixed cluster node name issues

diff --git a/Engine/Plugins/Runtime/nDisplay/Source/DisplayCluster/Private/Game/EngineClasses/Scene/DisplayClusterRootActorEditor.cpp b/Engine/Plugins/Runtime/nDisplay/Source/DisplayCluster/Private/Game/EngineClasses/Scene/DisplayClusterRootActorEditor.cpp
--- a/Engine/Plugins/Runtime/nDisplay/Source/DisplayCluster/Private/Game/EngineClasses/Scene/DisplayClusterRootActorEditor.cpp
+++ b/Engine/Plugins/Runtime/nDisplay/Source/DisplayCluster/Private/Game/EngineClasses/Scene/DisplayClusterRootActorEditor.cpp
@@ -360,7 +360,7 @@ bool ADisplayClusterRootActor::ImplUpdatePreviewRenderFrame_Editor(const FString
 			// Update all preview components resources before render
 			for (const TTuple<FString, TObjectPtr<UDisplayClusterPreviewComponent>>& PreviewComponentIt : PreviewComponents)
 			{
-				if (PreviewComponentIt.Value && PreviewComponentIt.Value->GetClusterNodeId() == InClusterNodeId)
+				if (PreviewComponentIt.Value && (InClusterNodeId.IsEmpty() || PreviewComponentIt.Value->GetClusterNodeId() == InClusterNodeId))
 				{
 					PreviewComponentIt.Value->UpdatePreviewResources();
 				}
diff --git a/Engine/Plugins/Runtime/nDisplay/Source/DisplayCluster/Private/Render/Viewport/Configuration/DisplayClusterViewportConfiguration.cpp b/Engine/Plugins/Runtime/nDisplay/Source/DisplayCluster/Private/Render/Viewport/Configuration/DisplayClusterViewportConfiguration.cpp
--- a/Engine/Plugins/Runtime/nDisplay/Source/DisplayCluster/Private/Render/Viewport/Configuration/DisplayClusterViewportConfiguration.cpp
+++ b/Engine/Plugins/Runtime/nDisplay/Source/DisplayCluster/Private/Render/Viewport/Configuration/DisplayClusterViewportConfiguration.cpp
@@ -122,7 +122,7 @@ bool FDisplayClusterViewportConfiguration::ImplUpdateConfiguration(EDisplayClust
 
 			if (InViewportNames)
 			{
-				ConfigurationBase.Update(*InViewportNames);
+				ConfigurationBase.Update(*InViewportNames, RenderFrameSettings);
 			}
 			else
 			{
diff --git a/Engine/Plugins/Runtime/nDisplay/Source/DisplayCluster/Private/Render/Viewport/Configuration/DisplayClusterViewportConfigurationBase.cpp b/Engine/Plugins/Runtime/nDisplay/Source/DisplayCluster/Private/Render/Viewport/Configuration/DisplayClusterViewportConfigurationBase.cpp
--- a/Engine/Plugins/Runtime/nDisplay/Source/DisplayCluster/Private/Render/Viewport/Configuration/DisplayClusterViewportConfigurationBase.cpp
+++ b/Engine/Plugins/Runtime/nDisplay/Source/DisplayCluster/Private/Render/Viewport/Configuration/DisplayClusterViewportConfigurationBase.cpp
@@ -49,7 +49,7 @@ void FDisplayClusterViewportConfigurationBase::Update(const FString& ClusterNode
 			if ((It->GetRenderSettingsICVFX().RuntimeFlags & ViewportRuntime_InternalResource) == 0)
 			{
 				// Only viewports from cluster node in render
-				if (It->GetClusterNodeId() == ClusterNodeId)
+				if (ClusterNodeId.IsEmpty() || It->GetClusterNodeId() == ClusterNodeId)
 				{
 					if (!DesiredViewports.Contains(It->GetId()))
 					{
@@ -89,7 +89,7 @@ void FDisplayClusterViewportConfigurationBase::Update(const FString& ClusterNode
 	}
 }
 
-void FDisplayClusterViewportConfigurationBase::Update(const TArray<FString>& InViewportNames)
+void FDisplayClusterViewportConfigurationBase::Update(const TArray<FString>& InViewportNames, FDisplayClusterRenderFrameSettings& InOutRenderFrameSettings)
 {
 	// Collect unused viewports and delete
 	{
@@ -118,6 +118,9 @@ void FDisplayClusterViewportConfigurationBase::Update(const TArray<FString>& InV
 	{
 		if (const UDisplayClusterConfigurationClusterNode* ClusterNodeConfiguration = ClusterNodeConfigurationIt.Value)
 		{
+			// Assign the cluster node name to viewport internals
+			InOutRenderFrameSettings.ClusterNodeId = ClusterNodeConfigurationIt.Key;
+
 			for (const TPair<FString, TObjectPtr<UDisplayClusterConfigurationViewport>>& ViewportIt : ClusterNodeConfiguration->Viewports)
 			{
 				if (ViewportIt.Key.Len() && ViewportIt.Value && InViewportNames.Find(ViewportIt.Key) != INDEX_NONE)
@@ -135,10 +138,15 @@ void FDisplayClusterViewportConfigurationBase::Update(const TArray<FString>& InV
 			}
 		}
 	}
+
+	// Do not use the cluster node name for this pass type
+	InOutRenderFrameSettings.ClusterNodeId.Empty();
 }
 
 void FDisplayClusterViewportConfigurationBase::UpdateClusterNodePostProcess(const FString& InClusterNodeId, const FDisplayClusterRenderFrameSettings& InRenderFrameSettings)
 {
+	check(!InClusterNodeId.IsEmpty());
+
 	const UDisplayClusterConfigurationClusterNode* ClusterNode = ConfigurationData.Cluster->GetNode(InClusterNodeId);
 	if (ClusterNode)
 	{
diff --git a/Engine/Plugins/Runtime/nDisplay/Source/DisplayCluster/Private/Render/Viewport/Configuration/DisplayClusterViewportConfigurationBase.h b/Engine/Plugins/Runtime/nDisplay/Source/DisplayCluster/Private/Render/Viewport/Configuration/DisplayClusterViewportConfigurationBase.h
--- a/Engine/Plugins/Runtime/nDisplay/Source/DisplayCluster/Private/Render/Viewport/Configuration/DisplayClusterViewportConfigurationBase.h
+++ b/Engine/Plugins/Runtime/nDisplay/Source/DisplayCluster/Private/Render/Viewport/Configuration/DisplayClusterViewportConfigurationBase.h
@@ -25,7 +25,7 @@ public:
 
 public:
 	void Update(const FString& ClusterNodeId);
-	void Update(const TArray<FString>& InViewportNames);
+	void Update(const TArray<FString>& InViewportNames, FDisplayClusterRenderFrameSettings& InOutRenderFrameSettings);
 	void UpdateClusterNodePostProcess(const FString& ClusterNodeId, const FDisplayClusterRenderFrameSettings& InRenderFrameSettings);
 
 public:
diff --git a/Engine/Plugins/Runtime/nDisplay/Source/DisplayCluster/Private/Render/Viewport/Configuration/DisplayClusterViewportConfigurationHelpers_ICVFX.cpp b/Engine/Plugins/Runtime/nDisplay/Source/DisplayCluster/Private/Render/Viewport/Configuration/DisplayClusterViewportConfigurationHelpers_ICVFX.cpp
--- a/Engine/Plugins/Runtime/nDisplay/Source/DisplayCluster/Private/Render/Viewport/Configuration/DisplayClusterViewportConfigurationHelpers_ICVFX.cpp
+++ b/Engine/Plugins/Runtime/nDisplay/Source/DisplayCluster/Private/Render/Viewport/Configuration/DisplayClusterViewportConfigurationHelpers_ICVFX.cpp
@@ -76,6 +76,10 @@ namespace DisplayClusterViewportConfigurationHelpers_ICVFX_Impl
 	// Return unique ICVFX name
 	static FString ImplGetNameICVFX(const FString& InClusterNodeId, const FString& InViewportId, const FString& InResourceId)
 	{
+		check(!InClusterNodeId.IsEmpty());
+		check(!InViewportId.IsEmpty());
+		check(!InResourceId.IsEmpty());
+
 		return FString::Printf(TEXT("%s_%s_%s_%s"), *InClusterNodeId, DisplayClusterViewportStrings::icvfx::prefix, *InViewportId, *InResourceId);
 	}
 
@@ -113,8 +117,7 @@ FDisplayClusterViewport* FDisplayClusterViewportConfigurationHelpers_ICVFX::Impl
 	FDisplayClusterViewportManager* ViewportManager = FDisplayClusterViewportConfigurationHelpers_ICVFX::GetViewportManager(RootActor);
 	if (ViewportManager != nullptr)
 	{
-		const FDisplayClusterRenderFrameSettings& RenderFrameSettingsConstRef = ViewportManager->GetRenderFrameSettings();
-		return ViewportManager->ImplFindViewport(ImplGetNameICVFX(RenderFrameSettingsConstRef.ClusterNodeId, InViewportId, InResourceId));
+		return ViewportManager->ImplFindViewport(ImplGetNameICVFX(ViewportManager->GetRenderFrameSettings().ClusterNodeId, InViewportId, InResourceId));
 	}
 
 	return nullptr;
@@ -125,17 +128,17 @@ static bool ImplCreateProjectionPolicy(ADisplayClusterRootActor& RootActor, cons
 	FDisplayClusterViewportManager* ViewportManager = FDisplayClusterViewportConfigurationHelpers_ICVFX::GetViewportManager(RootActor);
 	if (ViewportManager != nullptr)
 	{
-		const FDisplayClusterRenderFrameSettings& RenderFrameSettingsConstRef = ViewportManager->GetRenderFrameSettings();
+		const FString& ClusterNodeId = ViewportManager->GetRenderFrameSettings().ClusterNodeId;
 
 		FDisplayClusterConfigurationProjection CameraProjectionPolicyConfig;
 		CameraProjectionPolicyConfig.Type = bIsCameraProjection ? DisplayClusterProjectionStrings::projection::Camera : DisplayClusterProjectionStrings::projection::Link;
 
 		// Create projection policy for viewport
-		OutProjPolicy = FDisplayClusterViewportManager::CreateProjectionPolicy(ImplGetNameICVFX(RenderFrameSettingsConstRef.ClusterNodeId, InViewportId, InResourceId), &CameraProjectionPolicyConfig);
+		OutProjPolicy = FDisplayClusterViewportManager::CreateProjectionPolicy(ImplGetNameICVFX(ClusterNodeId, InViewportId, InResourceId), &CameraProjectionPolicyConfig);
 
 		if (!OutProjPolicy.IsValid())
 		{
-			UE_LOG(LogDisplayClusterViewport, Error, TEXT("ICVFX Viewport '%s': projection policy for resource '%s' not created for node '%s'."), *InViewportId, *InResourceId, *RenderFrameSettingsConstRef.ClusterNodeId);
+			UE_LOG(LogDisplayClusterViewport, Error, TEXT("ICVFX Viewport '%s': projection policy for resource '%s' not created for node '%s'."), *InViewportId, *InResourceId, *ClusterNodeId);
 			return false;
 		}
 
@@ -156,10 +159,8 @@ FDisplayClusterViewport* FDisplayClusterViewportConfigurationHelpers_ICVFX::Impl
 	FDisplayClusterViewportManager* ViewportManager = FDisplayClusterViewportConfigurationHelpers_ICVFX::GetViewportManager(RootActor);
 	if (ViewportManager != nullptr)
 	{
-		const FDisplayClusterRenderFrameSettings& RenderFrameSettingsConstRef = ViewportManager->GetRenderFrameSettings();
-
 		// Create viewport for new projection policy
-		FDisplayClusterViewport* NewViewport = ViewportManager->ImplCreateViewport(ImplGetNameICVFX(RenderFrameSettingsConstRef.ClusterNodeId, InViewportId, InResourceId), InProjectionPolicy);
+		FDisplayClusterViewport* NewViewport = ViewportManager->ImplCreateViewport(ImplGetNameICVFX(ViewportManager->GetRenderFrameSettings().ClusterNodeId, InViewportId, InResourceId), InProjectionPolicy);
 		if (NewViewport != nullptr)
 		{
 			// Mark as internal resource
@@ -570,8 +571,7 @@ void FDisplayClusterViewportConfigurationHelpers_ICVFX::UpdateCameraViewportSett
 		const FDisplayClusterConfigurationICVFXMedia& MediaSettings = InCameraComponent.CameraSettings.RenderSettings.Media;
 		if (MediaSettings.bEnabled)
 		{
-			const FString ThisClusterNodeId = IDisplayCluster::Get().GetClusterMgr()->GetNodeId();
-			const bool bThisNodeSharesMedia = MediaSettings.MediaOutputNode.Equals(ThisClusterNodeId, ESearchCase::IgnoreCase);
+			const bool bThisNodeSharesMedia = MediaSettings.MediaOutputNode.Equals(DstViewport.GetClusterNodeId(), ESearchCase::IgnoreCase);
 
 			// In most cases there is no need to render ICVFX view if media input set up. The only exception
 			// is when this camera is used as media source for other cluster nodes. In this case media input
diff --git a/Engine/Plugins/Runtime/nDisplay/Source/DisplayCluster/Private/Render/Viewport/Configuration/DisplayClusterViewportConfigurationICVFX.cpp b/Engine/Plugins/Runtime/nDisplay/Source/DisplayCluster/Private/Render/Viewport/Configuration/DisplayClusterViewportConfigurationICVFX.cpp
--- a/Engine/Plugins/Runtime/nDisplay/Source/DisplayCluster/Private/Render/Viewport/Configuration/DisplayClusterViewportConfigurationICVFX.cpp
+++ b/Engine/Plugins/Runtime/nDisplay/Source/DisplayCluster/Private/Render/Viewport/Configuration/DisplayClusterViewportConfigurationICVFX.cpp
@@ -318,6 +318,16 @@ void FDisplayClusterViewportConfigurationICVFX::Update()
 	{
 		ImplBeginReallocateViewports(*ViewportManager);
 
+		// Disable ICVFX if cluster node rendering is not used
+		const FString& ClusterNodeId = ViewportManager->GetRenderFrameSettings().ClusterNodeId;
+		if(ClusterNodeId.IsEmpty())
+		{
+			// The nDisplay viewport should now always have the name of the cluster node.
+			// When rendering MRQ viewports, the list of viewports is used without cluster node names.
+			ImplFinishReallocateViewports(*ViewportManager);
+			return;
+		}
+
 		TArray<FDisplayClusterViewport*> TargetViewports;
 		const EDisplayClusterViewportICVFXFlags TargetViewportsFlags = ImplGetTargetViewports(*ViewportManager, TargetViewports);
 
diff --git a/Engine/Plugins/Runtime/nDisplay/Source/DisplayCluster/Private/Render/Viewport/DisplayClusterViewport.cpp b/Engine/Plugins/Runtime/nDisplay/Source/DisplayCluster/Private/Render/Viewport/DisplayClusterViewport.cpp
--- a/Engine/Plugins/Runtime/nDisplay/Source/DisplayCluster/Private/Render/Viewport/DisplayClusterViewport.cpp
+++ b/Engine/Plugins/Runtime/nDisplay/Source/DisplayCluster/Private/Render/Viewport/DisplayClusterViewport.cpp
@@ -32,6 +32,8 @@ FDisplayClusterViewport::FDisplayClusterViewport(FDisplayClusterViewportManager&
 	, ClusterNodeId(InClusterNodeId)
 	, Owner(InOwner)
 {
+	check(!ClusterNodeId.IsEmpty());
+	check(!ViewportId.IsEmpty());
 	check(UninitializedProjectionPolicy.IsValid());
 
 	// Create scene proxy pair with on game thread. Outside, in ViewportManager added to proxy array on render thread
diff --git a/Engine/Plugins/Runtime/nDisplay/Source/DisplayCluster/Private/Render/Viewport/DisplayClusterViewportManager.cpp b/Engine/Plugins/Runtime/nDisplay/Source/DisplayCluster/Private/Render/Viewport/DisplayClusterViewportManager.cpp
--- a/Engine/Plugins/Runtime/nDisplay/Source/DisplayCluster/Private/Render/Viewport/DisplayClusterViewportManager.cpp
+++ b/Engine/Plugins/Runtime/nDisplay/Source/DisplayCluster/Private/Render/Viewport/DisplayClusterViewportManager.cpp
@@ -394,7 +394,7 @@ void FDisplayClusterViewportManager::ImplUpdateClusterNodeViewports(const EDispl
 			TArray<FDisplayClusterViewport*> OverriddenViewports;
 			for (FDisplayClusterViewport* Viewport : Viewports)
 			{
-				if (Viewport && Viewport->GetClusterNodeId() == InClusterNodeId)
+				if (Viewport && (InClusterNodeId.IsEmpty() || Viewport->GetClusterNodeId() == InClusterNodeId))
 				{
 					bool bHasParentViewport = Viewport->RenderSettings.GetParentViewportId().IsEmpty() == false;
 					bool bHasViewportOverride = Viewport->RenderSettings.OverrideViewportId.IsEmpty() == false;
@@ -831,6 +831,7 @@ FDisplayClusterViewport* FDisplayClusterViewportManager::ImplCreateViewport(cons
 
 	// Create viewport for cluster node used for rendering
 	const FString& ClusterNodeId = GetRenderFrameSettings().ClusterNodeId;
+	check(!ClusterNodeId.IsEmpty());
 
 	// Create viewport
 	FDisplayClusterViewport* NewViewport = new FDisplayClusterViewport(*this, ClusterNodeId, ViewportId, InProjectionPolicy);
diff --git a/Engine/Plugins/Runtime/nDisplay/Source/DisplayCluster/Private/Render/Viewport/DisplayClusterViewportManagerProxy.cpp b/Engine/Plugins/Runtime/nDisplay/Source/DisplayCluster/Private/Render/Viewport/DisplayClusterViewportManagerProxy.cpp
--- a/Engine/Plugins/Runtime/nDisplay/Source/DisplayCluster/Private/Render/Viewport/DisplayClusterViewportManagerProxy.cpp
+++ b/Engine/Plugins/Runtime/nDisplay/Source/DisplayCluster/Private/Render/Viewport/DisplayClusterViewportManagerProxy.cpp
@@ -132,7 +132,7 @@ void FDisplayClusterViewportManagerProxy::ImplUpdateClusterNodeViewportProxies()
 	// Collect viewport proxies for rendered cluster node
 	for (FDisplayClusterViewportProxy* ViewportProxyIt : ViewportProxies)
 	{
-		if (ViewportProxyIt && ViewportProxyIt->GetClusterNodeId() == RenderFrameSettings.ClusterNodeId)
+		if (ViewportProxyIt && (RenderFrameSettings.ClusterNodeId.IsEmpty() || ViewportProxyIt->GetClusterNodeId() == RenderFrameSettings.ClusterNodeId))
 		{
 			ClusterNodeViewportProxies.Add(ViewportProxyIt);
 		}
diff --git a/Engine/Plugins/Runtime/nDisplay/Source/DisplayCluster/Private/Render/Viewport/DisplayClusterViewportManager_EditorPreview.cpp b/Engine/Plugins/Runtime/nDisplay/Source/DisplayCluster/Private/Render/Viewport/DisplayClusterViewportManager_EditorPreview.cpp
--- a/Engine/Plugins/Runtime/nDisplay/Source/DisplayCluster/Private/Render/Viewport/DisplayClusterViewportManager_EditorPreview.cpp
+++ b/Engine/Plugins/Runtime/nDisplay/Source/DisplayCluster/Private/Render/Viewport/DisplayClusterViewportManager_EditorPreview.cpp
@@ -65,7 +65,7 @@ void FDisplayClusterViewportManager::ImplUpdatePreviewRTTResources()
 	for (FDisplayClusterViewport* const ViewportIt : Viewports)
 	{
 		// update only current cluster node
-		if (ViewportIt->GetClusterNodeId() == ClusterNodeId)
+		if ((ClusterNodeId.IsEmpty() || ViewportIt->GetClusterNodeId() == ClusterNodeId))
 		{
 			if (ViewportIt->RenderSettings.bEnable && ViewportIt->RenderSettings.bVisible)
 			{
