nDisplay: Show categories of light card components added by [REDACTED_PERSON] bottom-most

diff --git a/Engine/Plugins/Runtime/nDisplay/Source/DisplayClusterLightCardEditor/Private/DetailCustomizations/DisplayClusterLightCardActorDetails.cpp b/Engine/Plugins/Runtime/nDisplay/Source/DisplayClusterLightCardEditor/Private/DetailCustomizations/DisplayClusterLightCardActorDetails.cpp
--- a/Engine/Plugins/Runtime/nDisplay/Source/DisplayClusterLightCardEditor/Private/DetailCustomizations/DisplayClusterLightCardActorDetails.cpp
+++ b/Engine/Plugins/Runtime/nDisplay/Source/DisplayClusterLightCardEditor/Private/DetailCustomizations/DisplayClusterLightCardActorDetails.cpp
@@ -17,8 +17,8 @@ TSharedRef<IDetailCustomization> FDisplayClusterLightCardActorDetails::MakeInsta
 }
 
 void FDisplayClusterLightCardActorDetails::CustomizeDetails(IDetailLayoutBuilder& InLayoutBuilder)
-{
-	// Add the detail views of components added for light card actor Extenders
+{		
+	// Add the detail views of components added for Light Card Actor Extenders
 	const TSharedRef<IPropertyHandle> ExtenderNameToComponentMapHandle = InLayoutBuilder.GetProperty(ADisplayClusterLightCardActor::GetExtenderNameToComponentMapMemberName());
 	ExtenderNameToComponentMapHandle->MarkHiddenByCustomization();
 
@@ -26,12 +26,17 @@ void FDisplayClusterLightCardActorDetails::CustomizeDetails(IDetailLayoutBuilder
 	ExtenderNameToComponentMapHandle->AccessRawData(RawDatas);
 	if (!RawDatas.IsEmpty())
 	{
+		// Remember the native categories, so extender categories can be ordered bottom-most later
+		TArray<FName> NativeCategories;
+		InLayoutBuilder.GetCategoryNames(NativeCategories);
+
 		IModularFeatures& ModularFeatures = IModularFeatures::Get();
 		const TArray<IDisplayClusterLightCardActorExtender*> Extenders = ModularFeatures.GetModularFeatureImplementations<IDisplayClusterLightCardActorExtender>(IDisplayClusterLightCardActorExtender::ModularFeatureName);
 
 		FName Category;
 		bool bShouldShowSubcategories = false;
 		TArray<UObject*> Components;
+		TArray<FName> ExtenderCategories;
 		for (const void* RawData : RawDatas)
 		{
 			TMap<FName, UActorComponent*>* ExtenderNameToComponentMapPtr = (TMap<FName, UActorComponent*>*)(RawData);
@@ -54,9 +59,14 @@ void FDisplayClusterLightCardActorDetails::CustomizeDetails(IDetailLayoutBuilder
 					continue;
 				}
 
-				Category = (*ExtenderPtr)->GetCategory();
 				bShouldShowSubcategories = (*ExtenderPtr)->ShouldShowSubcategories();
 				Components.Add(ExtenderNameToComponentPair.Value);
+
+				Category = (*ExtenderPtr)->GetCategory();
+				if (!NativeCategories.Contains(Category))
+				{
+					ExtenderCategories.AddUnique(Category);
+				}
 			}
 		}
 
@@ -67,5 +77,25 @@ void FDisplayClusterLightCardActorDetails::CustomizeDetails(IDetailLayoutBuilder
 		InLayoutBuilder
 			.EditCategory(Category)
 			.AddExternalObjects(Components, EPropertyLocation::Default, AddPropertyParams);
+
+		// Sort extender categories bottom-most
+		InLayoutBuilder.SortCategories([ExtenderCategories](const TMap<FName, IDetailCategoryBuilder*>& CategoryMap)
+			{
+				int32 NativeCategorySortOrder = 1;
+				int32 ExtenderCategorySortOrder = CategoryMap.Num();
+				for (const TTuple<FName, IDetailCategoryBuilder*>& CategoryPair : CategoryMap)
+				{
+					if (ExtenderCategories.Contains(CategoryPair.Key))
+					{
+						CategoryPair.Value->SetSortOrder(ExtenderCategorySortOrder);
+						ExtenderCategorySortOrder--;
+					}
+					else
+					{
+						CategoryPair.Value->SetSortOrder(NativeCategorySortOrder);
+						NativeCategorySortOrder++;
+					}
+				}
+			});
 	}
 }
