nDisplay: Fix lightcard editor widget scaling issue when in rotation mode.

diff --git a/Engine/Plugins/Runtime/nDisplay/Source/DisplayClusterLightCardEditor/Private/Viewport/DisplayClusterLightCardEditorWidget.cpp b/Engine/Plugins/Runtime/nDisplay/Source/DisplayClusterLightCardEditor/Private/Viewport/DisplayClusterLightCardEditorWidget.cpp
--- a/Engine/Plugins/Runtime/nDisplay/Source/DisplayClusterLightCardEditor/Private/Viewport/DisplayClusterLightCardEditorWidget.cpp
+++ b/Engine/Plugins/Runtime/nDisplay/Source/DisplayClusterLightCardEditor/Private/Viewport/DisplayClusterLightCardEditorWidget.cpp
@@ -26,7 +26,7 @@ void FDisplayClusterLightCardEditorWidget::Draw(const FSceneView* View, const FD
 	{
 		DrawAxis(PDI, EAxisList::Type::Y, SizeScalar, LengthScalar, OrthoScalar);
 		DrawCircle(PDI, EAxisList::Type::Z, SizeScalar, LengthScalar);
-		DrawOrigin(PDI, SizeScalar, OrthoScalar);
+		DrawOrigin(PDI, LengthScalar, 1.f);
 	}
 	else
 	{
