nDisplay: Fixed issue where UV light cards were not being alpha blended correctly because the UV light card map was coming in with pre-multiplied alpha

diff --git a/Engine/Plugins/Runtime/nDisplay/Shaders/Private/MPCDIOverlayShaders.usf b/Engine/Plugins/Runtime/nDisplay/Shaders/Private/MPCDIOverlayShaders.usf
--- a/Engine/Plugins/Runtime/nDisplay/Shaders/Private/MPCDIOverlayShaders.usf
+++ b/Engine/Plugins/Runtime/nDisplay/Shaders/Private/MPCDIOverlayShaders.usf
@@ -176,8 +176,9 @@ float4 GetInnerCameraColor(float4 WarpedUV, OutputVS IN)
 
 float3 GetUVLightCard(float3 SourceColor, float2 UV)
 {
+	// LightCardColor will have a pre-multiplied alpha, so use a pre-multiplied alpha compositing formula
 	float4 LightCardColor = UVLightCardMapTexture.Sample(UVLightCardMapSampler, UV);
-	return lerp(SourceColor.rgb, LightCardColor.rgb, LightCardColor.a);
+	return SourceColor.rgb * (1 - LightCardColor.a) + LightCardColor.rgb;
 }
 
 OutputPS IcvfxWarpPS(OutputVS IN)
