nDisplay: fixed FDisplayClusterTcpListener destructor issue

diff --git a/Engine/Plugins/Runtime/nDisplay/Source/DisplayCluster/Private/Network/Listener/DisplayClusterTcpListener.cpp b/Engine/Plugins/Runtime/nDisplay/Source/DisplayCluster/Private/Network/Listener/DisplayClusterTcpListener.cpp
--- a/Engine/Plugins/Runtime/nDisplay/Source/DisplayCluster/Private/Network/Listener/DisplayClusterTcpListener.cpp
+++ b/Engine/Plugins/Runtime/nDisplay/Source/DisplayCluster/Private/Network/Listener/DisplayClusterTcpListener.cpp
@@ -31,8 +31,12 @@ FDisplayClusterTcpListener::~FDisplayClusterTcpListener()
 {
 	// Just free resources by stopping the listening
 	StopListening(false);
-}
 
+	// ThreadObj is a TUniquePtr that is deleted after the destructor of this instance has completed.
+	// But internally ThreadObj accesses this instance and calls its methods from its destructor.
+	// Therefore, we need to delete the ThreadObj earlier, inside the destructor of this instance.
+	ThreadObj.Reset();
+}
 
 bool FDisplayClusterTcpListener::StartListening(const FString& InAddr, const uint16 InPort)
 {
