nDisplay: Added cvar to optionally disable icvfx panel auto-pan/tilt feature.

diff --git a/Engine/Plugins/Runtime/nDisplay/Source/DisplayClusterLightCardEditor/Private/Viewport/DisplayClusterLightCardEditorViewportClient.cpp b/Engine/Plugins/Runtime/nDisplay/Source/DisplayClusterLightCardEditor/Private/Viewport/DisplayClusterLightCardEditorViewportClient.cpp
--- a/Engine/Plugins/Runtime/nDisplay/Source/DisplayClusterLightCardEditor/Private/Viewport/DisplayClusterLightCardEditorViewportClient.cpp
+++ b/Engine/Plugins/Runtime/nDisplay/Source/DisplayClusterLightCardEditor/Private/Viewport/DisplayClusterLightCardEditorViewportClient.cpp
@@ -55,6 +55,14 @@
 
 #define LOCTEXT_NAMESPACE "DisplayClusterLightCardEditorViewportClient"
 
+
+/** Console variable used to control the size of the UV light card map texture */
+static TAutoConsoleVariable<bool> CVarICVFXPanelAutoPan(
+	TEXT("nDisplay.panel.autopan"),
+	true,
+	TEXT("When true, the stage view will automatically tilt and pan on supported map projections as you drag objects near the edges of the screen."));
+
+
 //////////////////////////////////////////////////////////////////////////
 // FDisplayClusterLightCardEditorViewportClient
 
@@ -1949,7 +1957,7 @@ void FDisplayClusterLightCardEditorViewportClient::MoveSelectedLightCards(FViewp
 	FVector LightCardProjectedLocation = ProjectWorldPosition(CachedEditorWidgetWorldTransform.GetTranslation(), ViewMatrices);
 	FVector2D ScreenPercentage;
 	
-	if (IsLocationCloseToEdge(LightCardProjectedLocation, InViewport, View, &ScreenPercentage))
+	if (CVarICVFXPanelAutoPan.GetValueOnGameThread() && IsLocationCloseToEdge(LightCardProjectedLocation, InViewport, View, &ScreenPercentage))
 	{
 		DesiredLookAtSpeed = FMath::Max(ScreenPercentage.X, ScreenPercentage.Y) * MaxDesiredLookAtSpeed;
 		DesiredLookAtLocation = LightCardProjectedLocation;
