nDisplay: Added ability to pan the camera when in UV projection mode in the light card editor

diff --git a/Engine/Plugins/Runtime/nDisplay/Shaders/Private/MeshProjectionShaders.usf b/Engine/Plugins/Runtime/nDisplay/Shaders/Private/MeshProjectionShaders.usf
--- a/Engine/Plugins/Runtime/nDisplay/Shaders/Private/MeshProjectionShaders.usf
+++ b/Engine/Plugins/Runtime/nDisplay/Shaders/Private/MeshProjectionShaders.usf
@@ -66,6 +66,7 @@ void PerspectiveVS(FVertexFactoryInput Input, out FMeshProjectionVSToPS Output)
 uint ProjectionParameters_UVProjectionIndex;
 float ProjectionParameters_UVProjectionPlaneSize;
 float ProjectionParameters_UVProjectionPlaneDistance;
+float3 ProjectionParameters_UVProjectionPlaneOffset;
 
 void UVProjectionVS(FVertexFactoryInput Input, out FMeshProjectionVSToPS Output)
 {
@@ -87,7 +88,8 @@ void UVProjectionVS(FVertexFactoryInput Input, out FMeshProjectionVSToPS Output)
 	
 	UV = GetUV(Output.FactoryInterpolants, UVIndex);
 	
-	const float3 PlanePosition = float3(ProjectionParameters_UVProjectionPlaneSize * (UV.x - 0.5),ProjectionParameters_UVProjectionPlaneSize * (0.5 - UV.y), ProjectionParameters_UVProjectionPlaneDistance);
+	const float3 PlanePosition = ProjectionParameters_UVProjectionPlaneOffset + 
+		float3(ProjectionParameters_UVProjectionPlaneSize * (UV.x - 0.5), ProjectionParameters_UVProjectionPlaneSize * (0.5 - UV.y), ProjectionParameters_UVProjectionPlaneDistance);
 	
 	Output.Azimuth = UVIndex;
 	Output.Position = mul(float4(PlanePosition, 1.0), ResolvedView.ViewToClip);
diff --git a/Engine/Plugins/Runtime/nDisplay/Source/DisplayClusterLightCardEditor/Private/Viewport/DisplayClusterLightCardEditorViewportClient.cpp b/Engine/Plugins/Runtime/nDisplay/Source/DisplayClusterLightCardEditor/Private/Viewport/DisplayClusterLightCardEditorViewportClient.cpp
--- a/Engine/Plugins/Runtime/nDisplay/Source/DisplayClusterLightCardEditor/Private/Viewport/DisplayClusterLightCardEditorViewportClient.cpp
+++ b/Engine/Plugins/Runtime/nDisplay/Source/DisplayClusterLightCardEditor/Private/Viewport/DisplayClusterLightCardEditorViewportClient.cpp
@@ -113,6 +113,7 @@ FDisplayClusterLightCardEditorViewportClient::FDisplayClusterLightCardEditorView
 	SetShowStats(true);
 
 	ResetFOVs();
+	ResetCameraProjectionTransforms();
 
 	//This seems to be needed to get the correct world time in the preview.
 	SetIsSimulateInEditorViewport(true);
@@ -123,6 +124,9 @@ FDisplayClusterLightCardEditorViewportClient::FDisplayClusterLightCardEditorView
 	
 	SetProjectionMode(static_cast<EDisplayClusterMeshProjectionType>(Settings->ProjectionMode),
 		static_cast<ELevelViewportType>(Settings->RenderViewportType));
+
+	const int32 CameraSpeed = 3;
+	SetCameraSpeedSetting(CameraSpeed);
 }
 
 FDisplayClusterLightCardEditorViewportClient::~FDisplayClusterLightCardEditorViewportClient()
@@ -154,7 +158,7 @@ void FDisplayClusterLightCardEditorViewportClient::Tick(float DeltaSeconds)
 	FEditorViewportClient::Tick(DeltaSeconds);
 
 	// Camera position is locked to a specific location
-	ResetCamera(/* bLocationOnly */ true);
+	FixCameraTransform();
 
 	CalcEditorWidgetTransform(CachedEditorWidgetWorldTransform);
 
@@ -340,6 +344,10 @@ void FDisplayClusterLightCardEditorViewportClient::Draw(FViewport* InViewport, F
 			RenderSettings.ProjectionTypeSettings.UVProjectionIndex = 1;
 			RenderSettings.ProjectionTypeSettings.UVProjectionPlaneSize = ADisplayClusterLightCardActor::UVPlaneDefaultSize;
 			RenderSettings.ProjectionTypeSettings.UVProjectionPlaneDistance = ADisplayClusterLightCardActor::UVPlaneDefaultDistance;
+
+			// Compute the UV plane offset to allow panning. Need to convert to view space, since the UV projection assumes all coordinates are in view space
+			const FMatrix WorldToViewTransform = FMatrix(FVector::ZAxisVector, FVector::XAxisVector, FVector::YAxisVector, FVector::ZeroVector);
+			RenderSettings.ProjectionTypeSettings.UVProjectionPlaneOffset = -WorldToViewTransform.TransformVector(GetViewTransform().GetLocation());
 		}
 
 		GetSceneViewInitOptions(RenderSettings.ViewInitOptions);
@@ -746,6 +754,7 @@ void FDisplayClusterLightCardEditorViewportClient::CreateDrawnLightCard(const TA
 
 	FSceneView* View = CalcSceneView(&ViewFamily);
 	const FVector ViewOrigin = View->ViewLocation;
+	const FVector OriginOffset = ProjectionMode == EDisplayClusterMeshProjectionType::UV ? ViewOrigin : FVector::ZeroVector;
 
 	TArray<FVector> MouseWorldDirections; // directions from view origin
 
@@ -755,7 +764,7 @@ void FDisplayClusterLightCardEditorViewportClient::CreateDrawnLightCard(const TA
 
 	for (int32 PointIdx = 0; PointIdx < MousePositions.Num(); ++PointIdx)
 	{
-		ProjectionHelper->CalculateOriginAndDirectionFromPixelPosition(MousePositions[PointIdx], *View, FVector::ZeroVector, ScreenOrigin, ScreenDirection);
+		ProjectionHelper->CalculateOriginAndDirectionFromPixelPosition(MousePositions[PointIdx], *View, OriginOffset, ScreenOrigin, ScreenDirection);
 		MouseWorldDirections[PointIdx] = ScreenDirection;
 	}
 
@@ -1042,6 +1051,56 @@ void FDisplayClusterLightCardEditorViewportClient::ProcessClick(FSceneView& View
 	FEditorViewportClient::ProcessClick(View, HitProxy, Key, Event, HitX, HitY);
 }
 
+void FDisplayClusterLightCardEditorViewportClient::BeginCameraMovement(bool bHasMovement)
+{
+	// A little bit of hackery here. FEditorViewportClient doesn't provide many ways to change how camera movement happens (e.g. what keys or mouse movements do what),
+	// but it does call BeginCameraMovement before it actually applies any of the movement stored in CameraUserImpulseData, so we can perform our own camera movement 
+	// impulses here. For most projection modes, we simply zero out any camera translation impulses since the user isn't generally allowed to move the camera away from
+	// the view origin. The UV projection mode is the exception, where the user is allowed to pan the camera around
+	if (ProjectionMode == EDisplayClusterMeshProjectionType::UV)
+	{
+		bool bIsLeftRightKeyPressed = false;
+		bool bIsForwardBackKeyPressed = false;
+
+		for (uint32 Index = 0; Index < static_cast<uint8>(EMultipleKeyBindingIndex::NumChords); ++Index)
+		{
+			EMultipleKeyBindingIndex ChordIndex = static_cast<EMultipleKeyBindingIndex>(Index);
+			bIsForwardBackKeyPressed |= Viewport->KeyState(FViewportNavigationCommands::Get().Forward->GetActiveChord(ChordIndex)->Key);
+			bIsForwardBackKeyPressed |= Viewport->KeyState(FViewportNavigationCommands::Get().Backward->GetActiveChord(ChordIndex)->Key);
+
+			bIsLeftRightKeyPressed |= Viewport->KeyState(FViewportNavigationCommands::Get().Left->GetActiveChord(ChordIndex)->Key);
+			bIsLeftRightKeyPressed |= Viewport->KeyState(FViewportNavigationCommands::Get().Right->GetActiveChord(ChordIndex)->Key);
+		}
+
+		bIsForwardBackKeyPressed |= Viewport->KeyState(EKeys::Up);
+		bIsForwardBackKeyPressed |= Viewport->KeyState(EKeys::Down);
+		bIsLeftRightKeyPressed |= Viewport->KeyState(EKeys::Left);
+		bIsLeftRightKeyPressed |= Viewport->KeyState(EKeys::Right);
+
+		// The editor viewport client is hardcoded to move left/right and forward/back using WSAD and arrows, but for UV, we want the arrow keys to 
+		// control up/down instead of forward/back. If the up/down keys are pressed, reroute the forward/back impulses to the up/down impulse data
+		if (bIsForwardBackKeyPressed)
+		{
+			CameraUserImpulseData->MoveUpDownImpulse = CameraUserImpulseData->MoveForwardBackwardImpulse;
+			CameraUserImpulseData->MoveForwardBackwardImpulse = 0.0f;
+		}
+		else if (!bIsForwardBackKeyPressed && !bIsLeftRightKeyPressed)
+		{
+			// Zero out all camera movement if WASD or arrow keys are not pressed, to keep the control scheme cleaner
+			CameraUserImpulseData->MoveForwardBackwardImpulse = 0.0f;
+			CameraUserImpulseData->MoveRightLeftImpulse = 0.0f;
+			CameraUserImpulseData->MoveUpDownImpulse = 0.0f;
+		}
+	}
+	else
+	{
+		// Always zero out the camera's movement impulses so that the user cannot move the view using any keyboard chords
+		CameraUserImpulseData->MoveForwardBackwardImpulse = 0.0f;
+		CameraUserImpulseData->MoveRightLeftImpulse = 0.0f;
+		CameraUserImpulseData->MoveUpDownImpulse = 0.0f;
+	}
+}
+
 EMouseCursor::Type FDisplayClusterLightCardEditorViewportClient::GetCursor(FViewport* InViewport, int32 X, int32 Y)
 {
 	EMouseCursor::Type MouseCursor = EMouseCursor::Default;
@@ -1516,6 +1575,8 @@ void FDisplayClusterLightCardEditorViewportClient::SetEditorWidgetMode(FDisplayC
 
 void FDisplayClusterLightCardEditorViewportClient::SetProjectionMode(EDisplayClusterMeshProjectionType InProjectionMode, ELevelViewportType InViewportType)
 {
+	SaveProjectionCameraTransform();
+
 	ProjectionMode = InProjectionMode;
 	RenderViewportType = InViewportType;
 
@@ -1529,17 +1590,10 @@ void FDisplayClusterLightCardEditorViewportClient::SetProjectionMode(EDisplayClu
 	ProjectionHelper->SetProjectionMode(ProjectionMode);
 	ProjectionHelper->SetIsOrthographic(InViewportType != LVT_Perspective);
 
-	if (ProjectionMode == EDisplayClusterMeshProjectionType::Linear || ProjectionMode == EDisplayClusterMeshProjectionType::UV)
-	{
-		// TODO: Do we want to cache the perspective rotation and restore it when the user switches back?
-		SetViewRotation(FVector::ForwardVector.Rotation());
-		EditorWidget->SetWidgetScale(1.f);
-	}
-	else if (ProjectionMode == EDisplayClusterMeshProjectionType::Azimuthal)
-	{
-		SetViewRotation(FVector::UpVector.Rotation());
-		EditorWidget->SetWidgetScale(0.5f);
-	}
+	RestoreProjectionCameraTransform();
+
+	const float WidgetScale = ProjectionMode == EDisplayClusterMeshProjectionType::Azimuthal ? 0.5f : 1.0f;
+	EditorWidget->SetWidgetScale(WidgetScale);
 
 	ProjectionOriginComponent = FindProjectionOriginComponent(RootActorProxy.Get());
 
@@ -1598,6 +1652,7 @@ void FDisplayClusterLightCardEditorViewportClient::ResetCamera(bool bLocationOnl
 	SetProjectionMode(GetProjectionMode(), GetRenderViewportType());
 
 	ResetFOVs();
+	ResetCameraProjectionTransforms();
 }
 
 void FDisplayClusterLightCardEditorViewportClient::CycleCoordinateSystem()
@@ -2036,7 +2091,7 @@ FVector2D FDisplayClusterLightCardEditorViewportClient::GetUVLightCardTranslatio
 	const FPlane UVProjectionPlane(ViewOrigin + FVector::ForwardVector * UVProjectionPlaneDistance, -FVector::ForwardVector);
 	const FVector PlaneIntersection = FMath::RayPlaneIntersection(Origin, Direction, UVProjectionPlane);
 
-	const FVector DesiredLocation = ((PlaneIntersection - DragWidgetOffset) - ViewOrigin);
+	const FVector DesiredLocation = PlaneIntersection - DragWidgetOffset;
 	const FVector2D DesiredUVLocation = FVector2D(DesiredLocation.Y / UVProjectionPlaneSize + 0.5f, 0.5f - DesiredLocation.Z / UVProjectionPlaneSize);
 
 	const FVector2D UVDelta = DesiredUVLocation - LightCard->UVCoordinates;
@@ -2417,6 +2472,76 @@ bool FDisplayClusterLightCardEditorViewportClient::IsLocationCloseToEdge(const F
 	return false;
 }
 
+void FDisplayClusterLightCardEditorViewportClient::FixCameraTransform()
+{
+	if (ProjectionMode == EDisplayClusterMeshProjectionType::UV)
+	{
+		// While in UV projection mode, the camera is not allowed to be rotated, but is allowed to be panned
+		SetViewRotation(FVector::ForwardVector.Rotation());
+
+		// Clamp the panning to ensure that the UV plane isn't accidentally dragged offscreen
+		const float PlaneSize = 0.5f * ADisplayClusterLightCardActor::UVPlaneDefaultSize;
+		const FVector ViewLocation = GetViewTransform().GetLocation();
+		const FVector ClampedLocation = FVector(0.0f, FMath::Clamp(ViewLocation.Y, -PlaneSize, PlaneSize), FMath::Clamp(ViewLocation.Z, -PlaneSize, PlaneSize));
+
+		SetViewLocation(ClampedLocation);
+	}
+	else
+	{
+		// All other projection modes besides UV are not allowed to move the camera's position, but can be rotated freely
+		ResetCamera(/* bLocationOnly */ true);
+	}
+}
+
+void FDisplayClusterLightCardEditorViewportClient::SaveProjectionCameraTransform()
+{
+	int32 ProjectionModeIndex = (int32)ProjectionMode;
+	if (ProjectionCameraTransforms.Num() > ProjectionModeIndex)
+	{
+		FViewportCameraTransform CameraTransform;
+		CameraTransform.SetRotation(GetViewRotation());
+
+		if (ProjectionMode == EDisplayClusterMeshProjectionType::UV)
+		{
+			CameraTransform.SetLocation(GetViewLocation());
+		}
+
+		ProjectionCameraTransforms[ProjectionModeIndex] = CameraTransform;
+	}
+}
+
+void FDisplayClusterLightCardEditorViewportClient::RestoreProjectionCameraTransform()
+{
+	int32 ProjectionModeIndex = (int32)ProjectionMode;
+	if (ProjectionCameraTransforms.Num() > ProjectionModeIndex)
+	{
+		const FViewportCameraTransform& CameraTransform = ProjectionCameraTransforms[ProjectionModeIndex];
+		SetViewRotation(CameraTransform.GetRotation());
+
+		if (ProjectionMode == EDisplayClusterMeshProjectionType::UV)
+		{
+			SetViewLocation(CameraTransform.GetLocation());
+		}
+	}
+}
+
+void FDisplayClusterLightCardEditorViewportClient::ResetCameraProjectionTransforms()
+{
+	constexpr int32 MaxProjections = 3;
+	if (ProjectionCameraTransforms.Num() < MaxProjections)
+	{
+		ProjectionCameraTransforms.AddDefaulted(MaxProjections - ProjectionCameraTransforms.Num());
+	}
+
+	ProjectionCameraTransforms[static_cast<int32>(EDisplayClusterMeshProjectionType::Linear)].SetRotation(FVector::ForwardVector.Rotation());
+	ProjectionCameraTransforms[static_cast<int32>(EDisplayClusterMeshProjectionType::Azimuthal)].SetRotation(FVector::UpVector.Rotation());
+	ProjectionCameraTransforms[static_cast<int32>(EDisplayClusterMeshProjectionType::UV)].SetRotation(FVector::ForwardVector.Rotation());
+
+	ProjectionCameraTransforms[static_cast<int32>(EDisplayClusterMeshProjectionType::UV)].SetLocation(FVector::ZeroVector);
+
+	RestoreProjectionCameraTransform();
+}
+
 void FDisplayClusterLightCardEditorViewportClient::ResetFOVs()
 {
 	constexpr int32 MaxFOVs = 3;
diff --git a/Engine/Plugins/Runtime/nDisplay/Source/DisplayClusterLightCardEditor/Private/Viewport/DisplayClusterLightCardEditorViewportClient.h b/Engine/Plugins/Runtime/nDisplay/Source/DisplayClusterLightCardEditor/Private/Viewport/DisplayClusterLightCardEditorViewportClient.h
--- a/Engine/Plugins/Runtime/nDisplay/Source/DisplayClusterLightCardEditor/Private/Viewport/DisplayClusterLightCardEditorViewportClient.h
+++ b/Engine/Plugins/Runtime/nDisplay/Source/DisplayClusterLightCardEditor/Private/Viewport/DisplayClusterLightCardEditorViewportClient.h
@@ -62,6 +62,7 @@ public:
 	virtual void TrackingStarted(const FInputEventState& InInputState, bool bIsDraggingWidget, bool bNudge) override;
 	virtual void TrackingStopped() override;
 	virtual void ProcessClick(FSceneView& View, HHitProxy* HitProxy, FKey Key, EInputEvent Event, uint32 HitX, uint32 HitY) override;
+	virtual void BeginCameraMovement(bool bHasMovement) override;
 	virtual EMouseCursor::Type GetCursor(FViewport* Viewport,int32 X,int32 Y) override;
 	virtual ELevelViewportType GetViewportType() const override { return LVT_Perspective; }
 	virtual bool HasDropPreviewActors() const override { return DropPreviewLightCards.Num() > 0; }
@@ -218,6 +219,18 @@ private:
 	/** Checks if the location is approaching the edge of the view space */
 	bool IsLocationCloseToEdge(const FVector& InPosition, const FViewport* InViewport = nullptr, const FSceneView* InView = nullptr, FVector2D* OutPercentageToEdge = nullptr);
 
+	/** Called each tick to ensure the viewport camera is in the appropriate place based on the projection mode */
+	void FixCameraTransform();
+
+	/** Saves the current camera transform for the current projection mode */
+	void SaveProjectionCameraTransform();
+
+	/** Restores the saved camera transform for the current projection mode */
+	void RestoreProjectionCameraTransform();
+
+	/** Resets the stored camera transforms to default values for each projection mode */
+	void ResetCameraProjectionTransforms();
+
 	/** Resets the camera FOVs */
 	void ResetFOVs();
 
@@ -316,6 +329,9 @@ private:
 	/** Stores each projection mode's field of view separately */
 	TArray<float> ProjectionFOVs;
 
+	/** Stores each projecion mode's camera transforms separately so they can be restored when the projection mode is selected  */
+	TArray<FViewportCameraTransform> ProjectionCameraTransforms;
+
 	/** The increment to change the FOV by when using the scroll wheel */
 	float FOVScrollIncrement = 5.0f;
 
diff --git a/Engine/Plugins/Runtime/nDisplay/Source/DisplayClusterLightCardEditorShaders/Private/DisplayClusterMeshProjectionRenderer.cpp b/Engine/Plugins/Runtime/nDisplay/Source/DisplayClusterLightCardEditorShaders/Private/DisplayClusterMeshProjectionRenderer.cpp
--- a/Engine/Plugins/Runtime/nDisplay/Source/DisplayClusterLightCardEditorShaders/Private/DisplayClusterMeshProjectionRenderer.cpp
+++ b/Engine/Plugins/Runtime/nDisplay/Source/DisplayClusterLightCardEditorShaders/Private/DisplayClusterMeshProjectionRenderer.cpp
@@ -180,6 +180,7 @@ public:
 		UVProjectionIndex.Bind(Initializer.ParameterMap, TEXT("ProjectionParameters_UVProjectionIndex"), SPF_Optional);
 		UVProjectionPlaneSize.Bind(Initializer.ParameterMap, TEXT("ProjectionParameters_UVProjectionPlaneSize"), SPF_Optional);
 		UVProjectionPlaneDistance.Bind(Initializer.ParameterMap, TEXT("ProjectionParameters_UVProjectionPlaneDistance"), SPF_Optional);
+		UVProjectionPlaneOffset.Bind(Initializer.ParameterMap, TEXT("ProjectionParameters_UVProjectionPlaneOffset"), SPF_Optional);
 	}
 
 	static bool ShouldCompilePermutation(const FMeshMaterialShaderPermutationParameters& Parameters)
@@ -204,12 +205,14 @@ public:
 		ShaderBindings.Add(UVProjectionIndex, ShaderElementData.ProjectionTypeSettings.UVProjectionIndex);
 		ShaderBindings.Add(UVProjectionPlaneSize, ShaderElementData.ProjectionTypeSettings.UVProjectionPlaneSize);
 		ShaderBindings.Add(UVProjectionPlaneDistance, ShaderElementData.ProjectionTypeSettings.UVProjectionPlaneDistance);
+		ShaderBindings.Add(UVProjectionPlaneOffset, FVector3f(ShaderElementData.ProjectionTypeSettings.UVProjectionPlaneOffset));
 	}
 
 private:
 	LAYOUT_FIELD(FShaderParameter, UVProjectionIndex);
 	LAYOUT_FIELD(FShaderParameter, UVProjectionPlaneSize);
 	LAYOUT_FIELD(FShaderParameter, UVProjectionPlaneDistance);
+	LAYOUT_FIELD(FShaderParameter, UVProjectionPlaneOffset);
 };
 
 using FMeshPerspectiveProjectionVS = FMeshProjectionVS<EDisplayClusterMeshProjectionType::Linear>;
diff --git a/Engine/Plugins/Runtime/nDisplay/Source/DisplayClusterLightCardEditorShaders/Public/DisplayClusterMeshProjectionRenderer.h b/Engine/Plugins/Runtime/nDisplay/Source/DisplayClusterLightCardEditorShaders/Public/DisplayClusterMeshProjectionRenderer.h
--- a/Engine/Plugins/Runtime/nDisplay/Source/DisplayClusterLightCardEditorShaders/Public/DisplayClusterMeshProjectionRenderer.h
+++ b/Engine/Plugins/Runtime/nDisplay/Source/DisplayClusterLightCardEditorShaders/Public/DisplayClusterMeshProjectionRenderer.h
@@ -69,6 +69,9 @@ struct DISPLAYCLUSTERLIGHTCARDEDITORSHADERS_API FDisplayClusterMeshProjectionTyp
 
 	/** The distance from the view origin of the plane the UVs are projected to, in view space */
 	float UVProjectionPlaneDistance = 100.0f;
+
+	/** A translation offset for the UV projection plane */
+	FVector UVProjectionPlaneOffset = FVector::ZeroVector;
 };
 
 /** Settings for producing a single render. */
diff --git a/Engine/Plugins/Runtime/nDisplay/Source/DisplayClusterScenePreview/Private/DisplayClusterLightCardEditorHelper.cpp b/Engine/Plugins/Runtime/nDisplay/Source/DisplayClusterScenePreview/Private/DisplayClusterLightCardEditorHelper.cpp
--- a/Engine/Plugins/Runtime/nDisplay/Source/DisplayClusterScenePreview/Private/DisplayClusterLightCardEditorHelper.cpp
+++ b/Engine/Plugins/Runtime/nDisplay/Source/DisplayClusterScenePreview/Private/DisplayClusterLightCardEditorHelper.cpp
@@ -718,8 +718,8 @@ void FDisplayClusterLightCardEditorHelper::GetSceneViewInitOptions(
 	const FMatrix* InRotationMatrix,
 	float InDPIScale)
 {
-	OutViewInitOptions.ViewLocation = (ProjectionMode == EDisplayClusterMeshProjectionType::UV) ? FVector::ZeroVector : InLocation;
-	OutViewInitOptions.ViewRotation = (ProjectionMode == EDisplayClusterMeshProjectionType::UV) ? FRotator::ZeroRotator : InRotation;
+	OutViewInitOptions.ViewLocation = InLocation;
+	OutViewInitOptions.ViewRotation = InRotation;
 	OutViewInitOptions.ViewOrigin = OutViewInitOptions.ViewLocation;
 
 	FIntPoint ViewportSize = InViewportSize;
