Fix deprecation warnings after robo merge changes

diff --git a/Engine/Plugins/Runtime/nDisplay/Source/DisplayClusterLightCardEditorShaders/Private/DisplayClusterMeshProjectionRenderer.cpp b/Engine/Plugins/Runtime/nDisplay/Source/DisplayClusterLightCardEditorShaders/Private/DisplayClusterMeshProjectionRenderer.cpp
--- a/Engine/Plugins/Runtime/nDisplay/Source/DisplayClusterLightCardEditorShaders/Private/DisplayClusterMeshProjectionRenderer.cpp
+++ b/Engine/Plugins/Runtime/nDisplay/Source/DisplayClusterLightCardEditorShaders/Private/DisplayClusterMeshProjectionRenderer.cpp
@@ -49,7 +49,7 @@ public:
 		const FSceneView* InView,
 		FMeshPassDrawListContext* InDrawListContext,
 		const FDisplayClusterMeshProjectionRenderSettings& InRenderSettings)
-		: FMeshPassProcessor(InScene, GMaxRHIFeatureLevel, InView, InDrawListContext)
+		: FMeshPassProcessor(EMeshPass::Num, InScene, GMaxRHIFeatureLevel, InView, InDrawListContext)
 		, ProjectionTypeSettings(InRenderSettings.ProjectionTypeSettings)
 		, NormalCorrectionMatrix(InRenderSettings.NormalCorrectionMatrix)
 		, StencilValue(0)
@@ -90,8 +90,8 @@ public:
 		Shaders.TryGetPixelShader(PassShaders.PixelShader);
 
 		FMeshDrawingPolicyOverrideSettings OverrideSettings = GetMeshOverrideSettings(MeshBatch);
-		ERasterizerFillMode MeshFillMode = ComputeMeshFillMode(MeshBatch, *Material, OverrideSettings);
-		ERasterizerCullMode MeshCullMode = ComputeMeshCullMode(MeshBatch, *Material, OverrideSettings);
+		ERasterizerFillMode MeshFillMode = ComputeMeshFillMode(*Material, OverrideSettings);
+		ERasterizerCullMode MeshCullMode = ComputeMeshCullMode(*Material, OverrideSettings);
 
 		ShaderElementDataType ShaderElementData = CreateShaderElementData(MeshBatch, PrimitiveSceneProxy, StaticMeshId);
 		FMeshDrawCommandSortKey SortKey = CreateMeshSortKey(MeshBatch, PrimitiveSceneProxy, *Material, PassShaders.VertexShader.GetShader(), PassShaders.PixelShader.GetShader());
