nDisplay: Fixed issue where stage geometry from nDisplay stage actors wouldn't show up in the Light Card Editor viewport if they had been marked as Hidden in Game. Now, the light card editor filters the proxy root actor's components to always include stage geometry that is referenced by a mesh projection policy in the cluster configuration

diff --git a/Engine/Plugins/Runtime/nDisplay/Source/DisplayClusterScenePreview/Private/DisplayClusterScenePreviewModule.cpp b/Engine/Plugins/Runtime/nDisplay/Source/DisplayClusterScenePreview/Private/DisplayClusterScenePreviewModule.cpp
--- a/Engine/Plugins/Runtime/nDisplay/Source/DisplayClusterScenePreview/Private/DisplayClusterScenePreviewModule.cpp
+++ b/Engine/Plugins/Runtime/nDisplay/Source/DisplayClusterScenePreview/Private/DisplayClusterScenePreviewModule.cpp
@@ -437,9 +437,19 @@ void FDisplayClusterScenePreviewModule::AutoPopulateScene(FRendererConfig& Rende
 
 	if (ADisplayClusterRootActor* RootActor = InternalGetRendererRootActor(RendererConfig))
 	{
-		RendererConfig.Renderer->AddActor(RootActor, [](const UPrimitiveComponent* PrimitiveComponent)
+		TArray<FString> ProjectionMeshNames;
+
+		if (UDisplayClusterConfigurationData* Config = RootActor->GetConfigData())
+		{
+			Config->GetReferencedMeshNames(ProjectionMeshNames);
+		}
+
+		RendererConfig.Renderer->AddActor(RootActor, [&ProjectionMeshNames](const UPrimitiveComponent* PrimitiveComponent)
 		{
-			return !PrimitiveComponent->bHiddenInGame || PrimitiveComponent->IsA<UDisplayClusterScreenComponent>();
+			// Filter out any primitive component that isn't a projection mesh (a static mesh that has a Mesh projection configured for it) or a screen component
+			const bool bIsProjectionMesh = PrimitiveComponent->IsA<UStaticMeshComponent>() && ProjectionMeshNames.Contains(PrimitiveComponent->GetName());
+			const bool bIsScreen = PrimitiveComponent->IsA<UDisplayClusterScreenComponent>();
+			return bIsProjectionMesh || bIsScreen;
 		});
 
 		if (RendererConfig.bAutoUpdateLightcards)
