nDisplay: Fix jumps when moving lightcards in orthographic map projection.

This was happening when the light card was out of bounds from the stage mesh. In this case it tries to use the normal map mesh to find its intersection with the orthographic ray, but was failing to find it (and was falling back to the point in the ray closest to the view origin, causing the jump). The reason seems to be that in a recent api refactor, the call to morph it into the north and south normal maps was removed, so the normal map mesh component still had the unmodified icosphere shape.

diff --git a/Engine/Plugins/Runtime/nDisplay/Source/DisplayClusterScenePreview/Private/DisplayClusterLightCardEditorHelper.cpp b/Engine/Plugins/Runtime/nDisplay/Source/DisplayClusterScenePreview/Private/DisplayClusterLightCardEditorHelper.cpp
--- a/Engine/Plugins/Runtime/nDisplay/Source/DisplayClusterScenePreview/Private/DisplayClusterLightCardEditorHelper.cpp
+++ b/Engine/Plugins/Runtime/nDisplay/Source/DisplayClusterScenePreview/Private/DisplayClusterLightCardEditorHelper.cpp
@@ -1466,6 +1466,9 @@ void FDisplayClusterLightCardEditorHelper::UpdateNormalMapMesh()
 				NormalMapMeshComponent.Get()->SetMaterial(Index, MaterialInterface);
 			}
 		}
+
+		NorthNormalMap.MorphProceduralMesh(NormalMapMeshComponent.Get());
+		SouthNormalMap.MorphProceduralMesh(NormalMapMeshComponent.Get());
 	}
 }
 
