nDisplay: Optimize label widget for editor use. This reapplies backed out changes from 21389818 directly to a DisplayClusterWidgetComponent with extra handling for releasing resources properly.

diff --git a/Engine/Plugins/Runtime/nDisplay/Source/DisplayCluster/Private/Game/EngineClasses/Scene/DisplayClusterLabelComponent.cpp b/Engine/Plugins/Runtime/nDisplay/Source/DisplayCluster/Private/Game/EngineClasses/Scene/DisplayClusterLabelComponent.cpp
--- a/Engine/Plugins/Runtime/nDisplay/Source/DisplayCluster/Private/Game/EngineClasses/Scene/DisplayClusterLabelComponent.cpp
+++ b/Engine/Plugins/Runtime/nDisplay/Source/DisplayCluster/Private/Game/EngineClasses/Scene/DisplayClusterLabelComponent.cpp
@@ -5,6 +5,7 @@
 #include "DisplayClusterLabelWidget.h"
 #include "DisplayClusterLightCardActor.h"
 #include "DisplayClusterRootActor.h"
+#include "DisplayClusterWidgetComponent.h"
 
 #include "EngineUtils.h"
 #include "Components/TextRenderComponent.h"
@@ -15,7 +16,7 @@ UDisplayClusterLabelComponent::UDisplayClusterLabelComponent()
 {
 	PrimaryComponentTick.bCanEverTick = true;
 	
-	WidgetComponent = CreateDefaultSubobject<UWidgetComponent>(TEXT("LabelWidget"));
+	WidgetComponent = CreateDefaultSubobject<UDisplayClusterWidgetComponent>(TEXT("LabelWidget"));
 	
 	// Keeps this visible even if its clipping through the owner.
 	// Requires the owner have a transparent blend mode on its material.
@@ -121,7 +122,8 @@ void UDisplayClusterLabelComponent::UpdateWidgetComponent()
 		LabelString = TEXT("Label");
 	}
 
-	if (WidgetComponent->GetWidgetClass() != WidgetClass.Get())
+	if (WidgetComponent->GetWidgetClass() != WidgetClass.Get() ||
+		(WidgetClass.IsValid() && WidgetComponent->GetWidget() == nullptr))
 	{
 		WidgetComponent->SetWidgetClass(WidgetClass.Get());
 		WidgetComponent->InitWidget();
diff --git a/Engine/Plugins/Runtime/nDisplay/Source/DisplayCluster/Private/Game/EngineClasses/Scene/DisplayClusterWidgetComponent.cpp b/Engine/Plugins/Runtime/nDisplay/Source/DisplayCluster/Private/Game/EngineClasses/Scene/DisplayClusterWidgetComponent.cpp
new file mode 100644
--- /dev/null
+++ b/Engine/Plugins/Runtime/nDisplay/Source/DisplayCluster/Private/Game/EngineClasses/Scene/DisplayClusterWidgetComponent.cpp
@@ -0,0 +1,89 @@
+// Copyright Epic Games, Inc. All Rights Reserved.
+
+#include "DisplayClusterWidgetComponent.h"
+
+#include "Slate/WidgetRenderer.h"
+
+UDisplayClusterWidgetComponent::UDisplayClusterWidgetComponent()
+{
+#if WITH_EDITOR
+	FWorldDelegates::OnWorldCleanup.AddUObject(this, &UDisplayClusterWidgetComponent::OnWorldCleanup);
+#endif
+}
+
+UDisplayClusterWidgetComponent::~UDisplayClusterWidgetComponent()
+{
+#if WITH_EDITOR
+	FWorldDelegates::OnWorldCleanup.RemoveAll(this);
+#endif
+}
+
+void UDisplayClusterWidgetComponent::OnRegister()
+{
+	// Set this prior to registering the scene component so that bounds are calculated correctly.
+	CurrentDrawSize = DrawSize;
+
+	UPrimitiveComponent::OnRegister();
+
+#if !UE_SERVER
+	FWorldDelegates::LevelRemovedFromWorld.AddUObject(this, &ThisClass::OnLevelRemovedFromWorld);
+
+	if ( !IsRunningDedicatedServer() )
+	{
+		const bool bIsGameWorld = GetWorld()->IsGameWorld();
+		if ( Space != EWidgetSpace::Screen )
+		{
+			if ( CanReceiveHardwareInput() && bIsGameWorld )
+			{
+				TSharedPtr<SViewport> GameViewportWidget = GEngine->GetGameViewportWidget();
+				RegisterHitTesterWithViewport(GameViewportWidget);
+			}
+
+			if ( !WidgetRenderer && !GUsingNullRHI )
+			{
+				WidgetRenderer = new FWidgetRenderer(bApplyGammaCorrection);
+			}
+		}
+
+		BodySetup = nullptr;
+
+		// No editor init -- optimization
+	}
+#endif // !UE_SERVER
+}
+
+void UDisplayClusterWidgetComponent::OnUnregister()
+{
+#if !UE_SERVER
+	FWorldDelegates::LevelRemovedFromWorld.RemoveAll(this);
+
+	if ( GetWorld()->IsGameWorld() )
+	{
+		TSharedPtr<SViewport> GameViewportWidget = GEngine->GetGameViewportWidget();
+		if ( GameViewportWidget.IsValid() )
+		{
+			UnregisterHitTesterWithViewport(GameViewportWidget);
+		}
+	}
+#endif
+
+	// No editor release -- optimization
+
+	UPrimitiveComponent::OnUnregister();
+}
+
+#if WITH_EDITOR
+void UDisplayClusterWidgetComponent::OnWorldCleanup(UWorld* InWorld, bool bSessionEnded, bool bCleanupResources)
+{
+	const UWorld* World = GetWorld();
+	
+	if (World == InWorld && bCleanupResources && World != nullptr)
+	{
+		if (!World->IsGameWorld())
+		{
+			// Prevents crash when a world is being cleaned up and CheckForWorldGCLeaks is called.
+			ReleaseResources();
+		}
+	}
+}
+#endif
diff --git a/Engine/Plugins/Runtime/nDisplay/Source/DisplayCluster/Private/Game/EngineClasses/Scene/DisplayClusterWidgetComponent.h b/Engine/Plugins/Runtime/nDisplay/Source/DisplayCluster/Private/Game/EngineClasses/Scene/DisplayClusterWidgetComponent.h
new file mode 100644
--- /dev/null
+++ b/Engine/Plugins/Runtime/nDisplay/Source/DisplayCluster/Private/Game/EngineClasses/Scene/DisplayClusterWidgetComponent.h
@@ -0,0 +1,27 @@
+// Copyright Epic Games, Inc. All Rights Reserved.
+
+#pragma once
+
+#include "Components/WidgetComponent.h"
+
+#include "DisplayClusterWidgetComponent.generated.h"
+
+/**
+ * Extend the widget component to optimize when the widget is initialized and released.
+ */
+UCLASS(transient, NotBlueprintable, NotBlueprintType)
+class UDisplayClusterWidgetComponent final : public UWidgetComponent
+{
+	GENERATED_BODY();
+public:
+	UDisplayClusterWidgetComponent();
+	virtual ~UDisplayClusterWidgetComponent() override;
+
+private:
+	virtual void OnRegister() override;
+	virtual void OnUnregister() override;
+
+#if WITH_EDITOR
+	void OnWorldCleanup(UWorld* InWorld, bool bSessionEnded, bool bCleanupResources);
+#endif
+};
diff --git a/Engine/Plugins/Runtime/nDisplay/Source/DisplayCluster/Public/Components/DisplayClusterLabelComponent.h b/Engine/Plugins/Runtime/nDisplay/Source/DisplayCluster/Public/Components/DisplayClusterLabelComponent.h
--- a/Engine/Plugins/Runtime/nDisplay/Source/DisplayCluster/Public/Components/DisplayClusterLabelComponent.h
+++ b/Engine/Plugins/Runtime/nDisplay/Source/DisplayCluster/Public/Components/DisplayClusterLabelComponent.h
@@ -56,7 +56,7 @@ private:
 	
 	/** The widget component to display for this label */
 	UPROPERTY(VisibleAnywhere, Category=LabelText)
-	TObjectPtr<UWidgetComponent> WidgetComponent;
+	TObjectPtr<class UDisplayClusterWidgetComponent> WidgetComponent;
 
 	/** The widget class to apply to the widget component */
 	UPROPERTY(EditAnywhere, Category=LabelText)
diff --git a/Engine/Plugins/Runtime/nDisplay/Source/DisplayClusterLightCardEditor/Private/Viewport/DisplayClusterLightcardEditorViewport.cpp b/Engine/Plugins/Runtime/nDisplay/Source/DisplayClusterLightCardEditor/Private/Viewport/DisplayClusterLightcardEditorViewport.cpp
--- a/Engine/Plugins/Runtime/nDisplay/Source/DisplayClusterLightCardEditor/Private/Viewport/DisplayClusterLightcardEditorViewport.cpp
+++ b/Engine/Plugins/Runtime/nDisplay/Source/DisplayClusterLightCardEditor/Private/Viewport/DisplayClusterLightcardEditorViewport.cpp
@@ -335,6 +335,20 @@ SDisplayClusterLightCardEditorViewport::~SDisplayClusterLightCardEditorViewport(
 		ViewportClient->Viewport = nullptr;
 		ViewportClient.Reset();
 	}
+
+	if (PreviewScene.IsValid())
+	{
+		if (UWorld* PreviewWorld = PreviewScene->GetWorld())
+		{
+			// Call DestroyWorld first so OnWorldCleanup is broadcasted while actors still have a valid world.
+			// Then mark as garbage to prevent AddReferencedObjects from continuing to be called.
+			PreviewWorld->DestroyWorld(true);
+			
+			PreviewWorld->MarkObjectsPendingKill();
+			PreviewWorld->MarkAsGarbage();
+		}
+		PreviewScene.Reset();
+	}
 }
 
 TSharedRef<SEditorViewport> SDisplayClusterLightCardEditorViewport::GetViewportWidget()
