nDisplay: fixed crash caused by lost references to the internal incamera viewport in optimized editor preview

diff --git a/Engine/Plugins/Runtime/nDisplay/Source/DisplayClusterProjection/Private/Policy/MPCDI/DisplayClusterProjectionMPCDIPolicy.cpp b/Engine/Plugins/Runtime/nDisplay/Source/DisplayClusterProjection/Private/Policy/MPCDI/DisplayClusterProjectionMPCDIPolicy.cpp
--- a/Engine/Plugins/Runtime/nDisplay/Source/DisplayClusterProjection/Private/Policy/MPCDI/DisplayClusterProjectionMPCDIPolicy.cpp
+++ b/Engine/Plugins/Runtime/nDisplay/Source/DisplayClusterProjection/Private/Policy/MPCDI/DisplayClusterProjectionMPCDIPolicy.cpp
@@ -323,6 +323,9 @@ void FDisplayClusterProjectionMPCDIPolicy::ApplyWarpBlend_RenderThread(FRHIComma
 					ShaderICVFX.UVLightCardMap = LightCardManager->GetUVLightCardMap_RenderThread();
 				}
 
+				// cleanup camera list for render
+				ShaderICVFX.CleanupCamerasForRender();
+
 				// Initialize shader input data
 				FDisplayClusterShaderParameters_WarpBlend WarpBlendParameters;
 
diff --git a/Engine/Plugins/Runtime/nDisplay/Source/DisplayClusterShaders/Private/Shaders/DisplayClusterShadersWarpblend_ICVFX.cpp b/Engine/Plugins/Runtime/nDisplay/Source/DisplayClusterShaders/Private/Shaders/DisplayClusterShadersWarpblend_ICVFX.cpp
--- a/Engine/Plugins/Runtime/nDisplay/Source/DisplayClusterShaders/Private/Shaders/DisplayClusterShadersWarpblend_ICVFX.cpp
+++ b/Engine/Plugins/Runtime/nDisplay/Source/DisplayClusterShaders/Private/Shaders/DisplayClusterShadersWarpblend_ICVFX.cpp
@@ -1,6 +1,7 @@
 // Copyright Epic Games, Inc. All Rights Reserved.
 
 #include "DisplayClusterShadersWarpblend_ICVFX.h"
+#include "DisplayClusterShadersLog.h"
 
 #include "RenderResource.h"
 #include "CommonRenderResources.h"
@@ -799,6 +800,22 @@ public:
 		TShaderMapRef<FIcvfxWarpVS> VertexShader(ShaderMap, RenderPassData.VSPermutationVector);
 		TShaderMapRef<FIcvfxWarpPS> PixelShader(ShaderMap, RenderPassData.PSPermutationVector);
 
+		if (!VertexShader.IsValid() || !PixelShader.IsValid())
+		{
+			// Invalid permutation vector
+			if (!PixelShader.IsValid())
+			{
+				UE_LOG(LogDisplayClusterShaders, Warning, TEXT("Invalid permutation vector %d for shader FIcvfxWarpPS"), RenderPassData.PSPermutationVector.ToDimensionValueId());
+			}
+
+			if (!VertexShader.IsValid())
+			{
+				UE_LOG(LogDisplayClusterShaders, Warning, TEXT("Invalid permutation vector %d for shader FIcvfxWarpVS"), RenderPassData.VSPermutationVector.ToDimensionValueId());
+			}
+
+			return false;
+		}
+
 		FGraphicsPipelineStateInitializer GraphicsPSOInit;
 		// Set the graphic pipeline state.
 		RHICmdList.ApplyCachedRenderTargets(GraphicsPSOInit);
diff --git a/Engine/Plugins/Runtime/nDisplay/Source/DisplayClusterShaders/Public/ShaderParameters/DisplayClusterShaderParameters_ICVFX.h b/Engine/Plugins/Runtime/nDisplay/Source/DisplayClusterShaders/Public/ShaderParameters/DisplayClusterShaderParameters_ICVFX.h
--- a/Engine/Plugins/Runtime/nDisplay/Source/DisplayClusterShaders/Public/ShaderParameters/DisplayClusterShaderParameters_ICVFX.h
+++ b/Engine/Plugins/Runtime/nDisplay/Source/DisplayClusterShaders/Public/ShaderParameters/DisplayClusterShaderParameters_ICVFX.h
@@ -228,6 +228,23 @@ public:
 		int32 RenderOrder = -1;
 	};
 
+	// Remove unused cameras from render
+	bool CleanupCamerasForRender()
+	{
+		const TArray<FCameraSettings> InCameras = Cameras;
+		Cameras.Empty();
+
+		for (const FCameraSettings& CameraIt : InCameras)
+		{
+			if (CameraIt.IsUsed())
+			{
+				Cameras.Add(CameraIt);
+			}
+		}
+
+		return Cameras.Num() == InCameras.Num();
+	}
+
 	TArray<FCameraSettings> Cameras;
 
 	// Lightcard settings
