nDisplay: Rename operator panel name to 'In-Camera VFX' and add icon.

diff --git a/Engine/Plugins/Runtime/nDisplay/Source/DisplayClusterOperator/DisplayClusterOperator.Build.cs b/Engine/Plugins/Runtime/nDisplay/Source/DisplayClusterOperator/DisplayClusterOperator.Build.cs
--- a/Engine/Plugins/Runtime/nDisplay/Source/DisplayClusterOperator/DisplayClusterOperator.Build.cs
+++ b/Engine/Plugins/Runtime/nDisplay/Source/DisplayClusterOperator/DisplayClusterOperator.Build.cs
@@ -17,6 +17,7 @@ public class DisplayClusterOperator : ModuleRules
 				"Engine",
 				"InputCore",
 				"Kismet",
+				"Projects",
 				"Slate",
 				"SlateCore",
 				"ToolWidgets",
diff --git a/Engine/Plugins/Runtime/nDisplay/Source/DisplayClusterOperator/Private/DisplayClusterOperatorModule.cpp b/Engine/Plugins/Runtime/nDisplay/Source/DisplayClusterOperator/Private/DisplayClusterOperatorModule.cpp
--- a/Engine/Plugins/Runtime/nDisplay/Source/DisplayClusterOperator/Private/DisplayClusterOperatorModule.cpp
+++ b/Engine/Plugins/Runtime/nDisplay/Source/DisplayClusterOperator/Private/DisplayClusterOperatorModule.cpp
@@ -2,6 +2,7 @@
 
 #include "DisplayClusterOperatorModule.h"
 
+#include "DisplayClusterOperatorStyle.h"
 #include "DisplayClusterOperatorViewModel.h"
 #include "SDisplayClusterOperatorPanel.h"
 #include "DisplayClusterRootActor.h"
@@ -24,6 +25,8 @@ void FDisplayClusterOperatorModule::StartupModule()
 	OperatorViewModel = MakeShared<FDisplayClusterOperatorViewModel>();
 	OperatorToolBarExtensibilityManager = MakeShared<FExtensibilityManager>();
 
+	FDisplayClusterOperatorStyle::Get();
+	
 	RegisterTabSpawners();
 }
 
@@ -92,8 +95,9 @@ void FDisplayClusterOperatorModule::RegisterTabSpawners()
 {
 	// Register the nDisplay operator panel tab spawner
 	FGlobalTabmanager::Get()->RegisterNomadTabSpawner(OperatorPanelTabName, FOnSpawnTab::CreateRaw(this, &FDisplayClusterOperatorModule::SpawnOperatorPanelTab))
-		.SetDisplayName(LOCTEXT("TabDisplayName", "nDisplay Operator"))
-		.SetTooltipText(LOCTEXT("TabTooltip", "Open the nDisplay Operator tab."))
+		.SetDisplayName(LOCTEXT("TabDisplayName", "In-Camera VFX"))
+		.SetTooltipText(LOCTEXT("TabTooltip", "Open the In-Camera VFX tab."))
+		.SetIcon(FSlateIcon(FDisplayClusterOperatorStyle::Get().GetStyleSetName(), (TEXT("OperatorPanel.Icon"))))
 		.SetGroup(WorkspaceMenu::GetMenuStructure().GetLevelEditorVirtualProductionCategory());
 }
 
@@ -106,6 +110,7 @@ TSharedRef<SDockTab> FDisplayClusterOperatorModule::SpawnOperatorPanelTab(const
 {
 	const TSharedRef<SDockTab> MajorTab = SNew(SDockTab)
 		.TabRole(ETabRole::NomadTab)
+		.IconColor(FDisplayClusterOperatorStyle::Get().GetColor(TEXT("OperatorPanel.AssetColor")))
 		.OnTabClosed_Raw(this, &FDisplayClusterOperatorModule::OnOperatorPanelTabClosed);
 
 	MajorTab->SetContent(SAssignNew(ActiveOperatorPanel, SDisplayClusterOperatorPanel, OperatorViewModel->CreateTabManager(MajorTab), SpawnTabArgs.GetOwnerWindow()));
diff --git a/Engine/Plugins/Runtime/nDisplay/Source/DisplayClusterOperator/Private/DisplayClusterOperatorStyle.cpp b/Engine/Plugins/Runtime/nDisplay/Source/DisplayClusterOperator/Private/DisplayClusterOperatorStyle.cpp
new file mode 100644
--- /dev/null
+++ b/Engine/Plugins/Runtime/nDisplay/Source/DisplayClusterOperator/Private/DisplayClusterOperatorStyle.cpp
@@ -0,0 +1,49 @@
+// Copyright Epic Games, Inc. All Rights Reserved.
+
+#include "DisplayClusterOperatorStyle.h"
+
+#include "Interfaces/IPluginManager.h"
+#include "Styling/SlateStyleRegistry.h"
+#include "Styling/SlateStyleMacros.h"
+
+FDisplayClusterOperatorStyle::FDisplayClusterOperatorStyle()
+	: FSlateStyleSet(TEXT("DisplayClusterOperatorStyle"))
+{
+	Initialize();
+
+	FSlateStyleRegistry::RegisterSlateStyle(*this);
+}
+
+FDisplayClusterOperatorStyle::~FDisplayClusterOperatorStyle()
+{
+	FSlateStyleRegistry::UnRegisterSlateStyle(*this);
+}
+
+#define EDITOR_IMAGE_BRUSH_SVG( RelativePath, ... ) FSlateVectorImageBrush(FPaths::EngineContentDir() / "Editor/Slate" / RelativePath + TEXT(".svg"), __VA_ARGS__ )
+#define EDITOR_BOX_BRUSH( RelativePath, ... ) FSlateBoxBrush(FPaths::EngineContentDir() / "Editor/Slate" / RelativePath + TEXT(".png"), __VA_ARGS__ )
+
+const FVector2D Icon16x16(16.0f, 16.0f);
+
+void FDisplayClusterOperatorStyle::Initialize()
+{
+	TSharedPtr<IPlugin> Plugin = IPluginManager::Get().FindPlugin(TEXT("nDisplay"));
+	check(Plugin.IsValid());
+	if (Plugin.IsValid())
+	{
+		SetContentRoot(FPaths::Combine(Plugin->GetBaseDir(), TEXT("Content")));
+		SetCoreContentRoot(FPaths::EngineContentDir() / TEXT("Slate"));
+	}
+
+	// Icons
+	{
+		Set("OperatorPanel.Icon", new IMAGE_BRUSH_SVG("Icons/Components/nDisplayCamera", Icon16x16));
+	}
+	
+	// Colors
+	{
+		Set("OperatorPanel.AssetColor", FLinearColor(0, 188, 212));  // From FDisplayClusterConfiguratorAssetTypeActions::GetTypeColor()
+	}
+}
+
+#undef EDITOR_IMAGE_BRUSH_SVG
+#undef EDITOR_BOX_BRUSH
diff --git a/Engine/Plugins/Runtime/nDisplay/Source/DisplayClusterOperator/Private/DisplayClusterOperatorStyle.h b/Engine/Plugins/Runtime/nDisplay/Source/DisplayClusterOperator/Private/DisplayClusterOperatorStyle.h
new file mode 100644
--- /dev/null
+++ b/Engine/Plugins/Runtime/nDisplay/Source/DisplayClusterOperator/Private/DisplayClusterOperatorStyle.h
@@ -0,0 +1,22 @@
+// Copyright Epic Games, Inc. All Rights Reserved.
+
+#pragma once
+
+#include "Styling/SlateStyle.h"
+
+class FDisplayClusterOperatorStyle : public FSlateStyleSet
+{
+public:
+	FDisplayClusterOperatorStyle();
+	~FDisplayClusterOperatorStyle();
+
+	static FDisplayClusterOperatorStyle& Get()
+	{
+		static FDisplayClusterOperatorStyle Inst;
+		return Inst;
+	}
+
+private:
+	void Initialize();
+
+};
diff --git a/Engine/Plugins/Runtime/nDisplay/Source/DisplayClusterOperator/Private/SDisplayClusterOperatorPanel.cpp b/Engine/Plugins/Runtime/nDisplay/Source/DisplayClusterOperator/Private/SDisplayClusterOperatorPanel.cpp
--- a/Engine/Plugins/Runtime/nDisplay/Source/DisplayClusterOperator/Private/SDisplayClusterOperatorPanel.cpp
+++ b/Engine/Plugins/Runtime/nDisplay/Source/DisplayClusterOperator/Private/SDisplayClusterOperatorPanel.cpp
@@ -32,7 +32,7 @@ SDisplayClusterOperatorPanel::~SDisplayClusterOperatorPanel()
 void SDisplayClusterOperatorPanel::Construct(const FArguments& InArgs, const TSharedRef<FTabManager>& InTabManager, const TSharedPtr<SWindow>& WindowOwner)
 {
 	TabManager = InTabManager;
-	TSharedRef<FWorkspaceItem> AppMenuGroup = TabManager->AddLocalWorkspaceMenuCategory(LOCTEXT("OperatorMenuGroupName", "nDisplay Operator"));
+	TSharedRef<FWorkspaceItem> AppMenuGroup = TabManager->AddLocalWorkspaceMenuCategory(LOCTEXT("OperatorMenuGroupName", "In-Camera VFX"));
 	TabManager->SetAllowWindowMenuBar(true);
 
 	TabManager->RegisterTabSpawner(ToolbarTabId, FOnSpawnTab::CreateSP(this, &SDisplayClusterOperatorPanel::SpawnToolbarTab))
