nDisplay: Hide empty folders in the light card editor outliner.

diff --git a/Engine/Plugins/Runtime/nDisplay/Source/DisplayClusterLightCardEditor/Private/Outliner/DisplayClusterLightCardOutlinerMode.cpp b/Engine/Plugins/Runtime/nDisplay/Source/DisplayClusterLightCardEditor/Private/Outliner/DisplayClusterLightCardOutlinerMode.cpp
--- a/Engine/Plugins/Runtime/nDisplay/Source/DisplayClusterLightCardEditor/Private/Outliner/DisplayClusterLightCardOutlinerMode.cpp
+++ b/Engine/Plugins/Runtime/nDisplay/Source/DisplayClusterLightCardEditor/Private/Outliner/DisplayClusterLightCardOutlinerMode.cpp
@@ -107,7 +107,7 @@ struct FActorSelector
 FDisplayClusterLightCardOutlinerMode::FDisplayClusterLightCardOutlinerMode(SSceneOutliner* InSceneOutliner,
 	TWeakPtr<SDisplayClusterLightCardOutliner> InLightCardOutliner, TWeakObjectPtr<UWorld> InSpecifiedWorldToDisplay)
 	: FActorMode(FActorModeParams(InSceneOutliner, InSpecifiedWorldToDisplay,
-	/* bHideComponents */ true, /* bHideLevelInstanceHierarchy */ true, /* bHideUnloadedActors */ true, /* bHideEmptyFolders */ false)),
+	/* bHideComponents */ true, /* bHideLevelInstanceHierarchy */ true, /* bHideUnloadedActors */ true, /* bHideEmptyFolders */ true)),
 	LightCardOutliner(InLightCardOutliner)
 {
 }
