nDisplay:  Support Color Correct Windows in the light card editor.
- Refactor light card editor to support any actor that implements a new IDisplayClusterStageActor interface. The ColorCorrectWindow and LightCardActor now implement this.
- Add new FDisplayClusterWeakStageActorPtr to use within the lightcard editor.
- Add support for spawning in actors that inherit the StageActor interface.
- Minor optimizations to refreshing the light card editor.

diff --git a/Engine/Plugins/Experimental/ColorCorrectRegions/ColorCorrectRegions.uplugin b/Engine/Plugins/Experimental/ColorCorrectRegions/ColorCorrectRegions.uplugin
--- a/Engine/Plugins/Experimental/ColorCorrectRegions/ColorCorrectRegions.uplugin
+++ b/Engine/Plugins/Experimental/ColorCorrectRegions/ColorCorrectRegions.uplugin
@@ -25,5 +25,11 @@
 			"Type": "Editor",
 			"LoadingPhase": "PostEngineInit"
 		}
+	],
+	"Plugins": [
+		{
+			"Name": "nDisplayModularFeatures",
+			"Enabled": true
+		}
 	]
 }
\ No newline at end of file
diff --git a/Engine/Plugins/Experimental/ColorCorrectRegions/Source/ColorCorrectRegions/ColorCorrectRegions.Build.cs b/Engine/Plugins/Experimental/ColorCorrectRegions/Source/ColorCorrectRegions/ColorCorrectRegions.Build.cs
--- a/Engine/Plugins/Experimental/ColorCorrectRegions/Source/ColorCorrectRegions/ColorCorrectRegions.Build.cs
+++ b/Engine/Plugins/Experimental/ColorCorrectRegions/Source/ColorCorrectRegions/ColorCorrectRegions.Build.cs
@@ -23,6 +23,7 @@ public class ColorCorrectRegions : ModuleRules
 			{
 				"Core",
 				"RenderCore",
+				"DisplayClusterLightCardExtender"
 			}
 		);
 			
diff --git a/Engine/Plugins/Experimental/ColorCorrectRegions/Source/ColorCorrectRegions/Private/ColorCorrectWindow.cpp b/Engine/Plugins/Experimental/ColorCorrectRegions/Source/ColorCorrectRegions/Private/ColorCorrectWindow.cpp
--- a/Engine/Plugins/Experimental/ColorCorrectRegions/Source/ColorCorrectRegions/Private/ColorCorrectWindow.cpp
+++ b/Engine/Plugins/Experimental/ColorCorrectRegions/Source/ColorCorrectRegions/Private/ColorCorrectWindow.cpp
@@ -16,6 +16,15 @@ AColorCorrectWindow::AColorCorrectWindow(const FObjectInitializer& ObjectInitial
 	: Super(ObjectInitializer)
 	, WindowType(EColorCorrectWindowType::Square)
 {
+	PositionalParams.DistanceFromCenter = 300.f;
+	PositionalParams.Longitude = 0.f;
+	PositionalParams.Latitude = 30.f;
+	PositionalParams.Spin = 0.f;
+	PositionalParams.Pitch = 0.f;
+	PositionalParams.Yaw = 0.5;
+	PositionalParams.RadialOffset = -0.5f;
+	PositionalParams.Scale = FVector2D(1.f);
+	
 	UMaterial* Material = LoadObject<UMaterial>(NULL, TEXT("/ColorCorrectRegions/Materials/M_ColorCorrectRegionTransparentPreview.M_ColorCorrectRegionTransparentPreview"), NULL, LOAD_None, NULL);
 	const TArray<UStaticMesh*> StaticMeshes =
 	{
@@ -50,6 +59,36 @@ void AColorCorrectWindow::PostEditChangeProperty(struct FPropertyChangedEvent& P
 	{
 		SetMeshVisibilityForWindowType();
 	}
+	else
+	{
+		const FStructProperty* StructProperty = CastField<FStructProperty>(PropertyChangedEvent.MemberProperty);
+		const bool bIsOrientation = StructProperty ? StructProperty->Struct == FDisplayClusterPositionalParams::StaticStruct() : false;
+	
+		if (bIsOrientation)
+		{
+			UpdateStageActorTransform();
+			// Updates MU in real-time. Skip our method as the positional coordinates are already correct.
+			AActor::PostEditMove(PropertyChangedEvent.ChangeType != EPropertyChangeType::Interactive);
+		}
+		else if (
+			PropertyName == USceneComponent::GetRelativeLocationPropertyName() ||
+			PropertyName == USceneComponent::GetRelativeRotationPropertyName() ||
+			PropertyName == USceneComponent::GetRelativeScale3DPropertyName())
+		{
+			bNotifyOnParamSetter = false;
+			UpdatePositionalParamsFromTransform();
+			bNotifyOnParamSetter = true;
+		}
+	}
+}
+
+void AColorCorrectWindow::PostEditMove(bool bFinished)
+{
+	Super::PostEditMove(bFinished);
+
+	bNotifyOnParamSetter = false;
+	UpdatePositionalParamsFromTransform();
+	bNotifyOnParamSetter = true;
 }
 #endif //WITH_EDITOR
 
@@ -113,4 +152,122 @@ void AColorCorrectWindow::SetMeshVisibilityForWindowType()
 			MeshComponents[TypeIndex]->SetVisibility(false, true);
 		}
 	}
-}
\ No newline at end of file
+}
+
+#define NOTIFY_PARAM_SETTER()\
+	if (bNotifyOnParamSetter)\
+	{\
+		UpdateStageActorTransform();\
+	}\
+
+void AColorCorrectWindow::SetLongitude(double InValue)
+{
+	PositionalParams.Longitude = InValue;
+	NOTIFY_PARAM_SETTER()
+}
+
+double AColorCorrectWindow::GetLongitude() const
+{
+	return PositionalParams.Longitude;
+}
+
+void AColorCorrectWindow::SetLatitude(double InValue)
+{
+	PositionalParams.Latitude = InValue;
+	NOTIFY_PARAM_SETTER()
+}
+
+double AColorCorrectWindow::GetLatitude() const
+{
+	return PositionalParams.Latitude;
+}
+
+void AColorCorrectWindow::SetDistanceFromCenter(double InValue)
+{
+	PositionalParams.DistanceFromCenter = InValue;
+	NOTIFY_PARAM_SETTER()
+}
+
+double AColorCorrectWindow::GetDistanceFromCenter() const
+{
+	return PositionalParams.DistanceFromCenter;
+}
+
+void AColorCorrectWindow::SetSpin(double InValue)
+{
+	PositionalParams.Spin = InValue;
+	NOTIFY_PARAM_SETTER()
+}
+
+double AColorCorrectWindow::GetSpin() const
+{
+	return PositionalParams.Spin;
+}
+
+void AColorCorrectWindow::SetPitch(double InValue)
+{
+	PositionalParams.Pitch = InValue;
+	NOTIFY_PARAM_SETTER()
+}
+
+double AColorCorrectWindow::GetPitch() const
+{
+	return PositionalParams.Pitch;
+}
+
+void AColorCorrectWindow::SetYaw(double InValue)
+{
+	PositionalParams.Yaw = InValue;
+	NOTIFY_PARAM_SETTER()
+}
+
+double AColorCorrectWindow::GetYaw() const
+{
+	return PositionalParams.Yaw;
+}
+
+void AColorCorrectWindow::SetRadialOffset(double InValue)
+{
+	PositionalParams.RadialOffset = InValue;
+	NOTIFY_PARAM_SETTER()
+}
+
+double AColorCorrectWindow::GetRadialOffset() const
+{
+	return PositionalParams.RadialOffset;
+}
+
+void AColorCorrectWindow::SetScale(const FVector2D& InScale)
+{
+	PositionalParams.Scale = InScale;
+	NOTIFY_PARAM_SETTER()
+}
+
+FVector2D AColorCorrectWindow::GetScale() const
+{
+	return PositionalParams.Scale;
+}
+
+void AColorCorrectWindow::SetOrigin(const FTransform& InOrigin)
+{
+	Origin = InOrigin;
+	NOTIFY_PARAM_SETTER()
+}
+
+FTransform AColorCorrectWindow::GetOrigin() const
+{
+	return Origin;
+}
+
+void AColorCorrectWindow::SetPositionalParams(const FDisplayClusterPositionalParams& InParams)
+{
+	PositionalParams = InParams;
+	NOTIFY_PARAM_SETTER()
+}
+
+FDisplayClusterPositionalParams AColorCorrectWindow::GetPositionalParams() const
+{
+	return PositionalParams;
+}
+
+#undef NOTIFY_PARAM_SETTER
diff --git a/Engine/Plugins/Experimental/ColorCorrectRegions/Source/ColorCorrectRegions/Public/ColorCorrectWindow.h b/Engine/Plugins/Experimental/ColorCorrectRegions/Source/ColorCorrectRegions/Public/ColorCorrectWindow.h
--- a/Engine/Plugins/Experimental/ColorCorrectRegions/Source/ColorCorrectRegions/Public/ColorCorrectWindow.h
+++ b/Engine/Plugins/Experimental/ColorCorrectRegions/Source/ColorCorrectRegions/Public/ColorCorrectWindow.h
@@ -8,6 +8,9 @@
 #include "GameFramework/Actor.h"
 #include "Engine/Classes/Components/MeshComponent.h"
 #include "Engine/Scene.h"
+
+#include "StageActor/IDisplayClusterStageActor.h"
+
 #include "ColorCorrectWindow.generated.h"
 
 
@@ -26,7 +29,7 @@ enum class EColorCorrectWindowType : uint8
  * 
  */
 UCLASS(Blueprintable, notplaceable)
-class COLORCORRECTREGIONS_API AColorCorrectWindow : public AColorCorrectRegion
+class COLORCORRECTREGIONS_API AColorCorrectWindow : public AColorCorrectRegion, public IDisplayClusterStageActor
 {
 	GENERATED_UCLASS_BODY()
 public:
@@ -41,6 +44,7 @@ public:
 #if WITH_EDITOR
 	/** Called when any of the properties are changed. */
 	virtual void PostEditChangeProperty(FPropertyChangedEvent& PropertyChangedEvent) override;
+	virtual void PostEditMove(bool bFinished) override;
 #endif
 
 private:
@@ -51,4 +55,67 @@ private:
 	/** Swaps meshes for different CCW. */
 	void SetMeshVisibilityForWindowType();
 
+public:
+
+	// ~Begin IDisplayClusterStageActor interface
+	UFUNCTION(BlueprintCallable, Category = Orientation)
+	virtual void SetLongitude(double InValue) override;
+	UFUNCTION(BlueprintCallable, Category = Orientation)
+	virtual double GetLongitude() const override;
+
+	UFUNCTION(BlueprintCallable, Category = Orientation)
+	virtual void SetLatitude(double InValue) override;
+	UFUNCTION(BlueprintCallable, Category = Orientation)
+	virtual double GetLatitude() const override;
+
+	UFUNCTION(BlueprintCallable, Category = Orientation)
+	virtual void SetDistanceFromCenter(double InValue) override;
+	UFUNCTION(BlueprintCallable, Category = Orientation)
+	virtual double GetDistanceFromCenter() const override;
+
+	UFUNCTION(BlueprintCallable, Category = Orientation)
+	virtual void SetSpin(double InValue) override;
+	UFUNCTION(BlueprintCallable, Category = Orientation)
+	virtual double GetSpin() const override;
+
+	UFUNCTION(BlueprintCallable, Category = Orientation)
+	virtual void SetPitch(double InValue) override;
+	UFUNCTION(BlueprintCallable, Category = Orientation)
+	virtual double GetPitch() const override;
+
+	UFUNCTION(BlueprintCallable, Category = Orientation)
+	virtual void SetYaw(double InValue) override;
+	UFUNCTION(BlueprintCallable, Category = Orientation)
+	virtual double GetYaw() const override;
+
+	UFUNCTION(BlueprintCallable, Category = Orientation)
+	virtual void SetRadialOffset(double InValue) override;
+	UFUNCTION(BlueprintCallable, Category = Orientation)
+	virtual double GetRadialOffset() const override;
+
+	UFUNCTION(BlueprintCallable, Category = Orientation)
+	virtual void SetScale(const FVector2D& InScale) override;
+	UFUNCTION(BlueprintCallable, Category = Orientation)
+	virtual FVector2D GetScale() const override;
+
+	UFUNCTION(BlueprintSetter)
+	virtual void SetOrigin(const FTransform& InOrigin) override;
+	UFUNCTION(BlueprintGetter)
+	virtual FTransform GetOrigin() const override;
+
+	UFUNCTION(BlueprintCallable, Category = Orientation)
+	virtual void SetPositionalParams(const FDisplayClusterPositionalParams& InParams) override;
+	UFUNCTION(BlueprintCallable, Category = Orientation)
+	virtual FDisplayClusterPositionalParams GetPositionalParams() const override;
+	// ~End IDisplayClusterStageActor interface
+	
+protected:
+	UPROPERTY(EditAnywhere, Category = Orientation, meta = (ShowOnlyInnerProperties))
+	FDisplayClusterPositionalParams PositionalParams;
+	
+	UPROPERTY(VisibleAnywhere, BlueprintSetter=SetOrigin, BlueprintGetter=GetOrigin, Category = Orientation)
+	FTransform Origin;
+
+	/** Update the transform when a positional setter is called */
+	bool bNotifyOnParamSetter = true;
 };
diff --git a/Engine/Plugins/Runtime/nDisplay/Source/DisplayCluster/Private/Game/EngineClasses/Scene/DisplayClusterLabelComponent.cpp b/Engine/Plugins/Runtime/nDisplay/Source/DisplayCluster/Private/Game/EngineClasses/Scene/DisplayClusterLabelComponent.cpp
--- a/Engine/Plugins/Runtime/nDisplay/Source/DisplayCluster/Private/Game/EngineClasses/Scene/DisplayClusterLabelComponent.cpp
+++ b/Engine/Plugins/Runtime/nDisplay/Source/DisplayCluster/Private/Game/EngineClasses/Scene/DisplayClusterLabelComponent.cpp
@@ -98,7 +98,7 @@ void UDisplayClusterLabelComponent::UpdateWidgetComponent()
 		const USceneComponent* Viewpoint = RootActor.IsValid() ? RootActor->GetCommonViewPoint() : nullptr;
 		if (Viewpoint != nullptr)
 		{
-			const FTransform LightCardTransform = LightCardActor->GetLightCardTransform();
+			const FTransform LightCardTransform = LightCardActor->GetStageActorTransform();
 			const FVector LightCardLocation = LightCardTransform.GetLocation();
 			const FVector DestinationPoint = Viewpoint->GetComponentLocation();
 
diff --git a/Engine/Plugins/Runtime/nDisplay/Source/DisplayCluster/Private/Game/EngineClasses/Scene/DisplayClusterLightCardActor.cpp b/Engine/Plugins/Runtime/nDisplay/Source/DisplayCluster/Private/Game/EngineClasses/Scene/DisplayClusterLightCardActor.cpp
--- a/Engine/Plugins/Runtime/nDisplay/Source/DisplayCluster/Private/Game/EngineClasses/Scene/DisplayClusterLightCardActor.cpp
+++ b/Engine/Plugins/Runtime/nDisplay/Source/DisplayCluster/Private/Game/EngineClasses/Scene/DisplayClusterLightCardActor.cpp
@@ -29,7 +29,6 @@ static FAutoConsoleVariableRef CVarDisplayClusterLightCardPolygonTextureSize(
 	ECVF_Default
 );
 
-const FRotator ADisplayClusterLightCardActor::PlaneMeshRotation = FRotator(0.0f, -90.0f, 90.0f);
 const float ADisplayClusterLightCardActor::UVPlaneDefaultSize = 200.0f;
 const float ADisplayClusterLightCardActor::UVPlaneDefaultDistance = 100.0f;
 
@@ -73,7 +72,7 @@ ADisplayClusterLightCardActor::ADisplayClusterLightCardActor(const FObjectInitia
 	LightCardComponent->Mobility = EComponentMobility::Movable;
 	LightCardComponent->SetStaticMesh(PlaneObj.Object);
 
-	UpdateLightCardTransform();
+	UpdateStageActorTransform();
 
 	LabelComponent = CreateOptionalDefaultSubobject<UDisplayClusterLabelComponent>(TEXT("Label"), true);
 	LabelComponent->AttachToComponent(LightCardComponent, FAttachmentTransformRules::KeepRelativeTransform);
@@ -109,30 +108,9 @@ void ADisplayClusterLightCardActor::Tick(float DeltaSeconds)
 {
 	Super::Tick(DeltaSeconds);
 
-	if (Longitude < 0 || Longitude > 360)
-	{
-		Longitude = FRotator::ClampAxis(Longitude);
-	}
-
-	if (Latitude < -90 || Latitude > 90)
-	{
-		// If latitude exceeds [-90, 90], mod it back into the appropriate range, and apply a shift of 180 degrees if
-		// needed to the longitude, to allow the latitude to be continuous (increasing latitude indefinitely should result in the LC 
-		// orbiting around a polar great circle)
-		double Parity = FMath::Fmod(FMath::Abs(Latitude) + 90, 360) - 180;
-		double DeltaLongitude = Parity > 1 ? 180.f : 0.f;
-
-		double LatMod = FMath::Fmod(Latitude + 90.f, 180.f);
-		if (LatMod < 0.f)
-		{
-			LatMod += 180.f;
-		}
+	ClampLatitudeAndLongitude(Latitude, Longitude);
 
-		Latitude = LatMod - 90;
-		Longitude = FRotator::ClampAxis(Longitude + DeltaLongitude);
-	}
-
-	UpdateLightCardTransform();
+	UpdateStageActorTransform();
 	UpdateLightCardMaterialInstance();
 	UpdateLightCardVisibility();
 }
@@ -152,7 +130,7 @@ void ADisplayClusterLightCardActor::PostEditChangeProperty(FPropertyChangedEvent
 		PropertyChangedEvent.Property->GetFName() == GET_MEMBER_NAME_CHECKED(ADisplayClusterLightCardActor, Scale) ||
 		PropertyChangedEvent.Property->GetFName() == GET_MEMBER_NAME_CHECKED(ADisplayClusterLightCardActor, bIsUVLightCard)))
 	{
-		UpdateLightCardTransform();
+		UpdateStageActorTransform();
 	}
 
 	if (PropertyChangedEvent.Property && (
@@ -196,74 +174,14 @@ FName ADisplayClusterLightCardActor::GetCustomIconName() const
 
 #endif
 
-FTransform ADisplayClusterLightCardActor::GetLightCardTransform(bool bIgnoreSpinYawPitch) const
-{
-	FTransform Transform;
-
-	Transform.SetLocation(LightCardComponent->GetComponentLocation());
-
-	FQuat LightCardOrientation;
-
-	if (!bIgnoreSpinYawPitch)
-	{
-		// Use the light card component's orientation, but remove the plane mesh rotation so that the returned transform's local x axis
-		// points radially inwards to match engine convention
-		LightCardOrientation = LightCardComponent->GetComponentQuat() * PlaneMeshRotation.Quaternion().Inverse();
-	}
-	else
-	{
-		LightCardOrientation = LightCardTransformerComponent->GetComponentQuat();
-	}
-
-	Transform.SetRotation(LightCardOrientation);
-
-	return Transform;
-}
-
-FBox ADisplayClusterLightCardActor::GetLightCardBounds(bool bLocalSpace) const
+UStaticMesh* ADisplayClusterLightCardActor::GetStaticMesh() const
 {
-	FTransform ObjectOrientedTransform;
-	ObjectOrientedTransform.SetRotation(PlaneMeshRotation.Quaternion());
-
-	if (!bLocalSpace)
-	{
-		ObjectOrientedTransform.SetTranslation(LightCardComponent->GetComponentLocation());
-		ObjectOrientedTransform.SetScale3D(LightCardComponent->GetComponentScale());
-	}
-
-	return LightCardComponent->CalcBounds(ObjectOrientedTransform).GetBox();
+	return LightCardComponent->GetStaticMesh();
 }
 
-void ADisplayClusterLightCardActor::UpdateLightCardTransform()
+void ADisplayClusterLightCardActor::SetStaticMesh(UStaticMesh* InStaticMesh)
 {
-	// If the light card is in UV space, set the spring arm's trasform to be zero, effectively removing it from the transform hierarchy
-	// This allows the light card to be positioned with the actor's cartesian coordinates instead of longitude and latitude
-	if (bIsUVLightCard)
-	{
-		MainSpringArmComponent->TargetArmLength = 0.0f;
-		MainSpringArmComponent->SetRelativeRotation(FRotator(0.0, 180.0, 0.0));
-
-		// Set world location and rotation such that the light card is always projected onto a YZ plane a distance of UVPlaneDefaultDistance from the world origin, facing in the -X direction
-		// This ensures that when the UV light cards are rendered to the light card map, they are always positioned and oriented correctly regardless of the stages location and rotation.
-		
-		// We place them slightly closer than UVPlaneDefaultDistance so that mouse clicks hit it first instead of DCRA meshes (avoids re-linetracing).
-		// UV LCs operate in orthographic projection so this should have no visual effect.
-		constexpr float DistanceFactor = 0.99f;
-
-		LightCardTransformerComponent->SetWorldLocation(FVector(DistanceFactor * UVPlaneDefaultDistance, -UVPlaneDefaultSize * (0.5 - UVCoordinates.X), UVPlaneDefaultSize * (0.5 - UVCoordinates.Y)));
-		LightCardTransformerComponent->SetWorldRotation(FVector(-1, 0, 0).Rotation());
-		LightCardTransformerComponent->SetWorldScale3D(FVector::OneVector);
-	}
-	else
-	{
-		MainSpringArmComponent->TargetArmLength = DistanceFromCenter + RadialOffset;
-		MainSpringArmComponent->SetRelativeRotation(FRotator(-Latitude, Longitude, 0.0));
-	}
-
-	FRotator LightCardOrientation = FRotator(-Pitch, Yaw, Spin);
-
-	LightCardComponent->SetRelativeRotation((LightCardOrientation.Quaternion() * PlaneMeshRotation.Quaternion()).Rotator());
-	LightCardComponent->SetRelativeScale3D(FVector(Scale, 1.f));
+	LightCardComponent->SetStaticMesh(InStaticMesh);
 }
 
 void ADisplayClusterLightCardActor::UpdateLightCardMaterialInstance()
@@ -280,7 +198,7 @@ void ADisplayClusterLightCardActor::UpdateLightCardMaterialInstance()
 		LightCardMaterialInstance->SetScalarParameterValue(TEXT("Tint"), Tint);
 		LightCardMaterialInstance->SetScalarParameterValue(TEXT("Exposure"), Exposure);
 		LightCardMaterialInstance->SetScalarParameterValue(TEXT("Gain"), Gain);
-		LightCardMaterialInstance->SetScalarParameterValue(TEXT("Opacity"), bIsProxy ? ProxyOpacity : Opacity);
+		LightCardMaterialInstance->SetScalarParameterValue(TEXT("Opacity"), IsProxy() ? ProxyOpacity : Opacity);
 		LightCardMaterialInstance->SetTextureParameterValue(TEXT("Texture"), Texture);
 		LightCardMaterialInstance->SetTextureParameterValue(TEXT("AlphaTexture"), PolygonMask);
 
@@ -408,7 +326,7 @@ void ADisplayClusterLightCardActor::UpdatePolygonTexture()
 
 void ADisplayClusterLightCardActor::UpdateLightCardVisibility()
 {
-	const bool bShouldBeVisible = !bIsUVLightCard || bIsProxy;
+	const bool bShouldBeVisible = !bIsUVLightCard || IsProxy();
 	const bool bIsVisible = !IsHidden();
 	if (bIsVisible != bShouldBeVisible)
 	{
@@ -461,40 +379,177 @@ void ADisplayClusterLightCardActor::SetRootActorOwner(ADisplayClusterRootActor*
 	RootActorOwner = MakeWeakObjectPtr(InRootActor);
 }
 
-UStaticMesh* ADisplayClusterLightCardActor::GetStaticMesh() const
+void ADisplayClusterLightCardActor::UpdateStageActorTransform()
 {
-	return LightCardComponent->GetStaticMesh();
+	// If the light card is in UV space, set the spring arm's trasform to be zero, effectively removing it from the transform hierarchy
+	// This allows the light card to be positioned with the actor's cartesian coordinates instead of longitude and latitude
+	if (bIsUVLightCard)
+	{
+		MainSpringArmComponent->TargetArmLength = 0.0f;
+		MainSpringArmComponent->SetRelativeRotation(FRotator(0.0, 180.0, 0.0));
+
+		// Set world location and rotation such that the light card is always projected onto a YZ plane a distance of UVPlaneDefaultDistance from the world origin, facing in the -X direction
+		// This ensures that when the UV light cards are rendered to the light card map, they are always positioned and oriented correctly regardless of the stages location and rotation.
+		
+		// We place them slightly closer than UVPlaneDefaultDistance so that mouse clicks hit it first instead of DCRA meshes (avoids re-linetracing).
+		// UV LCs operate in orthographic projection so this should have no visual effect.
+		constexpr float DistanceFactor = 0.99f;
+
+		LightCardTransformerComponent->SetWorldLocation(FVector(DistanceFactor * UVPlaneDefaultDistance, -UVPlaneDefaultSize * (0.5 - UVCoordinates.X), UVPlaneDefaultSize * (0.5 - UVCoordinates.Y)));
+		LightCardTransformerComponent->SetWorldRotation(FVector(-1, 0, 0).Rotation());
+		LightCardTransformerComponent->SetWorldScale3D(FVector::OneVector);
+	}
+	else
+	{
+		MainSpringArmComponent->TargetArmLength = DistanceFromCenter + RadialOffset;
+		MainSpringArmComponent->SetRelativeRotation(FRotator(-Latitude, Longitude, 0.0));
+	}
+
+	FRotator LightCardOrientation = FRotator(-Pitch, Yaw, Spin);
+
+	LightCardComponent->SetRelativeRotation((LightCardOrientation.Quaternion() * PlaneMeshRotation.Quaternion()).Rotator());
+	LightCardComponent->SetRelativeScale3D(FVector(Scale, 1.f));
 }
 
-void ADisplayClusterLightCardActor::SetStaticMesh(UStaticMesh* InStaticMesh)
+FTransform ADisplayClusterLightCardActor::GetStageActorTransform(bool bRemoveOrigin) const
 {
-	LightCardComponent->SetStaticMesh(InStaticMesh);
+	FTransform Transform;
+
+	FVector Position = LightCardComponent->GetComponentLocation();
+	if (bRemoveOrigin)
+	{
+		Position -= GetActorLocation();
+	}
+	Transform.SetLocation(MoveTemp(Position));
+
+	// Use the light card component's orientation, but remove the plane mesh rotation so that the returned transform's local x axis
+	// points radially inwards to match engine convention
+	const FQuat LightCardOrientation = LightCardComponent->GetComponentQuat() * PlaneMeshRotation.Quaternion().Inverse();
+
+	Transform.SetRotation(LightCardOrientation);
+
+	return Transform;
+}
+
+FBox ADisplayClusterLightCardActor::GetBoxBounds(bool bLocalSpace) const
+{
+	FTransform ObjectOrientedTransform;
+	ObjectOrientedTransform.SetRotation(PlaneMeshRotation.Quaternion());
+
+	if (!bLocalSpace)
+	{
+		ObjectOrientedTransform.SetTranslation(LightCardComponent->GetComponentLocation());
+		ObjectOrientedTransform.SetScale3D(LightCardComponent->GetComponentScale());
+	}
+
+	return LightCardComponent->CalcBounds(ObjectOrientedTransform).GetBox();
+}
+
+void ADisplayClusterLightCardActor::SetLongitude(double InValue)
+{
+	Longitude = InValue;
+}
+
+double ADisplayClusterLightCardActor::GetLongitude() const
+{
+	return Longitude;
 }
 
-ADisplayClusterLightCardActor::PositionalParams ADisplayClusterLightCardActor::GetPositionalParams()
+void ADisplayClusterLightCardActor::SetLatitude(double InValue)
 {
-	PositionalParams Params;
+	Latitude = InValue;
+}
 
-	Params.DistanceFromCenter = DistanceFromCenter;
-	Params.Latitude = Latitude;
-	Params.Longitude = Longitude;
-	Params.Pitch = Pitch;
-	Params.Spin = Spin;
-	Params.Yaw = Yaw;
-	Params.RadialOffset = RadialOffset;
+double ADisplayClusterLightCardActor::GetLatitude() const
+{
+	return Latitude;
+}
+
+void ADisplayClusterLightCardActor::SetDistanceFromCenter(double InValue)
+{
+	DistanceFromCenter = InValue;
+}
 
-	return Params;
+double ADisplayClusterLightCardActor::GetDistanceFromCenter() const
+{
+	return DistanceFromCenter;
+}
+
+void ADisplayClusterLightCardActor::SetSpin(double InValue)
+{
+	Spin = InValue;
+}
+
+void ADisplayClusterLightCardActor::SetPitch(double InValue)
+{
+	Pitch = InValue;
+}
+
+double ADisplayClusterLightCardActor::GetSpin() const
+{
+	return Spin;
+}
+
+double ADisplayClusterLightCardActor::GetPitch() const
+{
+	return Pitch;
+}
+
+void ADisplayClusterLightCardActor::SetYaw(double InValue)
+{
+	Yaw = InValue;
+}
+
+double ADisplayClusterLightCardActor::GetYaw() const
+{
+	return Yaw;
+}
+
+void ADisplayClusterLightCardActor::SetRadialOffset(double InValue)
+{
+	RadialOffset = InValue;
+}
+
+double ADisplayClusterLightCardActor::GetRadialOffset() const
+{
+	return RadialOffset;
+}
+
+bool ADisplayClusterLightCardActor::IsUVActor() const
+{
+	return bIsUVLightCard;
+}
+
+void ADisplayClusterLightCardActor::SetOrigin(const FTransform& InOrigin)
+{
+	SetActorLocation(InOrigin.GetLocation());
+	SetActorRotation(InOrigin.GetRotation());
+	// Scale not currently used for origin
+}
+
+FTransform ADisplayClusterLightCardActor::GetOrigin() const
+{
+	return {GetActorRotation(), GetActorLocation(), FVector::One()};
+}
+
+void ADisplayClusterLightCardActor::SetScale(const FVector2D& InScale)
+{
+	Scale = InScale;
+}
+
+FVector2D ADisplayClusterLightCardActor::GetScale() const
+{
+	return Scale;
+}
+
+void ADisplayClusterLightCardActor::SetUVCoordinates(const FVector2D& InUVCoordinates)
+{
+	UVCoordinates = InUVCoordinates;
 }
 
-void ADisplayClusterLightCardActor::SetPositionalParams(const PositionalParams& Params)
+FVector2D ADisplayClusterLightCardActor::GetUVCoordinates() const
 {
-	DistanceFromCenter = Params.DistanceFromCenter;
-	Latitude = Params.Latitude;
-	Longitude = Params.Longitude;
-	Pitch = Params.Pitch;
-	Spin = Params.Spin;
-	Yaw = Params.Yaw;
-	RadialOffset = Params.RadialOffset;
+	return UVCoordinates;
 }
 
 void ADisplayClusterLightCardActor::CreateComponentsForExtenders()
diff --git a/Engine/Plugins/Runtime/nDisplay/Source/DisplayCluster/Public/DisplayClusterLightCardActor.h b/Engine/Plugins/Runtime/nDisplay/Source/DisplayCluster/Public/DisplayClusterLightCardActor.h
--- a/Engine/Plugins/Runtime/nDisplay/Source/DisplayCluster/Public/DisplayClusterLightCardActor.h
+++ b/Engine/Plugins/Runtime/nDisplay/Source/DisplayCluster/Public/DisplayClusterLightCardActor.h
@@ -2,7 +2,8 @@
 
 #pragma once
 
-#include "CoreMinimal.h"
+#include "StageActor/IDisplayClusterStageActor.h"
+
 #include "GameFramework/Actor.h"
 
 #include "DisplayClusterLightCardActor.generated.h"
@@ -48,27 +49,11 @@ struct FLightCardAlphaGradientSettings
 };
 
 UCLASS(Blueprintable, DisplayName = "Light Card")
-class DISPLAYCLUSTER_API ADisplayClusterLightCardActor : public AActor
+class DISPLAYCLUSTER_API ADisplayClusterLightCardActor : public AActor, public IDisplayClusterStageActor
 {
 	GENERATED_BODY()
 
 public:
-
-	struct PositionalParams
-	{
-		double DistanceFromCenter;
-		double Longitude;
-		double Latitude;
-		double Spin;
-		double Pitch;
-		double Yaw;
-		double RadialOffset;
-	};
-
-public:
-	/** The rotation used to orient the plane mesh used for the light card so that its normal points radially inwards */
-	static const FRotator PlaneMeshRotation;
-
 	/** The default size of the porjection plane UV light cards are rendered to */
 	static const float UVPlaneDefaultSize;
 
@@ -88,15 +73,6 @@ public:
 	virtual FName GetCustomIconName() const override;
 #endif
 
-	/**
-	 * Gets the transform in world space of the light card component
-	 * @param bIgnoreSpinYawPitch - If the light card component's spin, yaw, and pitch should be ignored when computing the transform
-	 */
-	FTransform GetLightCardTransform(bool bIgnoreSpinYawPitch = false) const;
-
-	/** Gets the object oriented bounding box of the light card component */
-	FBox GetLightCardBounds(bool bLocalSpace = false) const;
-
 	/** Gets the light card mesh component */
 	UStaticMeshComponent* GetLightCardMeshComponent() const { return LightCardComponent.Get(); }
 
@@ -106,15 +82,6 @@ public:
 	/** Sets a new static mesh for the light card */
 	void SetStaticMesh(UStaticMesh* InStaticMesh);
 
-	/** Updates the Light Card transform based on its positional properties (Lat, Long, etc.) */
-	void UpdateLightCardTransform();
-
-	/** Retrieves positional parameters */
-	PositionalParams GetPositionalParams();
-
-	/** Set positional parameters */
-	void SetPositionalParams(const PositionalParams& Params);
-
 	/** Updates the card's material instance parameters */
 	void UpdateLightCardMaterialInstance();
 
@@ -132,6 +99,51 @@ public:
 	
 	/** Return the current owner, providing one was set */
 	TWeakObjectPtr<ADisplayClusterRootActor> GetRootActorOwner() const { return RootActorOwner; }
+
+	// ~Begin IDisplayClusterStageActor interface
+	/** Updates the Light Card transform based on its positional properties (Lat, Long, etc.) */
+	virtual void UpdateStageActorTransform() override;
+	
+	/**
+	 * Gets the transform in world space of the light card component
+	 */
+	virtual FTransform GetStageActorTransform(bool bRemoveOrigin = false) const override;
+
+	/** Gets the object oriented bounding box of the light card component */
+	virtual FBox GetBoxBounds(bool bLocalSpace = false) const override;
+	
+	virtual void SetLongitude(double InValue) override;
+	virtual double GetLongitude() const override;
+
+	virtual void SetLatitude(double InValue) override;
+	virtual double GetLatitude() const override;
+
+	virtual void SetDistanceFromCenter(double InValue) override;
+	virtual double GetDistanceFromCenter() const override;
+
+	virtual void SetSpin(double InValue) override;
+	virtual double GetSpin() const override;
+
+	virtual void SetPitch(double InValue) override;
+	virtual double GetPitch() const override;
+
+	virtual void SetYaw(double InValue) override;
+	virtual double GetYaw() const override;
+
+	virtual void SetRadialOffset(double InValue) override;
+	virtual double GetRadialOffset() const override;
+
+	virtual bool IsUVActor() const override;
+
+	virtual void SetOrigin(const FTransform& InOrigin) override;
+	virtual FTransform GetOrigin() const override;
+
+	virtual void SetScale(const FVector2D& InScale) override;
+	virtual FVector2D GetScale() const override;
+
+	virtual void SetUVCoordinates(const FVector2D& InUVCoordinates) override;
+	virtual FVector2D GetUVCoordinates() const override;
+	// ~End IDisplayClusterStageActor interface
 	
 public:
 
@@ -212,10 +224,6 @@ public:
 	UPROPERTY()
 	bool bIsUVLightCard = false;
 
-	/** Used to flag this light card as a proxy of a "real" light card. Used by the LightCard Editor */
-	UPROPERTY(Transient)
-	bool bIsProxy = false;
-
 	/** Polygon points when using this type of mask */
 	UPROPERTY(EditAnywhere, BlueprintReadWrite, Category = "Appearance")
 	TArray<FVector2D> Polygon;
diff --git a/Engine/Plugins/Runtime/nDisplay/Source/DisplayClusterLightCardEditor/DisplayClusterLightCardEditor.Build.cs b/Engine/Plugins/Runtime/nDisplay/Source/DisplayClusterLightCardEditor/DisplayClusterLightCardEditor.Build.cs
--- a/Engine/Plugins/Runtime/nDisplay/Source/DisplayClusterLightCardEditor/DisplayClusterLightCardEditor.Build.cs
+++ b/Engine/Plugins/Runtime/nDisplay/Source/DisplayClusterLightCardEditor/DisplayClusterLightCardEditor.Build.cs
@@ -31,13 +31,13 @@ public class DisplayClusterLightCardEditor : ModuleRules
 				"RenderCore",
 				"Renderer",
 				"RHI",
+				"SceneOutliner",
 				"Slate",
 				"SlateCore",
 				"UnrealEd",
-				"WorkspaceMenuStructure",
-				"ToolWidgets",
 				"ToolMenus",
-				"SceneOutliner"
+				"ToolWidgets",
+				"WorkspaceMenuStructure"
 			});
 	}
 }
diff --git a/Engine/Plugins/Runtime/nDisplay/Source/DisplayClusterLightCardEditor/Private/DisplayClusterLightCardEditorCommands.cpp b/Engine/Plugins/Runtime/nDisplay/Source/DisplayClusterLightCardEditor/Private/DisplayClusterLightCardEditorCommands.cpp
--- a/Engine/Plugins/Runtime/nDisplay/Source/DisplayClusterLightCardEditor/Private/DisplayClusterLightCardEditorCommands.cpp
+++ b/Engine/Plugins/Runtime/nDisplay/Source/DisplayClusterLightCardEditor/Private/DisplayClusterLightCardEditorCommands.cpp
@@ -42,6 +42,8 @@ void FDisplayClusterLightCardEditorCommands::RegisterCommands()
 
 	UI_COMMAND(ToggleAllLabels, "Labels", "Display labels in the preview and wall", EUserInterfaceActionType::ToggleButton, FInputChord());
 	UI_COMMAND(ToggleLightCardLabels, "Light Card Labels", "Display labels on light cards in the preview and wall", EUserInterfaceActionType::ToggleButton, FInputChord());
+
+	UI_COMMAND(ToggleIconVisibility, "Icons", "Display icons in the light card editor where applicable", EUserInterfaceActionType::ToggleButton, FInputChord());
 }
 
 #undef LOCTEXT_NAMESPACE
diff --git a/Engine/Plugins/Runtime/nDisplay/Source/DisplayClusterLightCardEditor/Private/DisplayClusterLightCardEditorCommands.h b/Engine/Plugins/Runtime/nDisplay/Source/DisplayClusterLightCardEditor/Private/DisplayClusterLightCardEditorCommands.h
--- a/Engine/Plugins/Runtime/nDisplay/Source/DisplayClusterLightCardEditor/Private/DisplayClusterLightCardEditorCommands.h
+++ b/Engine/Plugins/Runtime/nDisplay/Source/DisplayClusterLightCardEditor/Private/DisplayClusterLightCardEditorCommands.h
@@ -53,4 +53,6 @@ public:
 
 	TSharedPtr<FUICommandInfo> ToggleAllLabels;
 	TSharedPtr<FUICommandInfo> ToggleLightCardLabels;
+
+	TSharedPtr<FUICommandInfo> ToggleIconVisibility;
 };
diff --git a/Engine/Plugins/Runtime/nDisplay/Source/DisplayClusterLightCardEditor/Private/DisplayClusterLightCardEditorUtils.cpp b/Engine/Plugins/Runtime/nDisplay/Source/DisplayClusterLightCardEditor/Private/DisplayClusterLightCardEditorUtils.cpp
new file mode 100644
--- /dev/null
+++ b/Engine/Plugins/Runtime/nDisplay/Source/DisplayClusterLightCardEditor/Private/DisplayClusterLightCardEditorUtils.cpp
@@ -0,0 +1,73 @@
+// Copyright Epic Games, Inc. All Rights Reserved.
+
+#include "DisplayClusterLightCardEditorUtils.h"
+
+#include "DisplayClusterLightCardActor.h"
+
+#include "AssetRegistry/AssetRegistryModule.h"
+#include "AssetRegistry/IAssetRegistry.h"
+#include "Engine/Blueprint.h"
+#include "UObject/UObjectIterator.h"
+
+bool UE::DisplayClusterLightCardEditorUtils::IsManagedActor(const AActor* InActor)
+{
+	return InActor && InActor->Implements<UDisplayClusterStageActor>();
+}
+
+bool UE::DisplayClusterLightCardEditorUtils::IsProxySelectable(const AActor* InActor)
+{
+	// Currently managed actors & selectable proxies are 1 - 1
+	return IsManagedActor(InActor);
+}
+
+TSet<UClass*> UE::DisplayClusterLightCardEditorUtils::GetAllStageActorClasses()
+{
+	TSet<UClass*> OutClasses;
+
+	auto ShouldAddNativeClass = [] (const UClass* Class) -> bool
+	{
+		return Class->ImplementsInterface(UDisplayClusterStageActor::StaticClass()) && Class->IsNative() &&
+			!Class->HasAnyClassFlags(CLASS_Deprecated | CLASS_NewerVersionExists) &&
+			!Class->GetName().StartsWith(TEXT("SKEL_")) && !Class->GetName().StartsWith(TEXT("REINST_"));
+	};
+
+	// First search native classes
+	for (TObjectIterator<UClass> ClassIt; ClassIt; ++ClassIt)
+	{
+		UClass* Class = *ClassIt;
+
+		if (ShouldAddNativeClass(Class))
+		{
+			OutClasses.Add(Class);
+		}
+	}
+
+	// Next search blueprint classes inheriting the stage actor interface
+	const IAssetRegistry& AssetRegistry = FModuleManager::LoadModuleChecked<FAssetRegistryModule>(FName("AssetRegistry")).Get();
+
+	TSet<FTopLevelAssetPath> InheritedInterfacePaths;
+	AssetRegistry.GetDerivedClassNames({UDisplayClusterStageActor::StaticClass()->GetClassPathName()},
+		{}, InheritedInterfacePaths);
+
+	TArray<FAssetData> OutAssets;
+	AssetRegistry.GetAssetsByClass(UBlueprint::StaticClass()->GetClassPathName(), OutAssets, true);
+
+	for (const FAssetData& Asset : OutAssets)
+	{
+		FAssetDataTagMapSharedView::FFindTagResult Result = Asset.TagsAndValues.FindTag(TEXT("GeneratedClass"));
+		if (Result.IsSet())
+		{
+			const FString& GeneratedClassPathPtr = Result.GetValue();
+			const FTopLevelAssetPath ClassObjectPath(FPackageName::ExportTextPathToObjectPath(*GeneratedClassPathPtr));
+			if (InheritedInterfacePaths.Contains(ClassObjectPath))
+			{
+				TSoftObjectPtr<UClass> SoftObjectPtr(ClassObjectPath.ToString());
+				if (UClass* Class = SoftObjectPtr.LoadSynchronous())
+				{
+					OutClasses.Add(Class);
+				}
+			}
+		}
+	}
+	return OutClasses;
+}
diff --git a/Engine/Plugins/Runtime/nDisplay/Source/DisplayClusterLightCardEditor/Private/DisplayClusterLightCardEditorUtils.h b/Engine/Plugins/Runtime/nDisplay/Source/DisplayClusterLightCardEditor/Private/DisplayClusterLightCardEditorUtils.h
new file mode 100644
--- /dev/null
+++ b/Engine/Plugins/Runtime/nDisplay/Source/DisplayClusterLightCardEditor/Private/DisplayClusterLightCardEditorUtils.h
@@ -0,0 +1,19 @@
+// Copyright Epic Games, Inc. All Rights Reserved.
+
+#pragma once
+
+#include "CoreMinimal.h"
+
+class AActor;
+
+namespace UE::DisplayClusterLightCardEditorUtils
+{
+	/** Check if the given actor is of a type that should be managed by the light card editor */
+	bool IsManagedActor(const AActor* InActor);
+	
+	/** Check if the given actor is of a type that is selectable in the light card editor */
+	bool IsProxySelectable(const AActor* InActor);
+
+	/** Discover all classes implementing or inheriting the stage actor interface */
+	TSet<UClass*> GetAllStageActorClasses();
+}
\ No newline at end of file
diff --git a/Engine/Plugins/Runtime/nDisplay/Source/DisplayClusterLightCardEditor/Private/Outliner/SDisplayClusterLightCardOutliner.cpp b/Engine/Plugins/Runtime/nDisplay/Source/DisplayClusterLightCardEditor/Private/Outliner/SDisplayClusterLightCardOutliner.cpp
--- a/Engine/Plugins/Runtime/nDisplay/Source/DisplayClusterLightCardEditor/Private/Outliner/SDisplayClusterLightCardOutliner.cpp
+++ b/Engine/Plugins/Runtime/nDisplay/Source/DisplayClusterLightCardEditor/Private/Outliner/SDisplayClusterLightCardOutliner.cpp
@@ -13,6 +13,8 @@
 #include "DisplayClusterLightCardActor.h"
 #include "DisplayClusterRootActor.h"
 
+#include "StageActor/IDisplayClusterStageActor.h"
+
 #include "ActorTreeItem.h"
 #include "ClassIconFinder.h"
 #include "FolderTreeItem.h"
@@ -45,68 +47,60 @@ void SDisplayClusterLightCardOutliner::Construct(const FArguments& InArgs, TShar
 
 void SDisplayClusterLightCardOutliner::SetRootActor(ADisplayClusterRootActor* NewRootActor)
 {
-	TArray<ADisplayClusterLightCardActor*> PreviouslySelectedLightCards;
-	GetSelectedLightCards(PreviouslySelectedLightCards);
+	TArray<AActor*> PreviouslySelectedActors;
+	GetSelectedActors(PreviouslySelectedActors);
 	
 	RootActor = NewRootActor;
-	FillLightCardList();
+	FillActorList();
 
 	CreateWorldOutliner();
 	
-	for (const TSharedPtr<FLightCardTreeItem>& LightCardTreeItem : LightCardTree)
-	{
-		if (LightCardTreeItem->IsActorLayer())
-		{
-			LightCardTreeView->SetItemExpansion(LightCardTreeItem, true);
-		}
-	}
-	
 	// Select previously selected light cards. This fixes an issue where the details panel may clear in some cases
 	// and also supports maintaining the selection of shared light cards across different root actors.
-	SelectLightCards(MoveTemp(PreviouslySelectedLightCards));
+	SelectActors(MoveTemp(PreviouslySelectedActors));
 }
 
-void SDisplayClusterLightCardOutliner::GetSelectedLightCards(TArray<ADisplayClusterLightCardActor*>& OutSelectedLightCards) const
+void SDisplayClusterLightCardOutliner::GetSelectedActors(TArray<AActor*>& OutSelectedActors) const
 {
-	TArray<FSceneOutlinerTreeItemPtr> SelectedOutlinerLightCards = SceneOutliner->GetSelectedItems();
-	for (const FSceneOutlinerTreeItemPtr& SelectedLightCard : SelectedOutlinerLightCards)
+	TArray<FSceneOutlinerTreeItemPtr> SelectedOutlinerActors = SceneOutliner->GetSelectedItems();
+	for (const FSceneOutlinerTreeItemPtr& SelectedLightCard : SelectedOutlinerActors)
 	{
-		TSharedPtr<FLightCardTreeItem> MatchingLightCard = GetLightCardTreeItemFromOutliner(SelectedLightCard);
-		if (MatchingLightCard.IsValid() && MatchingLightCard->LightCardActor.IsValid())
+		TSharedPtr<FStageActorTreeItem> MatchingTreeItem = GetStageActorTreeItemFromOutliner(SelectedLightCard);
+		if (MatchingTreeItem.IsValid() && MatchingTreeItem->Actor.IsValid())
 		{
-			OutSelectedLightCards.Add(MatchingLightCard->LightCardActor.Get());
+			OutSelectedActors.Add(MatchingTreeItem->Actor.Get());
 		}
 	}
 }
 
-void SDisplayClusterLightCardOutliner::SelectLightCards(const TArray<ADisplayClusterLightCardActor*>& LightCardsToSelect)
+void SDisplayClusterLightCardOutliner::SelectActors(const TArray<AActor*>& ActorsToSelect)
 {
-	TArray<ADisplayClusterLightCardActor*> ActualSelectedLightCards;
+	TArray<AActor*> ActualSelectedActors;
 	TArray<FSceneOutlinerTreeItemPtr> SelectedTreeItems;
 	CachedSelectedActors.Reset();
 	
-	for (const TSharedPtr<FLightCardTreeItem>& TreeItem : LightCardActors)
+	for (const TSharedPtr<FStageActorTreeItem>& TreeItem : StageActorTreeItems)
 	{
-		if (LightCardsToSelect.Contains(TreeItem->LightCardActor))
+		if (ActorsToSelect.Contains(TreeItem->Actor))
 		{
-			if (FSceneOutlinerTreeItemPtr OutlinerItem = SceneOutliner->GetTreeItem(TreeItem->LightCardActor.Get()))
+			if (FSceneOutlinerTreeItemPtr OutlinerItem = SceneOutliner->GetTreeItem(TreeItem->Actor.Get()))
 			{
 				SelectedTreeItems.Add(OutlinerItem);
-				ActualSelectedLightCards.Add(TreeItem->LightCardActor.Get());
-				CachedSelectedActors.Add(TreeItem->LightCardActor);
+				ActualSelectedActors.Add(TreeItem->Actor.Get());
+				CachedSelectedActors.Add(TreeItem->Actor);
 			}
 		}
 	}
 	
 	SceneOutliner->SetItemSelection(SelectedTreeItems, true);
 	
-	IDisplayClusterOperator::Get().ShowDetailsForObjects(*reinterpret_cast<TArray<UObject*>*>(&ActualSelectedLightCards));
+	IDisplayClusterOperator::Get().ShowDetailsForObjects(*reinterpret_cast<TArray<UObject*>*>(&ActualSelectedActors));
 }
 
 void SDisplayClusterLightCardOutliner::RestoreCachedSelection()
 {
-	TArray<ADisplayClusterLightCardActor*> ActorsToSelect;
-	for (const TWeakObjectPtr<ADisplayClusterLightCardActor>& CachedActor : CachedSelectedActors)
+	TArray<AActor*> ActorsToSelect;
+	for (const TWeakObjectPtr<AActor>& CachedActor : CachedSelectedActors)
 	{
 		if (CachedActor.IsValid())
 		{
@@ -114,7 +108,7 @@ void SDisplayClusterLightCardOutliner::RestoreCachedSelection()
 		}
 	}
 	
-	SelectLightCards(MoveTemp(ActorsToSelect));
+	SelectActors(MoveTemp(ActorsToSelect));
 }
 
 void SDisplayClusterLightCardOutliner::CreateWorldOutliner()
@@ -196,82 +190,45 @@ void SDisplayClusterLightCardOutliner::CreateWorldOutliner()
 	];
 }
 
-bool SDisplayClusterLightCardOutliner::FillLightCardList()
+bool SDisplayClusterLightCardOutliner::FillActorList()
 {
-	TArray<TSharedPtr<FLightCardTreeItem>> OriginalTree = MoveTemp(LightCardTree);
-	LightCardActors.Empty();
+	TArray<TSharedPtr<FStageActorTreeItem>> OriginalTree = MoveTemp(StageActorTree);
+	StageActorTreeItems.Empty();
 
-	if (RootActor.IsValid())
+	if (LightCardEditorPtr.IsValid())
 	{
-		FDisplayClusterConfigurationICVFX_VisibilityList& RootActorLightCards = RootActor->GetConfigData()->StageSettings.Lightcard.ShowOnlyList;
-
-		for (const FActorLayer& ActorLayer : RootActorLightCards.ActorLayers)
-		{
-			TSharedPtr<FLightCardTreeItem> LightCardTreeItem = MakeShared<FLightCardTreeItem>();
-			LightCardTreeItem->ActorLayer = ActorLayer.Name;
-			LightCardTree.Add(LightCardTreeItem);
-		}
-
-		for (const TSoftObjectPtr<AActor>& LightCardActor : RootActorLightCards.Actors)
-		{
-			if (LightCardActor.IsValid() && LightCardActor->IsA<ADisplayClusterLightCardActor>())
-			{
-				TSharedPtr<FLightCardTreeItem> LightCardTreeItem = MakeShared<FLightCardTreeItem>();
-				LightCardTreeItem->LightCardActor = CastChecked<ADisplayClusterLightCardActor>(LightCardActor.Get());
-				LightCardTree.Add(LightCardTreeItem);
-				LightCardActors.Add(LightCardTreeItem);
-			}
-		}
-
-		// If there are any layers that are specified as light card layers, iterate over all actors in the world and 
-		// add any that are members of any of the light card layers to the list. Only add an actor once, even if it is
-		// in multiple layers
-		if (RootActorLightCards.ActorLayers.Num())
+		TArray<AActor*> ManagedActors = LightCardEditorPtr.Pin()->FindAllManagedActors();
+		StageActorTreeItems.Reserve(ManagedActors.Num());
+		
+		for (AActor* Actor : ManagedActors)
 		{
-			if (UWorld* World = RootActor->GetWorld())
-			{
-				for (const TWeakObjectPtr<ADisplayClusterLightCardActor> WeakActor : TActorRange<ADisplayClusterLightCardActor>(World))
-				{
-					if (WeakActor.IsValid())
-					{
-						for (const FActorLayer& ActorLayer : RootActorLightCards.ActorLayers)
-						{
-							if (WeakActor->Layers.Contains(ActorLayer.Name))
-							{
-								TSharedPtr<FLightCardTreeItem> LightCardReference = MakeShared<FLightCardTreeItem>();
-								LightCardReference->LightCardActor = WeakActor.Get();
-								LightCardReference->ActorLayer = ActorLayer.Name;
-								LightCardActors.Add(LightCardReference);
-								break;
-							}
-						}
-					}
-				}
-			}
+			const TSharedPtr<FStageActorTreeItem> TreeItem = MakeShared<FStageActorTreeItem>();
+			TreeItem->Actor = Actor;
+			StageActorTree.Add(TreeItem);
+			StageActorTreeItems.Add(TreeItem);
 		}
 	}
 
-	TrackedActors.Empty(LightCardActors.Num());
+	TrackedActors.Empty(StageActorTreeItems.Num());
 
-	for (const TSharedPtr<FLightCardTreeItem>& LightCard : LightCardActors)
+	for (const TSharedPtr<FStageActorTreeItem>& LightCard : StageActorTreeItems)
 	{
-		TrackedActors.Add(LightCard->LightCardActor.Get(), LightCard);
+		TrackedActors.Add(LightCard->Actor.Get(), LightCard);
 	}
 	
 	// Check if the tree items have changed.
 	{
-		if (OriginalTree.Num() != LightCardTree.Num())
+		if (OriginalTree.Num() != StageActorTree.Num())
 		{
 			return true;
 		}
 	
-		for (const TSharedPtr<FLightCardTreeItem>& OriginalTreeItem : OriginalTree)
+		for (const TSharedPtr<FStageActorTreeItem>& OriginalTreeItem : OriginalTree)
 		{
-			if (!LightCardTree.ContainsByPredicate([OriginalTreeItem](const TSharedPtr<FLightCardTreeItem>& LightCardTreeItem)
+			if (!StageActorTree.ContainsByPredicate([OriginalTreeItem](const TSharedPtr<FStageActorTreeItem>& LightCardTreeItem)
 			{
 				return OriginalTreeItem.IsValid() && LightCardTreeItem.IsValid() &&
-					OriginalTreeItem->LightCardActor.Get() == LightCardTreeItem->LightCardActor.Get() &&
-						OriginalTreeItem->ActorLayer == LightCardTreeItem->ActorLayer;
+					OriginalTreeItem->Actor.Get() == LightCardTreeItem->Actor.Get();
 			}))
 			{
 				return true;
@@ -287,23 +244,23 @@ void SDisplayClusterLightCardOutliner::OnOutlinerSelectionChanged(FSceneOutliner
 {
 	MostRecentSelectedItem = TreeItem;
 
-	TArray<ADisplayClusterLightCardActor*> SelectedLightCards;
-	GetSelectedLightCards(SelectedLightCards);
+	TArray<AActor*> SelectedLightCards;
+	GetSelectedActors(SelectedLightCards);
 
 	IDisplayClusterOperator::Get().ShowDetailsForObjects(*reinterpret_cast<TArray<UObject*>*>(&SelectedLightCards));
 
 	if (LightCardEditorPtr.IsValid())
 	{
-		LightCardEditorPtr.Pin()->SelectLightCardProxies(SelectedLightCards);
+		LightCardEditorPtr.Pin()->SelectActorProxies(SelectedLightCards);
 	}
 }
 
-TSharedPtr<SDisplayClusterLightCardOutliner::FLightCardTreeItem> SDisplayClusterLightCardOutliner::
-GetLightCardTreeItemFromOutliner(FSceneOutlinerTreeItemPtr InOutlinerTreeItem) const
+TSharedPtr<SDisplayClusterLightCardOutliner::FStageActorTreeItem> SDisplayClusterLightCardOutliner::
+GetStageActorTreeItemFromOutliner(FSceneOutlinerTreeItemPtr InOutlinerTreeItem) const
 {
 	if (const FActorTreeItem* ActorTreeItem = InOutlinerTreeItem->CastTo<FActorTreeItem>())
 	{
-		if (const TSharedPtr<FLightCardTreeItem>* MatchingLightCard = TrackedActors.Find(
+		if (const TSharedPtr<FStageActorTreeItem>* MatchingLightCard = TrackedActors.Find(
 			ActorTreeItem->Actor.Get()))
 		{
 			return *MatchingLightCard;
@@ -366,15 +323,15 @@ void SDisplayClusterLightCardOutliner::RegisterContextMenu(FName& InName, FToolM
 			{
 				// If a delete operation was undone the most recent selection will be invalid.
 				// Assume a single entry light card is correct.
-				TArray<ADisplayClusterLightCardActor*> SelectedLightCardActors;
-				GetSelectedLightCards(SelectedLightCardActors);
+				TArray<AActor*> SelectedLightCardActors;
+				GetSelectedActors(SelectedLightCardActors);
 				if (SelectedLightCardActors.Num() == 1)
 				{
 					Actor = SelectedLightCardActors[0];
 				}
 			}
 			
-			bool bAddFullEditMenu = false;
+			const bool bAddFullEditMenu = Actor && Actor->Implements<UDisplayClusterStageActor>();
 			
 			if (ADisplayClusterLightCardActor* LightCardActor = Cast<ADisplayClusterLightCardActor>(Actor))
 			{
@@ -382,7 +339,6 @@ void SDisplayClusterLightCardOutliner::RegisterContextMenu(FName& InName, FToolM
 
 				Section.AddMenuEntry(FDisplayClusterLightCardEditorCommands::Get().RemoveLightCard);
 				Section.AddMenuEntry(FDisplayClusterLightCardEditorCommands::Get().SaveLightCardTemplate);
-				bAddFullEditMenu = true;
 			}
 			
 			FToolMenuSection& Section = InMenu->AddSection("LightCardEditSection", LOCTEXT("LightCardEditSectionName", "Edit"));
diff --git a/Engine/Plugins/Runtime/nDisplay/Source/DisplayClusterLightCardEditor/Private/Outliner/SDisplayClusterLightCardOutliner.h b/Engine/Plugins/Runtime/nDisplay/Source/DisplayClusterLightCardEditor/Private/Outliner/SDisplayClusterLightCardOutliner.h
--- a/Engine/Plugins/Runtime/nDisplay/Source/DisplayClusterLightCardEditor/Private/Outliner/SDisplayClusterLightCardOutliner.h
+++ b/Engine/Plugins/Runtime/nDisplay/Source/DisplayClusterLightCardEditor/Private/Outliner/SDisplayClusterLightCardOutliner.h
@@ -2,17 +2,16 @@
 
 #pragma once
 
-#include "CoreMinimal.h"
 #include "SceneOutlinerFwd.h"
 #include "Widgets/DeclarativeSyntaxSupport.h"
 #include "Widgets/SCompoundWidget.h"
 
-class FExtender;
-class FUICommandList;
 class ADisplayClusterRootActor;
 class ADisplayClusterLightCardActor;
-class SDisplayClusterLightCardEditor;
 class ITableRow;
+class FExtender;
+class FUICommandList;
+class SDisplayClusterLightCardEditor;
 class STableViewBase;
 
 class SSceneOutliner;
@@ -23,18 +22,13 @@ class STreeView;
 class SDisplayClusterLightCardOutliner : public SCompoundWidget
 {
 public:
-	struct FLightCardTreeItem
+	struct FStageActorTreeItem
 	{
-		TWeakObjectPtr<ADisplayClusterLightCardActor> LightCardActor;
-		FName ActorLayer;
+		TWeakObjectPtr<AActor> Actor;
 
-		FLightCardTreeItem() :
-			LightCardActor(nullptr),
-			ActorLayer(NAME_None)
+		FStageActorTreeItem() :
+			Actor(nullptr)
 		{ }
-
-		bool IsInActorLayer() { return !ActorLayer.IsNone() && LightCardActor.IsValid(); }
-		bool IsActorLayer() { return !ActorLayer.IsNone() && !LightCardActor.IsValid(); }
 	};
 
 public:
@@ -48,13 +42,30 @@ public:
 
 	void SetRootActor(ADisplayClusterRootActor* NewRootActor);
 
-	const TArray<TSharedPtr<FLightCardTreeItem>>& GetLightCardActors() const { return LightCardActors; }
+	const TArray<TSharedPtr<FStageActorTreeItem>>& GetStageActorTreeItems() const { return StageActorTreeItems; }
 
 	/** Gets a list of light card actors that are currently selected in the list */
-	void GetSelectedLightCards(TArray<ADisplayClusterLightCardActor*>& OutSelectedLightCards) const;
+	void GetSelectedActors(TArray<AActor*>& OutSelectedActors) const;
+
+	template<typename T>
+	TArray<T*> GetSelectedActorsAs() const
+	{
+		TArray<AActor*> SelectedActors;
+		GetSelectedActors(SelectedActors);
+		TArray<T*> OutArray;
+		Algo::TransformIf(SelectedActors, OutArray, [](AActor* InItem)
+		{
+			return InItem && InItem->IsA<T>();
+		},
+		[](AActor* InItem)
+		{
+			return CastChecked<T>(InItem);
+		});
+		return OutArray;
+	}
 
 	/** Select light cards in the outliner and display their details */
-	void SelectLightCards(const TArray<ADisplayClusterLightCardActor*>& LightCardsToSelect);
+	void SelectActors(const TArray<AActor*>& ActorsToSelect);
 
 	/** Restores the last valid cached selection to the outliner */
 	void RestoreCachedSelection();
@@ -64,16 +75,16 @@ private:
 	void CreateWorldOutliner();
 	
 	/**
-	 * Fill the LightCard list with available LightCards
+	 * Fill the outliner with available actors
 	 * @return True if the list has been modified
 	 */
-	bool FillLightCardList();
+	bool FillActorList();
 
 	/** Called when the user clicks on a selection in the outliner */
 	void OnOutlinerSelectionChanged(FSceneOutlinerTreeItemPtr TreeItem, ESelectInfo::Type Type);
 
-	/** Return our light card tree item from the outliner tree item */
-	TSharedPtr<FLightCardTreeItem> GetLightCardTreeItemFromOutliner(FSceneOutlinerTreeItemPtr InOutlinerTreeItem) const;
+	/** Return our stage actor tree item from the outliner tree item */
+	TSharedPtr<FStageActorTreeItem> GetStageActorTreeItemFromOutliner(FSceneOutlinerTreeItemPtr InOutlinerTreeItem) const;
 
 	/** Return an actor if the tree item is an actor tree item */
 	AActor* GetActorFromTreeItem(FSceneOutlinerTreeItemPtr InOutlinerTreeItem) const;
@@ -87,17 +98,17 @@ private:
 	/** Pointer to the light card editor that owns this widget */
 	TWeakPtr<SDisplayClusterLightCardEditor> LightCardEditorPtr;
 
-	/** A hierarchical list of light card tree items to be displayed in the tree view */
-	TArray<TSharedPtr<FLightCardTreeItem>> LightCardTree;
+	/** A hierarchical list of stage actor tree items to be displayed in the tree view */
+	TArray<TSharedPtr<FStageActorTreeItem>> StageActorTree;
 
 	/** All actors currently tracked by the scene outliner */
-	TMap<TObjectPtr<AActor>, TSharedPtr<FLightCardTreeItem>> TrackedActors;
+	TMap<TObjectPtr<AActor>, TSharedPtr<FStageActorTreeItem>> TrackedActors;
 
 	/** Cached actors currently selected. Used for synchronizing with the outliner mode */
-	TArray<TWeakObjectPtr<ADisplayClusterLightCardActor>> CachedSelectedActors;
+	TArray<TWeakObjectPtr<AActor>> CachedSelectedActors;
 	
 	/** A flattened list of all the light card actors being displayed in the tree view */
-	TArray<TSharedPtr<FLightCardTreeItem>> LightCardActors;
+	TArray<TSharedPtr<FStageActorTreeItem>> StageActorTreeItems;
 
 	/** The active root actor whose light cards are being displayed */
 	TWeakObjectPtr<ADisplayClusterRootActor> RootActor;
@@ -111,9 +122,6 @@ private:
 	/** Called when the outliner makes a selection change */
 	FDelegateHandle SceneOutlinerSelectionChanged;
 	
-	/** The tree view widget displaying the light cards */
-	TSharedPtr<STreeView<TSharedPtr<FLightCardTreeItem>>> LightCardTreeView;
-	
 	/** Extenders for menus */
 	TSharedPtr<FExtender> Extenders;
 
diff --git a/Engine/Plugins/Runtime/nDisplay/Source/DisplayClusterLightCardEditor/Private/SDisplayClusterLightCardEditor.cpp b/Engine/Plugins/Runtime/nDisplay/Source/DisplayClusterLightCardEditor/Private/SDisplayClusterLightCardEditor.cpp
--- a/Engine/Plugins/Runtime/nDisplay/Source/DisplayClusterLightCardEditor/Private/SDisplayClusterLightCardEditor.cpp
+++ b/Engine/Plugins/Runtime/nDisplay/Source/DisplayClusterLightCardEditor/Private/SDisplayClusterLightCardEditor.cpp
@@ -3,6 +3,7 @@
 #include "SDisplayClusterLightCardEditor.h"
 
 #include "DisplayClusterLightCardEditorCommands.h"
+#include "DisplayClusterLightCardEditorUtils.h"
 #include "DisplayClusterLightCardEditorStyle.h"
 #include "IDisplayClusterLightCardEditor.h"
 #include "LightCardTemplates/DisplayClusterLightCardTemplate.h"
@@ -18,7 +19,9 @@
 #include "IDisplayClusterOperatorViewModel.h"
 
 #include "DisplayClusterConfigurationTypes.h"
+
 #include "DisplayClusterRootActor.h"
+#include "Blueprints/DisplayClusterBlueprintLib.h"
 #include "Components/DisplayClusterCameraComponent.h"
 
 #include "ContentBrowserModule.h"
@@ -172,37 +175,73 @@ void SDisplayClusterLightCardEditor::Construct(const FArguments& InArgs, const T
 	GEditor->RegisterForUndo(this);
 }
 
-void SDisplayClusterLightCardEditor::SelectLightCards(const TArray<ADisplayClusterLightCardActor*>& LightCardsToSelect)
+TArray<AActor*> SDisplayClusterLightCardEditor::FindAllManagedActors() const
 {
-	check(LightCardOutliner);
-	LightCardOutliner->SelectLightCards(LightCardsToSelect);
+	TArray<AActor*> ManagedActors;
+	if (ActiveRootActor.IsValid())
+	{
+		// First populate light cards
+		TSet<ADisplayClusterLightCardActor*> RootActorLightCardActors;
+		UDisplayClusterBlueprintLib::FindLightCardsForRootActor(ActiveRootActor.Get(), RootActorLightCardActors);
+
+		for (TObjectPtr<ADisplayClusterLightCardActor> LightCardActor : RootActorLightCardActors)
+		{
+			ManagedActors.Add(LightCardActor);
+		}
+		
+		// Find other actors
+		if (UWorld* World = ActiveRootActor->GetWorld())
+		{
+			for (const TWeakObjectPtr<AActor> WeakActor : TActorRange<AActor>(World))
+			{
+				if (WeakActor.IsValid() && !WeakActor->IsA<ADisplayClusterLightCardActor>() &&
+					UE::DisplayClusterLightCardEditorUtils::IsManagedActor(WeakActor.Get()))
+				{
+					ManagedActors.Add(WeakActor.Get());
+				}
+			}
+		}
+	}
+
+	return ManagedActors;
 }
 
-void SDisplayClusterLightCardEditor::GetSelectedLightCards(TArray<ADisplayClusterLightCardActor*>& OutSelectedLightCards)
+void SDisplayClusterLightCardEditor::SelectActors(const TArray<AActor*>& ActorsToSelect)
 {
 	check(LightCardOutliner);
-	LightCardOutliner->GetSelectedLightCards(OutSelectedLightCards);
+	LightCardOutliner->SelectActors(ActorsToSelect);
 }
 
-void SDisplayClusterLightCardEditor::SelectLightCardProxies(const TArray<ADisplayClusterLightCardActor*>& LightCardsToSelect)
+void SDisplayClusterLightCardEditor::SelectActorProxies(const TArray<AActor*>& ActorsToSelect)
 {
 	check(ViewportView);
-	ViewportView->GetLightCardEditorViewportClient()->SelectLightCards(LightCardsToSelect);
+	ViewportView->GetLightCardEditorViewportClient()->SelectActors(ActorsToSelect);
 }
 
-void SDisplayClusterLightCardEditor::CenterLightCardInView(ADisplayClusterLightCardActor& LightCard)
+void SDisplayClusterLightCardEditor::GetSelectedActors(TArray<AActor*>& OutSelectedActors)
+{
+	check(LightCardOutliner);
+	LightCardOutliner->GetSelectedActors(OutSelectedActors);
+}
+
+void SDisplayClusterLightCardEditor::CenterActorInView(AActor* Actor)
 {
 	check(ViewportView);
-	ViewportView->GetLightCardEditorViewportClient()->CenterLightCardInView(LightCard);
+	check(Actor);
+	ViewportView->GetLightCardEditorViewportClient()->CenterActorInView(Actor);
 }
 
-ADisplayClusterLightCardActor* SDisplayClusterLightCardEditor::SpawnLightCard(const FName& LightCardName)
+AActor* SDisplayClusterLightCardEditor::SpawnActor(TSubclassOf<AActor> InActorClass, const FName& InActorName)
 {
 	if (!ActiveRootActor.IsValid())
 	{
 		return nullptr;
 	}
 
+	FScopedTransaction Transaction(LOCTEXT("SpawnActorTransactionMessage", "Spawn Actor"));
+	
+	check(InActorClass);
+
 	const FVector SpawnLocation = ActiveRootActor->GetDefaultCamera()->GetComponentLocation();
 	FRotator SpawnRotation = ActiveRootActor->GetDefaultCamera()->GetComponentRotation();
 	SpawnRotation.Yaw -= 180.f;
@@ -210,31 +249,39 @@ ADisplayClusterLightCardActor* SDisplayClusterLightCardEditor::SpawnLightCard(co
 	FActorSpawnParameters SpawnParameters;
 	SpawnParameters.bNoFail = true;
 	SpawnParameters.SpawnCollisionHandlingOverride = ESpawnActorCollisionHandlingMethod::AdjustIfPossibleButAlwaysSpawn;
-	SpawnParameters.Name = LightCardName;
+	SpawnParameters.Name = InActorName;
 	SpawnParameters.NameMode = FActorSpawnParameters::ESpawnActorNameMode::Requested;
 	SpawnParameters.OverrideLevel = ActiveRootActor->GetWorld()->GetCurrentLevel();
 
-	ADisplayClusterLightCardActor* NewLightCard = CastChecked<ADisplayClusterLightCardActor>(
-		ActiveRootActor->GetWorld()->SpawnActor(ADisplayClusterLightCardActor::StaticClass(),
+	AActor* NewActor = CastChecked<AActor>(
+		ActiveRootActor->GetWorld()->SpawnActor(InActorClass,
 			&SpawnLocation, &SpawnRotation, MoveTemp(SpawnParameters)));
 
-	NewLightCard->SetActorLabel(NewLightCard->GetName());
+	NewActor->SetActorLabel(NewActor->GetName());
 
-	if (ViewportView.IsValid())
+	if (ADisplayClusterLightCardActor* NewLightCard = Cast<ADisplayClusterLightCardActor>(NewActor))
 	{
-		const EDisplayClusterMeshProjectionType ProjectionMode = ViewportView->GetLightCardEditorViewportClient()->GetProjectionMode();
-		if (ProjectionMode == EDisplayClusterMeshProjectionType::UV)
+		if (ViewportView.IsValid())
 		{
-			NewLightCard->bIsUVLightCard = true;
-			NewLightCard->Scale /= 4;
-			NewLightCard->Feathering = 0.05; // Just enough to avoid jagged look on UV lightcards.
+			const EDisplayClusterMeshProjectionType ProjectionMode = ViewportView->GetLightCardEditorViewportClient()->GetProjectionMode();
+			if (ProjectionMode == EDisplayClusterMeshProjectionType::UV)
+			{
+				NewLightCard->bIsUVLightCard = true;
+				NewLightCard->Scale /= 4;
+				NewLightCard->Feathering = 0.05; // Just enough to avoid jagged look on UV lightcards.
+			}
 		}
+		const TArray<ADisplayClusterLightCardActor*> LightCards { NewLightCard } ;
+		AddLightCardsToActor(LightCards);
 	}
+	else
+	{
+		RefreshPreviewActors(EDisplayClusterLightCardEditorProxyType::LightCards);
+	}
+	
+	SelectActors({ NewActor });
 
-	TArray<ADisplayClusterLightCardActor*> LightCards { NewLightCard } ;
-	AddLightCardsToActor(LightCards);
-
-	return NewLightCard;
+	return NewActor;
 }
 
 ADisplayClusterLightCardActor* SDisplayClusterLightCardEditor::SpawnLightCardFromTemplate(
@@ -275,7 +322,7 @@ ADisplayClusterLightCardActor* SDisplayClusterLightCardEditor::SpawnLightCardFro
 		RecentlyPlacedItem.ItemType = FDisplayClusterLightCardEditorRecentItem::Type_LightCardTemplate;
 		
 		AddRecentlyPlacedItem(MoveTemp(RecentlyPlacedItem));
-		SelectLightCards({NewLightCard});
+		SelectActors({NewLightCard});
 	}
 	
 	return NewLightCard;
@@ -287,19 +334,17 @@ void SDisplayClusterLightCardEditor::AddNewLightCard()
 
 	FScopedTransaction Transaction(LOCTEXT("AddNewLightCardTransactionMessage", "Add New Light Card"));
 
-	ADisplayClusterLightCardActor* NewLightCard = SpawnLightCard();
-
 	// When adding a new lightcard, usually the desired location is in the middle of the viewport
-	if (NewLightCard)
+	if (ADisplayClusterLightCardActor* NewLightCard = SpawnActorAs<ADisplayClusterLightCardActor>(TEXT("LightCard")))
 	{
-		CenterLightCardInView(*NewLightCard);
+		CenterActorInView(NewLightCard);
 
 		FDisplayClusterLightCardEditorRecentItem RecentlyPlacedItem;
 		RecentlyPlacedItem.ObjectPath = NewLightCard->GetClass();
 		RecentlyPlacedItem.ItemType = FDisplayClusterLightCardEditorRecentItem::Type_LightCard;
 		
 		AddRecentlyPlacedItem(MoveTemp(RecentlyPlacedItem));
-		SelectLightCards({NewLightCard});
+		SelectActors({NewLightCard});
 	}
 }
 
@@ -418,22 +463,41 @@ void SDisplayClusterLightCardEditor::AddNewFlag()
 	FScopedTransaction Transaction(LOCTEXT("AddNewFlagTransactionMessage", "Add New Flag"));
 
 	// When adding a new lightcard, usually the desired location is in the middle of the viewport
-	if (ADisplayClusterLightCardActor* NewLightCard = SpawnLightCard(TEXT("Flag")))
+	if (ADisplayClusterLightCardActor* NewLightCard = SpawnActorAs<ADisplayClusterLightCardActor>(TEXT("Flag")))
 	{
 		NewLightCard->Color = FLinearColor(0.f, 0.f, 0.f);
 		
-		CenterLightCardInView(*NewLightCard);
+		CenterActorInView(NewLightCard);
 
 		FDisplayClusterLightCardEditorRecentItem RecentlyPlacedItem;
 		RecentlyPlacedItem.ObjectPath = NewLightCard->GetClass();
 		RecentlyPlacedItem.ItemType = FDisplayClusterLightCardEditorRecentItem::Type_Flag;
 		
 		AddRecentlyPlacedItem(MoveTemp(RecentlyPlacedItem));
-		SelectLightCards({NewLightCard});
+		SelectActors({NewLightCard});
+	}
+}
+
+void SDisplayClusterLightCardEditor::AddNewDynamic(UClass* InClass)
+{
+	check(ActiveRootActor.IsValid());
+
+	FScopedTransaction Transaction(LOCTEXT("AddNewActorMessage", "Add New Actor"));
+	
+	if (AActor* NewActor = SpawnActor(InClass))
+	{
+		CenterActorInView(NewActor);
+
+		FDisplayClusterLightCardEditorRecentItem RecentlyPlacedItem;
+		RecentlyPlacedItem.ObjectPath = NewActor->GetClass();
+		RecentlyPlacedItem.ItemType = FDisplayClusterLightCardEditorRecentItem::Type_Dynamic;
+		
+		AddRecentlyPlacedItem(MoveTemp(RecentlyPlacedItem));
+		SelectActors({NewActor});
 	}
 }
 
-void SDisplayClusterLightCardEditor::AddLightCardsToActor(TArray<ADisplayClusterLightCardActor*> LightCards)
+void SDisplayClusterLightCardEditor::AddLightCardsToActor(const TArray<ADisplayClusterLightCardActor*>& LightCards)
 {
 	if (ActiveRootActor.IsValid())
 	{
@@ -475,25 +539,25 @@ bool SDisplayClusterLightCardEditor::CanAddLightCard() const
 	return ActiveRootActor.IsValid() && ActiveRootActor->GetWorld() != nullptr;
 }
 
-void SDisplayClusterLightCardEditor::CutLightCards()
+void SDisplayClusterLightCardEditor::CutSelectedActors()
 {
-	TArray<ADisplayClusterLightCardActor*> SelectedLightCards;
-	LightCardOutliner->GetSelectedLightCards(SelectedLightCards);
+	TArray<AActor*> SelectedActors;
+	LightCardOutliner->GetSelectedActors(SelectedActors);
 	
-	CopyLightCards();
+	CopySelectedActors();
 
-	RemoveLightCards(SelectedLightCards, /*bDeleteLightCards*/ true);
+	RemoveActors(SelectedActors, /*bDeleteActors*/ true);
 }
 
-bool SDisplayClusterLightCardEditor::CanCutLightCards()
+bool SDisplayClusterLightCardEditor::CanCutSelectedActors()
 {
-	return CanCopyLightCards() && CanRemoveLightCards();
+	return CanCopySelectedActors() && CanRemoveSelectedActors();
 }
 
-void SDisplayClusterLightCardEditor::CopyLightCards()
+void SDisplayClusterLightCardEditor::CopySelectedActors()
 {
-	TArray<ADisplayClusterLightCardActor*> SelectedLightCards;
-	LightCardOutliner->GetSelectedLightCards(SelectedLightCards);
+	TArray<AActor*> SelectedActors;
+	LightCardOutliner->GetSelectedActors(SelectedActors);
 
 	UWorld* EditorWorld = GEditor->GetEditorWorldContext().World();
 
@@ -502,7 +566,7 @@ void SDisplayClusterLightCardEditor::CopyLightCards()
 	const bool bWarnAboutManyActors = false;
 	GEditor->SelectNone(bNoteSelectionChange, bDeselectBSPSurfs, bWarnAboutManyActors);
 
-	for (AActor* LightCard : SelectedLightCards)
+	for (AActor* LightCard : SelectedActors)
 	{
 		const bool bInSelected = true;
 		const bool bNotify = false;
@@ -516,81 +580,88 @@ void SDisplayClusterLightCardEditor::CopyLightCards()
 	GEditor->CopySelectedActorsToClipboard(EditorWorld, bShouldCut, bIsMove, bWarnAboutReferences);
 }
 
-bool SDisplayClusterLightCardEditor::CanCopyLightCards() const
+bool SDisplayClusterLightCardEditor::CanCopySelectedActors() const
 {
-	TArray<ADisplayClusterLightCardActor*> SelectedLightCards;
-	LightCardOutliner->GetSelectedLightCards(SelectedLightCards);
+	TArray<AActor*> SelectedActors;
+	LightCardOutliner->GetSelectedActors(SelectedActors);
 
-	return SelectedLightCards.Num() > 0;
+	return SelectedActors.Num() > 0;
 }
 
-void SDisplayClusterLightCardEditor::PasteLightCards(bool bOffsetLightCardPosition)
+void SDisplayClusterLightCardEditor::PasteActors(bool bOffsetActorPosition)
 {
 	UWorld* EditorWorld = GEditor->GetEditorWorldContext().World();
-	GEditor->PasteSelectedActorsFromClipboard(EditorWorld, LOCTEXT("PasteLightCardsTransactionMessage", "Paste Light Cards"), EPasteTo::PT_OriginalLocation);
+	GEditor->PasteSelectedActorsFromClipboard(EditorWorld, LOCTEXT("PasteActorsTransactionMessage", "Paste Actors"), EPasteTo::PT_OriginalLocation);
 
+	TArray<AActor*> PastedActors;
 	TArray<ADisplayClusterLightCardActor*> PastedLightCards;
 	for (FSelectionIterator It(GEditor->GetSelectedActorIterator()); It; ++It)
 	{
-		if (ADisplayClusterLightCardActor* LightCard = Cast<ADisplayClusterLightCardActor>(*It))
+		if (AActor* Actor = Cast<AActor>(*It))
 		{
-			PastedLightCards.Add(LightCard);
-		}
-	}
+			if (!Actor->Implements<UDisplayClusterStageActor>())
+			{
+				continue;
+			}
 
-	for (ADisplayClusterLightCardActor* LightCard : PastedLightCards)
-	{
-		if (bOffsetLightCardPosition)
-		{
-			// If the light card should be offset from its pasted location, offset its longitude and latitude by a number of
-			// degrees equal to an arc length of 10 units (arc length = angle in radians * radius)
-			const float AngleOffset = FMath::RadiansToDegrees(10.0f / FMath::Max(LightCard->DistanceFromCenter, 1.0f));
-			LightCard->Latitude -= AngleOffset;
-			LightCard->Longitude += AngleOffset;
+			IDisplayClusterStageActor* StageActor = CastChecked<IDisplayClusterStageActor>(Actor);
+			PastedActors.Add(Actor);
+			if (bOffsetActorPosition)
+			{
+				// If the light card should be offset from its pasted location, offset its longitude and latitude by a number of
+				// degrees equal to an arc length of 10 units (arc length = angle in radians * radius)
+				const float AngleOffset = FMath::RadiansToDegrees(10.0f / FMath::Max(StageActor->GetDistanceFromCenter(), 1.0f));
+				StageActor->SetLatitude(StageActor->GetLatitude() - AngleOffset);
+				StageActor->SetLongitude(StageActor->GetLongitude() + AngleOffset);
+			}
+			
+			if (ADisplayClusterLightCardActor* LightCard = Cast<ADisplayClusterLightCardActor>(Actor))
+			{
+				PastedLightCards.Add(LightCard);
+			}
 		}
 	}
 
 	AddLightCardsToActor(PastedLightCards);
-
-	SelectLightCards(PastedLightCards);
+	SelectActors(PastedActors);
 }
 
-bool SDisplayClusterLightCardEditor::CanPasteLightCards() const
+bool SDisplayClusterLightCardEditor::CanPasteActors() const
 {
 	UWorld* EditorWorld = GEditor->GetEditorWorldContext().World();
 	return GEditor->CanPasteSelectedActorsFromClipboard(EditorWorld);
 }
 
-void SDisplayClusterLightCardEditor::DuplicateLightCards()
+void SDisplayClusterLightCardEditor::DuplicateSelectedActors()
 {
-	CopyLightCards();
+	CopySelectedActors();
 
-	const bool bOffsetLightCardPosition = true;
-	PasteLightCards(bOffsetLightCardPosition);
+	const bool bOffsetActorPosition = true;
+	PasteActors(bOffsetActorPosition);
 }
 
-bool SDisplayClusterLightCardEditor::CanDuplicateLightCards() const
+bool SDisplayClusterLightCardEditor::CanDuplicateSelectedActors() const
 {
-	return CanCopyLightCards();
+	return CanCopySelectedActors();
 }
 
-void SDisplayClusterLightCardEditor::RemoveLightCards(bool bDeleteLightCardActor)
+void SDisplayClusterLightCardEditor::RemoveSelectedActors(bool bDeleteLightCardActor)
 {
-	TArray<ADisplayClusterLightCardActor*> SelectedLightCards;
-	LightCardOutliner->GetSelectedLightCards(SelectedLightCards);
+	TArray<AActor*> SelectedActors;
+	LightCardOutliner->GetSelectedActors(SelectedActors);
 
-	RemoveLightCards(SelectedLightCards, bDeleteLightCardActor);
+	RemoveActors(SelectedActors, bDeleteLightCardActor);
 }
 
-void SDisplayClusterLightCardEditor::RemoveLightCards(
-	const TArray<ADisplayClusterLightCardActor*>& InLightCardsToRemove, bool bDeleteLightCards)
+void SDisplayClusterLightCardEditor::RemoveActors(
+	const TArray<AActor*>& InActorsToRemove, bool bDeleteActors)
 {
-	FScopedTransaction Transaction(LOCTEXT("RemoveLightCardTransactionMessage", "Remove Light Card(s)"));
+	FScopedTransaction Transaction(LOCTEXT("RemoveActorTransactionMessage", "Remove Light Card(s)"));
 
 	USelection* EdSelectionManager = GEditor->GetSelectedActors();
 	UWorld* WorldToUse = nullptr;
 
-	if (bDeleteLightCards)
+	if (bDeleteActors)
 	{
 		EdSelectionManager->BeginBatchSelectOperation();
 		EdSelectionManager->Modify();
@@ -604,26 +675,30 @@ void SDisplayClusterLightCardEditor::RemoveLightCards(
 
 		FDisplayClusterConfigurationICVFX_VisibilityList& RootActorLightCards = ConfigData->StageSettings.Lightcard.ShowOnlyList;
 
-		for (ADisplayClusterLightCardActor* LightCard : InLightCardsToRemove)
+		for (AActor* Actor : InActorsToRemove)
 		{
-			RootActorLightCards.Actors.RemoveAll([&](const TSoftObjectPtr<AActor>& Actor)
-				{
-					return Actor.Get() == LightCard;
-				});
-
-			if (bDeleteLightCards)
+			if (ADisplayClusterLightCardActor* LightCard = Cast<ADisplayClusterLightCardActor>(Actor))
 			{
-				WorldToUse = LightCard->GetWorld();
-				GEditor->SelectActor(LightCard, /*bSelect =*/true, /*bNotifyForActor =*/false, /*bSelectEvenIfHidden =*/true);
+				RootActorLightCards.Actors.RemoveAll([&](const TSoftObjectPtr<AActor>& Actor)
+					{
+						return Actor.Get() == LightCard;
+					});
+
+				if (!bDeleteActors)
+				{
+					LightCard->ShowLightCardLabel(false, *GetLightCardLabelScale(), ActiveRootActor.Get());
+				}
 			}
-			else
+
+			if (bDeleteActors)
 			{
-				LightCard->ShowLightCardLabel(false, *GetLightCardLabelScale(), ActiveRootActor.Get());
+				WorldToUse = Actor->GetWorld();
+				GEditor->SelectActor(Actor, /*bSelect =*/true, /*bNotifyForActor =*/false, /*bSelectEvenIfHidden =*/true);
 			}
 		}
 	}
 
-	if (bDeleteLightCards)
+	if (bDeleteActors)
 	{
 		EdSelectionManager->EndBatchSelectOperation();
 
@@ -636,21 +711,28 @@ void SDisplayClusterLightCardEditor::RemoveLightCards(
 	RefreshPreviewActors(EDisplayClusterLightCardEditorProxyType::LightCards);
 }
 
-bool SDisplayClusterLightCardEditor::CanRemoveLightCards() const
+bool SDisplayClusterLightCardEditor::CanRemoveSelectedActors() const
 {
-	TArray<ADisplayClusterLightCardActor*> SelectedLightCards;
-	LightCardOutliner->GetSelectedLightCards(SelectedLightCards);
+	TArray<AActor*> SelectedActors;
+	LightCardOutliner->GetSelectedActors(SelectedActors);
 
-	return SelectedLightCards.Num() > 0;
+	for (const AActor* Actor : SelectedActors)
+	{
+		if (Actor)
+		{
+			return true;
+		}
+	}
+	
+	return false;
 }
 
 void SDisplayClusterLightCardEditor::CreateLightCardTemplate()
 {
-	TArray<ADisplayClusterLightCardActor*> SelectedLightCards;
-	LightCardOutliner->GetSelectedLightCards(SelectedLightCards);
+	TArray<ADisplayClusterLightCardActor*> SelectedActors = LightCardOutliner->GetSelectedActorsAs<ADisplayClusterLightCardActor>();
 
-	check(SelectedLightCards.Num() == 1);
-	const ADisplayClusterLightCardActor* LightCardActor = SelectedLightCards[0];
+	check(SelectedActors.Num() == 1);
+	const ADisplayClusterLightCardActor* LightCardActor = SelectedActors[0];
 
 	const UDisplayClusterLightCardEditorProjectSettings* Settings = GetDefault<UDisplayClusterLightCardEditorProjectSettings>();
 	FString DefaultPath = Settings->LightCardTemplateDefaultPath.Path;
@@ -709,9 +791,7 @@ void SDisplayClusterLightCardEditor::CreateLightCardTemplate()
 
 bool SDisplayClusterLightCardEditor::CanCreateLightCardTemplate() const
 {
-	TArray<ADisplayClusterLightCardActor*> SelectedLightCards;
-	LightCardOutliner->GetSelectedLightCards(SelectedLightCards);
-
+	const TArray<ADisplayClusterLightCardActor*> SelectedLightCards = LightCardOutliner->GetSelectedActorsAs<ADisplayClusterLightCardActor>();
 	return SelectedLightCards.Num() == 1;
 }
 
@@ -770,6 +850,32 @@ void SDisplayClusterLightCardEditor::SetLightCardLabelScale(float NewValue)
 	LightCardEditorModule.ShowLabels(MoveTemp(Args));
 }
 
+void SDisplayClusterLightCardEditor::ShowIcons(bool bVisible)
+{
+	UDisplayClusterLightCardEditorSettings* Settings = GetMutableDefault<UDisplayClusterLightCardEditorSettings>();
+	Settings->bDisplayIcons = bVisible;
+	Settings->PostEditChange();
+	Settings->SaveConfig();
+}
+
+bool SDisplayClusterLightCardEditor::ShouldShowIcons() const
+{
+	return GetDefault<UDisplayClusterLightCardEditorSettings>()->bDisplayIcons;
+}
+
+TOptional<float> SDisplayClusterLightCardEditor::GetIconScale() const
+{
+	return GetDefault<UDisplayClusterLightCardEditorSettings>()->IconScale;
+}
+
+void SDisplayClusterLightCardEditor::SetIconScale(float NewValue)
+{
+	UDisplayClusterLightCardEditorSettings* Settings = GetMutableDefault<UDisplayClusterLightCardEditorSettings>();
+	Settings->IconScale = NewValue;
+	Settings->PostEditChange();
+	Settings->SaveConfig();
+}
+
 void SDisplayClusterLightCardEditor::OnActiveRootActorChanged(ADisplayClusterRootActor* NewRootActor)
 {
 	if (NewRootActor == ActiveRootActor.Get())
@@ -818,6 +924,21 @@ TSharedRef<SWidget> SDisplayClusterLightCardEditor::GeneratePlaceActorsMenu()
 		MenuBuilder.AddMenuEntry(FDisplayClusterLightCardEditorCommands::Get().AddExistingLightCard,
 			NAME_None, TAttribute<FText>(), TAttribute<FText>(), LightCardIcon);
 		
+		TSet<UClass*> StageActorClasses = UE::DisplayClusterLightCardEditorUtils::GetAllStageActorClasses();
+		for (UClass* Class : StageActorClasses)
+		{
+			if (Class == ADisplayClusterLightCardActor::StaticClass())
+			{
+				// Added manually already
+				continue;
+			}
+			
+			FString Label = FString::Printf(TEXT("Add %s"), *Class->GetDisplayNameText().ToString());
+			FSlateIcon StageActorIcon = FSlateIconFinder::FindIconForClass(Class);
+			MenuBuilder.AddMenuEntry(FText::FromString(Label), LOCTEXT("AddStageActorHeader", "Add a stage actor to the scene"), StageActorIcon,
+				FUIAction(FExecuteAction::CreateRaw(this, &SDisplayClusterLightCardEditor::AddNewDynamic, Class)));
+		}
+		
 		MenuBuilder.AddSubMenu(LOCTEXT("AllTemplatesLabel", "All Templates"),
 		LOCTEXT("AllTemplatesTooltip", "Select a template"),
 		FNewMenuDelegate::CreateSP(this, &SDisplayClusterLightCardEditor::GenerateTemplateSubMenu),
@@ -943,6 +1064,10 @@ TSharedRef<SWidget> SDisplayClusterLightCardEditor::GeneratePlaceActorsMenu()
 							SpawnLightCardFromTemplate(Template);
 						}
 					}
+					else if (RecentlyPlacedItem.ItemType == FDisplayClusterLightCardEditorRecentItem::Type_Dynamic)
+					{
+						AddNewDynamic(Cast<UClass>(ObjectWeakPtr.Get()));
+					}
 				}
 			});
 			
@@ -1018,6 +1143,42 @@ TSharedRef<SWidget> SDisplayClusterLightCardEditor::GenerateLabelsMenu()
 		LOCTEXT("LightCardLabelScale_Label", "Label Scale"));
 	}
 	MenuBuilder.EndSection();
+
+	MenuBuilder.BeginSection("Icons", LOCTEXT("IconsMenuHeader", "Icons"));
+	{
+		MenuBuilder.AddMenuEntry(FDisplayClusterLightCardEditorCommands::Get().ToggleIconVisibility,
+			NAME_None, LOCTEXT("IconVisibilityLabel", "Actor Icons"));
+	
+		MenuBuilder.AddWidget(
+		SNew(SBox)
+		.HAlign(HAlign_Right)
+		[
+			SNew(SBox)
+			.Padding(FMargin(8.0f, 0.0f, 0.0f, 0.0f))
+			.WidthOverride(100.0f)
+			[
+				SNew(SBorder)
+				.BorderImage(FAppStyle::Get().GetBrush("Menu.WidgetBorder"))
+				.Padding(FMargin(1.0f))
+				[
+					SNew(SBox)
+					.MinDesiredWidth(64)
+					[
+						SNew(SNumericEntryBox<float>)
+						.Value(this, &SDisplayClusterLightCardEditor::GetIconScale)
+						.OnValueChanged(this, &SDisplayClusterLightCardEditor::SetIconScale)
+						.MinValue(0)
+						.MaxValue(FLT_MAX)
+						.MinSliderValue(0)
+						.MaxSliderValue(4)
+						.AllowSpin(true)
+					]
+				]
+			]
+		],
+		LOCTEXT("LightCardIconScale_Label", "Icon Scale"));
+	}
+	MenuBuilder.EndSection();
 	
 	return MenuBuilder.MakeWidget();
 }
@@ -1050,33 +1211,33 @@ void SDisplayClusterLightCardEditor::BindCommands()
 
 	CommandList->MapAction(
 		FDisplayClusterLightCardEditorCommands::Get().RemoveLightCard,
-		FExecuteAction::CreateSP(this, &SDisplayClusterLightCardEditor::RemoveLightCards, false),
-		FCanExecuteAction::CreateSP(this, &SDisplayClusterLightCardEditor::CanRemoveLightCards));
+		FExecuteAction::CreateSP(this, &SDisplayClusterLightCardEditor::RemoveSelectedActors, false),
+		FCanExecuteAction::CreateSP(this, &SDisplayClusterLightCardEditor::CanRemoveSelectedActors));
 
 	CommandList->MapAction(
 		FGenericCommands::Get().Cut,
-		FExecuteAction::CreateSP(this, &SDisplayClusterLightCardEditor::CutLightCards),
-		FCanExecuteAction::CreateSP(this, &SDisplayClusterLightCardEditor::CanCutLightCards));
+		FExecuteAction::CreateSP(this, &SDisplayClusterLightCardEditor::CutSelectedActors),
+		FCanExecuteAction::CreateSP(this, &SDisplayClusterLightCardEditor::CanCutSelectedActors));
 
 	CommandList->MapAction(
 		FGenericCommands::Get().Copy,
-		FExecuteAction::CreateSP(this, &SDisplayClusterLightCardEditor::CopyLightCards),
-		FCanExecuteAction::CreateSP(this, &SDisplayClusterLightCardEditor::CanCopyLightCards));
+		FExecuteAction::CreateSP(this, &SDisplayClusterLightCardEditor::CopySelectedActors),
+		FCanExecuteAction::CreateSP(this, &SDisplayClusterLightCardEditor::CanCopySelectedActors));
 
 	CommandList->MapAction(
 		FGenericCommands::Get().Paste,
-		FExecuteAction::CreateSP(this, &SDisplayClusterLightCardEditor::PasteLightCards, false),
-		FCanExecuteAction::CreateSP(this, &SDisplayClusterLightCardEditor::CanPasteLightCards));
+		FExecuteAction::CreateSP(this, &SDisplayClusterLightCardEditor::PasteActors, false),
+		FCanExecuteAction::CreateSP(this, &SDisplayClusterLightCardEditor::CanPasteActors));
 
 	CommandList->MapAction(
 		FGenericCommands::Get().Duplicate,
-		FExecuteAction::CreateSP(this, &SDisplayClusterLightCardEditor::DuplicateLightCards),
-		FCanExecuteAction::CreateSP(this, &SDisplayClusterLightCardEditor::CanDuplicateLightCards));
+		FExecuteAction::CreateSP(this, &SDisplayClusterLightCardEditor::DuplicateSelectedActors),
+		FCanExecuteAction::CreateSP(this, &SDisplayClusterLightCardEditor::CanDuplicateSelectedActors));
 
 	CommandList->MapAction(
 		FGenericCommands::Get().Delete,
-		FExecuteAction::CreateSP(this, &SDisplayClusterLightCardEditor::RemoveLightCards, true),
-		FCanExecuteAction::CreateSP(this, &SDisplayClusterLightCardEditor::CanRemoveLightCards));
+		FExecuteAction::CreateSP(this, &SDisplayClusterLightCardEditor::RemoveSelectedActors, true),
+		FCanExecuteAction::CreateSP(this, &SDisplayClusterLightCardEditor::CanRemoveSelectedActors));
 
 	CommandList->MapAction(
 		FDisplayClusterLightCardEditorCommands::Get().SaveLightCardTemplate,
@@ -1178,9 +1339,9 @@ bool SDisplayClusterLightCardEditor::IsOurObject(UObject* InObject,
 	bool bIsOurActor = IsOurActor(GetActiveRootActor().Get());
 	if (!bIsOurActor && LightCardOutliner.IsValid())
 	{
-		for (const TSharedPtr<SDisplayClusterLightCardOutliner::FLightCardTreeItem>& LightCard : LightCardOutliner->GetLightCardActors())
+		for (const TSharedPtr<SDisplayClusterLightCardOutliner::FStageActorTreeItem>& LightCard : LightCardOutliner->GetStageActorTreeItems())
 		{
-			bIsOurActor = IsOurActor(LightCard->LightCardActor.Get());
+			bIsOurActor = IsOurActor(LightCard->Actor.Get());
 			if (bIsOurActor)
 			{
 				ProxyType = EDisplayClusterLightCardEditorProxyType::LightCards;
@@ -1197,11 +1358,11 @@ void SDisplayClusterLightCardEditor::BindCompileDelegates()
 {
 	if (LightCardOutliner.IsValid())
 	{
-		for (const TSharedPtr<SDisplayClusterLightCardOutliner::FLightCardTreeItem>& LightCardActor : LightCardOutliner->GetLightCardActors())
+		for (const TSharedPtr<SDisplayClusterLightCardOutliner::FStageActorTreeItem>& LightCardActor : LightCardOutliner->GetStageActorTreeItems())
 		{
-			if (LightCardActor.IsValid() && LightCardActor->LightCardActor.IsValid())
+			if (LightCardActor.IsValid() && LightCardActor->Actor.IsValid())
 			{
-				if (UBlueprint* Blueprint = UBlueprint::GetBlueprintFromClass(LightCardActor->LightCardActor->GetClass()))
+				if (UBlueprint* Blueprint = UBlueprint::GetBlueprintFromClass(LightCardActor->Actor->GetClass()))
 				{
 					Blueprint->OnCompiled().AddSP(this, &SDisplayClusterLightCardEditor::OnBlueprintCompiled);
 				}
@@ -1214,11 +1375,11 @@ void SDisplayClusterLightCardEditor::RemoveCompileDelegates()
 {
 	if (LightCardOutliner.IsValid())
 	{
-		for (const TSharedPtr<SDisplayClusterLightCardOutliner::FLightCardTreeItem>& LightCardActor : LightCardOutliner->GetLightCardActors())
+		for (const TSharedPtr<SDisplayClusterLightCardOutliner::FStageActorTreeItem>& LightCardActor : LightCardOutliner->GetStageActorTreeItems())
 		{
-			if (LightCardActor.IsValid() && LightCardActor->LightCardActor.IsValid())
+			if (LightCardActor.IsValid() && LightCardActor->Actor.IsValid())
 			{
-				if (UBlueprint* Blueprint = UBlueprint::GetBlueprintFromClass(LightCardActor->LightCardActor->GetClass()))
+				if (UBlueprint* Blueprint = UBlueprint::GetBlueprintFromClass(LightCardActor->Actor->GetClass()))
 				{
 					Blueprint->OnCompiled().RemoveAll(this);
 				}
@@ -1233,7 +1394,13 @@ void SDisplayClusterLightCardEditor::OnActorPropertyChanged(UObject* ObjectBeing
 	EDisplayClusterLightCardEditorProxyType ProxyType;
 	if (IsOurObject(ObjectBeingModified, ProxyType))
 	{
-		if (PropertyChangedEvent.ChangeType == EPropertyChangeType::Interactive)
+		const FName PropertyName = PropertyChangedEvent.GetPropertyName();
+		const bool bTransformationChanged =
+				(PropertyName == USceneComponent::GetRelativeLocationPropertyName() ||
+					PropertyName == USceneComponent::GetRelativeRotationPropertyName() ||
+					PropertyName == USceneComponent::GetRelativeScale3DPropertyName());
+		if (PropertyChangedEvent.ChangeType == EPropertyChangeType::Interactive ||
+			(bTransformationChanged && !ObjectBeingModified->IsA<ADisplayClusterRootActor>()))
 		{
 			// Real-time & efficient update when dragging a slider.
 			if (ViewportView.IsValid())
@@ -1252,9 +1419,9 @@ void SDisplayClusterLightCardEditor::OnActorPropertyChanged(UObject* ObjectBeing
 void SDisplayClusterLightCardEditor::OnLevelActorDeleted(AActor* Actor)
 {
 	if (LightCardOutliner.IsValid() &&
-		LightCardOutliner->GetLightCardActors().ContainsByPredicate([Actor](const TSharedPtr<SDisplayClusterLightCardOutliner::FLightCardTreeItem>& LightCardTreeItem)
+		LightCardOutliner->GetStageActorTreeItems().ContainsByPredicate([Actor](const TSharedPtr<SDisplayClusterLightCardOutliner::FStageActorTreeItem>& LightCardTreeItem)
 	{
-		return LightCardTreeItem.IsValid() && Actor == LightCardTreeItem->LightCardActor.Get();
+		return LightCardTreeItem.IsValid() && Actor == LightCardTreeItem->Actor.Get();
 	}))
 	{
 		if (Actor && Actor->GetClass()->HasAnyClassFlags(CLASS_NewerVersionExists))
diff --git a/Engine/Plugins/Runtime/nDisplay/Source/DisplayClusterLightCardEditor/Private/SDisplayClusterLightCardEditor.h b/Engine/Plugins/Runtime/nDisplay/Source/DisplayClusterLightCardEditor/Private/SDisplayClusterLightCardEditor.h
--- a/Engine/Plugins/Runtime/nDisplay/Source/DisplayClusterLightCardEditor/Private/SDisplayClusterLightCardEditor.h
+++ b/Engine/Plugins/Runtime/nDisplay/Source/DisplayClusterLightCardEditor/Private/SDisplayClusterLightCardEditor.h
@@ -70,21 +70,30 @@ public:
 	/** The current active root actor for this light card editor */
 	TWeakObjectPtr<ADisplayClusterRootActor> GetActiveRootActor() const { return ActiveRootActor; }
 
-	/** Selects the specified light cards in the light card list and details panel */
-	void SelectLightCards(const TArray<ADisplayClusterLightCardActor*>& LightCardsToSelect);
+	/** Iterate the world actors looking for all actors this editor is managing */
+	TArray<AActor*> FindAllManagedActors() const;
 
-	/** Gets the light cards that are selected in the light card list */
-	void GetSelectedLightCards(TArray<ADisplayClusterLightCardActor*>& OutSelectedLightCards);
+	/** Selects the specified actors in the outliner and details panel */
+	void SelectActors(const TArray<AActor*>& ActorsToSelect);
 
-	/** Selects the light card proxies that correspond to the specified light cards */
-	void SelectLightCardProxies(const TArray<ADisplayClusterLightCardActor*>& LightCardsToSelect);
+	/** Gets the actors that are selected in the outliner */
+	void GetSelectedActors(TArray<AActor*>& OutSelectedActors);
+
+	/** Selects the actor proxies that correspond to the specified actors */
+	void SelectActorProxies(const TArray<AActor*>& ActorsToSelect);
 
 	/** Places the given light card in the middle of the viewport */
-	void CenterLightCardInView(ADisplayClusterLightCardActor& LightCard);
+	void CenterActorInView(AActor* Actor);
 
-	/** Spawns a new light card and adds it to the root actor */
-	ADisplayClusterLightCardActor* SpawnLightCard(const FName& LightCardName = TEXT("LightCard"));
+	/** Spawns a new actor and adds it to the root actor if it is a light card */
+	AActor* SpawnActor(TSubclassOf<AActor> InActorClass, const FName& InActorName = NAME_None);
 
+	template<typename T>
+	T* SpawnActorAs(const FName& InActorName = NAME_None)
+	{
+		return Cast<T>(SpawnActor(T::StaticClass(), InActorName));
+	}
+	
 	/** Spawns a new light card from a light card template */
 	ADisplayClusterLightCardActor* SpawnLightCardFromTemplate(const UDisplayClusterLightCardTemplate* InTemplate, ULevel* InLevel = nullptr, bool bIsPreview = false);
 
@@ -97,51 +106,54 @@ public:
 	/** Adds a new light card configured as a flag */
 	void AddNewFlag();
 
+	/** Add a new actor dynamically based on a class */
+	void AddNewDynamic(UClass* InClass);
+
 	/** Adds the given Light Card to the root actor */
-	void AddLightCardsToActor(TArray<ADisplayClusterLightCardActor*> LightCards);
+	void AddLightCardsToActor(const TArray<ADisplayClusterLightCardActor*>& LightCards);
 
 	/** If a Light Card can currently be added */
 	bool CanAddLightCard() const;
 
-	/** Copies any selected light cards to the clipboard, and then deletes them */
-	void CutLightCards();
+	/** Copies any selected actors to the clipboard, and then deletes them */
+	void CutSelectedActors();
 
-	/** Determines if there are selected light cards that can be cut */
-	bool CanCutLightCards();
+	/** Determines if there are selected actors that can be cut */
+	bool CanCutSelectedActors();
 
-	/** Copies any selected light cards to the clipboard */
-	void CopyLightCards();
+	/** Copies any selected actors to the clipboard */
+	void CopySelectedActors();
 
-	/** Determines if there are selected light cards that can be copied */
-	bool CanCopyLightCards() const;
+	/** Determines if there are selected actors that can be copied */
+	bool CanCopySelectedActors() const;
 
-	/** Pastes any light cards in the clipboard to the current root actor */
-	void PasteLightCards(bool bOffsetLightCardPosition);
+	/** Pastes any actors in the clipboard to the current root actor */
+	void PasteActors(bool bOffsetActorPosition);
 
-	/** Determines if there are any light cards that can be pasted from the clipboard */
-	bool CanPasteLightCards() const;
+	/** Determines if there are any actors that can be pasted from the clipboard */
+	bool CanPasteActors() const;
 
-	/** Copies any selected light cards to the clipboard and then pastes them */
-	void DuplicateLightCards();
+	/** Copies any selected actors to the clipboard and then pastes them */
+	void DuplicateSelectedActors();
 
-	/** Determines if there are selected light cards that can be duplicated */
-	bool CanDuplicateLightCards() const;
+	/** Determines if there are selected actors that can be duplicated */
+	bool CanDuplicateSelectedActors() const;
 
 	/**
 	 * Remove the light card from the actor
 	 *@param bDeleteLightCardActor Delete the actor from the level
 	 */
-	void RemoveLightCards(bool bDeleteLightCardActor);
+	void RemoveSelectedActors(bool bDeleteLightCardActor);
 
 	/**
 	 * Remove the given light cards from the actor
-	 * @param InLightCardsToRemove Light cards which should be removed
-	 * @param bDeleteLightCards If the light cards should be deleted from the level
+	 * @param InActorsToRemove Actors which should be removed
+	 * @param bDeleteActors If the light cards should be deleted from the level
 	 */
-	void RemoveLightCards(const TArray<ADisplayClusterLightCardActor*>& InLightCardsToRemove, bool bDeleteLightCards);
+	void RemoveActors(const TArray<AActor*>& InActorsToRemove, bool bDeleteActors);
 	
-	/** If the selected Light Card can be removed */
-	bool CanRemoveLightCards() const;
+	/** If the selected actors can be removed */
+	bool CanRemoveSelectedActors() const;
 
 	/** Creates a template of the selected light card */
 	void CreateLightCardTemplate();
@@ -163,6 +175,18 @@ public:
 
 	/** Update light card label scale */
 	void SetLightCardLabelScale(float NewValue);
+
+	/** Set all icon visibility */
+	void ShowIcons(bool bVisible);
+
+	/** If all icons should be visible */
+	bool ShouldShowIcons() const;
+
+	/** Return the current icon scale */
+	TOptional<float> GetIconScale() const;
+
+	/** Update all icons scale */
+	void SetIconScale(float NewValue);
 	
 private:
 	/** Raised when the active Display cluster root actor has been changed in the operator panel */
diff --git a/Engine/Plugins/Runtime/nDisplay/Source/DisplayClusterLightCardEditor/Private/Settings/DisplayClusterLightCardEditorSettings.cpp b/Engine/Plugins/Runtime/nDisplay/Source/DisplayClusterLightCardEditor/Private/Settings/DisplayClusterLightCardEditorSettings.cpp
--- a/Engine/Plugins/Runtime/nDisplay/Source/DisplayClusterLightCardEditor/Private/Settings/DisplayClusterLightCardEditorSettings.cpp
+++ b/Engine/Plugins/Runtime/nDisplay/Source/DisplayClusterLightCardEditor/Private/Settings/DisplayClusterLightCardEditorSettings.cpp
@@ -25,6 +25,7 @@ UDisplayClusterLightCardEditorProjectSettings::UDisplayClusterLightCardEditorPro
 const FName FDisplayClusterLightCardEditorRecentItem::Type_LightCard = "LightCard";
 const FName FDisplayClusterLightCardEditorRecentItem::Type_Flag = "Flag";
 const FName FDisplayClusterLightCardEditorRecentItem::Type_LightCardTemplate = "LightCardTemplate";
+const FName FDisplayClusterLightCardEditorRecentItem::Type_Dynamic = "Dynamic";
 
 FText FDisplayClusterLightCardEditorRecentItem::GetItemDisplayName() const
 {
@@ -38,7 +39,9 @@ FText FDisplayClusterLightCardEditorRecentItem::GetItemDisplayName() const
 	}
 	if (const UObject* Object = ObjectPath.LoadSynchronous())
 	{
-		return FText::FromString(Object->GetName());
+		FString ObjectName = Object->GetName();
+		ObjectName.RemoveFromEnd(TEXT("_C"));
+		return FText::FromString(MoveTemp(ObjectName));
 	}
 
 	return FText::GetEmpty();
@@ -46,7 +49,7 @@ FText FDisplayClusterLightCardEditorRecentItem::GetItemDisplayName() const
 
 const FSlateBrush* FDisplayClusterLightCardEditorRecentItem::GetSlateBrush() const
 {
-	if (ItemType == Type_LightCard || ItemType == Type_Flag)
+	if (ItemType == Type_LightCard || ItemType == Type_Flag || ItemType == Type_Dynamic)
 	{
 		return FSlateIconFinder::FindIconBrushForClass(ADisplayClusterLightCardActor::StaticClass());
 	}
@@ -83,6 +86,9 @@ UDisplayClusterLightCardEditorSettings::UDisplayClusterLightCardEditorSettings()
 {
 	ProjectionMode = EDisplayClusterMeshProjectionType::Azimuthal;
 	RenderViewportType = ELevelViewportType::LVT_Perspective;
+	
+	bDisplayIcons = true;
+    IconScale = 0.5f;
 }
 
 #undef LOCTEXT_NAMESPACE
\ No newline at end of file
diff --git a/Engine/Plugins/Runtime/nDisplay/Source/DisplayClusterLightCardEditor/Private/Settings/DisplayClusterLightCardEditorSettings.h b/Engine/Plugins/Runtime/nDisplay/Source/DisplayClusterLightCardEditor/Private/Settings/DisplayClusterLightCardEditorSettings.h
--- a/Engine/Plugins/Runtime/nDisplay/Source/DisplayClusterLightCardEditor/Private/Settings/DisplayClusterLightCardEditorSettings.h
+++ b/Engine/Plugins/Runtime/nDisplay/Source/DisplayClusterLightCardEditor/Private/Settings/DisplayClusterLightCardEditorSettings.h
@@ -40,6 +40,7 @@ struct FDisplayClusterLightCardEditorRecentItem
 	static const FName Type_LightCard;
 	static const FName Type_Flag;
 	static const FName Type_LightCardTemplate;
+	static const FName Type_Dynamic;
 
 	/** Path of object to instantiate */
 	UPROPERTY()
@@ -83,4 +84,12 @@ public:
 	/** Last used viewport type user has set */
 	UPROPERTY(config)
 	uint8 RenderViewportType;
+
+	/** Display icons in the light card editor where applicable */
+	UPROPERTY(config, EditAnywhere, Category = Icons)
+	bool bDisplayIcons;
+
+	/** Scale icon size where applicable */
+	UPROPERTY(config, EditAnywhere, Category = Icons)
+	float IconScale;
 };
diff --git a/Engine/Plugins/Runtime/nDisplay/Source/DisplayClusterLightCardEditor/Private/Viewport/DisplayClusterLightCardEditorViewportClient.cpp b/Engine/Plugins/Runtime/nDisplay/Source/DisplayClusterLightCardEditor/Private/Viewport/DisplayClusterLightCardEditorViewportClient.cpp
--- a/Engine/Plugins/Runtime/nDisplay/Source/DisplayClusterLightCardEditor/Private/Viewport/DisplayClusterLightCardEditorViewportClient.cpp
+++ b/Engine/Plugins/Runtime/nDisplay/Source/DisplayClusterLightCardEditor/Private/Viewport/DisplayClusterLightCardEditorViewportClient.cpp
@@ -6,7 +6,6 @@
 #include "DisplayClusterLightCardEditorWidget.h"
 #include "LightCardTemplates/DisplayClusterLightCardTemplate.h"
 
-#include "Blueprints/DisplayClusterBlueprintLib.h"
 #include "Components/DisplayClusterPreviewComponent.h"
 #include "Components/DisplayClusterCameraComponent.h"
 #include "Components/DisplayClusterScreenComponent.h"
@@ -16,6 +15,7 @@
 #include "DisplayClusterRootActor.h"
 #include "DisplayClusterLightCardActor.h"
 #include "DisplayClusterLightCardEditorLog.h"
+#include "DisplayClusterLightCardEditorUtils.h"
 #include "IDisplayClusterScenePreview.h"
 #include "SDisplayClusterLightCardEditor.h"
 #include "Settings/DisplayClusterLightCardEditorSettings.h"
@@ -50,6 +50,7 @@
 #include "Slate/SceneViewport.h"
 #include "UnrealEdGlobals.h"
 #include "UnrealWidget.h"
+#include "Components/BillboardComponent.h"
 #include "Widgets/Docking/SDockTab.h"
 
 
@@ -80,7 +81,7 @@ FDisplayClusterLightCardEditorViewportClient::FDisplayClusterLightCardEditorView
 	PreviewRendererId = PreviewModule.CreateRenderer();
 	PreviewModule.SetRendererActorSelectedDelegate(
 		PreviewRendererId,
-		FDisplayClusterMeshProjectionRenderer::FSelection::CreateRaw(this, &FDisplayClusterLightCardEditorViewportClient::IsLightCardSelected)
+		FDisplayClusterMeshProjectionRenderer::FSelection::CreateRaw(this, &FDisplayClusterLightCardEditorViewportClient::IsActorSelected)
 	);
 	PreviewModule.SetRendererRenderSimpleElementsDelegate(
 		PreviewRendererId,
@@ -466,7 +467,64 @@ void FDisplayClusterLightCardEditorViewportClient::Draw(FViewport* InViewport, F
 void FDisplayClusterLightCardEditorViewportClient::Draw(const FSceneView* View, FPrimitiveDrawInterface* PDI)
 {
 	const bool bIsUVProjection = ProjectionMode == EDisplayClusterMeshProjectionType::UV;
-	if (LastSelectedLightCard.IsValid() && LastSelectedLightCard->bIsUVLightCard == bIsUVProjection)
+
+	// Draw any sprites
+	for (const TWeakObjectPtr<UBillboardComponent>& BillboardComponent : BillboardComponentProxies)
+	{
+		if (!BillboardComponent.IsValid())
+		{
+			continue;
+		}
+		FSpriteProxy SpriteProxy = FSpriteProxy::FromBillboard(BillboardComponent.Get());
+		
+		FSceneViewInitOptions SceneViewInitOptions;
+		GetSceneViewInitOptions(SceneViewInitOptions);
+		FViewMatrices ViewMatrices(SceneViewInitOptions);
+		
+		FVector ProjectedLocation = ProjectWorldPosition(SpriteProxy.WorldPosition, ViewMatrices);
+		
+		if (const FTexture* TextureResource = SpriteProxy.Sprite ? SpriteProxy.Sprite->GetResource() : nullptr)
+		{
+			const UDisplayClusterLightCardEditorSettings* Settings = GetDefault<UDisplayClusterLightCardEditorSettings>();
+			
+			PDI->SetHitProxy(new HActor(BillboardComponent->GetOwner(), BillboardComponent.Get()));
+
+			if (Settings->bDisplayIcons)
+			{
+				const float UL = SpriteProxy.UL == 0.0f ? TextureResource->GetSizeX() : SpriteProxy.UL;
+				const float VL = SpriteProxy.VL == 0.0f ? TextureResource->GetSizeY() : SpriteProxy.VL;
+			
+				// Calculate the view-dependent scaling factor.
+				float ViewedSizeX = SpriteProxy.SpriteScale * UL;
+				float ViewedSizeY = SpriteProxy.SpriteScale * VL;
+
+				if (SpriteProxy.bIsScreenSizeScaled && (View->ViewMatrices.GetProjectionMatrix().M[3][3] != 1.0f))
+				{
+					const float ZoomFactor = FMath::Min<float>(View->ViewMatrices.GetProjectionMatrix().M[0][0], View->ViewMatrices.GetProjectionMatrix().M[1][1]);
+					if(ZoomFactor != 0.0f)
+					{
+						const float Radius = View->WorldToScreen(ProjectedLocation).W * (SpriteProxy.ScreenSize / ZoomFactor);
+
+						if (Radius < 1.0f)
+						{
+							ViewedSizeX *= Radius;
+							ViewedSizeY *= Radius;
+						}					
+					}
+				}
+				
+				ViewedSizeX *= Settings->IconScale;
+				ViewedSizeY *= Settings->IconScale;
+			
+				PDI->DrawSprite(ProjectedLocation, ViewedSizeX, ViewedSizeY, TextureResource, FLinearColor::White, SDPG_World,
+					SpriteProxy.U, UL, SpriteProxy.V, VL, SE_BLEND_Masked, SpriteProxy.OpacityMaskRefVal);
+			}
+			
+			PDI->SetHitProxy(nullptr);
+		}
+	}
+		
+	if (LastSelectedActor.IsValid() && LastSelectedActor->IsUVActor() == bIsUVProjection)
 	{
 		// Project the editor widget's world position into projection space so that it renders at the appropriate screen location.
 		// This needs to be computed on the render thread using the render thread's scene view, which will be behind the game thread's scene view
@@ -628,27 +686,27 @@ bool FDisplayClusterLightCardEditorViewportClient::InputWidgetDelta(FViewport* I
 	}
 	else
 	{
-		if (CurrentAxis != EAxisList::Type::None && SelectedLightCards.Num())
+		if (CurrentAxis != EAxisList::Type::None && SelectedActors.Num())
 		{
 			switch (EditorWidget->GetWidgetMode())
 			{
 			case FDisplayClusterLightCardEditorWidget::EWidgetMode::WM_Translate:
 				if (ProjectionMode == EDisplayClusterMeshProjectionType::UV)
 				{
-					MoveSelectedUVLightCards(InViewport, CurrentAxis);
+					MoveSelectedUVActors(InViewport, CurrentAxis);
 				}
 				else
 				{
-					MoveSelectedLightCards(InViewport, CurrentAxis);
+					MoveSelectedActors(InViewport, CurrentAxis);
 				}
 				break;
 
 			case FDisplayClusterLightCardEditorWidget::EWidgetMode::WM_RotateZ:
-				SpinSelectedLightCards(InViewport);
+				SpinSelectedActors(InViewport);
 				break;
 
 			case FDisplayClusterLightCardEditorWidget::EWidgetMode::WM_Scale:
-				ScaleSelectedLightCards(InViewport, CurrentAxis);
+				ScaleSelectedActors(InViewport, CurrentAxis);
 				break;
 			}
 
@@ -666,7 +724,7 @@ void FDisplayClusterLightCardEditorViewportClient::TrackingStarted(
 	bool bIsDraggingWidget,
 	bool bNudge)
 {
-	if ((InputMode == EInputMode::Idle) && bIsDraggingWidget && InInputState.IsLeftMouseButtonPressed() && SelectedLightCards.Num())
+	if ((InputMode == EInputMode::Idle) && bIsDraggingWidget && InInputState.IsLeftMouseButtonPressed() && SelectedActors.Num())
 	{
 		// Start dragging actor
 		InputMode = EInputMode::DraggingActor;
@@ -726,7 +784,7 @@ void FDisplayClusterLightCardEditorViewportClient::TrackingStopped()
 		DragWidgetOffset = FVector::ZeroVector;
 		EndTransaction();
 
-		if (SelectedLightCards.Num())
+		if (SelectedActors.Num())
 		{
 			GEditor->DisableDeltaModification(false);
 		}
@@ -939,7 +997,7 @@ void FDisplayClusterLightCardEditorViewportClient::CreateDrawnLightCard(const TA
 
 	FScopedTransaction Transaction(LOCTEXT("AddNewLightCard", "Add New Light Card"));
 
-	ADisplayClusterLightCardActor* LightCard = LightCardEditorPtr.Pin()->SpawnLightCard();
+	ADisplayClusterLightCardActor* LightCard = LightCardEditorPtr.Pin()->SpawnActorAs<ADisplayClusterLightCardActor>(TEXT("LightCard"));
 
 	if (!LightCard)
 	{
@@ -989,11 +1047,11 @@ void FDisplayClusterLightCardEditorViewportClient::CreateDrawnLightCard(const TA
 	else
 	{
 		const FDisplayClusterLightCardEditorHelper::FSphericalCoordinates LightCardCoords(LightCardLocation - ViewOrigin);
-		MoveLightCardTo(*LightCard, LightCardLocation - ViewOrigin);
+		MoveActorTo(LightCard, LightCardLocation - ViewOrigin);
 	}
 }
 
-double FDisplayClusterLightCardEditorViewportClient::CalculateFinalLightCardDistance(double FlushDistance, double DesiredOffsetFromFlush) const
+double FDisplayClusterLightCardEditorViewportClient::CalculateFinalActorDistance(double FlushDistance, double DesiredOffsetFromFlush) const
 {
 	double Distance = FMath::Min(FlushDistance, RootActorBoundingRadius) + DesiredOffsetFromFlush;
 
@@ -1023,25 +1081,25 @@ void FDisplayClusterLightCardEditorViewportClient::ProcessClick(FSceneView& View
 				if (ActorHitProxy->PrimComponent && ActorHitProxy->PrimComponent->IsA<UStaticMeshComponent>())
 				{
 					ADisplayClusterLightCardActor* TracedLightCard = TraceScreenForLightCard(View, HitX, HitY);
-					SelectLightCard(TracedLightCard, bMultiSelect);
+					SelectActor(TracedLightCard, bMultiSelect);
 				}
 			}
-			else if (ActorHitProxy->Actor->IsA<ADisplayClusterLightCardActor>() && LightCardProxies.Contains(ActorHitProxy->Actor))
+			else if (UE::DisplayClusterLightCardEditorUtils::IsProxySelectable(ActorHitProxy->Actor) && ActorProxies.Contains(ActorHitProxy->Actor))
 			{
-				SelectLightCard(Cast<ADisplayClusterLightCardActor>(ActorHitProxy->Actor), bMultiSelect);
+				SelectActor(ActorHitProxy->Actor, bMultiSelect);
 			}
 			else if (!bMultiSelect)
 			{
-				SelectLightCard(nullptr);
+				SelectActor(nullptr);
 			}
 		}
 	}
 	else
 	{
-		SelectLightCard(nullptr);
+		SelectActor(nullptr);
 	}
 	
-	PropagateLightCardSelection();
+	PropagateActorSelection();
 
 	if (bIsRightClickSelection)
 	{
@@ -1144,7 +1202,7 @@ EMouseCursor::Type FDisplayClusterLightCardEditorViewportClient::GetCursor(FView
 						}
 					}
 				}
-				else if (LightCardProxies.Contains(ActorHitProxy->Actor))
+				else if (ActorProxies.Contains(ActorHitProxy->Actor))
 				{
 					MouseCursor = EMouseCursor::Crosshairs;
 				}
@@ -1171,7 +1229,7 @@ void FDisplayClusterLightCardEditorViewportClient::DestroyDropPreviewActors()
 	{
 		for (auto ActorIt = DropPreviewLightCards.CreateConstIterator(); ActorIt; ++ActorIt)
 		{
-			ADisplayClusterLightCardActor* PreviewActor = (*ActorIt).Get();
+			AActor* PreviewActor = (*ActorIt).AsActor();
 			if (PreviewActor)
 			{
 				IDisplayClusterScenePreview::Get().RemoveActorFromRenderer(PreviewRendererId, PreviewActor);
@@ -1195,15 +1253,15 @@ bool FDisplayClusterLightCardEditorViewportClient::UpdateDropPreviewActors(int32
 
 	const FIntPoint NewMousePos { MouseX, MouseY };
 
-	for (TWeakObjectPtr<ADisplayClusterLightCardActor>& LightCard : DropPreviewLightCards)
+	for (const FDisplayClusterWeakStageActorPtr& Actor : DropPreviewLightCards)
 	{
-		if (LightCard.IsValid())
+		if (Actor.IsValid())
 		{
-			ProjectionHelper->VerifyAndFixLightCardOrigin(*LightCard);
+			ProjectionHelper->VerifyAndFixActorOrigin(Actor);
 		}
 	}
 	
-	MoveLightCardsToPixel(NewMousePos, DropPreviewLightCards);
+	MoveActorsToPixel(NewMousePos, DropPreviewLightCards);
 	
 	return true;
 }
@@ -1228,14 +1286,14 @@ bool FDisplayClusterLightCardEditorViewportClient::DropObjectsAtCoordinates(int3
 
 	if (bSelectActors)
 	{
-		SelectedLightCards.Empty();
+		SelectedActors.Empty();
 	}
 	
-	TArray<TWeakObjectPtr<ADisplayClusterLightCardActor>> CreatedLightCards;
+	TArray<FDisplayClusterWeakStageActorPtr> CreatedLightCards;
 	CreatedLightCards.Reserve(DroppedObjects.Num());
 
-	SelectLightCard(nullptr);
-	PropagateLightCardSelection();
+	SelectActor(nullptr);
+	PropagateActorSelection();
 	
 	for (UObject* DroppedObject : DroppedObjects)
 	{
@@ -1245,7 +1303,7 @@ bool FDisplayClusterLightCardEditorViewportClient::DropObjectsAtCoordinates(int3
 				bCreateDropPreview ? PreviewWorld->GetCurrentLevel() : nullptr, bCreateDropPreview);
 			check(LightCardActor);
 
-			ProjectionHelper->VerifyAndFixLightCardOrigin(*LightCardActor);
+			ProjectionHelper->VerifyAndFixActorOrigin(LightCardActor);
 			
 			CreatedLightCards.Add(LightCardActor);
 			
@@ -1253,7 +1311,7 @@ bool FDisplayClusterLightCardEditorViewportClient::DropObjectsAtCoordinates(int3
 			{
 				DropPreviewLightCards.Add(LightCardActor);
 				
-				LightCardActor->bIsProxy = true;
+				LightCardActor->SetFlags(RF_Transient);
 				IDisplayClusterScenePreview::Get().AddActorToRenderer(PreviewRendererId, LightCardActor);
 			}
 			else
@@ -1265,8 +1323,8 @@ bool FDisplayClusterLightCardEditorViewportClient::DropObjectsAtCoordinates(int3
 					GetOnNextSceneRefresh().AddLambda([this, LightCardActor]()
 					{
 						// Select on next refresh so the persistent level and corresponding proxy has spawned.
-						SelectLightCard(LightCardActor, true);
-						PropagateLightCardSelection();
+						SelectActor(LightCardActor, true);
+						PropagateActorSelection();
 					});
 				}
 			}
@@ -1276,7 +1334,7 @@ bool FDisplayClusterLightCardEditorViewportClient::DropObjectsAtCoordinates(int3
 	}
 
 	const FIntPoint NewMousePos { MouseX, MouseY };
-	MoveLightCardsToPixel(NewMousePos, CreatedLightCards);
+	MoveActorsToPixel(NewMousePos, CreatedLightCards);
 
 	return bSuccess;
 }
@@ -1374,62 +1432,78 @@ void FDisplayClusterLightCardEditorViewportClient::UpdatePreviewActor(ADisplayCl
 			if (ProxyType == EDisplayClusterLightCardEditorProxyType::All ||
 				ProxyType == EDisplayClusterLightCardEditorProxyType::LightCards)
 			{
-				TSet<ADisplayClusterLightCardActor*> LightCards;
-				UDisplayClusterBlueprintLib::FindLightCardsForRootActor(RootActorPtr.Get(), LightCards);
-				
-				SelectLightCard(nullptr);
-				
+				SelectActor(nullptr);
 				const FTransform RALevelTransformNoScale(RootActorLevelInstance->GetActorRotation(), RootActorLevelInstance->GetActorLocation(), FVector::OneVector);
 
-				for (ADisplayClusterLightCardActor* LightCard : LightCards)
+				if (LightCardEditorPtr.IsValid())
 				{
-					FObjectDuplicationParameters DupeActorParameters(LightCard, PreviewWorld->GetCurrentLevel());
-					DupeActorParameters.FlagMask = RF_AllFlags & ~(RF_ArchetypeObject | RF_Transactional);
-					DupeActorParameters.PortFlags = PPF_DuplicateVerbatim;
-				
-					ADisplayClusterLightCardActor* LightCardProxy = CastChecked<ADisplayClusterLightCardActor>(StaticDuplicateObjectEx(DupeActorParameters));
-					PreviewWorld->GetCurrentLevel()->AddLoadedActor(LightCardProxy);
-
-					const FTransform LCLevelRelativeToRALevel = LightCard->GetTransform().GetRelativeTransform(RALevelTransformNoScale);
-
-					LightCardProxy->SetActorTransform(LCLevelRelativeToRALevel);
-					LightCardProxy->PolygonMask = LightCard->PolygonMask;
-					LightCardProxy->bIsProxy = true;
-
-					// Change mesh to proxy mesh with more vertices
-					if (const UStaticMesh* LightCardMesh = LightCardProxy->GetStaticMesh())
-					{
-						// Only change the mesh if we are using the default one.
+					TArray<AActor*> ManagedActors = LightCardEditorPtr.Pin()->FindAllManagedActors();
 					
-						const FString DefaultPlanePath = TEXT("/nDisplay/LightCard/SM_LightCardPlane.SM_LightCardPlane");
-						const UStaticMesh* DefaultPlane = Cast<UStaticMesh>(FSoftObjectPath(DefaultPlanePath).TryLoad());
-
-						if (DefaultPlane == LightCardMesh)
+					for (AActor* Actor : ManagedActors)
+					{
+						// Create proxy
+						FObjectDuplicationParameters DupeActorParameters(Actor, PreviewWorld->GetCurrentLevel());
+						DupeActorParameters.FlagMask = RF_AllFlags & ~(RF_ArchetypeObject | RF_Transactional);
+						DupeActorParameters.PortFlags = PPF_DuplicateVerbatim;
+
+						AActor* ActorProxy = CastChecked<AActor>(StaticDuplicateObjectEx(DupeActorParameters));
+						ActorProxy->SetFlags(RF_Transient); // This signals to the stage actor it is a proxy
+						PreviewWorld->GetCurrentLevel()->AddLoadedActor(ActorProxy);
+
+						const FTransform LCLevelRelativeToRALevel = ActorProxy->GetTransform().GetRelativeTransform(RALevelTransformNoScale);
+						ActorProxy->SetActorTransform(LCLevelRelativeToRALevel);
+						
+						if (ADisplayClusterLightCardActor* LightCard = Cast<ADisplayClusterLightCardActor>(Actor))
 						{
-							const FString LightCardPlanePath = TEXT("/nDisplay/LightCard/SM_LightCardPlaneSubdivided.SM_LightCardPlaneSubdivided");
-							if (UStaticMesh* LightCardPlaneMesh = Cast<UStaticMesh>(FSoftObjectPath(LightCardPlanePath).TryLoad()))
+							ADisplayClusterLightCardActor* LightCardProxy = CastChecked<ADisplayClusterLightCardActor>(ActorProxy);
+							LightCardProxy->PolygonMask = LightCard->PolygonMask;
+
+							// Change mesh to proxy mesh with more vertices
+							if (const UStaticMesh* LightCardMesh = LightCardProxy->GetStaticMesh())
 							{
-								LightCardProxy->SetStaticMesh(LightCardPlaneMesh);
+								// Only change the mesh if we are using the default one.
+					
+								const FString DefaultPlanePath = TEXT("/nDisplay/LightCard/SM_LightCardPlane.SM_LightCardPlane");
+								const UStaticMesh* DefaultPlane = Cast<UStaticMesh>(FSoftObjectPath(DefaultPlanePath).TryLoad());
+
+								if (DefaultPlane == LightCardMesh)
+								{
+									const FString LightCardPlanePath = TEXT("/nDisplay/LightCard/SM_LightCardPlaneSubdivided.SM_LightCardPlaneSubdivided");
+									if (UStaticMesh* LightCardPlaneMesh = Cast<UStaticMesh>(FSoftObjectPath(LightCardPlanePath).TryLoad()))
+									{
+										LightCardProxy->SetStaticMesh(LightCardPlaneMesh);
+									}
+								}
 							}
 						}
-					}
+						else
+						{
+							// Misc actor proxies
+							TArray<UBillboardComponent*> BillboardComponents;
+							ActorProxy->GetComponents<UBillboardComponent>(BillboardComponents);
 
-					LightCardProxies.Add(FLightCardProxy(LightCard, LightCardProxy));
+							BillboardComponentProxies.Append(BillboardComponents);
+						}
+						
+						ActorProxies.Add(FActorProxy(Actor, ActorProxy));
+						ProjectionHelper->VerifyAndFixActorOrigin(Actor);
+						ProjectionHelper->VerifyAndFixActorOrigin(ActorProxy);
+					}
 				}
 
 				// Update the selected light card proxies to match the currently selected light cards in the light card list
-				TArray<ADisplayClusterLightCardActor*> CurrentlySelectedLightCards;
+				TArray<AActor*> CurrentlySelectedActors;
 				if (LightCardEditorPtr.IsValid())
 				{
-					LightCardEditorPtr.Pin()->GetSelectedLightCards(CurrentlySelectedLightCards);
+					LightCardEditorPtr.Pin()->GetSelectedActors(CurrentlySelectedActors);
 				}
 
-				SelectLightCards(CurrentlySelectedLightCards);
+				SelectActors(CurrentlySelectedActors);
 			}
 
-			for (const FLightCardProxy& LightCardProxy : LightCardProxies)
+			for (const FActorProxy& ActorProxy : ActorProxies)
 			{
-				IDisplayClusterScenePreview::Get().AddActorToRenderer(PreviewRendererId, LightCardProxy.Proxy.Get(), [this](const UPrimitiveComponent* PrimitiveComponent)
+				IDisplayClusterScenePreview::Get().AddActorToRenderer(PreviewRendererId, ActorProxy.Proxy.AsActor(), [this, ActorProxy](const UPrimitiveComponent* PrimitiveComponent)
 				{
 					// Always add the light card mesh component to the renderer's scene even if it is marked hidden in game, since UV light cards will purposefully
 					// hide the light card mesh since it isn't supposed to exist in 3D space. The light card mesh will be appropriately filtered when the scene is
@@ -1458,25 +1532,20 @@ void FDisplayClusterLightCardEditorViewportClient::UpdateProxyTransforms()
 			RootActorProxy->SetActorScale3D(RootActorLevelInstance->GetActorScale3D());
 		}
 		
-		for (const FLightCardProxy& LightCardProxy : LightCardProxies)
+		for (const FActorProxy& ActorProxy : ActorProxies)
 		{
-			if (LightCardProxy.LevelInstance.IsValid() && LightCardProxy.Proxy.IsValid())
+			if (ActorProxy.LevelInstance.IsValid() && ActorProxy.Proxy.IsValid())
 			{
 				const FTransform RALevelTransformNoScale(RootActorLevelInstance->GetActorRotation(), RootActorLevelInstance->GetActorLocation(), FVector::OneVector);
-				const FTransform LCLevelRelativeToRALevel = LightCardProxy.LevelInstance->GetTransform().GetRelativeTransform(RALevelTransformNoScale);
-
-				LightCardProxy.Proxy->SetActorTransform(LCLevelRelativeToRALevel);
-
+				const FTransform LCLevelRelativeToRALevel = ActorProxy.LevelInstance.AsActorChecked()->GetTransform().GetRelativeTransform(RALevelTransformNoScale);
+				ActorProxy.Proxy.AsActorChecked()->SetActorTransform(LCLevelRelativeToRALevel);
+				
+				// When dealing with light card actors the transform scale can be set separately from the stage actor 2d scale. When updating the
+				// transform scale on the level instance it won't impact the positional params so the proxy transform scale won't update unless we manually set it
+				ActorProxy.Proxy.AsActorChecked()->SetActorScale3D(ActorProxy.LevelInstance.AsActorChecked()->GetActorScale3D());
+				
 				// Need to update these manually or the proxy's position will be out of sync next update
-				LightCardProxy.Proxy->Longitude = LightCardProxy.LevelInstance->Longitude;
-				LightCardProxy.Proxy->Latitude = LightCardProxy.LevelInstance->Latitude;
-				LightCardProxy.Proxy->DistanceFromCenter = LightCardProxy.LevelInstance->DistanceFromCenter;
-				LightCardProxy.Proxy->Spin = LightCardProxy.LevelInstance->Spin;
-				LightCardProxy.Proxy->Pitch = LightCardProxy.LevelInstance->Pitch;
-				LightCardProxy.Proxy->Yaw = LightCardProxy.LevelInstance->Yaw;
-				LightCardProxy.Proxy->Scale = LightCardProxy.LevelInstance->Scale;
-				LightCardProxy.Proxy->UVCoordinates = LightCardProxy.LevelInstance->UVCoordinates;
-				LightCardProxy.Proxy->RadialOffset = LightCardProxy.LevelInstance->RadialOffset;
+				ActorProxy.Proxy->SetPositionalParams(ActorProxy.LevelInstance->GetPositionalParams());
 			}
 		}
 	}
@@ -1496,11 +1565,11 @@ void FDisplayClusterLightCardEditorViewportClient::DestroyProxies(
 		break;
 
 	case EDisplayClusterLightCardEditorProxyType::LightCards:
-		for (const FLightCardProxy& LightCardProxy : LightCardProxies)
+		for (const FActorProxy& ActorProxy : ActorProxies)
 		{
-			if (LightCardProxy.Proxy.IsValid())
+			if (ActorProxy.Proxy.IsValid())
 			{
-				IDisplayClusterScenePreview::Get().RemoveActorFromRenderer(PreviewRendererId, LightCardProxy.Proxy.Get());
+				IDisplayClusterScenePreview::Get().RemoveActorFromRenderer(PreviewRendererId, ActorProxy.Proxy.AsActor());
 			}
 		}
 		break;
@@ -1535,28 +1604,29 @@ void FDisplayClusterLightCardEditorViewportClient::DestroyProxies(
 	if (ProxyType == EDisplayClusterLightCardEditorProxyType::All ||
 		ProxyType == EDisplayClusterLightCardEditorProxyType::LightCards)
 	{
-		for (const FLightCardProxy& LightCardProxy : LightCardProxies)
+		for (const FActorProxy& ActorProxy : ActorProxies)
 		{
-			if (LightCardProxy.Proxy.IsValid())
+			if (ActorProxy.Proxy.IsValid())
 			{
-				PreviewWorld->EditorDestroyActor(LightCardProxy.Proxy.Get(), false);
+				PreviewWorld->EditorDestroyActor(ActorProxy.Proxy.AsActor(), false);
 			}
 		}
 
-		LightCardProxies.Empty();	
+		ActorProxies.Empty();
+		BillboardComponentProxies.Empty();
 	}
 }
 
-void FDisplayClusterLightCardEditorViewportClient::SelectLightCards(const TArray<ADisplayClusterLightCardActor*>& LightCardsToSelect)
+void FDisplayClusterLightCardEditorViewportClient::SelectActors(const TArray<AActor*>& ActorsToSelect)
 {
-	SelectLightCard(nullptr);
-	for (AActor* LightCard : LightCardsToSelect)
+	SelectActor(nullptr);
+	for (AActor* Actor : ActorsToSelect)
 	{
-		if (FLightCardProxy* FoundProxy = LightCardProxies.FindByKey(LightCard))
+		if (const FActorProxy* FoundProxy = ActorProxies.FindByKey(Actor))
 		{
 			if (FoundProxy->Proxy.IsValid())
 			{
-				SelectLightCard(FoundProxy->Proxy.Get(), true);
+				SelectActor(FoundProxy->Proxy.AsActor(), true);
 			}
 		}
 	}
@@ -1676,36 +1746,34 @@ void FDisplayClusterLightCardEditorViewportClient::SetCoordinateSystem(FDisplayC
 	}
 }
 
-void FDisplayClusterLightCardEditorViewportClient::MoveLightCardTo(ADisplayClusterLightCardActor& LightCard, const FDisplayClusterLightCardEditorHelper::FSphericalCoordinates& SphericalCoords) const
+void FDisplayClusterLightCardEditorViewportClient::MoveActorTo(const FDisplayClusterWeakStageActorPtr& Actor, const FDisplayClusterLightCardEditorHelper::FSphericalCoordinates& SphericalCoords) const
 {
-	const FVector Origin = GetViewLocation(); // Assumed the same as the LC origin in proxy space. Call VerifyAndFixLightCardOrigin to ensure this.
-
-	ProjectionHelper->MoveLightCardsTo({ &LightCard }, SphericalCoords);
+	ProjectionHelper->MoveActorsTo({ Actor }, SphericalCoords);
 }
 
-void FDisplayClusterLightCardEditorViewportClient::CenterLightCardInView(ADisplayClusterLightCardActor& LightCard)
+void FDisplayClusterLightCardEditorViewportClient::CenterActorInView(const FDisplayClusterWeakStageActorPtr& Actor)
 {
-	ProjectionHelper->VerifyAndFixLightCardOrigin(LightCard);
+	ProjectionHelper->VerifyAndFixActorOrigin(Actor);
 
-	if (LightCard.bIsUVLightCard)
+	if (Actor->IsUVActor())
 	{
-		LightCard.UVCoordinates = FVector2D(0.5, 0.5);
+		Actor->SetUVCoordinates(FVector2D(0.5, 0.5));
 	}
 	else
 	{
-		MoveLightCardTo(LightCard, FDisplayClusterLightCardEditorHelper::FSphericalCoordinates(GetViewRotation().RotateVector(FVector::ForwardVector)));
+		MoveActorTo(Actor, FDisplayClusterLightCardEditorHelper::FSphericalCoordinates(GetViewRotation().RotateVector(FVector::ForwardVector)));
 	}
 
 	// If this is a proxy light, propagate to its counterpart in the level.
-	if (LightCard.bIsProxy)
+	if (Actor->IsProxy())
 	{
-		PropagateLightCardTransform(&LightCard);
+		PropagateActorTransform(Actor);
 	}
 }
 
-void FDisplayClusterLightCardEditorViewportClient::MoveSelectedLightCardsToPixel(const FIntPoint& PixelPos)
+void FDisplayClusterLightCardEditorViewportClient::MoveSelectedActorsToPixel(const FIntPoint& PixelPos)
 {
-	MoveLightCardsToPixel(PixelPos, SelectedLightCards);
+	MoveActorsToPixel(PixelPos, SelectedActors);
 }
 
 void FDisplayClusterLightCardEditorViewportClient::BeginTransaction(const FText& Description)
@@ -1872,34 +1940,34 @@ USceneComponent* FDisplayClusterLightCardEditorViewportClient::FindProjectionOri
 	return InRootActor->GetCommonViewPoint();
 }
 
-bool FDisplayClusterLightCardEditorViewportClient::IsLightCardSelected(const AActor* Actor)
+bool FDisplayClusterLightCardEditorViewportClient::IsActorSelected(const AActor* Actor)
 {
-	return SelectedLightCards.Contains(Actor);
+	return Actor->Implements<UDisplayClusterStageActor>() && SelectedActors.Contains(Actor);
 }
 
-void FDisplayClusterLightCardEditorViewportClient::SelectLightCard(ADisplayClusterLightCardActor* Actor, bool bAddToSelection)
+void FDisplayClusterLightCardEditorViewportClient::SelectActor(AActor* Actor, bool bAddToSelection)
 {
-	TArray<ADisplayClusterLightCardActor*> UpdatedActors;
+	TArray<AActor*> UpdatedActors;
 
 	if (!bAddToSelection)
 	{
-		for (const TWeakObjectPtr<ADisplayClusterLightCardActor>& LightCard : SelectedLightCards)
+		for (const FDisplayClusterWeakStageActorPtr& SelectedActor : SelectedActors)
 		{
-			if (LightCard.IsValid())
+			if (SelectedActor.IsValid())
 			{
-				UpdatedActors.Add(LightCard.Get());
+				UpdatedActors.Add(SelectedActor.AsActor());
 			}
 		}
 
-		SelectedLightCards.Empty();
-		LastSelectedLightCard = nullptr;
+		SelectedActors.Empty();
+		LastSelectedActor.Reset();
 	}
 
 	if (Actor)
 	{
-		SelectedLightCards.Add(Actor);
+		SelectedActors.Add(Actor);
 		UpdatedActors.Add(Actor);
-		LastSelectedLightCard = Actor;
+		LastSelectedActor = Actor;
 	}
 
 	for (AActor* UpdatedActor : UpdatedActors)
@@ -1908,48 +1976,45 @@ void FDisplayClusterLightCardEditorViewportClient::SelectLightCard(ADisplayClust
 	}
 }
 
-void FDisplayClusterLightCardEditorViewportClient::PropagateLightCardSelection()
+void FDisplayClusterLightCardEditorViewportClient::PropagateActorSelection()
 {
-	TArray<ADisplayClusterLightCardActor*> SelectedLevelInstances;
-	for (const TWeakObjectPtr<ADisplayClusterLightCardActor>& SelectedLightCard : SelectedLightCards)
+	TArray<AActor*> SelectedLevelInstances;
+	for (const FDisplayClusterWeakStageActorPtr& SelectedActor : SelectedActors)
 	{
-		if (FLightCardProxy* FoundProxy = LightCardProxies.FindByKey(SelectedLightCard.Get()))
+		if (const FActorProxy* FoundProxy = ActorProxies.FindByKey(SelectedActor.AsActor()))
 		{
 			if (FoundProxy->LevelInstance.IsValid())
 			{
-				SelectedLevelInstances.Add(FoundProxy->LevelInstance.Get());
+				SelectedLevelInstances.Add(FoundProxy->LevelInstance.AsActor());
 			}
 		}
 	}
 
-	LightCardEditorPtr.Pin()->SelectLightCards(SelectedLevelInstances);
+	LightCardEditorPtr.Pin()->SelectActors(SelectedLevelInstances);
 }
 
-void FDisplayClusterLightCardEditorViewportClient::PropagateLightCardTransform(ADisplayClusterLightCardActor* LightCardProxy)
+void FDisplayClusterLightCardEditorViewportClient::PropagateActorTransform(const FDisplayClusterWeakStageActorPtr& ActorProxy)
 {
-	const FLightCardProxy* FoundProxy = LightCardProxies.FindByKey(LightCardProxy);
-	if (FoundProxy && FoundProxy->Proxy == LightCardProxy && FoundProxy->LevelInstance.IsValid())
+	const FActorProxy* FoundProxy = ActorProxies.FindByKey(ActorProxy.AsActor());
+	if (FoundProxy && FoundProxy->Proxy == ActorProxy && FoundProxy->LevelInstance.IsValid())
 	{
-		ADisplayClusterLightCardActor* LevelInstance = FoundProxy->LevelInstance.Get();
-
+		AActor* LevelInstance = FoundProxy->LevelInstance.AsActorChecked();
 		LevelInstance->Modify();
-
-		TArray<const FProperty*> ChangedProperties;
 		
 		// Set the level instance property value to our proxy property value.
-		auto TryChangeProperty = [&](FName InPropertyName) -> void
+		auto TryChangeProperty = [&](FName InPropertyName, TArray<const FProperty*>& InOutChangedProperties) -> void
 		{
-			const FProperty* Property = FindFProperty<FProperty>(LevelInstance->GetClass(), InPropertyName);
-			check(Property);
-			
-			// Only change if values are different.
-			if (!Property->Identical_InContainer(LightCardProxy, LevelInstance))
+			if (const FProperty* Property = FindFProperty<FProperty>(LevelInstance->GetClass(), InPropertyName))
 			{
-				void* NewValue = nullptr;
-				Property->GetValue_InContainer(LightCardProxy, &NewValue);
-				Property->SetValue_InContainer(LevelInstance, &NewValue);
+				// Only change if values are different.
+				if (!Property->Identical_InContainer(ActorProxy.AsActorChecked(), LevelInstance))
+				{
+					void* NewValue = nullptr;
+					Property->GetValue_InContainer(ActorProxy.AsActorChecked(), &NewValue);
+					Property->SetValue_InContainer(LevelInstance, &NewValue);
 
-				ChangedProperties.Add(Property);
+					InOutChangedProperties.Add(Property);
+				}
 			}
 		};
 		
@@ -1962,27 +2027,29 @@ void FDisplayClusterLightCardEditorViewportClient::PropagateLightCardTransform(A
 
 		const FTransform RALevelTransformNoScale(RALevelRotation, RALevelLocation, FVector::OneVector);
 
-		LevelInstance->SetActorTransform(LightCardProxy->GetTransform() * RALevelTransformNoScale);
-
-		TryChangeProperty(GET_MEMBER_NAME_CHECKED(ADisplayClusterLightCardActor, Longitude));
-		TryChangeProperty(GET_MEMBER_NAME_CHECKED(ADisplayClusterLightCardActor, Latitude));
-		TryChangeProperty(GET_MEMBER_NAME_CHECKED(ADisplayClusterLightCardActor, DistanceFromCenter));
-		TryChangeProperty(GET_MEMBER_NAME_CHECKED(ADisplayClusterLightCardActor, Spin));
-		TryChangeProperty(GET_MEMBER_NAME_CHECKED(ADisplayClusterLightCardActor, Pitch));
-		TryChangeProperty(GET_MEMBER_NAME_CHECKED(ADisplayClusterLightCardActor, Yaw));
-		TryChangeProperty(GET_MEMBER_NAME_CHECKED(ADisplayClusterLightCardActor, Scale));
-		TryChangeProperty(GET_MEMBER_NAME_CHECKED(ADisplayClusterLightCardActor, UVCoordinates));
-		TryChangeProperty(GET_MEMBER_NAME_CHECKED(ADisplayClusterLightCardActor, RadialOffset));
+		LevelInstance->SetActorTransform(ActorProxy.AsActorChecked()->GetTransform() * RALevelTransformNoScale);
+
+		const TSet<FName>& PropertyNames = ActorProxy->GetPositionalPropertyNames();
+
+		TArray<const FProperty*> ChangedProperties;
+		ChangedProperties.Reserve(PropertyNames.Num());
+		for (const FName& PropertyName : PropertyNames)
+		{
+			TryChangeProperty(PropertyName, ChangedProperties);
+		}
 		
 		// Snapshot the changed properties so multi-user can update while dragging.
 		if (ChangedProperties.Num() > 0)
 		{
 			SnapshotTransactionBuffer(LevelInstance, MakeArrayView(ChangedProperties));
 		}
+
+		// Allows MU to receive the update in real-time.
+		LevelInstance->PostEditMove(false);
 	}
 }
 
-void FDisplayClusterLightCardEditorViewportClient::MoveSelectedLightCards(FViewport* InViewport, EAxisList::Type CurrentAxis)
+void FDisplayClusterLightCardEditorViewportClient::MoveSelectedActors(FViewport* InViewport, EAxisList::Type CurrentAxis)
 {
 	FSceneViewFamilyContext ViewFamily(FSceneViewFamily::ConstructionValues(
 		InViewport,
@@ -1997,12 +2064,12 @@ void FDisplayClusterLightCardEditorViewportClient::MoveSelectedLightCards(FViewp
 	InViewport->GetMousePos(MousePos);
 
 	// Move the light cards
-	ProjectionHelper->DragLightCards(SelectedLightCards, MousePos, *View, EditorWidgetCoordinateSystem, DragWidgetOffset, CurrentAxis, LastSelectedLightCard.Get());
+	ProjectionHelper->DragActors(SelectedActors, MousePos, *View, EditorWidgetCoordinateSystem, DragWidgetOffset, CurrentAxis, LastSelectedActor);
 
 	// Update the level instances
-	for (const TWeakObjectPtr<ADisplayClusterLightCardActor>& LightCard : SelectedLightCards)
+	for (const FDisplayClusterWeakStageActorPtr& SelectedActor : SelectedActors)
 	{
-		PropagateLightCardTransform(LightCard.Get());
+		PropagateActorTransform(SelectedActor);
 	}
 
 	FSceneViewInitOptions SceneViewInitOptions;
@@ -2023,7 +2090,7 @@ void FDisplayClusterLightCardEditorViewportClient::MoveSelectedLightCards(FViewp
 	}
 }
 
-void FDisplayClusterLightCardEditorViewportClient::MoveLightCardsToPixel(const FIntPoint& PixelPos, const TArray<TWeakObjectPtr<ADisplayClusterLightCardActor>>& InLightCards)
+void FDisplayClusterLightCardEditorViewportClient::MoveActorsToPixel(const FIntPoint& PixelPos, const TArray<FDisplayClusterWeakStageActorPtr>& InActors)
 {
 	FSceneViewFamilyContext ViewFamily(FSceneViewFamily::ConstructionValues(
 		Viewport,
@@ -2034,41 +2101,40 @@ void FDisplayClusterLightCardEditorViewportClient::MoveLightCardsToPixel(const F
 
 	FSceneView* View = CalcSceneView(&ViewFamily);
 
-	ProjectionHelper->MoveLightCardsToPixel(SelectedLightCards, PixelPos, *View);
+	ProjectionHelper->MoveActorsToPixel(SelectedActors, PixelPos, *View);
 
-	// Update each light card with the delta coordinates; the flush constraint is applied by MoveLightCardTo, ensuring the light card is always flush to screens
-	for (const TWeakObjectPtr<ADisplayClusterLightCardActor>& LightCard : InLightCards)
+	// Update each light card with the delta coordinates; the flush constraint is applied by MoveActorTo, ensuring the light card is always flush to screens
+	for (const FDisplayClusterWeakStageActorPtr& LightCard : InActors)
 	{
 		if (LightCard.IsValid() &&
-			((LightCard->bIsUVLightCard && ProjectionMode == EDisplayClusterMeshProjectionType::UV) ||
-			(!LightCard->bIsUVLightCard && ProjectionMode != EDisplayClusterMeshProjectionType::UV)))
+			((LightCard->IsUVActor() && ProjectionMode == EDisplayClusterMeshProjectionType::UV) ||
+			(!LightCard->IsUVActor() && ProjectionMode != EDisplayClusterMeshProjectionType::UV)))
 		{
-			PropagateLightCardTransform(LightCard.Get());
+			PropagateActorTransform(LightCard);
 		}
 	}
 }
 
-
-void FDisplayClusterLightCardEditorViewportClient::MoveSelectedUVLightCards(FViewport* InViewport, EAxisList::Type CurrentAxis)
+void FDisplayClusterLightCardEditorViewportClient::MoveSelectedUVActors(FViewport* InViewport, EAxisList::Type CurrentAxis)
 {
-	if (LastSelectedLightCard.IsValid() && LastSelectedLightCard->bIsUVLightCard)
+	if (LastSelectedActor.IsValid() && LastSelectedActor->IsUVActor())
 	{
-		const FVector2D DeltaUV = GetUVLightCardTranslationDelta(InViewport, LastSelectedLightCard.Get(), CurrentAxis);
-		for (const TWeakObjectPtr<ADisplayClusterLightCardActor>& LightCard : SelectedLightCards)
+		const FVector2D DeltaUV = GetUVActorTranslationDelta(InViewport, LastSelectedActor, CurrentAxis);
+		for (const FDisplayClusterWeakStageActorPtr& Actor : SelectedActors)
 		{
-			if (!LightCard.IsValid() || !LightCard->bIsUVLightCard)
+			if (!Actor.IsValid() || !Actor->IsUVActor())
 			{
 				continue;
 			}
 
-			LightCard->UVCoordinates += DeltaUV;
+			Actor->SetUVCoordinates(Actor->GetUVCoordinates() + DeltaUV);
 
-			PropagateLightCardTransform(LightCard.Get());
+			PropagateActorTransform(Actor);
 		}
 	}
 }
 
-FVector2D FDisplayClusterLightCardEditorViewportClient::GetUVLightCardTranslationDelta(FViewport* InViewport, ADisplayClusterLightCardActor* LightCard, EAxisList::Type CurrentAxis)
+FVector2D FDisplayClusterLightCardEditorViewportClient::GetUVActorTranslationDelta(FViewport* InViewport, const FDisplayClusterWeakStageActorPtr& InActor, EAxisList::Type CurrentAxis)
 {
 	FIntPoint MousePos;
 	InViewport->GetMousePos(MousePos);
@@ -2094,7 +2160,7 @@ FVector2D FDisplayClusterLightCardEditorViewportClient::GetUVLightCardTranslatio
 	const FVector DesiredLocation = PlaneIntersection - DragWidgetOffset;
 	const FVector2D DesiredUVLocation = FVector2D(DesiredLocation.Y / UVProjectionPlaneSize + 0.5f, 0.5f - DesiredLocation.Z / UVProjectionPlaneSize);
 
-	const FVector2D UVDelta = DesiredUVLocation - LightCard->UVCoordinates;
+	const FVector2D UVDelta = DesiredUVLocation - InActor->GetUVCoordinates();
 
 	FVector2D UVAxis = FVector2D::ZeroVector;
 	if (CurrentAxis & EAxisList::Type::X)
@@ -2110,27 +2176,26 @@ FVector2D FDisplayClusterLightCardEditorViewportClient::GetUVLightCardTranslatio
 	return UVDelta * UVAxis;
 }
 
-void FDisplayClusterLightCardEditorViewportClient::ScaleSelectedLightCards(FViewport* InViewport, EAxisList::Type CurrentAxis)
+void FDisplayClusterLightCardEditorViewportClient::ScaleSelectedActors(FViewport* InViewport, EAxisList::Type CurrentAxis)
 {
-	if (LastSelectedLightCard.IsValid())
+	if (LastSelectedActor.IsValid())
 	{
-		const FVector2D DeltaScale = GetLightCardScaleDelta(InViewport, LastSelectedLightCard.Get(), CurrentAxis);
-		for (const TWeakObjectPtr<ADisplayClusterLightCardActor>& LightCard : SelectedLightCards)
+		const FVector2D DeltaScale = GetActorScaleDelta(InViewport, LastSelectedActor, CurrentAxis);
+		for (const FDisplayClusterWeakStageActorPtr& SelectedActor : SelectedActors)
 		{
-			if (!LightCard.IsValid())
-			{
-				continue;
-			}
-
-			LightCard->Scale += DeltaScale;
-
-			PropagateLightCardTransform(LightCard.Get());
+			SelectedActor->SetScale(SelectedActor->GetScale() + DeltaScale);
+			PropagateActorTransform(SelectedActor);
 		}
 	}
 }
 
-FVector2D FDisplayClusterLightCardEditorViewportClient::GetLightCardScaleDelta(FViewport* InViewport, ADisplayClusterLightCardActor* LightCard, EAxisList::Type CurrentAxis)
+FVector2D FDisplayClusterLightCardEditorViewportClient::GetActorScaleDelta(FViewport* InViewport, const FDisplayClusterWeakStageActorPtr& InActor, EAxisList::Type CurrentAxis)
 {
+	if (InActor == nullptr)
+	{
+		return FVector2D();
+	}
+	
 	FIntPoint MousePos;
 	InViewport->GetMousePos(MousePos);
 
@@ -2174,7 +2239,7 @@ FVector2D FDisplayClusterLightCardEditorViewportClient::GetLightCardScaleDelta(F
 
 	// First, obtain the size of the unscaled light card in the direction the user is scaling. Convert to screen space as the scale and drag vectors are in screen space
 	const bool bLocalSpace = true;
-	const FVector LightCardSize3D = LightCard->GetLightCardBounds(bLocalSpace).GetSize();
+	const FVector LightCardSize3D = InActor->GetBoxBounds(bLocalSpace).GetSize();
 	const FVector2D SizeToScale = FVector2D((CurrentAxis & EAxisList::X) * LightCardSize3D.Y, (CurrentAxis & EAxisList::Y) * LightCardSize3D.Z);
 	const double DistanceFromCamera = FMath::Max(FVector::Dist(WorldWidgetOrigin, View->ViewMatrices.GetViewOrigin()), 1.0f);
 	const double ScreenSize = RenderViewportType == LVT_Perspective 
@@ -2190,32 +2255,33 @@ FVector2D FDisplayClusterLightCardEditorViewportClient::GetLightCardScaleDelta(F
 	if ((CurrentAxis & EAxisList::Type::X) && (CurrentAxis & EAxisList::Type::Y))
 	{
 		// Ensure the signs of the deltas remain the same, and avoid potential divide by zero
-		ScaleDelta.Y = ScaleDelta.X * FMath::Abs(LightCard->Scale.Y) / FMath::Max(0.001, FMath::Abs(LightCard->Scale.X));
+		FVector2D Scale = InActor->GetScale();
+		ScaleDelta.Y = ScaleDelta.X * FMath::Abs(Scale.Y) / FMath::Max(0.001, FMath::Abs(Scale.X));
 	}
 
 	return ScaleDelta;
 }
 
-void FDisplayClusterLightCardEditorViewportClient::SpinSelectedLightCards(FViewport* InViewport)
+void FDisplayClusterLightCardEditorViewportClient::SpinSelectedActors(FViewport* InViewport)
 {
-	if (LastSelectedLightCard.IsValid())
+	if (LastSelectedActor.IsValid())
 	{
-		const double DeltaSpin = GetLightCardSpinDelta(InViewport, LastSelectedLightCard.Get());
-		for (const TWeakObjectPtr<ADisplayClusterLightCardActor>& LightCard : SelectedLightCards)
+		const double DeltaSpin = GetActorSpinDelta(InViewport);
+		for (const FDisplayClusterWeakStageActorPtr& Actor : SelectedActors)
 		{
-			if (!LightCard.IsValid())
+			if (!Actor.IsValid())
 			{
 				continue;
 			}
 
-			LightCard->Spin += DeltaSpin;
+			Actor->SetSpin(Actor->GetSpin() + DeltaSpin);
 
-			PropagateLightCardTransform(LightCard.Get());
+			PropagateActorTransform(Actor);
 		}
 	}
 }
 
-double FDisplayClusterLightCardEditorViewportClient::GetLightCardSpinDelta(FViewport* InViewport, ADisplayClusterLightCardActor* LightCard)
+double FDisplayClusterLightCardEditorViewportClient::GetActorSpinDelta(FViewport* InViewport)
 {
 	FIntPoint MousePos;
 	InViewport->GetMousePos(MousePos);
@@ -2297,7 +2363,7 @@ ADisplayClusterLightCardActor* FDisplayClusterLightCardEditorViewportClient::Tra
 							{
 								if (ADisplayClusterLightCardActor* LightCardActor = Cast<ADisplayClusterLightCardActor>(HitResult.GetActor()))
 								{
-									if (LightCardProxies.Contains(LightCardActor))
+									if (ActorProxies.Contains(LightCardActor))
 									{
 										return LightCardActor;
 									}
@@ -2307,7 +2373,7 @@ ADisplayClusterLightCardActor* FDisplayClusterLightCardEditorViewportClient::Tra
 					}
 				}
 			}
-			else if (HitActor->IsA<ADisplayClusterLightCardActor>() && LightCardProxies.Contains(HitActor))
+			else if (UE::DisplayClusterLightCardEditorUtils::IsProxySelectable(HitActor) && ActorProxies.Contains(HitActor))
 			{
 				return Cast<ADisplayClusterLightCardActor>(HitActor);
 			}
@@ -2356,19 +2422,19 @@ bool FDisplayClusterLightCardEditorViewportClient::WorldToScreenDirection(const
 
 bool FDisplayClusterLightCardEditorViewportClient::CalcEditorWidgetTransform(FTransform& WidgetTransform)
 {
-	if (!SelectedLightCards.Num())
+	if (!SelectedActors.Num())
 	{
 		return false;
 	}
 
-	if (!LastSelectedLightCard.IsValid())
+	if (!LastSelectedActor.IsValid())
 	{
 		return false;
 	}
 
-	FVector LightCardPosition = LastSelectedLightCard->GetLightCardTransform().GetTranslation();
+	const FVector ActorPosition = LastSelectedActor->GetStageActorTransform(false).GetTranslation();
 
-	WidgetTransform = FTransform(FRotator::ZeroRotator, LightCardPosition, FVector::OneVector);
+	WidgetTransform = FTransform(FRotator::ZeroRotator, ActorPosition, FVector::OneVector);
 
 	FQuat WidgetOrientation;
 	if (EditorWidget->GetWidgetMode() == FDisplayClusterLightCardEditorWidget::EWidgetMode::WM_Translate)
@@ -2383,7 +2449,7 @@ bool FDisplayClusterLightCardEditorViewportClient::CalcEditorWidgetTransform(FTr
 			{
 				// The translation widget should be oriented to show the x axis pointing in the longitudinal direction and the y axis pointing in the latitudinal direction
 				const FVector ProjectionOrigin = ProjectionOriginComponent.IsValid() ? ProjectionOriginComponent->GetComponentLocation() : FVector::ZeroVector;
-				const FVector RadialVector = (LightCardPosition - ProjectionOrigin).GetSafeNormal();
+				const FVector RadialVector = (ActorPosition - ProjectionOrigin).GetSafeNormal();
 				const FVector AzimuthalVector = (FVector::ZAxisVector ^ RadialVector).GetSafeNormal();
 				const FVector InclinationVector = RadialVector ^ AzimuthalVector;
 
@@ -2398,7 +2464,7 @@ bool FDisplayClusterLightCardEditorViewportClient::CalcEditorWidgetTransform(FTr
 	else
 	{
 		// Otherwise, orient the widget to match the light card's local orientation (spin, pitch, yaw)
-		const FQuat LightCardOrientation = LastSelectedLightCard->GetLightCardTransform().GetRotation();
+		const FQuat LightCardOrientation = LastSelectedActor->GetStageActorTransform(false).GetRotation();
 		const FVector Normal = LightCardOrientation.RotateVector(FVector::XAxisVector);
 		const FVector Tangent = LightCardOrientation.RotateVector(FVector::YAxisVector);
 		const FVector Binormal = LightCardOrientation.RotateVector(FVector::ZAxisVector);
@@ -2559,7 +2625,7 @@ void FDisplayClusterLightCardEditorViewportClient::EnterDrawingLightCardMode()
 	if (InputMode == EInputMode::Idle)
 	{
 		InputMode = EInputMode::DrawingLightCard;
-		SelectLightCard(nullptr);
+		SelectActor(nullptr);
 	}
 }
 
@@ -2601,4 +2667,22 @@ bool FDisplayClusterLightCardEditorViewportClient::ShouldApplyProjectionToPrimit
 	return true;
 }
 
+FDisplayClusterLightCardEditorViewportClient::FSpriteProxy FDisplayClusterLightCardEditorViewportClient::FSpriteProxy::
+FromBillboard(const UBillboardComponent* InBillboardComponent)
+{
+	check(InBillboardComponent);
+	return FSpriteProxy{
+		InBillboardComponent->Sprite,
+		InBillboardComponent->GetComponentLocation(),
+		InBillboardComponent->U,
+		InBillboardComponent->UL,
+		InBillboardComponent->V,
+		InBillboardComponent->VL,
+		InBillboardComponent->ScreenSize,
+		InBillboardComponent->GetComponentTransform().GetMaximumAxisScale() * InBillboardComponent->GetOwner()->SpriteScale * 0.25f,
+		InBillboardComponent->OpacityMaskRefVal,
+		static_cast<bool>(InBillboardComponent->bIsScreenSizeScaled)
+	};
+}
+
 #undef LOCTEXT_NAMESPACE
diff --git a/Engine/Plugins/Runtime/nDisplay/Source/DisplayClusterLightCardEditor/Private/Viewport/DisplayClusterLightCardEditorViewportClient.h b/Engine/Plugins/Runtime/nDisplay/Source/DisplayClusterLightCardEditor/Private/Viewport/DisplayClusterLightCardEditorViewportClient.h
--- a/Engine/Plugins/Runtime/nDisplay/Source/DisplayClusterLightCardEditor/Private/Viewport/DisplayClusterLightCardEditorViewportClient.h
+++ b/Engine/Plugins/Runtime/nDisplay/Source/DisplayClusterLightCardEditor/Private/Viewport/DisplayClusterLightCardEditorViewportClient.h
@@ -12,6 +12,7 @@
 #include "DisplayClusterLightCardActor.h"
 #include "DisplayClusterLightCardEditorWidget.h"
 #include "DisplayClusterLightCardEditorHelper.h"
+#include "StageActor/DisplayClusterWeakStageActorPtr.h"
 
 class ADisplayClusterRootActor;
 class SDisplayClusterLightCardEditor;
@@ -90,8 +91,8 @@ public:
 	/** Remove proxies of the specified type */
 	void DestroyProxies(EDisplayClusterLightCardEditorProxyType ProxyType);
 	
-	/** Selects the light card proxies that correspond to the specified light cards */
-	void SelectLightCards(const TArray<ADisplayClusterLightCardActor*>& LightCardsToSelect);
+	/** Selects the actor proxies that correspond to the specified actors */
+	void SelectActors(const TArray<AActor*>& ActorsToSelect);
 
 	FDisplayClusterLightCardEditorWidget::EWidgetMode GetEditorWidgetMode() const { return EditorWidget->GetWidgetMode(); }
 	void SetEditorWidgetMode(FDisplayClusterLightCardEditorWidget::EWidgetMode InWidgetMode);
@@ -123,19 +124,19 @@ public:
 	/** Sets the editor widget's coordinate system to the specified coordinate system */
 	void SetCoordinateSystem(FDisplayClusterLightCardEditorHelper::ECoordinateSystem NewCoordinateSystem);
 
-	/** Moves specified card to desired coordinates. Actual radius will be based on flush constraint and LightCard's RadialOffset.
+	/** Moves specified card to desired coordinates. Actual radius will be based on flush constraint and Actor's RadialOffset.
 	 *
-	 * @param LightCard The light card that we are moving
-	 * @param SphericalCoords specifies desired location of light card in spherical coordinates with respect to view origin.
+	 * @param Actor The actor that we are moving
+	 * @param SphericalCoords specifies desired location of actor in spherical coordinates with respect to view origin.
 	 * 
 	*/
-	void MoveLightCardTo(ADisplayClusterLightCardActor& LightCard, const FDisplayClusterLightCardEditorHelper::FSphericalCoordinates& SphericalCoords) const;
+	void MoveActorTo(const FDisplayClusterWeakStageActorPtr& Actor, const FDisplayClusterLightCardEditorHelper::FSphericalCoordinates& SphericalCoords) const;
 
 	/** Places the given light card in the middle of the current viewport */
-	void CenterLightCardInView(ADisplayClusterLightCardActor& LightCard);
+	void CenterActorInView(const FDisplayClusterWeakStageActorPtr& Actor);
 
 	/** Moves all selected light cards to the specified pixel position */
-	void MoveSelectedLightCardsToPixel(const FIntPoint& PixelPos);
+	void MoveSelectedActorsToPixel(const FIntPoint& PixelPos);
 
 	/** Returns the current input mode */
 	EInputMode GetInputMode() { return InputMode; }
@@ -168,41 +169,41 @@ private:
 	/** Finds a suitable primitive component on the stage actor to use as a projection origin */
 	USceneComponent* FindProjectionOriginComponent(const ADisplayClusterRootActor* InRootActor) const;
 
-	/** Callback to check if an light card actor is among the list of selected light card actors */
-	bool IsLightCardSelected(const AActor* Actor);
+	/** Callback to check if an actor is among the list of selected actors */
+	bool IsActorSelected(const AActor* Actor);
 
-	/** Adds the specified light card actor the the list of selected light cards */
-	void SelectLightCard(ADisplayClusterLightCardActor* Actor, bool bAddToSelection = false);
+	/** Adds the specified actor the the list of selected actors */
+	void SelectActor(AActor* Actor, bool bAddToSelection = false);
 
-	/** Notifies the light card editor of the currently selected light cards so that it may update other UI components to match */
-	void PropagateLightCardSelection();
+	/** Notifies the light card editor of the currently selected actors so that it may update other UI components to match */
+	void PropagateActorSelection();
 
-	/** Propagates the specified light card proxy's transform back to its level instance version */
-	void PropagateLightCardTransform(ADisplayClusterLightCardActor* LightCardProxy);
+	/** Propagates the specified actor proxy's transform back to its level instance version */
+	void PropagateActorTransform(const FDisplayClusterWeakStageActorPtr& ActorProxy);
 
-	/** Moves the currently selected light cards */
-	void MoveSelectedLightCards(FViewport* InViewport, EAxisList::Type CurrentAxis);
+	/** Moves the currently selected actors */
+	void MoveSelectedActors(FViewport* InViewport, EAxisList::Type CurrentAxis);
 
-	/** Moves all given light cards to the specified pixel position */
-	void MoveLightCardsToPixel(const FIntPoint& PixelPos, const TArray<TWeakObjectPtr<ADisplayClusterLightCardActor>>& InLightCards);
+	/** Moves all given actors to the specified pixel position */
+	void MoveActorsToPixel(const FIntPoint& PixelPos, const TArray<FDisplayClusterWeakStageActorPtr>& InActors);
 
-	/** Moves the currently selected UV light cards */
-	void MoveSelectedUVLightCards(FViewport* InViewport, EAxisList::Type CurrentAxis);
+	/** Moves the currently selected UV actors */
+	void MoveSelectedUVActors(FViewport* InViewport, EAxisList::Type CurrentAxis);
 
-	/** Determines the appropriate delta in UV coordinates needed to move the specified light card to the mouse's location */
-	FVector2D GetUVLightCardTranslationDelta(FViewport* InViewport, ADisplayClusterLightCardActor* LightCard, EAxisList::Type CurrentAxis);
+	/** Determines the appropriate delta in UV coordinates needed to move the specified actor to the mouse's location */
+	FVector2D GetUVActorTranslationDelta(FViewport* InViewport, const FDisplayClusterWeakStageActorPtr& InActor, EAxisList::Type CurrentAxis);
 
-	/** Scales the currently selected light cards */
-	void ScaleSelectedLightCards(FViewport* InViewport, EAxisList::Type CurrentAxis);
+	/** Scales the currently selected actors */
+	void ScaleSelectedActors(FViewport* InViewport, EAxisList::Type CurrentAxis);
 
-	/** Determines the appropriate scale delta needed to scale the light card */
-	FVector2D GetLightCardScaleDelta(FViewport* InViewport, ADisplayClusterLightCardActor* LightCard, EAxisList::Type CurrentAxis);
+	/** Determines the appropriate scale delta needed to scale the actor */
+	FVector2D GetActorScaleDelta(FViewport* InViewport, const FDisplayClusterWeakStageActorPtr& InActor, EAxisList::Type CurrentAxis);
 
-	/** Rotates the currently selected light cards around the light card's normal axis */
-	void SpinSelectedLightCards(FViewport* InViewport);
+	/** Rotates the currently selected actors around the actor's normal axis */
+	void SpinSelectedActors(FViewport* InViewport);
 
 	/** Determines the appropriate spin delta needed to rotate the light card */
-	double GetLightCardSpinDelta(FViewport* InViewport, ADisplayClusterLightCardActor* LightCard);
+	double GetActorSpinDelta(FViewport* InViewport);
 
 	/** Traces to find the light card corresponding to a click on a stage screen */
 	ADisplayClusterLightCardActor* TraceScreenForLightCard(const FSceneView& View, int32 HitX, int32 HitY);
@@ -213,7 +214,7 @@ private:
 	/** Converts a direction vector from world space to screen screen space, and returns true of the direction vector is on the screen */
 	bool WorldToScreenDirection(const FSceneView& View, const FVector& WorldPos, const FVector& WorldDirection, FVector2D& OutScreenDir);
 
-	/** Calculates the world transform to render the editor widget with to align it with the selected light card */
+	/** Calculates the world transform to render the editor widget with to align it with the selected actor */
 	bool CalcEditorWidgetTransform(FTransform& WidgetTransform);
 
 	/** Checks if the location is approaching the edge of the view space */
@@ -237,8 +238,8 @@ private:
 	/** Creates a new light card using a polygon alpha mask as defined by the given mouse positions on the viewport */
 	void CreateDrawnLightCard(const TArray<FIntPoint>& MousePositions);
 
-	/** Calculates the final distance from the origin of a light card, given its flush distance and a desired offset */
-	double CalculateFinalLightCardDistance(double FlushDistance, double DesiredOffsetFromFlush = 0.) const;
+	/** Calculates the final distance from the origin of an actor, given its flush distance and a desired offset */
+	double CalculateFinalActorDistance(double FlushDistance, double DesiredOffsetFromFlush = 0.) const;
 
 	/** Callback passed into the mesh projection renderer to filter which primitives are drawn */
 	bool ShouldRenderPrimitive(const UPrimitiveComponent* PrimitiveComponent);
@@ -262,35 +263,55 @@ private:
 	/** The radius of the bounding sphere that entirely encapsulates the root actor */
 	float RootActorBoundingRadius = 0.0f;
 
-	struct FLightCardProxy
+	struct FActorProxy
 	{
-		TWeakObjectPtr<ADisplayClusterLightCardActor> LevelInstance;
-		TWeakObjectPtr<ADisplayClusterLightCardActor> Proxy;
+		FDisplayClusterWeakStageActorPtr LevelInstance;
+		FDisplayClusterWeakStageActorPtr Proxy;
 
-		FLightCardProxy(ADisplayClusterLightCardActor* InLevelInstance, ADisplayClusterLightCardActor* InProxy)
+		FActorProxy(const AActor* InLevelInstance, const AActor* InProxy)
 			: LevelInstance(InLevelInstance)
 			, Proxy(InProxy)
-		{ }
+		{
+		}
 
-		bool operator==(AActor* Actor) const { return (LevelInstance.IsValid() && LevelInstance.Get() == Actor) || (Proxy.IsValid() && Proxy.Get() == Actor); }
+		bool operator==(const AActor* Actor) const { return (LevelInstance.IsValid() && LevelInstance.AsActor() == Actor) || (Proxy.IsValid() && Proxy.AsActor() == Actor); }
 	};
 
-	TArray<FLightCardProxy> LightCardProxies;
-	TArray<TWeakObjectPtr<ADisplayClusterLightCardActor>> SelectedLightCards;
+	struct FSpriteProxy
+	{
+		TObjectPtr<UTexture2D> Sprite;
+		FVector WorldPosition;
+		float U;
+		float UL;
+		float V;
+		float VL;
+		float ScreenSize;
+		float SpriteScale;
+		float OpacityMaskRefVal;
+
+		bool bIsScreenSizeScaled = false;
+
+		/** Construct a sprite proxy from a billboard component */
+		static FSpriteProxy FromBillboard(const UBillboardComponent* InBillboardComponent);
+	};
+
+	TArray<FActorProxy> ActorProxies;
+	TArray<TWeakObjectPtr<UBillboardComponent>> BillboardComponentProxies;
+	TArray<FDisplayClusterWeakStageActorPtr> SelectedActors;
 
-	/** The light card in the selected light card list that was selected last */
-	TWeakObjectPtr<ADisplayClusterLightCardActor> LastSelectedLightCard = nullptr;
+	/** The actor in the selected actor list that was selected last */
+	FDisplayClusterWeakStageActorPtr LastSelectedActor;
 
 	/** Light card preview actors being dropped on the scene */
-	TArray<TWeakObjectPtr<ADisplayClusterLightCardActor>> DropPreviewLightCards;
+	TArray<FDisplayClusterWeakStageActorPtr> DropPreviewLightCards;
 	
 	/** The index of the scene preview renderer returned from IDisplayClusterScenePreview */
 	int32 PreviewRendererId = -1;
 
-	/** Helper used to convert between screen and world coordinates for lightcard positions. */
+	/** Helper used to convert between screen and world coordinates for actor positions. */
 	TSharedPtr<FDisplayClusterLightCardEditorHelper> ProjectionHelper;
 	
-	/** The LC editor widget used to manipulate light cards */
+	/** The LC editor widget used to manipulate actors */
 	TSharedPtr<FDisplayClusterLightCardEditorWidget> EditorWidget;
 
 	/** The cached editor widget transform in unprojected world space */
diff --git a/Engine/Plugins/Runtime/nDisplay/Source/DisplayClusterLightCardEditor/Private/Viewport/DisplayClusterLightcardEditorViewport.cpp b/Engine/Plugins/Runtime/nDisplay/Source/DisplayClusterLightCardEditor/Private/Viewport/DisplayClusterLightcardEditorViewport.cpp
--- a/Engine/Plugins/Runtime/nDisplay/Source/DisplayClusterLightCardEditor/Private/Viewport/DisplayClusterLightcardEditorViewport.cpp
+++ b/Engine/Plugins/Runtime/nDisplay/Source/DisplayClusterLightCardEditor/Private/Viewport/DisplayClusterLightcardEditorViewport.cpp
@@ -202,6 +202,7 @@ public:
 			MenuBuilder.BeginSection("LightCardEditorShow", LOCTEXT("ShowMenuHeader", "Show Flags"));
 			{
 				MenuBuilder.AddMenuEntry(FDisplayClusterLightCardEditorCommands::Get().ToggleAllLabels);
+				MenuBuilder.AddMenuEntry(FDisplayClusterLightCardEditorCommands::Get().ToggleIconVisibility);
 			}
 			MenuBuilder.EndSection();
 		}
@@ -545,6 +546,12 @@ void SDisplayClusterLightCardEditorViewport::BindCommands()
 			FExecuteAction::CreateSP(this, &SDisplayClusterLightCardEditorViewport::ToggleLabels),
 			FCanExecuteAction(),
 			FIsActionChecked::CreateSP(this, &SDisplayClusterLightCardEditorViewport::AreLabelsToggled));
+
+		CommandList->MapAction(
+			Commands.ToggleIconVisibility,
+			FExecuteAction::CreateSP(this, &SDisplayClusterLightCardEditorViewport::ToggleIcons),
+			FCanExecuteAction(),
+			FIsActionChecked::CreateSP(this, &SDisplayClusterLightCardEditorViewport::AreIconsToggled));
 	}
 }
 
@@ -752,8 +759,8 @@ void SDisplayClusterLightCardEditorViewport::PasteLightCardsHere()
 {
 	if (LightCardEditorPtr.IsValid())
 	{
-		const bool bOffsetLightCardPositions = false;
-		LightCardEditorPtr.Pin()->PasteLightCards(bOffsetLightCardPositions);
+		const bool bOffsetActorPosition = false;
+		LightCardEditorPtr.Pin()->PasteActors(bOffsetActorPosition);
 
 		if (ViewportClient.IsValid())
 		{
@@ -761,7 +768,7 @@ void SDisplayClusterLightCardEditorViewport::PasteLightCardsHere()
 			// regenerated until then
 			ViewportClient->GetOnNextSceneRefresh().AddLambda([this]()
 			{
-				ViewportClient->MoveSelectedLightCardsToPixel(FIntPoint(PasteHerePos.X, PasteHerePos.Y));
+				ViewportClient->MoveSelectedActorsToPixel(FIntPoint(PasteHerePos.X, PasteHerePos.Y));
 			});
 		}
 	}
@@ -771,7 +778,7 @@ bool SDisplayClusterLightCardEditorViewport::CanPasteLightCardsHere() const
 {
 	if (LightCardEditorPtr.IsValid())
 	{
-		return LightCardEditorPtr.Pin()->CanPasteLightCards();
+		return LightCardEditorPtr.Pin()->CanPasteActors();
 	}
 
 	return false;
@@ -792,4 +799,17 @@ bool SDisplayClusterLightCardEditorViewport::AreLabelsToggled() const
 	return LightCardEditorPtr.IsValid() && LightCardEditorPtr.Pin()->ShouldShowLightCardLabels();
 }
 
+void SDisplayClusterLightCardEditorViewport::ToggleIcons()
+{
+	if (LightCardEditorPtr.IsValid())
+	{
+		LightCardEditorPtr.Pin()->ShowIcons(!AreIconsToggled());
+	}
+}
+
+bool SDisplayClusterLightCardEditorViewport::AreIconsToggled() const
+{
+	return LightCardEditorPtr.IsValid() && LightCardEditorPtr.Pin()->ShouldShowIcons();
+}
+
 #undef LOCTEXT_NAMESPACE
diff --git a/Engine/Plugins/Runtime/nDisplay/Source/DisplayClusterLightCardEditor/Private/Viewport/DisplayClusterLightcardEditorViewport.h b/Engine/Plugins/Runtime/nDisplay/Source/DisplayClusterLightCardEditor/Private/Viewport/DisplayClusterLightcardEditorViewport.h
--- a/Engine/Plugins/Runtime/nDisplay/Source/DisplayClusterLightCardEditor/Private/Viewport/DisplayClusterLightcardEditorViewport.h
+++ b/Engine/Plugins/Runtime/nDisplay/Source/DisplayClusterLightCardEditor/Private/Viewport/DisplayClusterLightcardEditorViewport.h
@@ -103,6 +103,12 @@ private:
 
 	/** If light card labels are currently toggled */
 	bool AreLabelsToggled() const;
+
+	/** Toggles all icons on or off */
+	void ToggleIcons();
+
+	/** If icons are currently toggled */
+	bool AreIconsToggled() const;
 	
 private:
 	/** The preview scene to use by the viewport client to manage the preview 3D world */
diff --git a/Engine/Plugins/Runtime/nDisplay/Source/DisplayClusterScenePreview/DisplayClusterScenePreview.Build.cs b/Engine/Plugins/Runtime/nDisplay/Source/DisplayClusterScenePreview/DisplayClusterScenePreview.Build.cs
--- a/Engine/Plugins/Runtime/nDisplay/Source/DisplayClusterScenePreview/DisplayClusterScenePreview.Build.cs
+++ b/Engine/Plugins/Runtime/nDisplay/Source/DisplayClusterScenePreview/DisplayClusterScenePreview.Build.cs
@@ -9,6 +9,7 @@ public class DisplayClusterScenePreview : ModuleRules
 		PublicDependencyModuleNames.AddRange(
 			new string[] {
 				"DisplayClusterLightCardEditorShaders",
+				"DisplayClusterLightCardExtender",
 				"RHI"
 			});
 
diff --git a/Engine/Plugins/Runtime/nDisplay/Source/DisplayClusterScenePreview/Private/DisplayClusterLightCardEditorHelper.cpp b/Engine/Plugins/Runtime/nDisplay/Source/DisplayClusterScenePreview/Private/DisplayClusterLightCardEditorHelper.cpp
--- a/Engine/Plugins/Runtime/nDisplay/Source/DisplayClusterScenePreview/Private/DisplayClusterLightCardEditorHelper.cpp
+++ b/Engine/Plugins/Runtime/nDisplay/Source/DisplayClusterScenePreview/Private/DisplayClusterLightCardEditorHelper.cpp
@@ -401,8 +401,8 @@ const UTexture2D* FDisplayClusterLightCardEditorHelper::GetNormalMapTexture(bool
 	return NormalMapTexture;
 }
 
-void FDisplayClusterLightCardEditorHelper::MoveLightCardsToPixel(
-	const TArray<TWeakObjectPtr<ADisplayClusterLightCardActor>>& LightCards, const FIntPoint& PixelPos, const FSceneView& SceneView)
+void FDisplayClusterLightCardEditorHelper::MoveActorsToPixel(
+	const TArray<FDisplayClusterWeakStageActorPtr>& Actors, const FIntPoint& PixelPos, const FSceneView& SceneView)
 {
 	FVector Origin;
 	FVector Direction;
@@ -410,15 +410,15 @@ void FDisplayClusterLightCardEditorHelper::MoveLightCardsToPixel(
 
 	if (ProjectionMode == EDisplayClusterMeshProjectionType::UV)
 	{
-		// Find the average position of all selected light cards. This group average is what is moved to the specified pixel
+		// Find the average position of all selected actors. This group average is what is moved to the specified pixel
 		FVector2D AverageUVCoords = FVector2D::ZeroVector;
 		int32 NumLightCards = 0;
 
-		for (const TWeakObjectPtr<ADisplayClusterLightCardActor>& LightCard : LightCards)
+		for (const FDisplayClusterWeakStageActorPtr& Actor : Actors)
 		{
-			if (LightCard.IsValid() && LightCard->bIsUVLightCard)
+			if (Actor.IsValid() && Actor->IsUVActor())
 			{
-				AverageUVCoords += LightCard->UVCoordinates;
+				AverageUVCoords += Actor->GetUVCoordinates();
 				++NumLightCards;
 			}
 		}
@@ -435,11 +435,11 @@ void FDisplayClusterLightCardEditorHelper::MoveLightCardsToPixel(
 
 		const FVector2D DeltaUVCoords = DesiredUVCoords - AverageUVCoords;
 
-		for (const TWeakObjectPtr<ADisplayClusterLightCardActor>& LightCard : LightCards)
+		for (const FDisplayClusterWeakStageActorPtr& Actor : Actors)
 		{
-			if (LightCard.IsValid() && LightCard->bIsUVLightCard)
+			if (Actor.IsValid() && Actor->IsUVActor())
 			{
-				LightCard->UVCoordinates += DeltaUVCoords;
+				Actor->SetUVCoordinates(Actor->GetUVCoordinates() + DeltaUVCoords);
 			}
 		}
 	}
@@ -447,22 +447,27 @@ void FDisplayClusterLightCardEditorHelper::MoveLightCardsToPixel(
 	{
 		// Find the average position of all selected light cards. This group average is what is moved to the specified pixel
 		FSphericalCoordinates AverageCoords = FSphericalCoordinates();
-		int32 NumLightCards = 0;
+		int32 NumActors = 0;
 
-		for (const TWeakObjectPtr<ADisplayClusterLightCardActor> LightCard : LightCards)
+		for (const FDisplayClusterWeakStageActorPtr& Actor : Actors)
 		{
-			if (LightCard.IsValid())
+			if (Actor.IsValid())
 			{
-				const FSphericalCoordinates LightCardCoords = GetLightCardCoordinates(*LightCard);
+				const FSphericalCoordinates ActorCoords = GetActorCoordinates(Actor);
 
-				AverageCoords = AverageCoords + LightCardCoords;
-				++NumLightCards;
+				AverageCoords = AverageCoords + ActorCoords;
+				++NumActors;
 			}
 		}
 
-		AverageCoords.Radius /= NumLightCards;
-		AverageCoords.Azimuth /= NumLightCards;
-		AverageCoords.Inclination /= NumLightCards;
+		if (NumActors == 0)
+		{
+			NumActors = 1;
+		}
+		
+		AverageCoords.Radius /= NumActors;
+		AverageCoords.Azimuth /= NumActors;
+		AverageCoords.Inclination /= NumActors;
 		AverageCoords.Conform();
 
 		// Compute desired coordinates (radius doesn't matter here since we will use the flush constraint on the light cards after moving them)
@@ -470,21 +475,21 @@ void FDisplayClusterLightCardEditorHelper::MoveLightCardsToPixel(
 		const FSphericalCoordinates DeltaCoords = DesiredCoords - AverageCoords;
 
 		// Update each light card with the delta coordinates; the flush constraint is applied by MoveLightCardTo, ensuring the light card is always flush to screens
-		for (const TWeakObjectPtr<ADisplayClusterLightCardActor> LightCard : LightCards)
+		for (const FDisplayClusterWeakStageActorPtr& LightCard : Actors)
 		{
-			if (LightCard.IsValid() && !LightCard->bIsUVLightCard)
+			if (LightCard.IsValid() && !LightCard->IsUVActor())
 			{
-				const FSphericalCoordinates LightCardCoords = GetLightCardCoordinates(*LightCard);
+				const FSphericalCoordinates LightCardCoords = GetActorCoordinates(LightCard);
 				const FSphericalCoordinates NewCoords = LightCardCoords + DeltaCoords;
 
-				MoveLightCardsTo({ LightCard }, NewCoords);
+				MoveActorsTo({ LightCard }, NewCoords);
 			}
 		}
 	}
 }
 
-void FDisplayClusterLightCardEditorHelper::MoveLightCardsTo(
-	const TArray<TWeakObjectPtr<ADisplayClusterLightCardActor>>& LightCards, const FSphericalCoordinates& SphericalCoords)
+void FDisplayClusterLightCardEditorHelper::MoveActorsTo(
+	const TArray<FDisplayClusterWeakStageActorPtr>& Actors, const FSphericalCoordinates& SphericalCoords)
 {
 	ADisplayClusterRootActor* RootActor = UpdateRootActor();
 	if (!RootActor)
@@ -497,26 +502,28 @@ void FDisplayClusterLightCardEditorHelper::MoveLightCardsTo(
 		return;
 	}
 
-	for (const TWeakObjectPtr<ADisplayClusterLightCardActor> LightCard : LightCards)
+	for (const FDisplayClusterWeakStageActorPtr& Actor : Actors)
 	{
-		if (!LightCard.IsValid())
+		if (!Actor.IsValid())
 		{
 			continue;
 		}
 
-		InternalMoveLightCardTo(*LightCard, SphericalCoords, true);
+		InternalMoveActorTo(Actor, SphericalCoords, true);
 	}
 }
 
-void FDisplayClusterLightCardEditorHelper::DragLightCards(const TArray<TWeakObjectPtr<ADisplayClusterLightCardActor>>& LightCards, const FIntPoint& PixelPos, const FSceneView& SceneView,
-	ECoordinateSystem CoordinateSystem, const FVector& DragWidgetOffset, EAxisList::Type DragAxis, ADisplayClusterLightCardActor* PrimaryLightCard)
+void FDisplayClusterLightCardEditorHelper::DragActors(
+	const TArray<FDisplayClusterWeakStageActorPtr>& Actors, const FIntPoint& PixelPos,
+	const FSceneView& SceneView, ECoordinateSystem CoordinateSystem, const FVector& DragWidgetOffset, EAxisList::Type DragAxis,
+	FDisplayClusterWeakStageActorPtr PrimaryActor)
 {
-	if (LightCards.IsEmpty())
+	if (Actors.IsEmpty())
 	{
 		return;
 	}
 
-	ADisplayClusterRootActor* RootActor = UpdateRootActor();
+	const ADisplayClusterRootActor* RootActor = UpdateRootActor();
 	if (!RootActor)
 	{
 		return;
@@ -526,23 +533,23 @@ void FDisplayClusterLightCardEditorHelper::DragLightCards(const TArray<TWeakObje
 	{
 		return;
 	}
-
+	
 	// Maps are already ready, so we can move immediately
-	InternalDragLightCards(LightCards, PixelPos, SceneView, CoordinateSystem, DragWidgetOffset, DragAxis, PrimaryLightCard);
+	InternalDragActors(Actors, PixelPos, SceneView, CoordinateSystem, DragWidgetOffset, DragAxis, PrimaryActor);
 }
 
-void FDisplayClusterLightCardEditorHelper::VerifyAndFixLightCardOrigin(ADisplayClusterLightCardActor& LightCard)
+void FDisplayClusterLightCardEditorHelper::VerifyAndFixActorOrigin(const FDisplayClusterWeakStageActorPtr& Actor)
 {
-	// Center lightcard on the current view origin, let it keep its current world placement 
-	// (but not its spin/yaw/pitch since that will be happen later using the cache).
+	// Center actor on the current view origin, let it keep its current world placement 
+	// (but not its spin/yaw/pitch since that will be happen later using the cache)
 
-	ADisplayClusterRootActor* RootActor = UpdateRootActor();
-	const ADisplayClusterRootActor* LightCardRootActor = (LightCard.bIsProxy || !LevelInstanceRootActor.IsValid()) ? RootActor : LevelInstanceRootActor.Get();
+	const ADisplayClusterRootActor* RootActor = UpdateRootActor();
+	const ADisplayClusterRootActor* OwningRootActor = (Actor->IsProxy() || !LevelInstanceRootActor.IsValid()) ? RootActor : LevelInstanceRootActor.Get();
 
 	const USceneComponent* OriginComponent = nullptr;
-	if (LightCardRootActor)
+	if (OwningRootActor)
 	{
-		OriginComponent = LightCard.bIsProxy ? ProjectionOriginComponent.Get() : LightCardRootActor->GetCommonViewPoint();
+		OriginComponent = Actor->IsProxy() && ProjectionOriginComponent.IsValid() ? ProjectionOriginComponent.Get() : OwningRootActor->GetCommonViewPoint();
 	}
 
 	if (!OriginComponent)
@@ -550,27 +557,24 @@ void FDisplayClusterLightCardEditorHelper::VerifyAndFixLightCardOrigin(ADisplayC
 		return;
 	}
 
-	// Set location at the view origin
-	const FVector& NewLightCardActorLocation = OriginComponent ? OriginComponent->GetComponentLocation() : FVector::ZeroVector;
-
-	// Set rotation to match the root actor
-	const FRotator& NewLightCardActorRotation = LightCardRootActor ? LightCardRootActor->GetActorRotation() : FRotator::ZeroRotator;
-
-	const FVector LightCardEndEffectorLocation = LightCard.GetLightCardTransform().GetLocation();
-
-	LightCard.SetActorLocation(NewLightCardActorLocation);
-	LightCard.SetActorRotation(NewLightCardActorRotation);
-
-	// Update the light card spherical coordinates to match its current world coordinates
-
-	const FVector LightCardRelativeLocation = NewLightCardActorRotation.UnrotateVector(LightCardEndEffectorLocation - NewLightCardActorLocation);
-
-	const FSphericalCoordinates SphericalCoords(LightCardRelativeLocation);
+	// Set location and rotation to match the root actor
+	const FVector& NewActorLocation = OriginComponent ? OriginComponent->GetComponentLocation() : FVector::ZeroVector;
+	const FRotator& NewActorRotation = OwningRootActor ? OwningRootActor->GetActorRotation() : FRotator::ZeroRotator;
+	
+	Actor->SetOrigin({ NewActorRotation, NewActorLocation, FVector::One() });
+	
+	if (Actor.AsActorChecked()->IsA<ADisplayClusterLightCardActor>())
+	{
+		// Update the light card spherical coordinates to match its current world coordinates
+		const FVector LightCardEndEffectorLocation = Actor->GetStageActorTransform(true).GetLocation();
+		const FVector ActorRelativeLocation = NewActorRotation.UnrotateVector(LightCardEndEffectorLocation);
 
-	SetLightCardCoordinates(LightCard, SphericalCoords);
+		const FSphericalCoordinates SphericalCoords(ActorRelativeLocation);
+		SetActorCoordinates(Actor, SphericalCoords);
+	}
 
 #if WITH_EDITOR
-	PostEditChangePropertiesForMovedLightCard(LightCard);
+	PostEditChangePropertiesForMovedActor(Actor);
 #endif
 }
 
@@ -838,22 +842,21 @@ void FDisplayClusterLightCardEditorHelper::GetSceneViewInitOptions(
 	OutViewInitOptions.FOV = InFOV;
 }
 
-FDisplayClusterLightCardEditorHelper::FSphericalCoordinates FDisplayClusterLightCardEditorHelper::GetLightCardCoordinates(const ADisplayClusterLightCardActor& LightCard)
+FDisplayClusterLightCardEditorHelper::FSphericalCoordinates FDisplayClusterLightCardEditorHelper::GetActorCoordinates(const FDisplayClusterWeakStageActorPtr& Actor)
 {
-	const FVector LightCardLocation = LightCard.GetLightCardTransform().GetTranslation() - LightCard.GetActorLocation();
-
-	FSphericalCoordinates LightCardCoords(LightCardLocation);
+	const FVector ActorLocation = Actor->GetStageActorTransform(true).GetTranslation();
+	FSphericalCoordinates ActorSphericalCoords(ActorLocation);
 
 	// If the light card points at any of the poles, the spherical coordinates will have an "undefined" azimuth value. 
 	// For continuity when dragging a light card positioned there, 
 	// we can manually set the azimuthal value to match the light card's configured longitude
 
-	if (LightCardCoords.IsPointingAtPole())
+	if (ActorSphericalCoords.IsPointingAtPole())
 	{
-		LightCardCoords.Azimuth = FMath::DegreesToRadians(LightCard.Longitude + 180);
+		ActorSphericalCoords.Azimuth = FMath::DegreesToRadians(Actor->GetLongitude() + 180.f);
 	}
 
-	return LightCardCoords;
+	return ActorSphericalCoords;
 }
 
 ADisplayClusterRootActor* FDisplayClusterLightCardEditorHelper::UpdateRootActor()
@@ -931,11 +934,11 @@ FVector FDisplayClusterLightCardEditorHelper::GetProjectionOrigin() const
 	return FVector::Zero();
 }
 
-FRotator FDisplayClusterLightCardEditorHelper::GetLightCardRotationDelta(const FIntPoint& PixelPos, const FSceneView& View, ADisplayClusterLightCardActor& LightCard,
+FRotator FDisplayClusterLightCardEditorHelper::GetActorRotationDelta(const FIntPoint& PixelPos, const FSceneView& View, const FDisplayClusterWeakStageActorPtr& Actor,
 	ECoordinateSystem CoordinateSystem, EAxisList::Type DragAxis, const FVector& DragWidgetOffset)
 {
-	const FSphericalCoordinates DeltaCoords = GetLightCardTranslationDelta(PixelPos, View, LightCard, CoordinateSystem, DragAxis, DragWidgetOffset);
-	const FSphericalCoordinates LightCardCoords = GetLightCardCoordinates(LightCard);
+	const FSphericalCoordinates DeltaCoords = GetActorTranslationDelta(PixelPos, View, Actor, CoordinateSystem, DragAxis, DragWidgetOffset);
+	const FSphericalCoordinates LightCardCoords = GetActorCoordinates(Actor);
 	const FSphericalCoordinates NewCoords = LightCardCoords + DeltaCoords;
 
 	const FVector LightCardPos = LightCardCoords.AsCartesian();
@@ -954,10 +957,10 @@ FRotator FDisplayClusterLightCardEditorHelper::GetLightCardRotationDelta(const F
 	return FQuat(AxisOfRotation, Angle).Rotator();
 }
 
-FDisplayClusterLightCardEditorHelper::FSphericalCoordinates FDisplayClusterLightCardEditorHelper::GetLightCardTranslationDelta(
+FDisplayClusterLightCardEditorHelper::FSphericalCoordinates FDisplayClusterLightCardEditorHelper::GetActorTranslationDelta(
 	const FIntPoint& PixelPos,
 	const FSceneView& View,
-	ADisplayClusterLightCardActor& LightCard,
+	const FDisplayClusterWeakStageActorPtr& Actor,
 	ECoordinateSystem CoordinateSystem,
 	EAxisList::Type DragAxis,
 	const FVector& DragWidgetOffset)
@@ -971,10 +974,11 @@ FDisplayClusterLightCardEditorHelper::FSphericalCoordinates FDisplayClusterLight
 		Direction = (Direction - DragWidgetOffset).GetSafeNormal();
 	}
 
-	const FVector LocalDirection = LightCard.GetActorRotation().RotateVector(Direction);
-	const FVector LightCardLocation = LightCard.GetLightCardTransform().GetTranslation() - Origin;
-	const FSphericalCoordinates LightCardCoords = GetLightCardCoordinates(LightCard);
+	checkSlow(CachedRootActor.IsValid());
 
+	const FVector LocalDirection = CachedRootActor->GetActorRotation().RotateVector(Direction);
+	const FVector LightCardLocation = Actor->GetStageActorTransform().GetTranslation() - Origin;
+	const FSphericalCoordinates ActorCoords = GetActorCoordinates(Actor);
 	FSphericalCoordinates DeltaCoords;
 
 	// If we are in a cartesian coordinate system and are constraining to an axis, perform the constraint calculations in
@@ -1000,7 +1004,7 @@ FDisplayClusterLightCardEditorHelper::FSphericalCoordinates FDisplayClusterLight
 			Axis = FVector::ZAxisVector;
 		}
 
-		const FVector LocalAxis = LightCard.GetActorRotation().RotateVector(Axis);
+		const FVector LocalAxis = CachedRootActor->GetActorRotation().RotateVector(Axis);
 
 		// Compute the offset between the requested location and the light card's current location, and project that
 		// offset onto the constraint axis
@@ -1009,7 +1013,7 @@ FDisplayClusterLightCardEditorHelper::FSphericalCoordinates FDisplayClusterLight
 		const FVector ConstrainedLocation = LightCardLocation + DeltaLocation;
 
 		const FSphericalCoordinates ConstrainedCoords(ConstrainedLocation);
-		DeltaCoords = ConstrainedCoords - LightCardCoords;
+		DeltaCoords = ConstrainedCoords - ActorCoords;
 	}
 	else
 	{
@@ -1019,16 +1023,16 @@ FDisplayClusterLightCardEditorHelper::FSphericalCoordinates FDisplayClusterLight
 		// If the light card is in the southern hemisphere of the view origin, use the southern normal map; otherwise, use the north normal map
 		if (LightCardLocation.Z < 0.0f)
 		{
-			SouthNormalMap.GetNormalAndDistanceAtPosition(LightCard.GetLightCardTransform().GetTranslation(), Normal, Distance);
+			SouthNormalMap.GetNormalAndDistanceAtPosition(Actor->GetStageActorTransform(true).GetTranslation(), Normal, Distance);
 		}
 		else
 		{
-			NorthNormalMap.GetNormalAndDistanceAtPosition(LightCard.GetLightCardTransform().GetTranslation(), Normal, Distance);
+			NorthNormalMap.GetNormalAndDistanceAtPosition(Actor->GetStageActorTransform(true).GetTranslation(), Normal, Distance);
 		}
 
 		const FSphericalCoordinates RequestedCoords(LocalDirection * Distance);
 
-		DeltaCoords = RequestedCoords - LightCardCoords;
+		DeltaCoords = RequestedCoords - ActorCoords;
 
 		if (CoordinateSystem == ECoordinateSystem::Spherical)
 		{
@@ -1049,7 +1053,7 @@ FDisplayClusterLightCardEditorHelper::FSphericalCoordinates FDisplayClusterLight
 				const double FixedAzimuth = FMath::RoundToInt(DeltaCoords.Azimuth / PI) * PI;
 
 				DeltaCoords.Azimuth = FixedAzimuth;
-				DeltaCoords.Inclination = FixedInclination - LightCardCoords.Inclination;
+				DeltaCoords.Inclination = FixedInclination - ActorCoords.Inclination;
 			}
 		}
 	}
@@ -1057,8 +1061,8 @@ FDisplayClusterLightCardEditorHelper::FSphericalCoordinates FDisplayClusterLight
 	return DeltaCoords;
 }
 
-void FDisplayClusterLightCardEditorHelper::InternalMoveLightCardTo(
-	ADisplayClusterLightCardActor& LightCard, const FSphericalCoordinates& Position, bool bIsFinalChange) const
+void FDisplayClusterLightCardEditorHelper::InternalMoveActorTo(
+	const FDisplayClusterWeakStageActorPtr& Actor, const FSphericalCoordinates& Position, bool bIsFinalChange) const
 {
 	if (bNormalMapInvalid || !CachedRootActor.IsValid())
 	{
@@ -1086,68 +1090,70 @@ void FDisplayClusterLightCardEditorHelper::InternalMoveLightCardTo(
 	// which already takes the root actor rotation into account as part of its view matrix.
 	const FQuat RootRotation = CachedRootActor->GetTransform().GetRotation().Inverse();
 	const FVector InverseRotatedPosition = RootRotation.RotateVector(Position.AsCartesian());
-	SetLightCardCoordinates(LightCard, FSphericalCoordinates(InverseRotatedPosition));
 
-	LightCard.DistanceFromCenter = CalculateFinalLightCardDistance(DesiredDistance);
+	SetActorCoordinates(Actor, FSphericalCoordinates(InverseRotatedPosition));
+	
+	const double FinalDistance = CalculateFinalLightCardDistance(DesiredDistance);
+	Actor->SetDistanceFromCenter(FinalDistance);
 
 	const FRotator Rotation = FRotationMatrix::MakeFromX(-DesiredNormal).Rotator();
-
-	LightCard.Pitch = Rotation.Pitch;
-	LightCard.Yaw = Rotation.Yaw;
+	
+	Actor->SetPitch(Rotation.Pitch);
+	Actor->SetYaw(Rotation.Yaw);
 
 #if WITH_EDITOR
-	if (bIsFinalChange && !LightCard.bIsProxy)
+	if (bIsFinalChange && !Actor->IsProxy())
 	{
-		PostEditChangePropertiesForMovedLightCard(LightCard);
+		PostEditChangePropertiesForMovedActor(Actor);
 	}
 #endif
 }
 
-void FDisplayClusterLightCardEditorHelper::InternalDragLightCards(const TArray<TWeakObjectPtr<ADisplayClusterLightCardActor>>& LightCards, const FIntPoint& PixelPos, const FSceneView& View,
-	ECoordinateSystem CoordinateSystem, const FVector& DragWidgetOffset, EAxisList::Type DragAxis, ADisplayClusterLightCardActor* PrimaryLightCard)
+void FDisplayClusterLightCardEditorHelper::InternalDragActors(const TArray<FDisplayClusterWeakStageActorPtr>& Actors, const FIntPoint& PixelPos, const FSceneView& View,
+	ECoordinateSystem CoordinateSystem, const FVector& DragWidgetOffset, EAxisList::Type DragAxis, FDisplayClusterWeakStageActorPtr PrimaryActor)
 {
-	if (LightCards.IsEmpty() || !ensure(!bNormalMapInvalid) || !UpdateRootActor())
+	if (Actors.IsEmpty() || !ensure(!bNormalMapInvalid) || !UpdateRootActor())
 	{
 		return;
 	}
 
-	if (!PrimaryLightCard)
+	if (!PrimaryActor.IsValid())
 	{
-		PrimaryLightCard = LightCards.Last().Get();
+		PrimaryActor = Actors.Last();
 	}
 
-	if (PrimaryLightCard && !PrimaryLightCard->bIsUVLightCard)
+	if (PrimaryActor.IsValid() && !PrimaryActor->IsUVActor())
 	{
 		const bool bUseDeltaRotation = (DragAxis == EAxisList::Type::XYZ) || (DragAxis == EAxisList::Type::Y) || (CoordinateSystem == ECoordinateSystem::Cartesian);
 
 		const FRotator DeltaRotation =
 			bUseDeltaRotation ?
-			GetLightCardRotationDelta(PixelPos, View, *PrimaryLightCard, CoordinateSystem, DragAxis, DragWidgetOffset)
+			GetActorRotationDelta(PixelPos, View, PrimaryActor, CoordinateSystem, DragAxis, DragWidgetOffset)
 			: FRotator::ZeroRotator;
 
 		const FSphericalCoordinates DeltaCoords =
 			bUseDeltaRotation ?
 			FSphericalCoordinates()
-			: GetLightCardTranslationDelta(PixelPos, View, *PrimaryLightCard, CoordinateSystem, DragAxis, DragWidgetOffset);
+			: GetActorTranslationDelta(PixelPos, View, PrimaryActor, CoordinateSystem, DragAxis, DragWidgetOffset);
 
-		for (const TWeakObjectPtr<ADisplayClusterLightCardActor> LightCard : LightCards)
+		for (const FDisplayClusterWeakStageActorPtr& Actor : Actors)
 		{
-			if (!LightCard.IsValid() || LightCard->bIsUVLightCard)
+			if (!Actor.IsValid() || Actor->IsUVActor())
 			{
 				continue;
 			}
 
-			VerifyAndFixLightCardOrigin(*LightCard);
+			VerifyAndFixActorOrigin(Actor);
 
-			// Note: GetLightCardCoordinates maintains last known Azimuth when looking at the poles
-			const FSphericalCoordinates CurrentCoords = GetLightCardCoordinates(*LightCard.Get());
+			// Note: GetActorCoordinates maintains last known Azimuth when looking at the poles
+			const FSphericalCoordinates CurrentCoords = GetActorCoordinates(Actor);
 
 			// We will adjust the spin (to maintain the apparent spin) when using center of gizmo
 			// or dragging longitudinally (this seems to provide an intuitive behavior)
 			if (bUseDeltaRotation) // Dragging center of gizmo
 			{
 				// We might need this to put back the LightCard exactly as it was
-				const ADisplayClusterLightCardActor::PositionalParams OriginalPositionalParams = LightCard->GetPositionalParams();
+				const FDisplayClusterPositionalParams OriginalPositionalParams = Actor->GetPositionalParams();
 
 				const FVector CurrentPos = CurrentCoords.AsCartesian();
 				const FVector NewPos = DeltaRotation.RotateVector(CurrentPos);
@@ -1155,13 +1161,13 @@ void FDisplayClusterLightCardEditorHelper::InternalDragLightCards(const TArray<T
 				// Calculations are only valid if translation is not too small
 				const FSphericalCoordinates NewCoords(NewPos);
 
-				const FTransform Transform_A = LightCard->GetLightCardTransform(false /*bIgnoreSpinYawPitch*/);
+				const FTransform Transform_A = Actor->GetStageActorTransform(true);
 
 				// Don't fire off property change events yet since we're not done modifying the lightcard
-				InternalMoveLightCardTo(*LightCard, NewCoords, false);
-				LightCard->UpdateLightCardTransform(); // We must call this for GetLightCardTransform to be valid
+				InternalMoveActorTo(Actor, NewCoords, false);
+				Actor->UpdateStageActorTransform(); // We must call this for GetStageActorTransform to be valid
 
-				const FTransform Transform_B = LightCard->GetLightCardTransform(false /*bIgnoreSpinYawPitch*/);
+				const FTransform Transform_B = Actor->GetStageActorTransform(true);
 
 				// Calculate world delta translation of moving from A to B
 				FVector WorldDelta = Transform_B.GetLocation() - Transform_A.GetLocation(); // X towards front of stage. Y towards right of stage. Z towards ceiling.
@@ -1195,21 +1201,21 @@ void FDisplayClusterLightCardEditorHelper::InternalDragLightCards(const TArray<T
 					const double DeltaSpin = RelativeSpinAngle_B - RelativeSpinAngle_A;
 
 					// Apply delta spin to lightcard
-					LightCard->Spin += FMath::RadiansToDegrees(DeltaSpin);
+					Actor->SetSpin(Actor->GetSpin() + FMath::RadiansToDegrees(DeltaSpin));
 				}
 				else
 				{
 					// Leave it where it was to avoid apparent spins even though motion would have been insignificant.
-					LightCard->SetPositionalParams(OriginalPositionalParams);
+					Actor->SetPositionalParams(OriginalPositionalParams);
 				}
 
 #if WITH_EDITOR
-				PostEditChangePropertiesForMovedLightCard(*LightCard);
+				PostEditChangePropertiesForMovedActor(Actor);
 #endif
 			}
 			else // Dragging latitudinally
 			{
-				InternalMoveLightCardTo(*LightCard, CurrentCoords + DeltaCoords, true);
+				InternalMoveActorTo(Actor, CurrentCoords + DeltaCoords, true);
 			}
 		}
 	}
@@ -1293,16 +1299,16 @@ void FDisplayClusterLightCardEditorHelper::GetNormalMapSceneViewInitOptions(cons
 #endif
 }
 
-void FDisplayClusterLightCardEditorHelper::SetLightCardCoordinates(ADisplayClusterLightCardActor& LightCard, const FSphericalCoordinates& SphericalCoords) const
+void FDisplayClusterLightCardEditorHelper::SetActorCoordinates(const FDisplayClusterWeakStageActorPtr& Actor, const FSphericalCoordinates& SphericalCoords) const
 {
-	LightCard.DistanceFromCenter = SphericalCoords.Radius;
-	LightCard.Latitude = 90.f - FMath::RadiansToDegrees(SphericalCoords.Inclination);
+	Actor->SetDistanceFromCenter(SphericalCoords.Radius);
+	Actor->SetLatitude(90.f - FMath::RadiansToDegrees(SphericalCoords.Inclination));
 
 	// Keep the same longitude when pointing at the pole. This helps with continuity 
 	// and also mitigates sudden changes in apparent spin when moving around the poles
 	if (!SphericalCoords.IsPointingAtPole())
 	{
-		LightCard.Longitude = FRotator::ClampAxis(FMath::RadiansToDegrees(SphericalCoords.Azimuth) - 180);
+		Actor->SetLongitude(FRotator::ClampAxis(FMath::RadiansToDegrees(SphericalCoords.Azimuth) - 180.f));
 	}
 }
 
@@ -1568,38 +1574,40 @@ void FDisplayClusterLightCardEditorHelper::OnRootActorBlueprintCompiled(UBluepri
 #endif
 
 #if WITH_EDITOR
-void FDisplayClusterLightCardEditorHelper::PostEditChangePropertiesForMovedLightCard(ADisplayClusterLightCardActor& LightCard) const
+void FDisplayClusterLightCardEditorHelper::PostEditChangePropertiesForMovedActor(const FDisplayClusterWeakStageActorPtr& Actor) const
 {
-	if (LightCard.bIsProxy)
+	if (Actor->IsProxy())
 	{
 		return;
 	}
 
-	LightCard.Modify();
+	Actor.AsActorChecked()->Modify();
 
-	auto ModifyLightCardProperty = [&](const FName& PropertyName)
+	auto ModifyLightCardProperty = [&](const FName& PropertyName) -> FProperty*
 	{
-		FProperty* Property = FindFProperty<FProperty>(LightCard.GetClass(), PropertyName);
-
-		// Broadcast the event directly instead of PostEditChangeProperty on the object itself, which would attempt to create unnecessary snapshots for each call
-		FPropertyChangedEvent PropertyChangedEvent(Property, EPropertyChangeType::Interactive);
-		FCoreUObjectDelegates::OnObjectPropertyChanged.Broadcast(&LightCard, PropertyChangedEvent);
+		if (FProperty* Property = FindFProperty<FProperty>(Actor.AsActorChecked()->GetClass(), PropertyName))
+		{
+			// Broadcast the event directly instead of PostEditChangeProperty on the object itself, which would attempt to create unnecessary snapshots for each call
+			FPropertyChangedEvent PropertyChangedEvent(Property, EPropertyChangeType::Interactive);
+			FCoreUObjectDelegates::OnObjectPropertyChanged.Broadcast(Actor.AsActorChecked(), PropertyChangedEvent);
 
-		return Property;
+			return Property;
+		}
+		return nullptr;
 	};
 
-	const FProperty* ChangedProperties[] =
+	const TSet<FName>& PropertyNames = Actor->GetPositionalPropertyNames();
+	
+	TArray<const FProperty*> ChangedProperties;
+	ChangedProperties.Reserve(PropertyNames.Num());
+	for (const FName& PropertyName : PropertyNames)
 	{
-		ModifyLightCardProperty(GET_MEMBER_NAME_CHECKED(ADisplayClusterLightCardActor, Longitude)),
-		ModifyLightCardProperty(GET_MEMBER_NAME_CHECKED(ADisplayClusterLightCardActor, Latitude)),
-		ModifyLightCardProperty(GET_MEMBER_NAME_CHECKED(ADisplayClusterLightCardActor, DistanceFromCenter)),
-		ModifyLightCardProperty(GET_MEMBER_NAME_CHECKED(ADisplayClusterLightCardActor, Spin)),
-		ModifyLightCardProperty(GET_MEMBER_NAME_CHECKED(ADisplayClusterLightCardActor, Pitch)),
-		ModifyLightCardProperty(GET_MEMBER_NAME_CHECKED(ADisplayClusterLightCardActor, Yaw)),
-		ModifyLightCardProperty(GET_MEMBER_NAME_CHECKED(ADisplayClusterLightCardActor, Scale)),
-		ModifyLightCardProperty(GET_MEMBER_NAME_CHECKED(ADisplayClusterLightCardActor, UVCoordinates)),
-	};
+		if (const FProperty* Property = ModifyLightCardProperty(PropertyName))
+		{
+			ChangedProperties.Add(Property);
+		}
+	}
 
-	SnapshotTransactionBuffer(&LightCard, ChangedProperties);
+	SnapshotTransactionBuffer(Actor.AsActorChecked(), MakeArrayView(ChangedProperties.GetData(), ChangedProperties.Num()));
 }
 #endif
diff --git a/Engine/Plugins/Runtime/nDisplay/Source/DisplayClusterScenePreview/Public/DisplayClusterLightCardEditorHelper.h b/Engine/Plugins/Runtime/nDisplay/Source/DisplayClusterScenePreview/Public/DisplayClusterLightCardEditorHelper.h
--- a/Engine/Plugins/Runtime/nDisplay/Source/DisplayClusterScenePreview/Public/DisplayClusterLightCardEditorHelper.h
+++ b/Engine/Plugins/Runtime/nDisplay/Source/DisplayClusterScenePreview/Public/DisplayClusterLightCardEditorHelper.h
@@ -2,12 +2,14 @@
 
 #pragma once
 
-#include "CoreMinimal.h"
-#include "Containers/Union.h"
+#include "StageActor/DisplayClusterWeakStageActorPtr.h"
+
 #include "DisplayClusterMeshProjectionRenderer.h"
+
 #include "SceneView.h"
-#include "UObject/StrongObjectPtr.h"
 #include "UnrealClient.h"
+#include "Containers/Union.h"
+#include "UObject/StrongObjectPtr.h"
 
 class ADisplayClusterLightCardActor;
 class ADisplayClusterRootActor;
@@ -176,36 +178,36 @@ public:
 	/**
 	 * Moves the given light cards to the specified pixel position within the provided scene view.
 	 *
-	 * @param LightCards The light cards that we are moving
-	 * @param PixelPos The pixel location to move the light cards to
+	 * @param Actors The actors that we are moving
+	 * @param PixelPos The pixel location to move the actors to
 	 * @param SceneView The scene view used to convert from pixel position to 3D position
 	*/
-	DISPLAYCLUSTERSCENEPREVIEW_API void MoveLightCardsToPixel(const TArray<TWeakObjectPtr<ADisplayClusterLightCardActor>>& LightCards, const FIntPoint& PixelPos, const FSceneView& SceneView);
+	DISPLAYCLUSTERSCENEPREVIEW_API void MoveActorsToPixel(const TArray<FDisplayClusterWeakStageActorPtr>& Actors, const FIntPoint& PixelPos, const FSceneView& SceneView);
 
 	/**
-	 * Moves specified cards to desired coordinates. Actual radius will be based on flush constraint and LightCard's RadialOffset.
+	 * Moves specified actors to desired coordinates. Actual radius will be based on flush constraint and actor's RadialOffset.
 	 *
-	 * @param LightCards The light cards that we are moving
-	 * @param SphericalCoords The desired location of light card in spherical coordinates with respect to view origin.
+	 * @param Actors The actors that we are moving
+	 * @param SphericalCoords The desired location of the actors in spherical coordinates with respect to view origin.
 	*/
-	DISPLAYCLUSTERSCENEPREVIEW_API void MoveLightCardsTo(const TArray<TWeakObjectPtr<ADisplayClusterLightCardActor>>& LightCards, const FSphericalCoordinates& SphericalCoords);
-
+	DISPLAYCLUSTERSCENEPREVIEW_API void MoveActorsTo(const TArray<FDisplayClusterWeakStageActorPtr>& Actors, const FSphericalCoordinates& SphericalCoords);
+	
 	/**
 	 * Moves specified cards to a coordinate in viewport space as if dragged by a translate widget.
 	 * 
-	 * @param LightCards The light cards that we are moving
+	 * @param Actors The actors that we are moving
 	 * @param PixelPos The screen pixel position of the widget
 	 * @param SceneView The scene view used to convert from pixel position to 3D position
 	 * @param CoordinateSystem The coordinate system to use when computing drag constraints
 	 * @param DragWidgetOffset The offset between the actual cursor position and the position of the widget when the drag action started
 	 * @param DragAxis The axis along which the widget is being dragged
-	 * @param PrimaryLightCard The light card used to calculate the translation/rotation delta. If not provided, the last entry in LightCards will be used.
+	 * @param PrimaryActor The actor used to calculate the translation/rotation delta. If not provided, the last entry in Actors will be used.
 	 */
-	DISPLAYCLUSTERSCENEPREVIEW_API void DragLightCards(const TArray<TWeakObjectPtr<ADisplayClusterLightCardActor>>& LightCards, const FIntPoint& PixelPos, const FSceneView& SceneView,
-		ECoordinateSystem CoordinateSystem, const FVector& DragWidgetOffset, EAxisList::Type DragAxis, ADisplayClusterLightCardActor* PrimaryLightCard = nullptr);
-
-	/** Ensures that the light card root component is at the same location as the projection/view origin */
-	DISPLAYCLUSTERSCENEPREVIEW_API void VerifyAndFixLightCardOrigin(ADisplayClusterLightCardActor& LightCard);
+	DISPLAYCLUSTERSCENEPREVIEW_API void DragActors(const TArray<FDisplayClusterWeakStageActorPtr>& Actors, const FIntPoint& PixelPos, const FSceneView& SceneView,
+		ECoordinateSystem CoordinateSystem, const FVector& DragWidgetOffset, EAxisList::Type DragAxis, FDisplayClusterWeakStageActorPtr PrimaryActor = nullptr);
+	
+	/** Ensures that the actor root component is at the same location as the projection/view origin */
+	DISPLAYCLUSTERSCENEPREVIEW_API void VerifyAndFixActorOrigin(const FDisplayClusterWeakStageActorPtr& Actor);
 
 	/**
 	 * Calculates the relative normal vector and world position in the specified direction from the given view origin.
@@ -254,8 +256,8 @@ public:
 		const FMatrix* InRotationMatrix = nullptr,
 		float InDPIScale = 1.f);
 
-	/** Gets the spherical coordinates of the specified light card. */
-	DISPLAYCLUSTERSCENEPREVIEW_API FSphericalCoordinates GetLightCardCoordinates(const ADisplayClusterLightCardActor& LightCard);
+	/** Gets the spherical coordinates of the specified actor. */
+	DISPLAYCLUSTERSCENEPREVIEW_API FSphericalCoordinates GetActorCoordinates(const FDisplayClusterWeakStageActorPtr& Actor);
 
 private:
 	/**
@@ -274,25 +276,25 @@ private:
 	FVector GetProjectionOrigin() const;
 
 	/** Determines the appropriate delta rotation needed to move the specified light card to the given position within the view. */
-	FRotator GetLightCardRotationDelta(const FIntPoint& PixelPos, const FSceneView& View, ADisplayClusterLightCardActor& LightCard,
+	FRotator GetActorRotationDelta(const FIntPoint& PixelPos, const FSceneView& View, const FDisplayClusterWeakStageActorPtr& Actor,
 		ECoordinateSystem CoordinateSystem, EAxisList::Type DragAxis, const FVector& DragWidgetOffset);
 
-	/** Determines the appropriate delta in spherical coordinates needed to move the specified light card to the given position within the view. */
-	FSphericalCoordinates GetLightCardTranslationDelta(const FIntPoint& PixelPos, const FSceneView& View, ADisplayClusterLightCardActor& LightCard,
+	/** Determines the appropriate delta in spherical coordinates needed to move the specified actor to the given position within the view. */
+	FSphericalCoordinates GetActorTranslationDelta(const FIntPoint& PixelPos, const FSceneView& View, const FDisplayClusterWeakStageActorPtr& Actor,
 		ECoordinateSystem CoordinateSystem, EAxisList::Type DragAxis, const FVector& DragWidgetOffset);
 
-	/** Moves specified card to desired coordinates immediately using the current normal maps. Requires valid normal maps. */
-	void InternalMoveLightCardTo(ADisplayClusterLightCardActor& LightCard, const FSphericalCoordinates& Position, bool bIsFinalChange) const;
+	/** Moves specified actor to desired coordinates immediately using the current normal maps. Requires valid normal maps. */
+	void InternalMoveActorTo(const FDisplayClusterWeakStageActorPtr& Actor, const FSphericalCoordinates& Position, bool bIsFinalChange) const;
 
-	/** Moves specified cards to a coordinate in viewport space as if dragged by a widget. Requires valid normal maps. */
-	void InternalDragLightCards(const TArray<TWeakObjectPtr<ADisplayClusterLightCardActor>>& LightCards, const FIntPoint& PixelPos, const FSceneView& View,
-		ECoordinateSystem CoordinateSystem, const FVector& DragWidgetOffset, EAxisList::Type DragAxis, ADisplayClusterLightCardActor* PrimaryLightCard);
+	/** Moves specified actors to a coordinate in viewport space as if dragged by a widget. Requires valid normal maps. */
+	void InternalDragActors(const TArray<FDisplayClusterWeakStageActorPtr>& Actors, const FIntPoint& PixelPos, const FSceneView& View,
+		ECoordinateSystem CoordinateSystem, const FVector& DragWidgetOffset, EAxisList::Type DragAxis, FDisplayClusterWeakStageActorPtr PrimaryActor);
 
 	/** Gets the scene view init options to use when rendering the normal map cache */
 	void GetNormalMapSceneViewInitOptions(const FVector& ViewDirection, FSceneViewInitOptions& OutViewInitOptions);
 
-	/** Sets the light card position to the given spherical coordinates */
-	void SetLightCardCoordinates(ADisplayClusterLightCardActor& LightCard, const FSphericalCoordinates& SphericalCoords) const;
+	/** Sets the actor position to the given spherical coordinates */
+	void SetActorCoordinates(const FDisplayClusterWeakStageActorPtr& Actor, const FSphericalCoordinates& SphericalCoords) const;
 
 	/** Performs a ray trace against the stage's geometry, and returns the hit point. */
 	bool TraceStage(const FVector& RayStart, const FVector& RayEnd, FVector& OutHitLocation);
@@ -327,8 +329,8 @@ private:
 #endif
 
 #if WITH_EDITOR
-	/** Send property change events for any properties that may have changed due to a lightcard being moved. */
-	void PostEditChangePropertiesForMovedLightCard(ADisplayClusterLightCardActor& LightCard) const;
+	/** Send property change events for any properties that may have changed due to a actor being moved. */
+	void PostEditChangePropertiesForMovedActor(const FDisplayClusterWeakStageActorPtr& Actor) const;
 #endif
 
 private:
diff --git a/Engine/Plugins/Runtime/nDisplayModularFeatures/Source/DisplayClusterLightCardExtender/DisplayClusterLightCardExtender.build.cs b/Engine/Plugins/Runtime/nDisplayModularFeatures/Source/DisplayClusterLightCardExtender/DisplayClusterLightCardExtender.build.cs
--- a/Engine/Plugins/Runtime/nDisplayModularFeatures/Source/DisplayClusterLightCardExtender/DisplayClusterLightCardExtender.build.cs
+++ b/Engine/Plugins/Runtime/nDisplayModularFeatures/Source/DisplayClusterLightCardExtender/DisplayClusterLightCardExtender.build.cs
@@ -11,6 +11,7 @@ public class DisplayClusterLightCardExtender : ModuleRules
 			new string[] {
 				"Core",
 				"CoreUObject",
+				"Engine"
 			});
 	}
 }
diff --git a/Engine/Plugins/Runtime/nDisplayModularFeatures/Source/DisplayClusterLightCardExtender/Private/StageActor/DisplayClusterWeakStageActorPtr.cpp b/Engine/Plugins/Runtime/nDisplayModularFeatures/Source/DisplayClusterLightCardExtender/Private/StageActor/DisplayClusterWeakStageActorPtr.cpp
new file mode 100644
--- /dev/null
+++ b/Engine/Plugins/Runtime/nDisplayModularFeatures/Source/DisplayClusterLightCardExtender/Private/StageActor/DisplayClusterWeakStageActorPtr.cpp
@@ -0,0 +1,37 @@
+// Copyright Epic Games, Inc. All Rights Reserved.
+
+#include "StageActor/DisplayClusterWeakStageActorPtr.h"
+#include "StageActor/IDisplayClusterStageActor.h"
+
+#include "GameFramework/Actor.h"
+
+FDisplayClusterWeakStageActorPtr::FDisplayClusterWeakStageActorPtr(const AActor* InActorPtr)
+{
+	if (InActorPtr)
+	{
+		if (ensureMsgf(InActorPtr->Implements<UDisplayClusterStageActor>(),
+			TEXT("Actor passed to FDisplayClusterWeakStageActorPtr must implement or inherit IDisplayClusterStageActor interface")))
+		{
+			ObjectPtr = InActorPtr;
+		}
+		else
+		{
+			ObjectPtr = nullptr;
+		}
+	}
+}
+
+IDisplayClusterStageActor* FDisplayClusterWeakStageActorPtr::Get() const
+{
+	return Cast<IDisplayClusterStageActor>(ObjectPtr.Get());
+}
+
+AActor* FDisplayClusterWeakStageActorPtr::AsActor() const
+{
+	return Cast<AActor>(ObjectPtr.Get());
+}
+
+AActor* FDisplayClusterWeakStageActorPtr::AsActorChecked() const
+{
+	return CastChecked<AActor>(ObjectPtr.Get());
+}
diff --git a/Engine/Plugins/Runtime/nDisplayModularFeatures/Source/DisplayClusterLightCardExtender/Private/StageActor/IDisplayClusterStageActor.cpp b/Engine/Plugins/Runtime/nDisplayModularFeatures/Source/DisplayClusterLightCardExtender/Private/StageActor/IDisplayClusterStageActor.cpp
new file mode 100644
--- /dev/null
+++ b/Engine/Plugins/Runtime/nDisplayModularFeatures/Source/DisplayClusterLightCardExtender/Private/StageActor/IDisplayClusterStageActor.cpp
@@ -0,0 +1,232 @@
+// Copyright Epic Games, Inc. All Rights Reserved.
+
+#include "StageActor/IDisplayClusterStageActor.h"
+
+#include "GameFramework/Actor.h"
+
+const FRotator IDisplayClusterStageActor::PlaneMeshRotation = FRotator(0.0f, -90.0f, 90.0f);
+
+struct FStageActorHelper
+{
+	static AActor* GetActor(IDisplayClusterStageActor* InInterface)
+	{
+		AActor* Actor = Cast<AActor>(InInterface);
+		ensureMsgf(Actor != nullptr, TEXT("IDisplayClusterStageActor should be implemented on Actor types only"));
+		return Actor;
+	}
+
+	static const AActor* GetActorConst(const IDisplayClusterStageActor* InInterface)
+	{
+		const AActor* Actor = Cast<AActor>(InInterface);
+		ensureMsgf(Actor != nullptr, TEXT("IDisplayClusterStageActor should be implemented on Actor types only"));
+		return Actor;
+	}
+};
+
+const TSet<FName>& IDisplayClusterStageActor::GetPositionalPropertyNames() const
+{
+	static TSet<FName> DefaultPositionalPropertyNames =
+		{
+			GET_MEMBER_NAME_CHECKED(FDisplayClusterPositionalParams, DistanceFromCenter),
+			GET_MEMBER_NAME_CHECKED(FDisplayClusterPositionalParams, Longitude),
+			GET_MEMBER_NAME_CHECKED(FDisplayClusterPositionalParams, Latitude),
+			GET_MEMBER_NAME_CHECKED(FDisplayClusterPositionalParams, Spin),
+			GET_MEMBER_NAME_CHECKED(FDisplayClusterPositionalParams, Pitch),
+			GET_MEMBER_NAME_CHECKED(FDisplayClusterPositionalParams, Yaw),
+			GET_MEMBER_NAME_CHECKED(FDisplayClusterPositionalParams, RadialOffset),
+			GET_MEMBER_NAME_CHECKED(FDisplayClusterPositionalParams, Scale),
+			TEXT("UVCoordinates") // Not currently a member of positional params
+		};
+	
+	return DefaultPositionalPropertyNames;
+}
+
+FBox IDisplayClusterStageActor::GetBoxBounds(bool bLocalSpace) const
+{
+	if (const AActor* Actor = FStageActorHelper::GetActorConst(this))
+	{
+		return Actor->GetComponentsBoundingBox();
+	}
+	return FBox();
+}
+
+bool IDisplayClusterStageActor::IsProxy() const
+{
+	if (const AActor* Actor = FStageActorHelper::GetActorConst(this))
+	{
+		return Actor->HasAnyFlags(RF_Transient);
+	}
+
+	return false;
+}
+
+void IDisplayClusterStageActor::UpdateStageActorTransform()
+{
+	const FTransform Transform = PositionalParamsToActorTransform(GetPositionalParams(), GetOrigin());
+	if (AActor* Actor = FStageActorHelper::GetActor(this))
+	{
+		Actor->SetActorTransform(Transform);
+	}
+}
+
+FTransform IDisplayClusterStageActor::GetStageActorTransform(bool bRemoveOrigin) const
+{
+	FTransform Transform;
+
+	if (const AActor* Actor = FStageActorHelper::GetActorConst(this))
+	{
+		FVector Location = Actor->GetActorLocation();
+		if (bRemoveOrigin)
+		{
+			Location -= GetOrigin().GetLocation();
+		}
+		Transform.SetLocation(Location);
+	
+		// Use the actor's orientation, but remove the plane mesh rotation so that the returned transform's local x axis
+		// points radially inwards to match engine convention
+		const FQuat StageActorOrientation = Actor->GetActorQuat() * PlaneMeshRotation.Quaternion().Inverse();
+
+		Transform.SetRotation(StageActorOrientation);
+	}
+	
+	return Transform;
+}
+
+void IDisplayClusterStageActor::SetPositionalParams(
+	const FDisplayClusterPositionalParams& InParams)
+{
+	SetPitch(InParams.Pitch);
+	SetYaw(InParams.Yaw);
+	SetSpin(InParams.Spin);
+	SetLatitude(InParams.Latitude);
+	SetLongitude(InParams.Longitude);
+	SetDistanceFromCenter(InParams.DistanceFromCenter);
+	SetScale(InParams.Scale);
+}
+
+FDisplayClusterPositionalParams IDisplayClusterStageActor::GetPositionalParams() const
+{
+	FDisplayClusterPositionalParams Params = FDisplayClusterPositionalParams();
+
+	Params.Pitch = GetPitch();
+	Params.Yaw = GetYaw();
+	Params.Spin = GetSpin();
+	Params.Latitude = GetLatitude();
+	Params.Longitude = GetLongitude();
+	Params.DistanceFromCenter = GetDistanceFromCenter();
+	Params.Scale = GetScale();
+	
+	ClampLatitudeAndLongitude(Params.Latitude, Params.Longitude);
+	
+	return Params;
+}
+
+void IDisplayClusterStageActor::UpdatePositionalParamsFromTransform()
+{
+	if (const AActor* Actor = FStageActorHelper::GetActorConst(this))
+	{
+		const FDisplayClusterPositionalParams PositionalParams = TransformToPositionalParams(
+		Actor->GetTransform(), GetOrigin(), GetRadialOffset());
+
+		SetPositionalParams(PositionalParams);
+	}
+}
+
+void IDisplayClusterStageActor::ClampLatitudeAndLongitude(double& InOutLatitude, double& InOutLongitude)
+{
+	double Latitude = InOutLatitude;
+	double Longitude = InOutLongitude;
+	if (Longitude < 0 || Longitude > 360)
+	{
+		Longitude = FRotator::ClampAxis(Longitude);
+	}
+
+	if (Latitude < -90 || Latitude > 90)
+	{
+		// If latitude exceeds [-90, 90], mod it back into the appropriate range, and apply a shift of 180 degrees if
+		// needed to the longitude, to allow the latitude to be continuous (increasing latitude indefinitely should result in the LC 
+		// orbiting around a polar great circle)
+		const double Parity = FMath::Fmod(FMath::Abs(Latitude) + 90.f, 360.f) - 180.f;
+		const double DeltaLongitude = Parity > 1 ? 180.f : 0.f;
+
+		double LatMod = FMath::Fmod(Latitude + 90.f, 180.f);
+		if (LatMod < 0.f)
+		{
+			LatMod += 180.f;
+		}
+
+		Latitude = LatMod - 90.f;
+		Longitude = FRotator::ClampAxis(Longitude + DeltaLongitude);
+	}
+
+	InOutLatitude = Latitude;
+	InOutLongitude = Longitude;
+}
+
+FTransform IDisplayClusterStageActor::PositionalParamsToActorTransform(const FDisplayClusterPositionalParams& InParams, const FTransform& InOrigin)
+{
+	const double LocalLatitude = FMath::DegreesToRadians(180.f - InParams.Latitude);
+	const double LocalLongitude = FMath::DegreesToRadians(InParams.Longitude);
+	const double Radius = InParams.DistanceFromCenter + InParams.RadialOffset;
+	
+	FVector Location;
+	{
+		const double CosLat = FMath::Cos(LocalLatitude);
+		const double SinLat = FMath::Sin(LocalLatitude);
+		const double CosLong = FMath::Cos(LocalLongitude);
+		const double SinLong = FMath::Sin(LocalLongitude);
+		Location.X = Radius * CosLat * CosLong;
+		Location.Y = Radius * CosLat * SinLong;
+		Location.Z = Radius * SinLat;
+	}
+	
+	const FRotator SphericalRotation = FRotator(-InParams.Latitude, InParams.Longitude, 0.0);
+	const FRotator StageActorOrientation = FRotator(-InParams.Pitch, InParams.Yaw, InParams.Spin);
+
+	// The origin point may be rotated (ie the DCRA may have a rotation applied)
+	const FVector RotatedLocation = InOrigin.GetRotation().RotateVector(Location);
+	const FQuat InitialRotation = InOrigin.GetRotation();
+	
+	FTransform Result;
+	Result.SetLocation(InOrigin.GetLocation() + RotatedLocation);
+	Result.SetRotation((InitialRotation * SphericalRotation.Quaternion() * StageActorOrientation.Quaternion() * PlaneMeshRotation.Quaternion()).Rotator().Quaternion());
+	Result.SetScale3D(FVector(InParams.Scale, 1.f));
+
+	return Result;
+}
+
+FDisplayClusterPositionalParams IDisplayClusterStageActor::TransformToPositionalParams(
+	const FTransform& InTransform, const FTransform& InOrigin, double InRadialOffset)
+{
+	FDisplayClusterPositionalParams PositionalParams;
+	
+	const FVector Location = InOrigin.GetRotation().UnrotateVector(InTransform.GetLocation() - InOrigin.GetLocation());
+	
+	const double Radius = FVector::Dist(InTransform.GetLocation(), InOrigin.GetLocation());
+	const double DistanceFromCenter = Radius - InRadialOffset;
+	
+	double Latitude = FMath::RadiansToDegrees(FMath::Asin(Location.Z / Radius)) + 180.f;
+	double Longitude = FMath::RadiansToDegrees(FMath::Atan2(Location.Y, Location.X));
+
+	// Clamp is necessary to maintain correct rotation
+	ClampLatitudeAndLongitude(Latitude, Longitude);
+
+	// Find pitch/yaw/spin from the rotation
+	const FQuat SphericalRotation = FRotator(-Latitude, Longitude, 0.0).Quaternion();
+	const FQuat InitialRotation_SphericalRotation_LightCardOrientation = InTransform.GetRotation() * PlaneMeshRotation.Quaternion().Inverse();
+	const FQuat SphericalRotation_LightCardOrientation = InOrigin.GetRotation().Inverse() * InitialRotation_SphericalRotation_LightCardOrientation;
+	const FQuat LightCardOrientation = SphericalRotation.Inverse() * SphericalRotation_LightCardOrientation;
+	
+	const FRotator LightCardRotation = LightCardOrientation.Rotator();
+	
+	PositionalParams.Latitude = Latitude;
+	PositionalParams.Longitude = Longitude;
+	PositionalParams.DistanceFromCenter = DistanceFromCenter;
+	PositionalParams.RadialOffset = InRadialOffset;
+	PositionalParams.Pitch = -LightCardRotation.Pitch;
+	PositionalParams.Yaw = LightCardRotation.Yaw;
+	PositionalParams.Spin = LightCardRotation.Roll;
+	PositionalParams.Scale = FVector2D(InTransform.GetScale3D());
+	
+	return PositionalParams;
+}
diff --git a/Engine/Plugins/Runtime/nDisplayModularFeatures/Source/DisplayClusterLightCardExtender/Public/StageActor/DisplayClusterPositionalParams.h b/Engine/Plugins/Runtime/nDisplayModularFeatures/Source/DisplayClusterLightCardExtender/Public/StageActor/DisplayClusterPositionalParams.h
new file mode 100644
--- /dev/null
+++ b/Engine/Plugins/Runtime/nDisplayModularFeatures/Source/DisplayClusterLightCardExtender/Public/StageActor/DisplayClusterPositionalParams.h
@@ -0,0 +1,42 @@
+// Copyright Epic Games, Inc. All Rights Reserved.
+
+#pragma once
+
+#include "CoreMinimal.h"
+
+#include "DisplayClusterPositionalParams.generated.h"
+
+/**
+ * Positional location and rotation parameters for use with nDisplay stage actors.
+ * Note that the origin point is purposely not included as these parameters are meant to be shared between actors
+ * with different origins.
+ */
+USTRUCT(BlueprintType)
+struct FDisplayClusterPositionalParams
+{
+	GENERATED_BODY()
+	
+	UPROPERTY(EditAnywhere, BlueprintReadWrite, Category = Orientation)
+	double DistanceFromCenter = 0.f;
+	
+	UPROPERTY(EditAnywhere, BlueprintReadWrite, Category = Orientation, meta = (UIMin = 0, ClampMin = 0, UIMax = 360, ClampMax = 360))
+	double Longitude = 0.f;
+	
+	UPROPERTY(EditAnywhere, BlueprintReadWrite, Category = Orientation, meta = (UIMin = -90, ClampMin = -90, UIMax = 90, ClampMax = 90))
+	double Latitude = 0.f;
+	
+	UPROPERTY(EditAnywhere, BlueprintReadWrite, Category = Orientation, meta = (UIMin = -360, ClampMin = -360, UIMax = 360, ClampMax = 360))
+	double Spin = 0.f;
+	
+	UPROPERTY(EditAnywhere, BlueprintReadWrite, Category = Orientation, meta = (UIMin = -360, ClampMin = -360, UIMax = 360, ClampMax = 360))
+	double Pitch = 0.f;
+	
+	UPROPERTY(EditAnywhere, BlueprintReadWrite, Category = Orientation, meta = (UIMin = -360, ClampMin = -360, UIMax = 360, ClampMax = 360))
+	double Yaw = 0.f;
+	
+	UPROPERTY(EditAnywhere, BlueprintReadWrite, Category = Orientation)
+	double RadialOffset = 0.f;
+
+	UPROPERTY(EditAnywhere, BlueprintReadWrite, Category = Orientation)
+	FVector2D Scale;
+};
diff --git a/Engine/Plugins/Runtime/nDisplayModularFeatures/Source/DisplayClusterLightCardExtender/Public/StageActor/DisplayClusterWeakStageActorPtr.h b/Engine/Plugins/Runtime/nDisplayModularFeatures/Source/DisplayClusterLightCardExtender/Public/StageActor/DisplayClusterWeakStageActorPtr.h
new file mode 100644
--- /dev/null
+++ b/Engine/Plugins/Runtime/nDisplayModularFeatures/Source/DisplayClusterLightCardExtender/Public/StageActor/DisplayClusterWeakStageActorPtr.h
@@ -0,0 +1,110 @@
+// Copyright Epic Games, Inc. All Rights Reserved.
+
+#pragma once
+
+#include "UObject/WeakObjectPtr.h"
+
+class IDisplayClusterStageActor;
+class AActor;
+
+/**
+ * Wrapper around a WeakObjectPtr to manage actors implementing IDisplayClusterStageActor.
+ */
+struct DISPLAYCLUSTERLIGHTCARDEXTENDER_API FDisplayClusterWeakStageActorPtr
+{
+private:
+	FWeakObjectPtr ObjectPtr;
+
+public:
+	virtual ~FDisplayClusterWeakStageActorPtr() = default;
+	FDisplayClusterWeakStageActorPtr() = default;
+	FDisplayClusterWeakStageActorPtr(const AActor* InActorPtr);
+
+	/** Dereference the weak pointer */
+	IDisplayClusterStageActor* Get() const;
+	
+	AActor* AsActor() const;
+	AActor* AsActorChecked() const;
+
+protected:
+	FWeakObjectPtr& GetWeakObjectPtr() { return ObjectPtr; }
+	const FWeakObjectPtr& GetWeakObjectPtrConst() const { return ObjectPtr; }
+	
+public:
+	/** Is the weak pointer valid */
+	FORCEINLINE bool IsValid() const
+	{
+		return GetWeakObjectPtrConst().IsValid();
+	}
+
+	/** Reset the weak pointer */
+	FORCEINLINE void Reset()
+	{
+		GetWeakObjectPtr().Reset();
+	}
+
+	/** Dereference the weak pointer */
+	FORCEINLINE IDisplayClusterStageActor* operator*() const
+	{
+		return Get();
+	}
+
+	/** Dereference the weak pointer */
+	FORCEINLINE IDisplayClusterStageActor* operator->() const
+	{
+		return Get();
+	}
+
+	friend bool operator==(const FDisplayClusterWeakStageActorPtr& Lhs, const FDisplayClusterWeakStageActorPtr& Rhs)
+	{
+		return Lhs.GetWeakObjectPtrConst() == Rhs.GetWeakObjectPtrConst();
+	}
+
+	template <typename LhsT, typename RhsT, typename = decltype((LhsT*)nullptr == (RhsT*)nullptr)>
+	friend bool operator==(const FDisplayClusterWeakStageActorPtr& Lhs, const RhsT* Rhs)
+	{
+		return Lhs.Get() == Rhs;
+	}
+
+	template <typename LhsT, typename RhsT, typename = decltype((LhsT*)nullptr == (RhsT*)nullptr)>
+	friend bool operator==(const LhsT* Lhs, const FDisplayClusterWeakStageActorPtr& Rhs)
+	{
+		return Lhs == Rhs.Get();
+	}
+
+	template <typename LhsT>
+	friend bool operator==(const FDisplayClusterWeakStageActorPtr& Lhs, TYPE_OF_NULLPTR)
+	{
+		return !Lhs.IsValid();
+	}
+
+	template <typename RhsT>
+	friend bool operator==(TYPE_OF_NULLPTR, const FDisplayClusterWeakStageActorPtr& Rhs)
+	{
+		return !Rhs.IsValid();
+	}
+
+	template <typename LhsT, typename RhsT, typename = decltype((LhsT*)nullptr != (RhsT*)nullptr)>
+	friend bool operator!=(const FDisplayClusterWeakStageActorPtr& Lhs, const RhsT* Rhs)
+	{
+		return Lhs.Get() != Rhs;
+	}
+
+	template <typename LhsT, typename RhsT, typename = decltype((LhsT*)nullptr != (RhsT*)nullptr)>
+	friend bool operator!=(const LhsT* Lhs, const FDisplayClusterWeakStageActorPtr& Rhs)
+	{
+		return Lhs != Rhs.Get();
+	}
+
+	template <typename LhsT>
+	friend bool operator!=(const FDisplayClusterWeakStageActorPtr& Lhs, TYPE_OF_NULLPTR)
+	{
+		return Lhs.IsValid();
+	}
+
+	template <typename RhsT>
+	friend bool operator!=(TYPE_OF_NULLPTR, const FDisplayClusterWeakStageActorPtr& Rhs)
+	{
+		return Rhs.IsValid();
+	}
+};
diff --git a/Engine/Plugins/Runtime/nDisplayModularFeatures/Source/DisplayClusterLightCardExtender/Public/StageActor/IDisplayClusterStageActor.h b/Engine/Plugins/Runtime/nDisplayModularFeatures/Source/DisplayClusterLightCardExtender/Public/StageActor/IDisplayClusterStageActor.h
new file mode 100644
--- /dev/null
+++ b/Engine/Plugins/Runtime/nDisplayModularFeatures/Source/DisplayClusterLightCardExtender/Public/StageActor/IDisplayClusterStageActor.h
@@ -0,0 +1,89 @@
+// Copyright Epic Games, Inc. All Rights Reserved.
+
+#pragma once
+
+#include "DisplayClusterPositionalParams.h"
+
+#include "UObject/Interface.h"
+
+#include "IDisplayClusterStageActor.generated.h"
+
+UINTERFACE(MinimalAPI, meta = (CannotImplementInterfaceInBlueprint))
+class UDisplayClusterStageActor : public UInterface
+{
+	GENERATED_BODY()
+};
+
+class DISPLAYCLUSTERLIGHTCARDEXTENDER_API IDisplayClusterStageActor
+{
+	GENERATED_BODY()
+
+public:
+	/** The rotation used to orient the plane mesh used for the light card so that its normal points radially inwards */
+	static const FRotator PlaneMeshRotation;
+
+	virtual void SetLongitude(double InValue) = 0;
+	virtual double GetLongitude() const = 0;
+
+	virtual void SetLatitude(double InValue) = 0;
+	virtual double GetLatitude() const = 0;
+
+	virtual void SetDistanceFromCenter(double InValue) = 0;
+	virtual double GetDistanceFromCenter() const = 0;
+
+	virtual void SetSpin(double InValue) = 0;
+	virtual double GetSpin() const = 0;
+
+	virtual void SetPitch(double InValue) = 0;
+	virtual double GetPitch() const = 0;
+
+	virtual void SetYaw(double InValue) = 0;
+	virtual double GetYaw() const = 0;
+
+	virtual void SetRadialOffset(double InValue) = 0;
+	virtual double GetRadialOffset() const = 0;
+
+	virtual void SetOrigin(const FTransform& InOrigin) = 0;
+	virtual FTransform GetOrigin() const = 0;
+
+	virtual void SetScale(const FVector2D& InScale) = 0;
+	virtual FVector2D GetScale() const = 0;
+
+	/** Return property names defined for this stage actor for use in property notifies */
+	virtual const TSet<FName>& GetPositionalPropertyNames() const;
+
+	/** Return the bounding box for the stage actor. The default implementation returns the components bounding box */
+	virtual FBox GetBoxBounds(bool bLocalSpace = false) const;
+
+	/** Whether this actor supports movement in UV space */
+	virtual bool IsUVActor() const { return false; }
+	virtual void SetUVCoordinates(const FVector2D& InUVCoordinates) {}
+	virtual FVector2D GetUVCoordinates() const { return FVector2D(); }
+
+	/** If this stage actor is a copy of a level instance actor */
+	virtual bool IsProxy() const;
+
+	/** Update positional properties of the stage actor. The default implementation updates the actor's transform */
+	virtual void UpdateStageActorTransform();
+
+	/** Return the adjusted transform of the stage actor. The default implementation uses the actor's location as a base */
+	virtual FTransform GetStageActorTransform(bool bRemoveOrigin = false) const;
+
+	/** Sets positional parameters */
+	virtual void SetPositionalParams(const FDisplayClusterPositionalParams& InParams);
+	
+	/** Retrieves positional parameters */
+	virtual FDisplayClusterPositionalParams GetPositionalParams() const;
+	
+	/** Update positional params from current actor transform */
+	virtual void UpdatePositionalParamsFromTransform();
+	
+	/** Clamp the given latitude and longitude */
+	static void ClampLatitudeAndLongitude(double& InOutLatitude, double& InOutLongitude);
+
+	/** Convert latitude/longitude/yaw/pitch/roll to a transform that could be applied to an actor  */
+	static FTransform PositionalParamsToActorTransform(const FDisplayClusterPositionalParams& InParams, const FTransform& InOrigin);
+
+	/** Convert an actor transform to positional params */
+	static FDisplayClusterPositionalParams TransformToPositionalParams(const FTransform& InTransform, const FTransform& InOrigin, double InRadialOffset);
+};
diff --git a/Engine/Plugins/VirtualProduction/DMX/DMXDisplayCluster/Source/DMXDisplayClusterLightCard/Private/DMXDisplayClusterLightCardComponent.cpp b/Engine/Plugins/VirtualProduction/DMX/DMXDisplayCluster/Source/DMXDisplayClusterLightCard/Private/DMXDisplayClusterLightCardComponent.cpp
--- a/Engine/Plugins/VirtualProduction/DMX/DMXDisplayCluster/Source/DMXDisplayClusterLightCard/Private/DMXDisplayClusterLightCardComponent.cpp
+++ b/Engine/Plugins/VirtualProduction/DMX/DMXDisplayCluster/Source/DMXDisplayClusterLightCard/Private/DMXDisplayClusterLightCardComponent.cpp
@@ -47,7 +47,7 @@ struct FDMXDisplayClusterLightCardActorData
 {
 	void Apply(ADisplayClusterLightCardActor* Actor)
 	{
-		if (!Actor || Actor->bIsProxy) 
+		if (!Actor || Actor->IsProxy()) 
 		{
 			return;
 		}
@@ -223,7 +223,7 @@ struct FDMXDisplayClusterLightCardActorData
 		// Apply changes
 		if (bNeedsUpdateLightCardTransform)
 		{
-			Actor->UpdateLightCardTransform();
+			Actor->UpdateStageActorTransform();
 		}
 
 		if (bNeedsUpdateLightCardMaterialInstance)
diff --git a/Engine/Plugins/VirtualProduction/EpicStageApp/EpicStageApp.uplugin b/Engine/Plugins/VirtualProduction/EpicStageApp/EpicStageApp.uplugin
--- a/Engine/Plugins/VirtualProduction/EpicStageApp/EpicStageApp.uplugin
+++ b/Engine/Plugins/VirtualProduction/EpicStageApp/EpicStageApp.uplugin
@@ -33,6 +33,10 @@
 		{
 			"Name": "nDisplay",
 			"Enabled": true
+		},
+		{
+			"Name": "nDisplayModularFeatures",
+			"Enabled": true
 		}
 	]
 }
\ No newline at end of file
diff --git a/Engine/Plugins/VirtualProduction/EpicStageApp/Source/EpicStageApp/EpicStageApp.build.cs b/Engine/Plugins/VirtualProduction/EpicStageApp/Source/EpicStageApp/EpicStageApp.build.cs
--- a/Engine/Plugins/VirtualProduction/EpicStageApp/Source/EpicStageApp/EpicStageApp.build.cs
+++ b/Engine/Plugins/VirtualProduction/EpicStageApp/Source/EpicStageApp/EpicStageApp.build.cs
@@ -19,6 +19,7 @@ public class EpicStageApp : ModuleRules
 				"DeveloperSettings",
 				"DisplayCluster",
 				"DisplayClusterLightCardEditorShaders",
+				"DisplayClusterLightCardExtender",
 				"DisplayClusterScenePreview",
 				"Engine",
 				"ImageWrapper",
diff --git a/Engine/Plugins/VirtualProduction/EpicStageApp/Source/EpicStageApp/Private/StageAppRouteHandler.cpp b/Engine/Plugins/VirtualProduction/EpicStageApp/Source/EpicStageApp/Private/StageAppRouteHandler.cpp
--- a/Engine/Plugins/VirtualProduction/EpicStageApp/Source/EpicStageApp/Private/StageAppRouteHandler.cpp
+++ b/Engine/Plugins/VirtualProduction/EpicStageApp/Source/EpicStageApp/Private/StageAppRouteHandler.cpp
@@ -15,6 +15,7 @@
 #include "RemoteControlWebsocketRoute.h"
 #include "StageAppResponse.h"
 #include "WebRemoteControlUtils.h"
+#include "StageActor/DisplayClusterWeakStageActorPtr.h"
 
 #if WITH_EDITOR
 #include "Editor.h"
@@ -515,7 +516,7 @@ void FStageAppRouteHandler::HandleWebSocketNDisplayPreviewRender(const FRemoteCo
 									continue;
 								}
 
-								WorldPosition = LightCard->GetLightCardTransform().GetTranslation();
+								WorldPosition = LightCard->GetStageActorTransform().GetTranslation();
 							}
 							else
 							{
@@ -730,10 +731,24 @@ void FStageAppRouteHandler::DragLightCards(FPerRendererData& PerRendererData, FV
 		FMath::RoundToInt(DragPosition.Y * Resolution.Y)
 	);
 	const FVector DragWidgetOffset = FVector::ZeroVector;
+
+	// Light card editor supports any type of actor implementing IDisplayClusterStageActor which is managed by
+	// a custom weak ptr. Convert the LightCard actors to the actor weak ptr array.
+	TArray<FDisplayClusterWeakStageActorPtr> StageActors;
+	StageActors.Reserve(PerRendererData.DraggedLightCards.Num());
+	Algo::TransformIf(PerRendererData.DraggedLightCards, StageActors,
+		[](const TWeakObjectPtr<ADisplayClusterLightCardActor>& WeakPtr)
+		{
+			return WeakPtr.IsValid();
+		},
+		[](const TWeakObjectPtr<ADisplayClusterLightCardActor>& WeakPtr)
+		{
+			return FDisplayClusterWeakStageActorPtr(WeakPtr.Get());
+		});
 	
 	FDisplayClusterLightCardEditorHelper& LightCardHelper = PerRendererData.GetLightCardHelper();
-	LightCardHelper.DragLightCards(
-		PerRendererData.DraggedLightCards,
+	LightCardHelper.DragActors(
+		MoveTemp(StageActors),
 		PixelPos,
 		View,
 		FDisplayClusterLightCardEditorHelper::ECoordinateSystem::Spherical,
