nDisplay: Marked specific categories as hidden for the light card actor as they had little to no value as properties. Also hid the label component, which was marked Visibile, and was causing the light card actor's details panel to show all of the label component's properties.

diff --git a/Engine/Plugins/Runtime/nDisplay/Source/DisplayCluster/Public/DisplayClusterLightCardActor.h b/Engine/Plugins/Runtime/nDisplay/Source/DisplayCluster/Public/DisplayClusterLightCardActor.h
--- a/Engine/Plugins/Runtime/nDisplay/Source/DisplayCluster/Public/DisplayClusterLightCardActor.h
+++ b/Engine/Plugins/Runtime/nDisplay/Source/DisplayCluster/Public/DisplayClusterLightCardActor.h
@@ -48,7 +48,7 @@ struct FLightCardAlphaGradientSettings
 	float Angle = 0;
 };
 
-UCLASS(Blueprintable, DisplayName = "Light Card")
+UCLASS(Blueprintable, DisplayName = "Light Card", HideCategories = (Tick, Physics, Collision, Replication, Cooking, Input, Actor))
 class DISPLAYCLUSTER_API ADisplayClusterLightCardActor : public AActor, public IDisplayClusterStageActor
 {
 	GENERATED_BODY()
@@ -260,7 +260,7 @@ protected:
 	UPROPERTY(VisibleAnywhere, Category = "Default")
 	TMap<FName, TObjectPtr<UActorComponent>> ExtenderNameToComponentMap;
 
-	UPROPERTY(VisibleAnywhere, transient, BlueprintReadOnly, Category = "Default")
+	UPROPERTY(transient, BlueprintReadOnly, Category = "Default")
 	TObjectPtr<UDisplayClusterLabelComponent> LabelComponent;
 
 	/** The current owner of the light card */
