nDisplay: Fix possible crash when changing scale of a lightcard in the operator panel.

diff --git a/Engine/Plugins/Runtime/nDisplay/Source/DisplayClusterLightCardEditor/Private/Viewport/DisplayClusterLightCardEditorViewportClient.cpp b/Engine/Plugins/Runtime/nDisplay/Source/DisplayClusterLightCardEditor/Private/Viewport/DisplayClusterLightCardEditorViewportClient.cpp
--- a/Engine/Plugins/Runtime/nDisplay/Source/DisplayClusterLightCardEditor/Private/Viewport/DisplayClusterLightCardEditorViewportClient.cpp
+++ b/Engine/Plugins/Runtime/nDisplay/Source/DisplayClusterLightCardEditor/Private/Viewport/DisplayClusterLightCardEditorViewportClient.cpp
@@ -2009,10 +2009,7 @@ void FDisplayClusterLightCardEditorViewportClient::PropagateActorTransform(const
 				// Only change if values are different.
 				if (!Property->Identical_InContainer(ActorProxy.AsActorChecked(), LevelInstance))
 				{
-					void* NewValue = nullptr;
-					Property->GetValue_InContainer(ActorProxy.AsActorChecked(), &NewValue);
-					Property->SetValue_InContainer(LevelInstance, &NewValue);
-
+					Property->CopyCompleteValue_InContainer(LevelInstance, ActorProxy.AsActorChecked());
 					InOutChangedProperties.Add(Property);
 				}
 			}
