nDisplay: Fixed issue where 3D and UV light cards were unable to be marked Hidden in Game to prevent them from showing up on a stage.

diff --git a/Engine/Plugins/Runtime/nDisplay/Source/DisplayCluster/Private/Game/EngineClasses/Scene/DisplayClusterLightCardActor.cpp b/Engine/Plugins/Runtime/nDisplay/Source/DisplayCluster/Private/Game/EngineClasses/Scene/DisplayClusterLightCardActor.cpp
--- a/Engine/Plugins/Runtime/nDisplay/Source/DisplayCluster/Private/Game/EngineClasses/Scene/DisplayClusterLightCardActor.cpp
+++ b/Engine/Plugins/Runtime/nDisplay/Source/DisplayCluster/Private/Game/EngineClasses/Scene/DisplayClusterLightCardActor.cpp
@@ -326,22 +326,15 @@ void ADisplayClusterLightCardActor::UpdatePolygonTexture()
 
 void ADisplayClusterLightCardActor::UpdateLightCardVisibility()
 {
-	const bool bShouldBeVisible = !bIsUVLightCard || IsProxy();
-	const bool bIsVisible = !IsHidden();
-	if (bIsVisible != bShouldBeVisible)
-	{
-		SetActorHiddenInGame(!bShouldBeVisible);
-	}
-
-#if WITH_EDITORONLY_DATA
-	const bool bShouldHideFromEditor = !bShouldBeVisible;
-	if (bShouldHideFromEditor != bHiddenEdLevel)
+	if (bIsUVLightCard)
 	{
-		bHiddenEdLevel = bShouldHideFromEditor;
-		bHiddenEdLayer = bShouldHideFromEditor;
-		MarkComponentsRenderStateDirty();
+		// Only render the UV light card if this is a proxy actor
+		const bool bShouldBeVisible = IsProxy();
+		if (LightCardComponent->IsVisible() != bShouldBeVisible)
+		{
+			LightCardComponent->SetVisibility(bShouldBeVisible);
+		}
 	}
-#endif
 }
 
 void ADisplayClusterLightCardActor::ShowLightCardLabel(bool bValue, float ScaleValue, ADisplayClusterRootActor* InRootActor)
diff --git a/Engine/Plugins/Runtime/nDisplay/Source/DisplayCluster/Private/Render/Viewport/LightCard/DisplayClusterViewportLightCardManager.cpp b/Engine/Plugins/Runtime/nDisplay/Source/DisplayCluster/Private/Render/Viewport/LightCard/DisplayClusterViewportLightCardManager.cpp
--- a/Engine/Plugins/Runtime/nDisplay/Source/DisplayCluster/Private/Render/Viewport/LightCard/DisplayClusterViewportLightCardManager.cpp
+++ b/Engine/Plugins/Runtime/nDisplay/Source/DisplayCluster/Private/Render/Viewport/LightCard/DisplayClusterViewportLightCardManager.cpp
@@ -121,6 +121,11 @@ void FDisplayClusterViewportLightCardManager::RenderFrame()
 			bool bLoadedPrimitives = false;
 			for (ADisplayClusterLightCardActor* LightCard : UVLightCards)
 			{
+				if (LightCard->IsHidden())
+				{
+					continue;
+				}
+
 				UStaticMeshComponent* LightCardMeshComp = LightCard->GetLightCardMeshComponent();
 				if (LightCardMeshComp && LightCardMeshComp->SceneProxy == nullptr)
 				{
