Removed trailing ';' which C++ UHT considers invalid.

#rnx

diff --git a/Engine/Plugins/Runtime/nDisplay/Source/DisplayCluster/Private/Game/EngineClasses/Scene/DisplayClusterWidgetComponent.h b/Engine/Plugins/Runtime/nDisplay/Source/DisplayCluster/Private/Game/EngineClasses/Scene/DisplayClusterWidgetComponent.h
--- a/Engine/Plugins/Runtime/nDisplay/Source/DisplayCluster/Private/Game/EngineClasses/Scene/DisplayClusterWidgetComponent.h
+++ b/Engine/Plugins/Runtime/nDisplay/Source/DisplayCluster/Private/Game/EngineClasses/Scene/DisplayClusterWidgetComponent.h
@@ -12,7 +12,7 @@
 UCLASS(transient, NotBlueprintable, NotBlueprintType)
 class UDisplayClusterWidgetComponent final : public UWidgetComponent
 {
-	GENERATED_BODY();
+	GENERATED_BODY()
 public:
 	UDisplayClusterWidgetComponent();
 	virtual ~UDisplayClusterWidgetComponent() override;
