[EpicStageApp] Duplicate lightcards and create them with better defaults

diff --git a/Engine/Plugins/Runtime/nDisplay/Source/DisplayCluster/Private/Blueprints/DisplayClusterBlueprintLib.cpp b/Engine/Plugins/Runtime/nDisplay/Source/DisplayCluster/Private/Blueprints/DisplayClusterBlueprintLib.cpp
--- a/Engine/Plugins/Runtime/nDisplay/Source/DisplayCluster/Private/Blueprints/DisplayClusterBlueprintLib.cpp
+++ b/Engine/Plugins/Runtime/nDisplay/Source/DisplayCluster/Private/Blueprints/DisplayClusterBlueprintLib.cpp
@@ -52,6 +52,10 @@ ADisplayClusterLightCardActor* UDisplayClusterBlueprintLib::CreateLightCard(ADis
 		RootActor->GetWorld()->SpawnActor(ADisplayClusterLightCardActor::StaticClass(),
 			&SpawnLocation, &SpawnRotation, MoveTemp(SpawnParameters)));
 
+	NewActor->Latitude = 30.0;
+	NewActor->Longitude = 180.0;
+	NewActor->Color = FLinearColor::Gray;
+
 #if WITH_EDITOR
 	NewActor->SetActorLabel(NewActor->GetName());
 #endif
@@ -69,6 +73,70 @@ ADisplayClusterLightCardActor* UDisplayClusterBlueprintLib::CreateLightCard(ADis
 	return NewActor;
 }
 
+void UDisplayClusterBlueprintLib::DuplicateLightCards(TArray<ADisplayClusterLightCardActor*> OriginalLightcards, TArray<ADisplayClusterLightCardActor*>& OutNewLightCards)
+{
+#if WITH_EDITOR
+	TMap<UWorld*, FCachedActorLabels> ActorLabelsByWorld;
+	FScopedTransaction Transaction(LOCTEXT("DuplicateLightCard", "Duplicate Light Cards"));
+#endif
+
+	for (ADisplayClusterLightCardActor* OriginalLightcard : OriginalLightcards)
+	{
+		if (!OriginalLightcard)
+		{
+			continue;
+		}
+
+		ADisplayClusterRootActor* RootActor = OriginalLightcard->GetRootActorOwner().Get();
+		if (!RootActor)
+		{
+			continue;
+		}
+
+		UWorld* World = RootActor->GetWorld();
+		ULevel* Level = World ? World->GetCurrentLevel() : nullptr;
+		if (!Level)
+		{
+			continue;
+		}
+
+		const FName UniqueName = MakeUniqueObjectName(Level, OriginalLightcard->GetClass());
+		ADisplayClusterLightCardActor* NewLightCard = CastChecked<ADisplayClusterLightCardActor>(StaticDuplicateObject(OriginalLightcard, Level, UniqueName));
+
+#if WITH_EDITOR
+		Level->AddLoadedActor(NewLightCard);
+
+		FCachedActorLabels* ActorLabels = ActorLabelsByWorld.Find(World);
+		if (!ActorLabels)
+		{
+			ActorLabels = &ActorLabelsByWorld.Emplace(World, World);
+		}
+
+		FActorLabelUtilities::SetActorLabelUnique(NewLightCard, NewLightCard->GetActorLabel(), ActorLabels);
+		ActorLabels->Add(NewLightCard->GetActorLabel());
+#endif
+
+		FAttachmentTransformRules AttachmentRules(EAttachmentRule::KeepWorld, false);
+		NewLightCard->AttachToActor(RootActor, AttachmentRules);
+
+		// Add it to the root actor
+		UDisplayClusterConfigurationData* ConfigData = RootActor->GetConfigData();
+		ConfigData->Modify();
+		FDisplayClusterConfigurationICVFX_VisibilityList& RootActorLightCards = ConfigData->StageSettings.Lightcard.ShowOnlyList;
+
+		RootActorLightCards.Actors.Add(NewLightCard);
+
+#if WITH_EDITOR
+		if (GIsEditor)
+		{
+			GEditor->BroadcastLevelActorAdded(NewLightCard);
+		}
+#endif
+
+		OutNewLightCards.Add(NewLightCard);
+	}
+}
+
 void UDisplayClusterBlueprintLib::FindLightCardsForRootActor(ADisplayClusterRootActor* RootActor, TSet<ADisplayClusterLightCardActor*>& OutLightCards)
 {
 	if (!RootActor)
diff --git a/Engine/Plugins/Runtime/nDisplay/Source/DisplayCluster/Public/Blueprints/DisplayClusterBlueprintLib.h b/Engine/Plugins/Runtime/nDisplay/Source/DisplayCluster/Public/Blueprints/DisplayClusterBlueprintLib.h
--- a/Engine/Plugins/Runtime/nDisplay/Source/DisplayCluster/Public/Blueprints/DisplayClusterBlueprintLib.h
+++ b/Engine/Plugins/Runtime/nDisplay/Source/DisplayCluster/Public/Blueprints/DisplayClusterBlueprintLib.h
@@ -28,6 +28,10 @@ public:
 	UFUNCTION(BlueprintPure, meta = (DisplayName = "DisplayCluster Module API"), Category = "NDisplay")
 	static ADisplayClusterLightCardActor* CreateLightCard(ADisplayClusterRootActor* RootActor);
 
+	/** Create duplicates of a list of existing light cards. */
+	UFUNCTION(BlueprintPure, meta = (DisplayName = "DisplayCluster Module API"), Category = "NDisplay")
+	static void DuplicateLightCards(TArray<ADisplayClusterLightCardActor*> OriginalLightcards, TArray<ADisplayClusterLightCardActor*>& OutNewLightCards);
+
 	/** Gets a list of all light card actors on the level linked to the specified root actor. */
 	UFUNCTION(BlueprintPure, meta = (DisplayName = "DisplayCluster Module API"), Category = "NDisplay")
 	static DISPLAYCLUSTER_API void FindLightCardsForRootActor(ADisplayClusterRootActor* RootActor, TSet<ADisplayClusterLightCardActor*>& OutLightCards);
