nDisplay: added HighlightsMax color grading property

diff --git a/Engine/Plugins/Runtime/nDisplay/Source/DisplayCluster/Private/Render/Viewport/Configuration/DisplayClusterViewportConfigurationHelpers_Postprocess.cpp b/Engine/Plugins/Runtime/nDisplay/Source/DisplayCluster/Private/Render/Viewport/Configuration/DisplayClusterViewportConfigurationHelpers_Postprocess.cpp
--- a/Engine/Plugins/Runtime/nDisplay/Source/DisplayCluster/Private/Render/Viewport/Configuration/DisplayClusterViewportConfigurationHelpers_Postprocess.cpp
+++ b/Engine/Plugins/Runtime/nDisplay/Source/DisplayCluster/Private/Render/Viewport/Configuration/DisplayClusterViewportConfigurationHelpers_Postprocess.cpp
@@ -426,6 +426,7 @@ static void ImplBlendPostProcessSettings(FPostProcessSettings& OutputPP, const F
 {
 	PP_CONDITIONAL_BLEND(+, , , , AutoExposureBias, , );
 	PP_CONDITIONAL_BLEND(+, , , , ColorCorrectionHighlightsMin, , );
+	PP_CONDITIONAL_BLEND(+, , , , ColorCorrectionHighlightsMax, , );
 	PP_CONDITIONAL_BLEND(+, , , , ColorCorrectionShadowsMax, , );
 
 	PP_CONDITIONAL_OVERRIDE(, , WhiteBalance., TemperatureType);
@@ -488,6 +489,7 @@ static void ImplCopyPPSStruct(bool bIsConditionalCopy, FDisplayClusterConfigurat
 	{
 		PP_CONDITIONAL_COPY(, , , AutoExposureBias);
 		PP_CONDITIONAL_COPY(, , , ColorCorrectionHighlightsMin);
+		PP_CONDITIONAL_COPY(, , , ColorCorrectionHighlightsMax);
 		PP_CONDITIONAL_COPY(, , , ColorCorrectionShadowsMax);
 
 		PP_CONDITIONAL_COPY(, WhiteBalance., , TemperatureType);
diff --git a/Engine/Plugins/Runtime/nDisplay/Source/DisplayClusterColorGrading/Private/DataModelGenerators/DisplayClusterColorGradingGenerator_RootActor.cpp b/Engine/Plugins/Runtime/nDisplay/Source/DisplayClusterColorGrading/Private/DataModelGenerators/DisplayClusterColorGradingGenerator_RootActor.cpp
--- a/Engine/Plugins/Runtime/nDisplay/Source/DisplayClusterColorGrading/Private/DataModelGenerators/DisplayClusterColorGradingGenerator_RootActor.cpp
+++ b/Engine/Plugins/Runtime/nDisplay/Source/DisplayClusterColorGrading/Private/DataModelGenerators/DisplayClusterColorGradingGenerator_RootActor.cpp
@@ -42,6 +42,7 @@ FDisplayClusterColorGradingDataModel::FColorGradingGroup FDisplayClusterColorGra
 	AddDetailsViewPropertyToGroup(GroupPropertyHandle, ColorGradingGroup, GET_MEMBER_NAME_CHECKED(FDisplayClusterConfigurationViewport_ColorGradingRenderingSettings, AutoExposureBias), TEXT("Exposure"));
 	AddDetailsViewPropertyToGroup(GroupPropertyHandle, ColorGradingGroup, GET_MEMBER_NAME_CHECKED(FDisplayClusterConfigurationViewport_ColorGradingRenderingSettings, ColorCorrectionShadowsMax), TEXT("Color Grading"));
 	AddDetailsViewPropertyToGroup(GroupPropertyHandle, ColorGradingGroup, GET_MEMBER_NAME_CHECKED(FDisplayClusterConfigurationViewport_ColorGradingRenderingSettings, ColorCorrectionHighlightsMin), TEXT("Color Grading"));
+	AddDetailsViewPropertyToGroup(GroupPropertyHandle, ColorGradingGroup, GET_MEMBER_NAME_CHECKED(FDisplayClusterConfigurationViewport_ColorGradingRenderingSettings, ColorCorrectionHighlightsMax), TEXT("Color Grading"));
 	AddDetailsViewPropertyToGroup(GroupPropertyHandle, ColorGradingGroup, GET_MEMBER_NAME_CHECKED(FDisplayClusterConfigurationViewport_ColorGradingWhiteBalanceSettings, TemperatureType), TEXT("White Balance"));
 	AddDetailsViewPropertyToGroup(GroupPropertyHandle, ColorGradingGroup, GET_MEMBER_NAME_CHECKED(FDisplayClusterConfigurationViewport_ColorGradingWhiteBalanceSettings, WhiteTemp), TEXT("White Balance"));
 	AddDetailsViewPropertyToGroup(GroupPropertyHandle, ColorGradingGroup, GET_MEMBER_NAME_CHECKED(FDisplayClusterConfigurationViewport_ColorGradingWhiteBalanceSettings, WhiteTint), TEXT("White Balance"));
diff --git a/Engine/Plugins/Runtime/nDisplay/Source/DisplayClusterConfiguration/Public/DisplayClusterConfigurationTypes_Postprocess.h b/Engine/Plugins/Runtime/nDisplay/Source/DisplayClusterConfiguration/Public/DisplayClusterConfigurationTypes_Postprocess.h
--- a/Engine/Plugins/Runtime/nDisplay/Source/DisplayClusterConfiguration/Public/DisplayClusterConfigurationTypes_Postprocess.h
+++ b/Engine/Plugins/Runtime/nDisplay/Source/DisplayClusterConfiguration/Public/DisplayClusterConfigurationTypes_Postprocess.h
@@ -175,6 +175,7 @@ struct DISPLAYCLUSTERCONFIGURATION_API FDisplayClusterConfigurationViewport_Colo
 		FDisplayClusterConfigurationViewport_ColorGradingRenderingSettings()
 		: bOverride_AutoExposureBias(0)
 		, bOverride_ColorCorrectionHighlightsMin(0)
+		, bOverride_ColorCorrectionHighlightsMax(0)
 		, bOverride_ColorCorrectionShadowsMax(0)
 		, Global()
 		, Shadows()
@@ -182,6 +183,7 @@ struct DISPLAYCLUSTERCONFIGURATION_API FDisplayClusterConfigurationViewport_Colo
 		, Midtones()
 		, Highlights()
 		, ColorCorrectionHighlightsMin(0)
+		, ColorCorrectionHighlightsMax(1)
 	{
 	}
 	
@@ -190,6 +192,8 @@ struct DISPLAYCLUSTERCONFIGURATION_API FDisplayClusterConfigurationViewport_Colo
 	UPROPERTY(EditAnywhere, BlueprintReadWrite, Category = "Overrides", meta = (PinHiddenByDefault, InlineEditConditionToggle))
 	uint8 bOverride_ColorCorrectionHighlightsMin:1;
 	UPROPERTY(EditAnywhere, BlueprintReadWrite, Category = "Overrides", meta = (PinHiddenByDefault, InlineEditConditionToggle))
+	uint8 bOverride_ColorCorrectionHighlightsMax:1;
+	UPROPERTY(EditAnywhere, BlueprintReadWrite, Category = "Overrides", meta = (PinHiddenByDefault, InlineEditConditionToggle))
 	uint8 bOverride_ColorCorrectionShadowsMax:1;
 
 	// Not Implemented: Blend weight
@@ -226,6 +230,9 @@ struct DISPLAYCLUSTERCONFIGURATION_API FDisplayClusterConfigurationViewport_Colo
 	UPROPERTY(EditAnywhere, BlueprintReadWrite, Category = "Viewport Settings", meta = (UIMin = "-1.0", UIMax = "1.0", EditCondition = "bOverride_ColorCorrectionHighlightsMin", DisplayName = "HighlightsMin"))
 	float ColorCorrectionHighlightsMin;
 
+	UPROPERTY(EditAnywhere, BlueprintReadWrite, Category = "Viewport Settings", meta = (UIMin = "1.0", UIMax = "10.0", EditCondition = "bOverride_ColorCorrectionHighlightsMax", DisplayName = "HighlightsMax"))
+	float ColorCorrectionHighlightsMax;
+
 	// Highlights color grading misc settings
 	UPROPERTY(EditAnywhere, BlueprintReadWrite, Category = "Viewport Settings")
 	FDisplayClusterConfigurationViewport_ColorGradingMiscSettings Misc;
diff --git a/Engine/Plugins/Runtime/nDisplay/Source/DisplayClusterTests/Private/DisplayCluster.spec.cpp b/Engine/Plugins/Runtime/nDisplay/Source/DisplayClusterTests/Private/DisplayCluster.spec.cpp
--- a/Engine/Plugins/Runtime/nDisplay/Source/DisplayClusterTests/Private/DisplayCluster.spec.cpp
+++ b/Engine/Plugins/Runtime/nDisplay/Source/DisplayClusterTests/Private/DisplayCluster.spec.cpp
@@ -896,6 +896,7 @@ void FDisplayClusterSpec::DoColorGradingRenderingSettingsPropagationTests(TFunct
 	COLOR_GRADING_RENDERING_SETTINGS_TEST_INTERNAL(UDisplayClusterConfigurationData, 13.7f, 11.5f, AutoExposureBias);
 	COLOR_GRADING_RENDERING_SETTINGS_TEST_INTERNAL(UDisplayClusterConfigurationData, 0.78f, 0.6f, ColorCorrectionShadowsMax);
 	COLOR_GRADING_RENDERING_SETTINGS_TEST_INTERNAL(UDisplayClusterConfigurationData, -0.34f, -0.2f, ColorCorrectionHighlightsMin);
+	COLOR_GRADING_RENDERING_SETTINGS_TEST_INTERNAL(UDisplayClusterConfigurationData, 1.34f, 3.2f, ColorCorrectionHighlightsMax);
 	COLOR_GRADING_RENDERING_SETTINGS_MISC_TEST_INTERNAL(UDisplayClusterConfigurationData, 0.22f, 0.42f, BlueCorrection);
 	COLOR_GRADING_RENDERING_SETTINGS_MISC_TEST_INTERNAL(UDisplayClusterConfigurationData, 0.67f, 0.45f, ExpandGamut);
 	COLOR_GRADING_RENDERING_SETTINGS_MISC_TEST_INTERNAL(UDisplayClusterConfigurationData, FLinearColor(0.6f, 0.2f, 0.9f, 0.3f), FLinearColor(0.2f, 0.5f, 0.1f, 0.8f), SceneColorTint);
