nDisplay: Implemented change to the display cluster color grading drawer that automatically switched back to color grading mode if the selected item had no details properties to display in details view mode, as well as automatically selecting the stage actor when the drawer is first opened.

diff --git a/Engine/Plugins/Runtime/nDisplay/Source/DisplayClusterColorGrading/Private/Drawer/DisplayClusterColorGradingDrawerSingleton.cpp b/Engine/Plugins/Runtime/nDisplay/Source/DisplayClusterColorGrading/Private/Drawer/DisplayClusterColorGradingDrawerSingleton.cpp
--- a/Engine/Plugins/Runtime/nDisplay/Source/DisplayClusterColorGrading/Private/Drawer/DisplayClusterColorGradingDrawerSingleton.cpp
+++ b/Engine/Plugins/Runtime/nDisplay/Source/DisplayClusterColorGrading/Private/Drawer/DisplayClusterColorGradingDrawerSingleton.cpp
@@ -96,6 +96,10 @@ TSharedRef<SWidget> FDisplayClusterColorGradingDrawerSingleton::CreateDrawerCont
 			ColorGradingDrawer.Pin()->SetDrawerState(PreviousDrawerState.GetValue());
 			PreviousDrawerState.Reset();
 		}
+		else
+		{
+			ColorGradingDrawer.Pin()->SetDrawerStateToDefault();
+		}
 
 		return Drawer.ToSharedRef();
 	}
diff --git a/Engine/Plugins/Runtime/nDisplay/Source/DisplayClusterColorGrading/Private/Drawer/SDisplayClusterColorGradingDrawer.cpp b/Engine/Plugins/Runtime/nDisplay/Source/DisplayClusterColorGrading/Private/Drawer/SDisplayClusterColorGradingDrawer.cpp
--- a/Engine/Plugins/Runtime/nDisplay/Source/DisplayClusterColorGrading/Private/Drawer/SDisplayClusterColorGradingDrawer.cpp
+++ b/Engine/Plugins/Runtime/nDisplay/Source/DisplayClusterColorGrading/Private/Drawer/SDisplayClusterColorGradingDrawer.cpp
@@ -532,6 +532,15 @@ void SDisplayClusterColorGradingDrawer::SetDrawerState(const FDisplayClusterColo
 	}
 }
 
+void SDisplayClusterColorGradingDrawer::SetDrawerStateToDefault()
+{
+	// The nDisplay stage actor is always the first item in the root actor color grading items list, so set that as the currently selected item
+	if (RootActorList.IsValid() && RootActorColorGradingItems.Num() > 0)
+	{
+		RootActorList->SetSelectedItems({ RootActorColorGradingItems[0] });
+	}
+}
+
 TSharedRef<SWidget> SDisplayClusterColorGradingDrawer::CreateDockInLayoutButton()
 {
 	if (bIsInDrawer)
@@ -949,6 +958,13 @@ void SDisplayClusterColorGradingDrawer::OnListSelectionChanged(TSharedRef<SDispl
 	}
 
 	ColorGradingDataModel->SetObjects(ObjectsToColorGrade);
+
+	// If the current drawer mode is the details view but the newly selected objects don't have any details sections to display,
+	// switch to the color grading drawer mode
+	if (CurrentDrawerMode == EDisplayClusterColorGradingDrawerMode::DetailsView && ColorGradingDataModel->DetailsSections.Num() == 0)
+	{
+		CurrentDrawerMode = EDisplayClusterColorGradingDrawerMode::ColorGrading;
+	}
 }
 
 FReply SDisplayClusterColorGradingDrawer::DockInLayout()
diff --git a/Engine/Plugins/Runtime/nDisplay/Source/DisplayClusterColorGrading/Private/Drawer/SDisplayClusterColorGradingDrawer.h b/Engine/Plugins/Runtime/nDisplay/Source/DisplayClusterColorGrading/Private/Drawer/SDisplayClusterColorGradingDrawer.h
--- a/Engine/Plugins/Runtime/nDisplay/Source/DisplayClusterColorGrading/Private/Drawer/SDisplayClusterColorGradingDrawer.h
+++ b/Engine/Plugins/Runtime/nDisplay/Source/DisplayClusterColorGrading/Private/Drawer/SDisplayClusterColorGradingDrawer.h
@@ -85,6 +85,9 @@ public:
 	/** Sets the state of the drawer UI */
 	void SetDrawerState(const FDisplayClusterColorGradingDrawerState& InDrawerState);
 
+	/** Sets the state of the drawer UI to its default value, which is to have the nDisplay stage actor selected */
+	void SetDrawerStateToDefault();
+
 private:
 	/** Creates the button used to dock the drawer in the operator panel */
 	TSharedRef<SWidget> CreateDockInLayoutButton();
