nDisplay: rendering latency for media + fix for backbuffer capture

diff --git a/Engine/Plugins/Runtime/nDisplay/Source/DisplayCluster/Private/DisplayClusterCallbacks.h b/Engine/Plugins/Runtime/nDisplay/Source/DisplayCluster/Private/DisplayClusterCallbacks.h
--- a/Engine/Plugins/Runtime/nDisplay/Source/DisplayCluster/Private/DisplayClusterCallbacks.h
+++ b/Engine/Plugins/Runtime/nDisplay/Source/DisplayCluster/Private/DisplayClusterCallbacks.h
@@ -102,9 +102,14 @@ public:
 		return DisplayClusterPostWarpViewportEvent;
 	}
 
-	virtual FDisplayClusterPreFrameRender_RenderThread& OnDisplayClusterPreFrameRender_RenderThread() override
+	virtual FDisplayClusterPostCrossGpuTransfer_RenderThread& OnDisplayClusterPostCrossGpuTransfer_RenderThread() override
 	{
-		return DisplayClusterPreFrameRenderEvent;
+		return DisplayClusterPostCrossGpuTransferEvent;
+	}
+
+	virtual FDisplayClusterProcessLatency_RenderThread& OnDisplayClusterProcessLatency_RenderThread() override
+	{
+		return FDisplayClusterProcessLatencyEvent;
 	}
 
 	virtual FDisplayClusterPostFrameRender_RenderThread& OnDisplayClusterPostFrameRender_RenderThread() override
@@ -112,6 +117,16 @@ public:
 		return DisplayClusterPostFrameRenderEvent;
 	}
 
+	virtual FDisplayClusterPostBackbufferUpdate_RenderThread& OnDisplayClusterPostBackbufferUpdate_RenderThread() override
+	{
+		return DisplayClusterPostBackbufferUpdateEvent;
+	}
+
+	virtual FDisplayClusterPreProcessIcvfx_RenderThread& OnDisplayClusterPreProcessIcvfx_RenderThread() override
+	{
+		return FDisplayClusterPreProcessIcvfxEvent;
+	}
+
 private:
 	FDisplayClusterStartSessionEvent         DisplayClusterStartSessionEvent;
 	FDisplayClusterEndSessionEvent           DisplayClusterEndSessionEvent;
@@ -129,10 +144,13 @@ private:
 	FDisplayClusterPresentationPostSynchronization_RHIThread DisplayClusterPresentationPostSynchronizationEvent;
 
 	FDisplayClusterPostRenderViewFamily_RenderThread DisplayClusterPostRenderViewFamily;
-	FDisplayClusterPreWarp_RenderThread          DisplayClusterPreWarpEvent;
-	FDisplayClusterPreWarpViewport_RenderThread  DisplayClusterPreWarpViewportEvent;
-	FDisplayClusterPostWarp_RenderThread         DisplayClusterPostWarpEvent;
-	FDisplayClusterPostWarpViewport_RenderThread DisplayClusterPostWarpViewportEvent;
-	FDisplayClusterPreFrameRender_RenderThread   DisplayClusterPreFrameRenderEvent;
-	FDisplayClusterPostFrameRender_RenderThread  DisplayClusterPostFrameRenderEvent;
+	FDisplayClusterPreWarp_RenderThread              DisplayClusterPreWarpEvent;
+	FDisplayClusterPreWarpViewport_RenderThread      DisplayClusterPreWarpViewportEvent;
+	FDisplayClusterPostWarp_RenderThread             DisplayClusterPostWarpEvent;
+	FDisplayClusterPostWarpViewport_RenderThread     DisplayClusterPostWarpViewportEvent;
+	FDisplayClusterPostCrossGpuTransfer_RenderThread DisplayClusterPostCrossGpuTransferEvent;
+	FDisplayClusterProcessLatency_RenderThread       FDisplayClusterProcessLatencyEvent;
+	FDisplayClusterPostFrameRender_RenderThread      DisplayClusterPostFrameRenderEvent;
+	FDisplayClusterPostBackbufferUpdate_RenderThread DisplayClusterPostBackbufferUpdateEvent;
+	FDisplayClusterPreProcessIcvfx_RenderThread      FDisplayClusterPreProcessIcvfxEvent;
 };
diff --git a/Engine/Plugins/Runtime/nDisplay/Source/DisplayCluster/Private/Render/Viewport/Configuration/DisplayClusterViewportConfigurationHelpers.cpp b/Engine/Plugins/Runtime/nDisplay/Source/DisplayCluster/Private/Render/Viewport/Configuration/DisplayClusterViewportConfigurationHelpers.cpp
--- a/Engine/Plugins/Runtime/nDisplay/Source/DisplayCluster/Private/Render/Viewport/Configuration/DisplayClusterViewportConfigurationHelpers.cpp
+++ b/Engine/Plugins/Runtime/nDisplay/Source/DisplayCluster/Private/Render/Viewport/Configuration/DisplayClusterViewportConfigurationHelpers.cpp
@@ -210,13 +210,18 @@ void FDisplayClusterViewportConfigurationHelpers::UpdateBaseViewportSetting(FDis
 	{
 		const FDisplayClusterConfigurationMedia& MediaSettings = InConfigurationViewport.RenderSettings.Media;
 
-		if (MediaSettings.bEnabled)
-		{
-			// In case this viewport consumes data from a media input device, there is no need to render scene
-			DstViewport.RenderSettings.bSkipSceneRenderingButLeaveResourcesAvailable = MediaSettings.MediaInput.bEnabled;
-			// Set media capture flag if media capture is used
-			DstViewport.RenderSettings.bIsBeingCaptured = InConfigurationViewport.RenderSettings.Media.MediaOutput.bEnabled;
-		}
+		const FString ThisClusterNodeId = IDisplayCluster::Get().GetClusterMgr()->GetNodeId();
+		const bool bThisNodeSharesMedia = MediaSettings.IsMediaSharingUsed() && MediaSettings.MediaSharingNode.Equals(ThisClusterNodeId, ESearchCase::IgnoreCase);
+
+		// Don't render the viewport if media input assigned
+		DstViewport.RenderSettings.bSkipSceneRenderingButLeaveResourcesAvailable = MediaSettings.IsMediaSharingUsed() ?
+			!bThisNodeSharesMedia :
+			!!MediaSettings.MediaSource;
+
+		// Mark this viewport is going to be captured by a capture device
+		DstViewport.RenderSettings.bIsBeingCaptured = MediaSettings.IsMediaSharingUsed() ?
+			bThisNodeSharesMedia :
+			!!MediaSettings.MediaOutput;
 	}
 
 	// FDisplayClusterConfigurationViewport_ICVFX property:
diff --git a/Engine/Plugins/Runtime/nDisplay/Source/DisplayCluster/Private/Render/Viewport/Configuration/DisplayClusterViewportConfigurationHelpers_ICVFX.cpp b/Engine/Plugins/Runtime/nDisplay/Source/DisplayCluster/Private/Render/Viewport/Configuration/DisplayClusterViewportConfigurationHelpers_ICVFX.cpp
--- a/Engine/Plugins/Runtime/nDisplay/Source/DisplayCluster/Private/Render/Viewport/Configuration/DisplayClusterViewportConfigurationHelpers_ICVFX.cpp
+++ b/Engine/Plugins/Runtime/nDisplay/Source/DisplayCluster/Private/Render/Viewport/Configuration/DisplayClusterViewportConfigurationHelpers_ICVFX.cpp
@@ -568,21 +568,20 @@ void FDisplayClusterViewportConfigurationHelpers_ICVFX::UpdateCameraViewportSett
 	// Set media related configuration (runtime only for now)
 	if (IDisplayCluster::Get().GetOperationMode() == EDisplayClusterOperationMode::Cluster)
 	{
-		const FDisplayClusterConfigurationICVFXMedia& MediaSettings = InCameraComponent.CameraSettings.RenderSettings.Media;
-		if (MediaSettings.bEnabled)
-		{
-			const bool bThisNodeSharesMedia = MediaSettings.MediaOutputNode.Equals(DstViewport.GetClusterNodeId(), ESearchCase::IgnoreCase);
+		const FDisplayClusterConfigurationMedia& MediaSettings = InCameraComponent.CameraSettings.RenderSettings.Media;
 
-			// In most cases there is no need to render ICVFX view if media input set up. The only exception
-			// is when this camera is used as media source for other cluster nodes. In this case media input
-			// settings are ignored, and media capture is used.
-			DstViewport.RenderSettings.bSkipSceneRenderingButLeaveResourcesAvailable = MediaSettings.MediaInput.bEnabled ?
-					!(MediaSettings.MediaOutput.bEnabled && bThisNodeSharesMedia) :
-					false;
+		const FString ThisClusterNodeId = DstViewport.GetClusterNodeId();
+		const bool bThisNodeSharesMedia = MediaSettings.IsMediaSharingUsed() && MediaSettings.MediaSharingNode.Equals(ThisClusterNodeId, ESearchCase::IgnoreCase);
 
-			// Set media capture flag if media capture is used
-			DstViewport.RenderSettings.bIsBeingCaptured = InCameraComponent.CameraSettings.RenderSettings.Media.MediaOutput.bEnabled;
-		}
+		// Don't render the viewport if media input assigned
+		DstViewport.RenderSettings.bSkipSceneRenderingButLeaveResourcesAvailable = MediaSettings.IsMediaSharingUsed() ?
+			!bThisNodeSharesMedia :
+			!!MediaSettings.MediaSource;
+
+		// Mark this viewport is going to be captured by a capture device
+		DstViewport.RenderSettings.bIsBeingCaptured = MediaSettings.IsMediaSharingUsed() ?
+			bThisNodeSharesMedia :
+			!!MediaSettings.MediaOutput;
 	}
 }
 
diff --git a/Engine/Plugins/Runtime/nDisplay/Source/DisplayCluster/Private/Render/Viewport/DisplayClusterViewportManagerProxy.cpp b/Engine/Plugins/Runtime/nDisplay/Source/DisplayCluster/Private/Render/Viewport/DisplayClusterViewportManagerProxy.cpp
--- a/Engine/Plugins/Runtime/nDisplay/Source/DisplayCluster/Private/Render/Viewport/DisplayClusterViewportManagerProxy.cpp
+++ b/Engine/Plugins/Runtime/nDisplay/Source/DisplayCluster/Private/Render/Viewport/DisplayClusterViewportManagerProxy.cpp
@@ -268,7 +268,10 @@ void FDisplayClusterViewportManagerProxy::ImplRenderFrame(FViewport* InViewport)
 			// Now all resources on GPU#0
 		}
 
-		IDisplayCluster::Get().GetCallbacks().OnDisplayClusterPreFrameRender_RenderThread().Broadcast(RHICmdList, ViewportManagerProxy, InViewport);
+		// PostCrossGpuTransfer notification
+		IDisplayCluster::Get().GetCallbacks().OnDisplayClusterPostCrossGpuTransfer_RenderThread().Broadcast(RHICmdList, ViewportManagerProxy, InViewport);
+		// Latency processing
+		IDisplayCluster::Get().GetCallbacks().OnDisplayClusterProcessLatency_RenderThread().Broadcast(RHICmdList, ViewportManagerProxy, InViewport);
 
 		// Update viewports resources: overlay, vp-overla, blur, nummips, etc
 		ViewportManagerProxy->UpdateDeferredResources_RenderThread(RHICmdList);
@@ -286,6 +289,9 @@ void FDisplayClusterViewportManagerProxy::ImplRenderFrame(FViewport* InViewport)
 			ViewportManagerProxy->PostProcessManager->HandleEndUpdateFrameResources_RenderThread(RHICmdList, ViewportManagerProxy);
 		}
 
+		// Postrender notification before copying final image to the backbuffer
+		IDisplayCluster::Get().GetCallbacks().OnDisplayClusterPostFrameRender_RenderThread().Broadcast(RHICmdList, ViewportManagerProxy, InViewport);
+
 		if (InViewport)
 		{
 			if (FRHITexture2D* FrameOutputRTT = InViewport->GetRenderTargetTexture())
@@ -304,10 +310,11 @@ void FDisplayClusterViewportManagerProxy::ImplRenderFrame(FViewport* InViewport)
 				}
 
 				ViewportManagerProxy->ResolveFrameTargetToBackBuffer_RenderThread(RHICmdList, 0, 0, FrameOutputRTT, FrameOutputRTT->GetSizeXY());
+
+				// Finally, notify about backbuffer update
+				IDisplayCluster::Get().GetCallbacks().OnDisplayClusterPostBackbufferUpdate_RenderThread().Broadcast(RHICmdList, ViewportManagerProxy, InViewport);
 			}
 		}
-
-		IDisplayCluster::Get().GetCallbacks().OnDisplayClusterPostFrameRender_RenderThread().Broadcast(RHICmdList, ViewportManagerProxy, InViewport);
 	});
 }
 
diff --git a/Engine/Plugins/Runtime/nDisplay/Source/DisplayCluster/Private/Render/Viewport/DisplayClusterViewportManagerViewExtension.cpp b/Engine/Plugins/Runtime/nDisplay/Source/DisplayCluster/Private/Render/Viewport/DisplayClusterViewportManagerViewExtension.cpp
--- a/Engine/Plugins/Runtime/nDisplay/Source/DisplayCluster/Private/Render/Viewport/DisplayClusterViewportManagerViewExtension.cpp
+++ b/Engine/Plugins/Runtime/nDisplay/Source/DisplayCluster/Private/Render/Viewport/DisplayClusterViewportManagerViewExtension.cpp
@@ -44,35 +44,18 @@ bool FDisplayClusterViewportManagerViewExtension::IsActiveThisFrame_Internal(con
 
 void FDisplayClusterViewportManagerViewExtension::PostRenderViewFamily_RenderThread(FRDGBuilder& GraphBuilder, FSceneViewFamily& InViewFamily)
 {
-	// Since FSceneViewFamily has no direct references to the nD viewport it relates to,
-	// we find a proper viewport ID by comparing RenderTarget addresses.
+	// FSceneViewFamily::ProfileDescription is used currently as a viewport ID holder, so we use it to find a proper viewport proxy.
 	if (ViewportManager)
 	{
-		const IDisplayClusterViewportManagerProxy* const ViewportMgrPoxy = ViewportManager->GetProxy();
-		if (ViewportMgrPoxy)
+		if (const IDisplayClusterViewportManagerProxy* const ViewportMgrPoxy = ViewportManager->GetProxy())
 		{
-			// Get all available viewport proxies
-			const TArrayView<IDisplayClusterViewportProxy*> ViewportProxies = ViewportMgrPoxy->GetViewports_RenderThread();
-
-			// Filter those viewports that refer to the RenderTarget used in the ViewFamily
-			const TArray<IDisplayClusterViewportProxy*> FoundViewportProxies = ViewportProxies.FilterByPredicate([&InViewFamily](const IDisplayClusterViewportProxy* ViewportProxy)
-			{	
-				TArray<FRHITexture*> RenderTargets;
-				if (ViewportProxy->GetResources_RenderThread(EDisplayClusterViewportResourceType::InternalRenderTargetResource, RenderTargets))
+			if (!InViewFamily.ProfileDescription.IsEmpty())
+			{
+				if (IDisplayClusterViewportProxy* ViewportProxy = ViewportMgrPoxy->FindViewport_RenderThread(InViewFamily.ProfileDescription))
 				{
-					if (RenderTargets.Contains(InViewFamily.RenderTarget->GetRenderTargetTexture()))
-					{
-						return true;
-					}
+					// Notify about viewport rendering finish
+					IDisplayCluster::Get().GetCallbacks().OnDisplayClusterPostRenderViewFamily_RenderThread().Broadcast(GraphBuilder, InViewFamily, ViewportProxy);
 				}
-
-				return false;
-			});
-			
-			// Now we can perform per-viewport notification
-			for (const IDisplayClusterViewportProxy* ViewportProxy : FoundViewportProxies)
-			{
-				IDisplayCluster::Get().GetCallbacks().OnDisplayClusterPostRenderViewFamily_RenderThread().Broadcast(GraphBuilder, InViewFamily, ViewportProxy);
 			}
 		}
 	}
diff --git a/Engine/Plugins/Runtime/nDisplay/Source/DisplayCluster/Private/Render/Viewport/DisplayClusterViewportManagerViewExtension.h b/Engine/Plugins/Runtime/nDisplay/Source/DisplayCluster/Private/Render/Viewport/DisplayClusterViewportManagerViewExtension.h
--- a/Engine/Plugins/Runtime/nDisplay/Source/DisplayCluster/Private/Render/Viewport/DisplayClusterViewportManagerViewExtension.h
+++ b/Engine/Plugins/Runtime/nDisplay/Source/DisplayCluster/Private/Render/Viewport/DisplayClusterViewportManagerViewExtension.h
@@ -4,7 +4,9 @@
 #include "SceneViewExtension.h"
 
 class IDisplayClusterViewportManager;
-/** 
+
+
+/**
  * View extension applying an DC Viewport features
  */
 class FDisplayClusterViewportManagerViewExtension : public FSceneViewExtensionBase
@@ -12,15 +14,24 @@ class FDisplayClusterViewportManagerViewExtension : public FSceneViewExtensionBa
 public:
 	FDisplayClusterViewportManagerViewExtension(const FAutoRegister& AutoRegister, const IDisplayClusterViewportManager* InViewportManager);
 
+public:
 	//~ Begin ISceneViewExtension interface
-	virtual int32 GetPriority() const override { return -1; }
+	virtual int32 GetPriority() const override
+	{
+		return -1;
+	}
 
-	virtual void SetupViewFamily(FSceneViewFamily& InViewFamily) override {}
-	virtual void SetupView(FSceneViewFamily& InViewFamily, FSceneView& InView) override {}
+	virtual void SetupViewFamily(FSceneViewFamily& InViewFamily) override
+	{ }
 
-	virtual void BeginRenderViewFamily(FSceneViewFamily& InViewFamily) override {};
+	virtual void SetupView(FSceneViewFamily& InViewFamily, FSceneView& InView) override
+	{ }
+
+	virtual void BeginRenderViewFamily(FSceneViewFamily& InViewFamily) override
+	{ }
 
 	virtual void PreRenderView_RenderThread(FRDGBuilder& GraphBuilder, FSceneView& InView) override;
+
 	virtual bool IsActiveThisFrame_Internal(const FSceneViewExtensionContext& Context) const override;
 
 	virtual void PostRenderViewFamily_RenderThread(FRDGBuilder& GraphBuilder, FSceneViewFamily& InViewFamily) override;
diff --git a/Engine/Plugins/Runtime/nDisplay/Source/DisplayCluster/Public/DisplayClusterRootActor.h b/Engine/Plugins/Runtime/nDisplay/Source/DisplayCluster/Public/DisplayClusterRootActor.h
--- a/Engine/Plugins/Runtime/nDisplay/Source/DisplayCluster/Public/DisplayClusterRootActor.h
+++ b/Engine/Plugins/Runtime/nDisplay/Source/DisplayCluster/Public/DisplayClusterRootActor.h
@@ -228,6 +228,10 @@ private:
 
 	UPROPERTY(EditAnywhere, Transient, Category = "Light Cards", meta = (PropertyPath = "CurrentConfigData.StageSettings.Lightcard.ShowOnlyList", EditConditionPath = "CurrentConfigData.StageSettings.Lightcard.bEnable"))
 	FDisplayClusterEditorPropertyReference LightCardContentRef;
+
+	UPROPERTY(EditAnywhere, Transient, Category = "Media", meta = (PropertyPath = "CurrentConfigData.MediaSettings"))
+	FDisplayClusterEditorPropertyReference MediaSettingsRef;
+
 #endif // WITH_EDITORONLY_DATA
 
 private:
diff --git a/Engine/Plugins/Runtime/nDisplay/Source/DisplayCluster/Public/IDisplayClusterCallbacks.h b/Engine/Plugins/Runtime/nDisplay/Source/DisplayCluster/Public/IDisplayClusterCallbacks.h
--- a/Engine/Plugins/Runtime/nDisplay/Source/DisplayCluster/Public/IDisplayClusterCallbacks.h
+++ b/Engine/Plugins/Runtime/nDisplay/Source/DisplayCluster/Public/IDisplayClusterCallbacks.h
@@ -4,12 +4,14 @@
 
 #include "DisplayClusterEnums.h"
 
-class IDisplayClusterViewportProxy;
-class IDisplayClusterViewportManagerProxy;
-class FRHICommandListImmediate;
+class FDisplayClusterShaderParameters_ICVFX;
 class FRDGBuilder;
+class FRHICommandListImmediate;
 class FSceneViewFamily;
 class FViewport;
+class IDisplayClusterViewportProxy;
+class IDisplayClusterViewportManagerProxy;
+struct FDisplayClusterShaderParameters_WarpBlend;
 
 
 /**
@@ -93,11 +95,23 @@ public:
 	DECLARE_EVENT_TwoParams(IDisplayClusterCallbacks, FDisplayClusterPostWarpViewport_RenderThread, FRHICommandListImmediate&, const IDisplayClusterViewportProxy*);
 	virtual FDisplayClusterPostWarpViewport_RenderThread& OnDisplayClusterPostWarpViewport_RenderThread() = 0;
 
-	/** Called before staring nDisplay rendering chain **/
-	DECLARE_EVENT_ThreeParams(IDisplayClusterCallbacks, FDisplayClusterPreFrameRender_RenderThread, FRHICommandListImmediate&, const IDisplayClusterViewportManagerProxy*, FViewport*);
-	virtual FDisplayClusterPreFrameRender_RenderThread& OnDisplayClusterPreFrameRender_RenderThread() = 0;
+	/** Called after inter-GPU synchronization **/
+	DECLARE_EVENT_ThreeParams(IDisplayClusterCallbacks, FDisplayClusterPostCrossGpuTransfer_RenderThread, FRHICommandListImmediate&, const IDisplayClusterViewportManagerProxy*, FViewport*);
+	virtual FDisplayClusterPostCrossGpuTransfer_RenderThread& OnDisplayClusterPostCrossGpuTransfer_RenderThread() = 0;
+
+	/** Called to let the artificial latency subsystem do its job **/
+	DECLARE_EVENT_ThreeParams(IDisplayClusterCallbacks, FDisplayClusterProcessLatency_RenderThread, FRHICommandListImmediate&, const IDisplayClusterViewportManagerProxy*, FViewport*);
+	virtual FDisplayClusterProcessLatency_RenderThread& OnDisplayClusterProcessLatency_RenderThread() = 0;
 
-	/** Last call from the nDisplay rendering chain **/
+	/** Called in the end of the nD rendering chain, right before updating the backbuffer **/
 	DECLARE_EVENT_ThreeParams(IDisplayClusterCallbacks, FDisplayClusterPostFrameRender_RenderThread, FRHICommandListImmediate&, const IDisplayClusterViewportManagerProxy*, FViewport*);
 	virtual FDisplayClusterPostFrameRender_RenderThread& OnDisplayClusterPostFrameRender_RenderThread() = 0;
+
+	/** Called after backbuffer update **/
+	DECLARE_EVENT_ThreeParams(IDisplayClusterCallbacks, FDisplayClusterPostBackbufferUpdate_RenderThread, FRHICommandListImmediate&, const IDisplayClusterViewportManagerProxy*, FViewport*);
+	virtual FDisplayClusterPostBackbufferUpdate_RenderThread& OnDisplayClusterPostBackbufferUpdate_RenderThread() = 0;
+
+	/** Called before applying ICVFX shaders **/
+	DECLARE_EVENT_FourParams(IDisplayClusterCallbacks, FDisplayClusterPreProcessIcvfx_RenderThread, FRHICommandListImmediate&, const IDisplayClusterViewportProxy*, FDisplayClusterShaderParameters_WarpBlend&, FDisplayClusterShaderParameters_ICVFX&);
+	virtual FDisplayClusterPreProcessIcvfx_RenderThread& OnDisplayClusterPreProcessIcvfx_RenderThread() = 0;
 };
diff --git a/Engine/Plugins/Runtime/nDisplay/Source/DisplayClusterConfiguration/Public/DisplayClusterConfigurationTypes.h b/Engine/Plugins/Runtime/nDisplay/Source/DisplayClusterConfiguration/Public/DisplayClusterConfigurationTypes.h
--- a/Engine/Plugins/Runtime/nDisplay/Source/DisplayClusterConfiguration/Public/DisplayClusterConfigurationTypes.h
+++ b/Engine/Plugins/Runtime/nDisplay/Source/DisplayClusterConfiguration/Public/DisplayClusterConfigurationTypes.h
@@ -566,6 +566,9 @@ public:
 	UPROPERTY(EditDefaultsOnly, BlueprintReadWrite, Category = Configuration, meta = (DisplayName = "Override Viewports From .ndisplay", DisplayAfter = "bExitOnEsc"))
 	bool bOverrideViewportsFromExternalConfig = false;
 
+	UPROPERTY(EditAnywhere, BlueprintReadWrite, Category = Configuration, meta = (DisplayName = "Global Media Settings"))
+	FDisplayClusterConfigurationGlobalMediaSettings MediaSettings;
+
 	/** Create empty config data. */
 	static UDisplayClusterConfigurationData* CreateNewConfigData(UObject* Owner = nullptr, EObjectFlags ObjectFlags = RF_NoFlags);
 #if WITH_EDITORONLY_DATA
diff --git a/Engine/Plugins/Runtime/nDisplay/Source/DisplayClusterConfiguration/Public/DisplayClusterConfigurationTypes_ICVFX.h b/Engine/Plugins/Runtime/nDisplay/Source/DisplayClusterConfiguration/Public/DisplayClusterConfigurationTypes_ICVFX.h
--- a/Engine/Plugins/Runtime/nDisplay/Source/DisplayClusterConfiguration/Public/DisplayClusterConfigurationTypes_ICVFX.h
+++ b/Engine/Plugins/Runtime/nDisplay/Source/DisplayClusterConfiguration/Public/DisplayClusterConfigurationTypes_ICVFX.h
@@ -323,7 +323,7 @@ public:
 
 	// Media settings
 	UPROPERTY(EditAnywhere, BlueprintReadWrite, Category = NDisplay, meta = (DisplayName = "Media"))
-	FDisplayClusterConfigurationICVFXMedia Media;
+	FDisplayClusterConfigurationMedia Media;
 
 	UPROPERTY()
 	FDisplayClusterConfigurationPostRender_BlurPostprocess PostprocessBlur;
diff --git a/Engine/Plugins/Runtime/nDisplay/Source/DisplayClusterConfiguration/Public/DisplayClusterConfigurationTypes_Media.h b/Engine/Plugins/Runtime/nDisplay/Source/DisplayClusterConfiguration/Public/DisplayClusterConfigurationTypes_Media.h
--- a/Engine/Plugins/Runtime/nDisplay/Source/DisplayClusterConfiguration/Public/DisplayClusterConfigurationTypes_Media.h
+++ b/Engine/Plugins/Runtime/nDisplay/Source/DisplayClusterConfiguration/Public/DisplayClusterConfigurationTypes_Media.h
@@ -13,91 +13,49 @@
 
 
 /*
- * Media input settings
+ * Media settings
  */
 USTRUCT(Blueprintable)
-struct DISPLAYCLUSTERCONFIGURATION_API FDisplayClusterConfigurationMediaInput
+struct DISPLAYCLUSTERCONFIGURATION_API FDisplayClusterConfigurationMedia
 {
 	GENERATED_BODY()
 
 public:
-	/** Enable/disable media input */
-	UPROPERTY(EditAnywhere, BlueprintReadWrite, Category = "Media Input", meta = (DisplayName = "Enable"))
-	bool bEnabled = false;
+	/** Enable/disable in-cluster media sharing */
+	UPROPERTY(EditAnywhere, BlueprintReadWrite, Category = "Media Sharing")
+	bool bMediaSharing = false;
+
+	/** When in-cluster media sharing us used, the cluster node specified here will be used as a source (Tx node). The others will be Tx'es. */
+	UPROPERTY(EditAnywhere, BlueprintReadWrite, Category = "Media Sharing")
+	FString MediaSharingNode;
 
 	/** Media source to use */
-	UPROPERTY(EditAnywhere, BlueprintReadWrite, Category = "Media Input", meta = (EditCondition = "bEnabled"))
+	UPROPERTY(EditAnywhere, BlueprintReadWrite, Category = "Media Input")
 	TObjectPtr<UMediaSource> MediaSource = nullptr;
 
-	/** Media player to use */
-	UPROPERTY(EditAnywhere, BlueprintReadWrite, Category = "Media Input", meta = (EditCondition = "bEnabled"))
-	TObjectPtr<UMediaPlayer> MediaPlayer = nullptr;
-
-	/** Media texture to use */
-	UPROPERTY(EditAnywhere, BlueprintReadWrite, Category = "Media Input", meta = (EditCondition = "bEnabled"))
-	TObjectPtr<UMediaTexture> MediaTexture = nullptr;
-};
-
-
-/*
- * Media output settings
- */
-USTRUCT(Blueprintable)
-struct DISPLAYCLUSTERCONFIGURATION_API FDisplayClusterConfigurationMediaOutput
-{
-	GENERATED_BODY()
-
-public:
-	/** Enable/disable media capture */
-	UPROPERTY(EditAnywhere, BlueprintReadWrite, Category = "Media Output", meta = (DisplayName = "Enable"))
-	bool bEnabled = false;
-
 	/** Media output to use */
-	UPROPERTY(EditAnywhere, BlueprintReadWrite, Category = "Media Output", meta = (EditCondition = "bEnabled"))
+	UPROPERTY(EditAnywhere, BlueprintReadWrite, Category = "Media Output")
 	TObjectPtr<UMediaOutput> MediaOutput = nullptr;
-};
-
-
-/*
- * Media settings
- */
-USTRUCT(Blueprintable)
-struct DISPLAYCLUSTERCONFIGURATION_API FDisplayClusterConfigurationMedia
-{
-	GENERATED_BODY()
 
 public:
-	/** Enable/disable media features */
-	UPROPERTY(EditAnywhere, BlueprintReadWrite, Category = "Media", meta = (DisplayName = "Enable"))
-	bool bEnabled = false;
-
-	/** Media input settings */
-	UPROPERTY(EditAnywhere, BlueprintReadWrite, Category = "Media", meta = (EditCondition = "bEnabled"))
-	FDisplayClusterConfigurationMediaInput MediaInput;
-
-	/** Media capture settings */
-	UPROPERTY(EditAnywhere, BlueprintReadWrite, Category = "Media", meta = (EditCondition = "bEnabled"))
-	FDisplayClusterConfigurationMediaOutput MediaOutput;
+	/** Returns true if media sharing is used */
+	bool IsMediaSharingUsed() const
+	{
+		return bMediaSharing && MediaSource && MediaOutput && !MediaSharingNode.IsEmpty();
+	}
 };
 
+
 /*
- * ICVFX media settings
+ * Global media settings
  */
 USTRUCT(Blueprintable)
-struct DISPLAYCLUSTERCONFIGURATION_API FDisplayClusterConfigurationICVFXMedia
-	: public FDisplayClusterConfigurationMedia
+struct DISPLAYCLUSTERCONFIGURATION_API FDisplayClusterConfigurationGlobalMediaSettings
 {
 	GENERATED_BODY()
 
 public:
-	/** When in-cluster media sharing us used, the cluster node specified here will be used as a source (Tx node) */
-	UPROPERTY(EditAnywhere, BlueprintReadWrite, Category = "Media", meta = (EditCondition = "bEnabled"))
-	FString MediaOutputNode;
-
-public:
-	/** Returns true if media sharing is used */
-	bool IsMediaSharingRequired() const
-	{
-		return bEnabled && MediaInput.bEnabled && MediaOutput.bEnabled;
-	}
+	/** Media latency */
+	UPROPERTY(EditAnywhere, BlueprintReadWrite, Category = "Media", meta = (ClampMin = "0", ClampMax = "9", UIMin = "0", UIMax = "9"))
+	int32 Latency = 0;
 };
diff --git a/Engine/Plugins/Runtime/nDisplay/Source/DisplayClusterConfigurator/Private/DisplayClusterConfiguratorModule.cpp b/Engine/Plugins/Runtime/nDisplay/Source/DisplayClusterConfigurator/Private/DisplayClusterConfiguratorModule.cpp
--- a/Engine/Plugins/Runtime/nDisplay/Source/DisplayClusterConfigurator/Private/DisplayClusterConfiguratorModule.cpp
+++ b/Engine/Plugins/Runtime/nDisplay/Source/DisplayClusterConfigurator/Private/DisplayClusterConfiguratorModule.cpp
@@ -17,7 +17,7 @@
 #include "Views/Details/Cluster/DisplayClusterConfiguratorClusterDetailsCustomization.h"
 #include "Views/Details/Cluster/DisplayClusterConfiguratorExternalImageTypeCustomization.h"
 #include "Views/Details/Cluster/DisplayClusterConfiguratorGenerateMipsCustomization.h"
-#include "Views/Details/Cluster/DisplayClusterConfiguratorICVFXMediaCustomization.h"
+#include "Views/Details/Cluster/DisplayClusterConfiguratorMediaCustomization.h"
 #include "Views/Details/Cluster/DisplayClusterConfiguratorNodeSelectionCustomization.h"
 #include "Views/Details/Cluster/DisplayClusterConfiguratorClusterReferenceListCustomization.h"
 #include "Views/Details/Cluster/DisplayClusterConfiguratorViewportDetailsCustomization.h"
@@ -202,7 +202,7 @@ void FDisplayClusterConfiguratorModule::RegisterCustomLayouts()
 	REGISTER_PROPERTY_LAYOUT(FDisplayClusterConfigurationClusterItemReferenceList, FDisplayClusterConfiguratorClusterReferenceListCustomization);
 	REGISTER_PROPERTY_LAYOUT(FDisplayClusterConfigurationViewport_RemapData, FDisplayClusterConfiguratorViewportRemapCustomization);
 	REGISTER_PROPERTY_LAYOUT(FDisplayClusterConfigurationRectangle, FDisplayClusterConfiguratorRectangleCustomization);
-	REGISTER_PROPERTY_LAYOUT(FDisplayClusterConfigurationICVFXMedia, FDisplayClusterConfiguratorICVFXMediaCustomization);
+	REGISTER_PROPERTY_LAYOUT(FDisplayClusterConfigurationMedia, FDisplayClusterConfiguratorMediaCustomization);
 }
 
 void FDisplayClusterConfiguratorModule::UnregisterCustomLayouts()
diff --git a/Engine/Plugins/Runtime/nDisplay/Source/DisplayClusterConfigurator/Private/Views/Details/Cluster/DisplayClusterConfiguratorICVFXMediaCustomization.cpp b/Engine/Plugins/Runtime/nDisplay/Source/DisplayClusterConfigurator/Private/Views/Details/Cluster/DisplayClusterConfiguratorICVFXMediaCustomization.cpp
deleted file mode 100644
diff --git a/Engine/Plugins/Runtime/nDisplay/Source/DisplayClusterConfigurator/Private/Views/Details/Cluster/DisplayClusterConfiguratorICVFXMediaCustomization.h b/Engine/Plugins/Runtime/nDisplay/Source/DisplayClusterConfigurator/Private/Views/Details/Cluster/DisplayClusterConfiguratorICVFXMediaCustomization.h
deleted file mode 100644
diff --git a/Engine/Plugins/Runtime/nDisplay/Source/DisplayClusterConfigurator/Private/Views/Details/Cluster/DisplayClusterConfiguratorMediaCustomization.cpp b/Engine/Plugins/Runtime/nDisplay/Source/DisplayClusterConfigurator/Private/Views/Details/Cluster/DisplayClusterConfiguratorMediaCustomization.cpp
new file mode 100644
--- /dev/null
+++ b/Engine/Plugins/Runtime/nDisplay/Source/DisplayClusterConfigurator/Private/Views/Details/Cluster/DisplayClusterConfiguratorMediaCustomization.cpp
@@ -0,0 +1,118 @@
+// Copyright Epic Games, Inc. All Rights Reserved.
+
+#include "Views/Details/Cluster/DisplayClusterConfiguratorMediaCustomization.h"
+
+#include "DisplayClusterConfigurationTypes.h"
+#include "DisplayClusterConfiguratorBlueprintEditor.h"
+#include "DisplayClusterRootActor.h"
+
+#include "DetailWidgetRow.h"
+#include "IDetailChildrenBuilder.h"
+#include "IPropertyTypeCustomization.h"
+#include "PropertyHandle.h"
+
+
+void FDisplayClusterConfiguratorMediaCustomization::CustomizeChildren(TSharedRef<IPropertyHandle> InPropertyHandle, IDetailChildrenBuilder& InChildBuilder, IPropertyTypeCustomizationUtils& CustomizationUtils)
+{
+	OutputNodeProperty = InPropertyHandle->GetChildHandle(GET_MEMBER_NAME_CHECKED(FDisplayClusterConfigurationMedia, MediaSharingNode));
+	OutputNodeProperty->MarkHiddenByCustomization();
+
+	BuildClusterNodeOptionsList();
+
+	// Create GUI representation
+	{
+		InChildBuilder.AddCustomRow(OutputNodeProperty->GetPropertyDisplayName())
+			.NameContent()
+			[
+				OutputNodeProperty->CreatePropertyNameWidget()
+			]
+			.ValueContent()
+			[
+				SAssignNew(OutputNodeComboBox, SComboBox<TSharedPtr<FString>>)
+				.OptionsSource(&OutputNodeOptions)
+				.OnGenerateWidget(this, &FDisplayClusterConfiguratorMediaCustomization::CreateComboWidget)
+				.OnSelectionChanged(this, &FDisplayClusterConfiguratorMediaCustomization::OnNodeIdSelected)
+				.ContentPadding(2)
+				.Content()
+				[
+					SNew(STextBlock)
+					.Text(this, &FDisplayClusterConfiguratorMediaCustomization::GetSelectedNodeIdText)
+				]
+			];
+	}
+
+	// Set selected item
+	FString CurrentNodeId;
+	if (OutputNodeProperty->GetValue(CurrentNodeId))
+	{
+		TSharedPtr<FString>* FoundItem = OutputNodeOptions.FindByPredicate([CurrentNodeId](const TSharedPtr<FString>& Item)
+		{
+			return Item->Equals(CurrentNodeId, ESearchCase::IgnoreCase);
+		});
+
+		if (FoundItem)
+		{
+			OutputNodeComboBox->SetSelectedItem(*FoundItem);
+		}
+	}
+
+	FDisplayClusterConfiguratorBaseTypeCustomization::CustomizeChildren(InPropertyHandle, InChildBuilder, CustomizationUtils);
+}
+
+void FDisplayClusterConfiguratorMediaCustomization::BuildClusterNodeOptionsList()
+{
+	ADisplayClusterRootActor* DCRA = nullptr;
+
+	if (FDisplayClusterConfiguratorBlueprintEditor* Editor = FindBlueprintEditor())
+	{
+		DCRA = Editor->GetDefaultRootActor();
+	}
+
+	if (!DCRA)
+	{
+		DCRA = FindRootActor();
+	}
+	
+	if (!DCRA)
+	{
+		return;
+	}
+
+	// Get cluster node IDs
+	TArray<FString> NodeIds;
+	if (const UDisplayClusterConfigurationData* const ConfigurationData = DCRA->GetConfigData())
+	{
+		ConfigurationData->Cluster->GetNodeIds(NodeIds);
+	}
+
+	// Sort the node IDs in alphabet order
+	NodeIds.Sort();
+
+	// Prepare combobox options
+	OutputNodeOptions.Reset();
+	for (const FString& NodeId : NodeIds)
+	{
+		if (!NodeId.IsEmpty())
+		{
+			OutputNodeOptions.Emplace(MakeShared<FString>(NodeId));
+		}
+	}
+}
+
+
+void FDisplayClusterConfiguratorMediaCustomization::OnNodeIdSelected(TSharedPtr<FString> PreviewNodeId, ESelectInfo::Type SelectInfo)
+{
+	const FString NewValue = PreviewNodeId.IsValid() ? *PreviewNodeId : FString();
+	OutputNodeProperty->SetValue(NewValue);
+}
+
+FText FDisplayClusterConfiguratorMediaCustomization::GetSelectedNodeIdText() const
+{
+	TSharedPtr<FString> CurSelection = OutputNodeComboBox->GetSelectedItem();
+	return FText::FromString(CurSelection.IsValid() ? *CurSelection : *FString());
+}
+
+TSharedRef<SWidget> FDisplayClusterConfiguratorMediaCustomization::CreateComboWidget(TSharedPtr<FString> InItem)
+{
+	return SNew(STextBlock).Text(FText::FromString(*InItem));
+}
diff --git a/Engine/Plugins/Runtime/nDisplay/Source/DisplayClusterConfigurator/Private/Views/Details/Cluster/DisplayClusterConfiguratorMediaCustomization.h b/Engine/Plugins/Runtime/nDisplay/Source/DisplayClusterConfigurator/Private/Views/Details/Cluster/DisplayClusterConfiguratorMediaCustomization.h
new file mode 100644
--- /dev/null
+++ b/Engine/Plugins/Runtime/nDisplay/Source/DisplayClusterConfigurator/Private/Views/Details/Cluster/DisplayClusterConfiguratorMediaCustomization.h
@@ -0,0 +1,47 @@
+// Copyright Epic Games, Inc. All Rights Reserved.
+
+#pragma once
+
+#include "Views/Details/DisplayClusterConfiguratorBaseTypeCustomization.h"
+
+template<class T>
+class SComboBox;
+
+
+class FDisplayClusterConfiguratorMediaCustomization final
+	: public FDisplayClusterConfiguratorBaseTypeCustomization
+{
+public:
+	static TSharedRef<IPropertyTypeCustomization> MakeInstance()
+	{
+		return MakeShared<FDisplayClusterConfiguratorMediaCustomization>();
+	}
+
+protected:
+	//~ IPropertyTypeCustomization interface begin
+	virtual void CustomizeChildren(TSharedRef<IPropertyHandle> InPropertyHandle, IDetailChildrenBuilder& InChildBuilder, IPropertyTypeCustomizationUtils& CustomizationUtils) override;
+	//~ IPropertyTypeCustomization interface end
+
+private:
+	/**  Generates internal array of cluster node IDs */
+	void BuildClusterNodeOptionsList();
+
+	/**  Processes combobox selection changes */
+	void OnNodeIdSelected(TSharedPtr<FString> PreviewNodeId, ESelectInfo::Type SelectInfo);
+
+	/**  Returns text of currently selected item */
+	FText GetSelectedNodeIdText() const;
+
+	/** Creates combobox widget */
+	TSharedRef<SWidget> CreateComboWidget(TSharedPtr<FString> InItem);
+
+private:
+	/** OutputNode original property handler */
+	TSharedPtr<IPropertyHandle> OutputNodeProperty;
+
+	/** List of cluster node IDs for the combobox */
+	TArray<TSharedPtr<FString>> OutputNodeOptions;
+
+	/** Node ID selection combobox */
+	TSharedPtr<SComboBox<TSharedPtr<FString>>> OutputNodeComboBox;
+};
diff --git a/Engine/Plugins/Runtime/nDisplay/Source/DisplayClusterMedia/DisplayClusterMedia.Build.cs b/Engine/Plugins/Runtime/nDisplay/Source/DisplayClusterMedia/DisplayClusterMedia.Build.cs
--- a/Engine/Plugins/Runtime/nDisplay/Source/DisplayClusterMedia/DisplayClusterMedia.Build.cs
+++ b/Engine/Plugins/Runtime/nDisplay/Source/DisplayClusterMedia/DisplayClusterMedia.Build.cs
@@ -12,6 +12,11 @@ public class DisplayClusterMedia : ModuleRules
 		string EngineDir = Path.GetFullPath(Target.RelativeEnginePath);
 		PrivateIncludePaths.Add(Path.Combine(EngineDir, "Source", "Runtime", "Renderer", "Private"));
 
+		PublicIncludePathModuleNames.AddRange(
+			new string[] {
+				"DisplayClusterShaders",
+			});
+
 		PrivateDependencyModuleNames.AddRange(
 			new string[] {
 				"Core",
diff --git a/Engine/Plugins/Runtime/nDisplay/Source/DisplayClusterMedia/Private/Capture/DisplayClusterMediaCaptureNode.cpp b/Engine/Plugins/Runtime/nDisplay/Source/DisplayClusterMedia/Private/Capture/DisplayClusterMediaCaptureNode.cpp
--- a/Engine/Plugins/Runtime/nDisplay/Source/DisplayClusterMedia/Private/Capture/DisplayClusterMediaCaptureNode.cpp
+++ b/Engine/Plugins/Runtime/nDisplay/Source/DisplayClusterMedia/Private/Capture/DisplayClusterMediaCaptureNode.cpp
@@ -52,10 +52,28 @@ FIntPoint FDisplayClusterMediaCaptureNode::GetCaptureSize() const
 	return FIntPoint();
 }
 
-void FDisplayClusterMediaCaptureNode::OnPostFrameRender_RenderThread(FRHICommandListImmediate& RHICmdList, const IDisplayClusterViewportManagerProxy* ViewportManagerProxy, FViewport* Viewport)
+FTextureRHIRef FDisplayClusterMediaCaptureNode::CreateIntermediateTexture_RenderThread(EPixelFormat Format, ETextureCreateFlags Flags, const FIntPoint& Size)
 {
-	ensure(ViewportManagerProxy);
+	// Prepare description
+	FRHITextureCreateDesc Desc =
+		FRHITextureCreateDesc::Create2D(TEXT("DisplayClusterFrameQueueCacheTexture"), Size.X, Size.Y, Format)
+		.SetClearValue(FClearValueBinding::Black)
+		.SetNumMips(1)
+		.SetFlags(ETextureCreateFlags::Dynamic)
+		.AddFlags(ETextureCreateFlags::MultiGPUGraphIgnore)
+		.SetInitialState(ERHIAccess::SRVMask);
+
+	// Leave original flags, but make sure it's ResolveTargetable but not RenderTargetable
+	Flags &= ~ETextureCreateFlags::RenderTargetable;
+	Flags |= ETextureCreateFlags::ResolveTargetable;
+	Desc.SetFlags(Flags);
+
+	// Create texture
+	return RHICreateTexture(Desc);
+}
 
+void FDisplayClusterMediaCaptureNode::OnPostFrameRender_RenderThread(FRHICommandListImmediate& RHICmdList, const IDisplayClusterViewportManagerProxy* ViewportManagerProxy, FViewport* Viewport)
+{
 	TArray<FRHITexture*> Textures;
 	TArray<FIntPoint>    Regions;
 
@@ -63,11 +81,28 @@ void FDisplayClusterMediaCaptureNode::OnPostFrameRender_RenderThread(FRHICommand
 	{
 		if (Textures.Num() > 0 && Regions.Num() > 0 && Textures[0])
 		{
-			FMediaTextureInfo TextureInfo{ Textures[0], FIntRect(Regions[0], Regions[0] + Textures[0]->GetDesc().Extent) };
-
-			FRDGBuilder Builder(RHICmdList);
-			ExportMediaData(Builder, TextureInfo);
-			Builder.Execute();
+			// Create interim texture if not already available
+			if (!InterimTexture.IsValid() && Viewport)
+			{
+				if (FRHITexture* const BackbufferTexture = Viewport->GetRenderTargetTexture())
+				{
+					InterimTexture = CreateIntermediateTexture_RenderThread(Textures[0]->GetDesc().Format, Textures[0]->GetDesc().Flags, BackbufferTexture->GetDesc().Extent);
+				}
+			}
+
+			if (InterimTexture.IsValid())
+			{
+				// Copy texture to the intermediate buffer
+				FRHICopyTextureInfo CopyInfo;
+				CopyInfo.Size = Textures[0]->GetDesc().GetSize();
+				RHICmdList.CopyTexture(Textures[0], InterimTexture, CopyInfo);
+
+				// Capture intermediate texture
+				FMediaTextureInfo TextureInfo{ InterimTexture, FIntRect(FIntPoint::ZeroValue, InterimTexture->GetDesc().Extent) };
+				FRDGBuilder Builder(RHICmdList);
+				ExportMediaData(Builder, TextureInfo);
+				Builder.Execute();
+			}
 		}
 	}
 }
diff --git a/Engine/Plugins/Runtime/nDisplay/Source/DisplayClusterMedia/Private/Capture/DisplayClusterMediaCaptureNode.h b/Engine/Plugins/Runtime/nDisplay/Source/DisplayClusterMedia/Private/Capture/DisplayClusterMediaCaptureNode.h
--- a/Engine/Plugins/Runtime/nDisplay/Source/DisplayClusterMedia/Private/Capture/DisplayClusterMediaCaptureNode.h
+++ b/Engine/Plugins/Runtime/nDisplay/Source/DisplayClusterMedia/Private/Capture/DisplayClusterMediaCaptureNode.h
@@ -4,6 +4,8 @@
 
 #include "Capture/DisplayClusterMediaCaptureBase.h"
 
+#include "RHIResources.h"
+
 class FRHICommandListImmediate;
 class FViewport;
 class IDisplayClusterViewportManagerProxy;
@@ -27,4 +29,12 @@ protected:
 
 private:
 	void OnPostFrameRender_RenderThread(FRHICommandListImmediate& RHICmdList, const IDisplayClusterViewportManagerProxy* ViewportManagerProxy, FViewport* Viewport);
+
+	// Auxiliary method to create an intermediate texture
+	FTextureRHIRef CreateIntermediateTexture_RenderThread(EPixelFormat Format, ETextureCreateFlags Flags, const FIntPoint& Size);
+
+private:
+	// We need an intermediate texture with the same size as backbuffer has. This allows to avoid
+	// any texture size related problems caused by optimizations in the nD rendering pipeline.
+	FTextureRHIRef InterimTexture = nullptr;
 };
diff --git a/Engine/Plugins/Runtime/nDisplay/Source/DisplayClusterMedia/Private/DisplayClusterMediaModule.cpp b/Engine/Plugins/Runtime/nDisplay/Source/DisplayClusterMedia/Private/DisplayClusterMediaModule.cpp
--- a/Engine/Plugins/Runtime/nDisplay/Source/DisplayClusterMedia/Private/DisplayClusterMediaModule.cpp
+++ b/Engine/Plugins/Runtime/nDisplay/Source/DisplayClusterMedia/Private/DisplayClusterMediaModule.cpp
@@ -66,6 +66,9 @@ void FDisplayClusterMediaModule::InitializeMedia()
 		return;
 	}
 
+	// Instantiate latency queue
+	FrameQueue.Init();
+
 	// Parse DCRA configuration and initialize media
 	if (const ADisplayClusterRootActor* const RootActor = IDisplayCluster::Get().GetGameMgr()->GetRootActor())
 	{
@@ -74,39 +77,72 @@ void FDisplayClusterMediaModule::InitializeMedia()
 		{
 			///////////////////////////////
 			// Node backbuffer media setup
-			if (ClusterNode->Media.bEnabled)
 			{
 				const FDisplayClusterConfigurationMedia& MediaSettings = ClusterNode->Media;
 
-				const bool bInputDataValid   = MediaSettings.MediaInput.MediaPlayer && MediaSettings.MediaInput.MediaSource && MediaSettings.MediaInput.MediaTexture;
-				const bool bOutputDataValid  = !!MediaSettings.MediaOutput.MediaOutput;
+				const bool bShared             = MediaSettings.IsMediaSharingUsed();
+				const bool bClusterNodeMatches = MediaSettings.MediaSharingNode.Equals(ClusterNodeId, ESearchCase::IgnoreCase);
 
-				// Media input
-				if (MediaSettings.MediaInput.bEnabled && bInputDataValid)
+				// Shared backbuffer
+				if (bShared && MediaSettings.MediaSource && MediaSettings.MediaOutput)
 				{
-					const FString MediaInputId = RootActor->GetName() + FString("_") + ClusterNodeId + FString("_backbuffer_input");
+					// Tx node, initialize media capture
+					if (bClusterNodeMatches)
+					{
+						const FString MediaCaptureId = RootActor->GetName() + FString("_") + ClusterNodeId + FString("_backbuffer_capture");
 
-					UE_LOG(LogDisplayClusterMedia, Log, TEXT("Initializing backbuffer media input for node '%s'..."), *ClusterNodeId);
+						UE_LOG(LogDisplayClusterMedia, Log, TEXT("Initializing in-cluster (shared) backbuffer media capture for node '%s'..."), *ClusterNodeId);
 
-					InputNode = MakeUnique<FDisplayClusterMediaInputNode>(
-						MediaInputId, ClusterNodeId,
-						MediaSettings.MediaInput.MediaSource, MediaSettings.MediaInput.MediaPlayer, MediaSettings.MediaInput.MediaTexture);
+						CaptureNode = MakeUnique<FDisplayClusterMediaCaptureNode>(
+							MediaCaptureId, ClusterNodeId,
+							MediaSettings.MediaOutput);
 
-					AllInputs.Add(InputNode.Get());
-				}
+						AllCaptures.Add(CaptureNode.Get());
+					}
+					// Rx node, initialize media playback
+					else
+					{
+						const FString MediaInputId = RootActor->GetName() + FString("_") + ClusterNodeId + FString("_backbuffer_input");
+
+						UE_LOG(LogDisplayClusterMedia, Log, TEXT("Initializing in-cluster (shared) backbuffer media input for node '%s'..."), *ClusterNodeId);
 
-				// Media capture
-				if (MediaSettings.MediaOutput.bEnabled && bOutputDataValid)
+						InputNode = MakeUnique<FDisplayClusterMediaInputNode>(
+							MediaInputId, ClusterNodeId, 
+							MediaSettings.MediaSource);
+
+						AllInputs.Add(InputNode.Get());
+					}
+				}
+				// No in-cluster sharing
+				else
 				{
-					const FString MediaCaptureId = RootActor->GetName() + FString("_") + ClusterNodeId + FString("_backbuffer_capture");
+					// Media input
+					if (MediaSettings.MediaSource)
+					{
+						const FString MediaInputId = RootActor->GetName() + FString("_") + ClusterNodeId + FString("_backbuffer_input");
+
+						UE_LOG(LogDisplayClusterMedia, Log, TEXT("Initializing backbuffer media input for node '%s'..."), *ClusterNodeId);
+
+						InputNode = MakeUnique<FDisplayClusterMediaInputNode>(
+							MediaInputId, ClusterNodeId, 
+							MediaSettings.MediaSource);
+
+						AllInputs.Add(InputNode.Get());
+					}
+
+					// Media capture
+					if (MediaSettings.MediaOutput)
+					{
+						const FString MediaCaptureId = RootActor->GetName() + FString("_") + ClusterNodeId + FString("_backbuffer_capture");
 
-					UE_LOG(LogDisplayClusterMedia, Log, TEXT("Initializing backbuffer media capture for node '%s'..."), *ClusterNodeId);
+						UE_LOG(LogDisplayClusterMedia, Log, TEXT("Initializing backbuffer media capture for node '%s'..."), *ClusterNodeId);
 
-					CaptureNode = MakeUnique<FDisplayClusterMediaCaptureNode>(
-						MediaCaptureId, ClusterNodeId,
-						MediaSettings.MediaOutput.MediaOutput);
+						CaptureNode = MakeUnique<FDisplayClusterMediaCaptureNode>(
+							MediaCaptureId, ClusterNodeId, 
+							MediaSettings.MediaOutput);
 
-					AllCaptures.Add(CaptureNode.Get());
+						AllCaptures.Add(CaptureNode.Get());
+					}
 				}
 			}
 
@@ -117,29 +153,65 @@ void FDisplayClusterMediaModule::InitializeMedia()
 				if (const UDisplayClusterConfigurationViewport* const Viewport = ViewportIt.Value)
 				{
 					const FDisplayClusterConfigurationMedia& MediaSettings = Viewport->RenderSettings.Media;
-					if (MediaSettings.bEnabled)
+
+					const bool bShared             = MediaSettings.IsMediaSharingUsed();
+					const bool bClusterNodeMatches = MediaSettings.MediaSharingNode.Equals(ClusterNodeId, ESearchCase::IgnoreCase);
+
+					// Shared viewport
+					if (bShared && MediaSettings.MediaSource && MediaSettings.MediaOutput)
 					{
-						const bool bInputDataValid   = MediaSettings.MediaInput.MediaPlayer && MediaSettings.MediaInput.MediaSource && MediaSettings.MediaInput.MediaTexture;
-						const bool bOutputDataValid  = !!MediaSettings.MediaOutput.MediaOutput;
+						// Tx node, initialize media capture
+						if (bClusterNodeMatches)
+						{
+							const FString MediaCaptureId = RootActor->GetName() + FString("_") + ClusterNodeId + FString("_") + ViewportIt.Key + FString("_viewport_capture");
+
+							UE_LOG(LogDisplayClusterMedia, Log, TEXT("Initializing in-cluster (shared) viewport capture for viewport '%s'..."), *ViewportIt.Key);
+
+							TUniquePtr<FDisplayClusterMediaCaptureViewport> NewViewportCapture = MakeUnique<FDisplayClusterMediaCaptureViewport>(
+								MediaCaptureId, ClusterNodeId, 
+								ViewportIt.Key, 
+								MediaSettings.MediaOutput);
+
+							AllCaptures.Add(NewViewportCapture.Get());
+							CaptureViewports.Emplace(MediaCaptureId, MoveTemp(NewViewportCapture));
+						}
+						// Rx node, initialize media playback
+						else
+						{
+							const FString MediaInputId = RootActor->GetName() + FString("_") + ClusterNodeId + FString("_") + ViewportIt.Key + FString("_viewport_input");
+
+							UE_LOG(LogDisplayClusterMedia, Log, TEXT("Initializing in-cluster (shared) viewport media input for viewport '%s'..."), *ViewportIt.Key);
+
+							TUniquePtr<FDisplayClusterMediaInputViewport> NewViewportInput = MakeUnique<FDisplayClusterMediaInputViewport>(
+								MediaInputId, ClusterNodeId,
+								ViewportIt.Key,
+								MediaSettings.MediaSource);
 
+							AllInputs.Add(NewViewportInput.Get());
+							InputViewports.Emplace(MediaInputId, MoveTemp(NewViewportInput));
+						}
+					}
+					// No in-cluster sharing
+					else
+					{
 						// Media input
-						if (MediaSettings.MediaInput.bEnabled && bInputDataValid)
+						if (MediaSettings.MediaSource)
 						{
 							const FString MediaInputId = RootActor->GetName() + FString("_") + ClusterNodeId + FString("_") + ViewportIt.Key + FString("_viewport_input");
 
-							UE_LOG(LogDisplayClusterMedia, Log, TEXT("Initializing viewport media input for viewport '%s'..."), *ClusterNodeId);
+							UE_LOG(LogDisplayClusterMedia, Log, TEXT("Initializing viewport media input for viewport '%s'..."), *ViewportIt.Key);
 
 							TUniquePtr<FDisplayClusterMediaInputViewport> NewViewportInput = MakeUnique<FDisplayClusterMediaInputViewport>(
 								MediaInputId, ClusterNodeId,
 								ViewportIt.Key,
-								MediaSettings.MediaInput.MediaSource, MediaSettings.MediaInput.MediaPlayer, MediaSettings.MediaInput.MediaTexture);
+								MediaSettings.MediaSource);
 
 							AllInputs.Add(NewViewportInput.Get());
 							InputViewports.Emplace(MediaInputId, MoveTemp(NewViewportInput));
 						}
 
 						// Media capture
-						if (MediaSettings.MediaOutput.bEnabled && bOutputDataValid)
+						if (MediaSettings.MediaOutput)
 						{
 							const FString MediaCaptureId = RootActor->GetName() + FString("_") + ClusterNodeId + FString("_") + ViewportIt.Key + FString("_viewport_capture");
 
@@ -148,7 +220,7 @@ void FDisplayClusterMediaModule::InitializeMedia()
 							TUniquePtr<FDisplayClusterMediaCaptureViewport> NewViewportCapture = MakeUnique<FDisplayClusterMediaCaptureViewport>(
 								MediaCaptureId, ClusterNodeId,
 								ViewportIt.Key,
-								MediaSettings.MediaOutput.MediaOutput);
+								MediaSettings.MediaOutput);
 
 							AllCaptures.Add(NewViewportCapture.Get());
 							CaptureViewports.Emplace(MediaCaptureId, MoveTemp(NewViewportCapture));
@@ -167,74 +239,70 @@ void FDisplayClusterMediaModule::InitializeMedia()
 				{
 					if (const UDisplayClusterICVFXCameraComponent* const ICVFXCamera = Cast<UDisplayClusterICVFXCameraComponent>(Component))
 					{
-						const FDisplayClusterConfigurationICVFXMedia& MediaSettings = ICVFXCamera->CameraSettings.RenderSettings.Media;
+						const FDisplayClusterConfigurationMedia& MediaSettings = ICVFXCamera->CameraSettings.RenderSettings.Media;
 
-						if (MediaSettings.bEnabled)
+						const FString ICVFXCameraName  = ICVFXCamera->GetName();
+						const FString ICVFXViewportId  = DisplayClusterMediaHelpers::GenerateICVFXViewportName(ClusterNodeId, ICVFXCameraName);
+						const FString MediaCaptureId   = RootActor->GetName() + FString("_") + ICVFXCameraName + FString("_icvfx_capture");
+						const FString MediaInputId     = RootActor->GetName() + FString("_") + ICVFXCameraName + FString("_icvfx_input");
+						const bool bShared             = MediaSettings.IsMediaSharingUsed();
+						const bool bClusterNodeMatches = MediaSettings.MediaSharingNode.Equals(ClusterNodeId, ESearchCase::IgnoreCase);
+
+						// Shared camera
+						if (bShared && MediaSettings.MediaSource && MediaSettings.MediaOutput)
 						{
-							const FString ICVFXCameraName  = ICVFXCamera->GetName();
-							const FString ICVFXViewportId  = DisplayClusterMediaHelpers::GenerateICVFXViewportName(ClusterNodeId, ICVFXCameraName);
-							const FString MediaCaptureId   = RootActor->GetName() + FString("_") + ICVFXCameraName + FString("_icvfx_capture");
-							const FString MediaInputId     = RootActor->GetName() + FString("_") + ICVFXCameraName + FString("_icvfx_input");
-							const bool bInputDataValid     = MediaSettings.MediaInput.MediaPlayer && MediaSettings.MediaInput.MediaSource && MediaSettings.MediaInput.MediaTexture;
-							const bool bOutputDataValid    = !!MediaSettings.MediaOutput.MediaOutput;
-							const bool bClusterNodeMatches = MediaSettings.MediaOutputNode.Equals(ClusterNodeId, ESearchCase::IgnoreCase);
-
-							// Media sharing
-							if (MediaSettings.IsMediaSharingRequired() && bInputDataValid && bOutputDataValid)
+							// Tx node, initialize media capture
+							if (bClusterNodeMatches)
 							{
-								// Share media
-								if (bClusterNodeMatches)
-								{
-									UE_LOG(LogDisplayClusterMedia, Log, TEXT("Initializing ICVFX in-cluster capture for camera '%s' (in-cluster shared)..."), *ICVFXCameraName);
-
-									TUniquePtr<FDisplayClusterMediaCaptureViewport> NewICVFXCapture = MakeUnique<FDisplayClusterMediaCaptureCamera>(
-										MediaCaptureId, ClusterNodeId,
-										ICVFXCameraName, ICVFXViewportId,
-										MediaSettings.MediaOutput.MediaOutput);
-
-									AllCaptures.Add(NewICVFXCapture.Get());
-									CaptureViewports.Emplace(MediaCaptureId, MoveTemp(NewICVFXCapture));
-								}
-								// Consume media
-								else
-								{
-									UE_LOG(LogDisplayClusterMedia, Log, TEXT("Initializing ICVFX in-cluster media input for camera '%s' (in-cluster shared)..."), *ICVFXCameraName);
-
-									TUniquePtr<FDisplayClusterMediaInputViewport> NewICVFXInput = MakeUnique<FDisplayClusterMediaInputViewport>(
-										MediaInputId, ClusterNodeId,
-										ICVFXViewportId,
-										MediaSettings.MediaInput.MediaSource, MediaSettings.MediaInput.MediaPlayer, MediaSettings.MediaInput.MediaTexture);
-
-									AllInputs.Add(NewICVFXInput.Get());
-									InputViewports.Emplace(MediaInputId, MoveTemp(NewICVFXInput));
-								}
+								UE_LOG(LogDisplayClusterMedia, Log, TEXT("Initializing in-cluster (shared) ICVFX capture for camera '%s'..."), *ICVFXCameraName);
+
+								TUniquePtr<FDisplayClusterMediaCaptureViewport> NewICVFXCapture = MakeUnique<FDisplayClusterMediaCaptureCamera>(
+									MediaCaptureId, ClusterNodeId,
+									ICVFXCameraName, ICVFXViewportId,
+									MediaSettings.MediaOutput);
+
+								AllCaptures.Add(NewICVFXCapture.Get());
+								CaptureViewports.Emplace(MediaCaptureId, MoveTemp(NewICVFXCapture));
 							}
-							// Media input only
-							else if (MediaSettings.MediaInput.bEnabled && bInputDataValid)
+							// Rx node, initialize media playback
+							else
 							{
-								UE_LOG(LogDisplayClusterMedia, Log, TEXT("Initializing ICVFX media input for camera '%s'..."), *ICVFXCameraName);
+								UE_LOG(LogDisplayClusterMedia, Log, TEXT("Initializing in-cluster (shared) ICVFX media input for camera '%s'..."), *ICVFXCameraName);
 
 								TUniquePtr<FDisplayClusterMediaInputViewport> NewICVFXInput = MakeUnique<FDisplayClusterMediaInputViewport>(
 									MediaInputId, ClusterNodeId,
 									ICVFXViewportId,
-									MediaSettings.MediaInput.MediaSource, MediaSettings.MediaInput.MediaPlayer, MediaSettings.MediaInput.MediaTexture);
+									MediaSettings.MediaSource);
 
 								AllInputs.Add(NewICVFXInput.Get());
 								InputViewports.Emplace(MediaInputId, MoveTemp(NewICVFXInput));
 							}
-							// Media capture only
-							else if (MediaSettings.MediaOutput.bEnabled && bOutputDataValid)
-							{
-								UE_LOG(LogDisplayClusterMedia, Log, TEXT("Initializing ICVFX capture for camera '%s'..."), *ICVFXCameraName);
+						}
+						// Media input only
+						else if (MediaSettings.MediaSource)
+						{
+							UE_LOG(LogDisplayClusterMedia, Log, TEXT("Initializing ICVFX media input for camera '%s'..."), *ICVFXCameraName);
 
-								TUniquePtr<FDisplayClusterMediaCaptureViewport> NewICVFXCapture = MakeUnique<FDisplayClusterMediaCaptureCamera>(
-									MediaCaptureId, ClusterNodeId,
-									ICVFXCameraName, ICVFXViewportId,
-									MediaSettings.MediaOutput.MediaOutput);
+							TUniquePtr<FDisplayClusterMediaInputViewport> NewICVFXInput = MakeUnique<FDisplayClusterMediaInputViewport>(
+								MediaInputId, ClusterNodeId,
+								ICVFXViewportId,
+								MediaSettings.MediaSource);
 
-								AllCaptures.Add(NewICVFXCapture.Get());
-								CaptureViewports.Emplace(MediaCaptureId, MoveTemp(NewICVFXCapture));
-							}
+							AllInputs.Add(NewICVFXInput.Get());
+							InputViewports.Emplace(MediaInputId, MoveTemp(NewICVFXInput));
+						}
+						// Media capture only
+						else if (MediaSettings.MediaOutput)
+						{
+							UE_LOG(LogDisplayClusterMedia, Log, TEXT("Initializing ICVFX capture for camera '%s'..."), *ICVFXCameraName);
+
+							TUniquePtr<FDisplayClusterMediaCaptureViewport> NewICVFXCapture = MakeUnique<FDisplayClusterMediaCaptureCamera>(
+								MediaCaptureId, ClusterNodeId,
+								ICVFXCameraName, ICVFXViewportId,
+								MediaSettings.MediaOutput);
+
+							AllCaptures.Add(NewICVFXCapture.Get());
+							CaptureViewports.Emplace(MediaCaptureId, MoveTemp(NewICVFXCapture));
 						}
 					}
 				}
@@ -255,6 +323,8 @@ void FDisplayClusterMediaModule::ReleaseMedia()
 	InputViewports.Reset();
 	InputNode.Reset();
 	AllInputs.Reset();
+
+	FrameQueue.Release();
 }
 
 void FDisplayClusterMediaModule::StartCapture()
diff --git a/Engine/Plugins/Runtime/nDisplay/Source/DisplayClusterMedia/Private/DisplayClusterMediaModule.h b/Engine/Plugins/Runtime/nDisplay/Source/DisplayClusterMedia/Private/DisplayClusterMediaModule.h
--- a/Engine/Plugins/Runtime/nDisplay/Source/DisplayClusterMedia/Private/DisplayClusterMediaModule.h
+++ b/Engine/Plugins/Runtime/nDisplay/Source/DisplayClusterMedia/Private/DisplayClusterMediaModule.h
@@ -5,6 +5,8 @@
 #include "CoreMinimal.h"
 #include "Modules/ModuleInterface.h"
 
+#include "Synchronization/DisplayClusterFrameQueue.h"
+
 class FDisplayClusterMediaCaptureBase;
 class FDisplayClusterMediaCaptureICVFX;
 class FDisplayClusterMediaCaptureNode;
@@ -37,7 +39,7 @@ protected:
 	/** Start media capture (all sources) */
 	void StartCapture();
 
-	/** Stop media capture (all soureces) */
+	/** Stop media capture (all sources) */
 	void StopCapture();
 
 	/** Start playing media (all inputs) */
@@ -58,4 +60,7 @@ private:
 	TMap<FString, TUniquePtr<FDisplayClusterMediaInputViewport>>   InputViewports;
 	TUniquePtr<FDisplayClusterMediaInputNode>                      InputNode;
 	TArray<FDisplayClusterMediaInputBase*>                         AllInputs;
+
+	// Latency queue
+	FDisplayClusterFrameQueue FrameQueue;
 };
diff --git a/Engine/Plugins/Runtime/nDisplay/Source/DisplayClusterMedia/Private/Input/DisplayClusterMediaInputBase.cpp b/Engine/Plugins/Runtime/nDisplay/Source/DisplayClusterMedia/Private/Input/DisplayClusterMediaInputBase.cpp
--- a/Engine/Plugins/Runtime/nDisplay/Source/DisplayClusterMedia/Private/Input/DisplayClusterMediaInputBase.cpp
+++ b/Engine/Plugins/Runtime/nDisplay/Source/DisplayClusterMedia/Private/Input/DisplayClusterMediaInputBase.cpp
@@ -10,14 +10,28 @@
 #include "MediaTexture.h"
 
 
-FDisplayClusterMediaInputBase::FDisplayClusterMediaInputBase(const FString& InMediaId, const FString& InClusterNodeId, UMediaSource* InMediaSource, UMediaPlayer* InMediaPlayer, UMediaTexture* InMediaTexture)
+FDisplayClusterMediaInputBase::FDisplayClusterMediaInputBase(const FString& InMediaId, const FString& InClusterNodeId, UMediaSource* InMediaSource)
 	: FDisplayClusterMediaBase(InMediaId, InClusterNodeId)
 	, MediaSource(InMediaSource)
-	, MediaPlayer(InMediaPlayer)
-	, MediaTexture(InMediaTexture)
 {
-	// We expect these to always be valid
-	check(MediaSource && MediaPlayer && MediaTexture);
+	checkSlow(MediaSource);
+
+	// Instantiate media player
+	MediaPlayer = NewObject<UMediaPlayer>();
+	if (MediaPlayer)
+	{
+		MediaPlayer->SetLooping(false);
+		MediaPlayer->PlayOnOpen = false;
+
+		// Instantiate media texture
+		MediaTexture = NewObject<UMediaTexture>();
+		if (MediaTexture)
+		{
+			MediaTexture->NewStyleOutput = true;
+			MediaTexture->SetMediaPlayer(MediaPlayer);
+			MediaTexture->UpdateResource();
+		}
+	}
 }
 
 
@@ -74,8 +88,6 @@ void FDisplayClusterMediaInputBase::ImportMediaData(FRHICommandListImmediate& RH
 		}
 		else
 		{
-			UE_LOG(LogDisplayClusterMedia, Log, TEXT(""));
-
 			DisplayClusterMediaHelpers::ResampleTexture_RenderThread(
 				RHICmdList, SrcTexture, DstTexture,
 				FIntRect(FIntPoint(0, 0), SrcTexture->GetDesc().Extent),
diff --git a/Engine/Plugins/Runtime/nDisplay/Source/DisplayClusterMedia/Private/Input/DisplayClusterMediaInputBase.h b/Engine/Plugins/Runtime/nDisplay/Source/DisplayClusterMedia/Private/Input/DisplayClusterMediaInputBase.h
--- a/Engine/Plugins/Runtime/nDisplay/Source/DisplayClusterMedia/Private/Input/DisplayClusterMediaInputBase.h
+++ b/Engine/Plugins/Runtime/nDisplay/Source/DisplayClusterMedia/Private/Input/DisplayClusterMediaInputBase.h
@@ -22,7 +22,7 @@ class FDisplayClusterMediaInputBase
 	, public FGCObject
 {
 public:
-	FDisplayClusterMediaInputBase(const FString& MediaId, const FString& ClusterNodeId, UMediaSource* MediaSource, UMediaPlayer* MediaPlayer, UMediaTexture* MediaTexture);
+	FDisplayClusterMediaInputBase(const FString& MediaId, const FString& ClusterNodeId, UMediaSource* MediaSource);
 
 public:
 	//~ Begin FGCObject interface
diff --git a/Engine/Plugins/Runtime/nDisplay/Source/DisplayClusterMedia/Private/Input/DisplayClusterMediaInputNode.cpp b/Engine/Plugins/Runtime/nDisplay/Source/DisplayClusterMedia/Private/Input/DisplayClusterMediaInputNode.cpp
--- a/Engine/Plugins/Runtime/nDisplay/Source/DisplayClusterMedia/Private/Input/DisplayClusterMediaInputNode.cpp
+++ b/Engine/Plugins/Runtime/nDisplay/Source/DisplayClusterMedia/Private/Input/DisplayClusterMediaInputNode.cpp
@@ -12,8 +12,8 @@
 #include "RHIResources.h"
 
 
-FDisplayClusterMediaInputNode::FDisplayClusterMediaInputNode(const FString& InMediaId, const FString& InClusterNodeId, UMediaSource* InMediaSource, UMediaPlayer* InMediaPlayer, UMediaTexture* InMediaTexture)
-	: FDisplayClusterMediaInputBase(InMediaId, InClusterNodeId, InMediaSource, InMediaPlayer, InMediaTexture)
+FDisplayClusterMediaInputNode::FDisplayClusterMediaInputNode(const FString& InMediaId, const FString& InClusterNodeId, UMediaSource* InMediaSource)
+	: FDisplayClusterMediaInputBase(InMediaId, InClusterNodeId, InMediaSource)
 {
 }
 
diff --git a/Engine/Plugins/Runtime/nDisplay/Source/DisplayClusterMedia/Private/Input/DisplayClusterMediaInputNode.h b/Engine/Plugins/Runtime/nDisplay/Source/DisplayClusterMedia/Private/Input/DisplayClusterMediaInputNode.h
--- a/Engine/Plugins/Runtime/nDisplay/Source/DisplayClusterMedia/Private/Input/DisplayClusterMediaInputNode.h
+++ b/Engine/Plugins/Runtime/nDisplay/Source/DisplayClusterMedia/Private/Input/DisplayClusterMediaInputNode.h
@@ -16,7 +16,7 @@ class FDisplayClusterMediaInputNode
 	: public FDisplayClusterMediaInputBase
 {
 public:
-	FDisplayClusterMediaInputNode(const FString& MediaId, const FString& ClusterNodeId, UMediaSource* MediaSource, UMediaPlayer* MediaPlayer, UMediaTexture* MediaTexture);
+	FDisplayClusterMediaInputNode(const FString& MediaId, const FString& ClusterNodeId, UMediaSource* MediaSource);
 
 public:
 	virtual bool Play() override;
diff --git a/Engine/Plugins/Runtime/nDisplay/Source/DisplayClusterMedia/Private/Input/DisplayClusterMediaInputViewport.cpp b/Engine/Plugins/Runtime/nDisplay/Source/DisplayClusterMedia/Private/Input/DisplayClusterMediaInputViewport.cpp
--- a/Engine/Plugins/Runtime/nDisplay/Source/DisplayClusterMedia/Private/Input/DisplayClusterMediaInputViewport.cpp
+++ b/Engine/Plugins/Runtime/nDisplay/Source/DisplayClusterMedia/Private/Input/DisplayClusterMediaInputViewport.cpp
@@ -6,12 +6,14 @@
 #include "IDisplayCluster.h"
 #include "IDisplayClusterCallbacks.h"
 
+#include "MediaTexture.h"
+
 #include "RHICommandList.h"
 #include "RHIResources.h"
 
 
-FDisplayClusterMediaInputViewport::FDisplayClusterMediaInputViewport(const FString& InMediaId, const FString& InClusterNodeId, const FString& InViewportId, UMediaSource* InMediaSource, UMediaPlayer* InMediaPlayer, UMediaTexture* InMediaTexture)
-	: FDisplayClusterMediaInputBase(InMediaId, InClusterNodeId, InMediaSource, InMediaPlayer, InMediaTexture)
+FDisplayClusterMediaInputViewport::FDisplayClusterMediaInputViewport(const FString& InMediaId, const FString& InClusterNodeId, const FString& InViewportId, UMediaSource* InMediaSource)
+	: FDisplayClusterMediaInputBase(InMediaId, InClusterNodeId, InMediaSource)
 	, ViewportId(InViewportId)
 {
 }
@@ -22,7 +24,7 @@ bool FDisplayClusterMediaInputViewport::Play()
 	// If playback has started successfully, subscribe for rendering callbacks
 	if (FDisplayClusterMediaInputBase::Play())
 	{
-		IDisplayCluster::Get().GetCallbacks().OnDisplayClusterPreFrameRender_RenderThread().AddRaw(this, &FDisplayClusterMediaInputViewport::PreFrameRender_RenderThread);
+		IDisplayCluster::Get().GetCallbacks().OnDisplayClusterPostCrossGpuTransfer_RenderThread().AddRaw(this, &FDisplayClusterMediaInputViewport::PostCrossGpuTransfer_RenderThread);
 		return true;
 	}
 
@@ -32,12 +34,12 @@ bool FDisplayClusterMediaInputViewport::Play()
 void FDisplayClusterMediaInputViewport::Stop()
 {
 	// Stop receiving notifications
-	IDisplayCluster::Get().GetCallbacks().OnDisplayClusterPreFrameRender_RenderThread().RemoveAll(this);
+	IDisplayCluster::Get().GetCallbacks().OnDisplayClusterPostCrossGpuTransfer_RenderThread().RemoveAll(this);
 	// Stop playing
 	FDisplayClusterMediaInputBase::Stop();
 }
 
-void FDisplayClusterMediaInputViewport::PreFrameRender_RenderThread(FRHICommandListImmediate& RHICmdList, const IDisplayClusterViewportManagerProxy* ViewportManagerProxy, FViewport* Viewport)
+void FDisplayClusterMediaInputViewport::PostCrossGpuTransfer_RenderThread(FRHICommandListImmediate& RHICmdList, const IDisplayClusterViewportManagerProxy* ViewportManagerProxy, FViewport* Viewport)
 {
 	checkSlow(ViewportManagerProxy);
 
diff --git a/Engine/Plugins/Runtime/nDisplay/Source/DisplayClusterMedia/Private/Input/DisplayClusterMediaInputViewport.h b/Engine/Plugins/Runtime/nDisplay/Source/DisplayClusterMedia/Private/Input/DisplayClusterMediaInputViewport.h
--- a/Engine/Plugins/Runtime/nDisplay/Source/DisplayClusterMedia/Private/Input/DisplayClusterMediaInputViewport.h
+++ b/Engine/Plugins/Runtime/nDisplay/Source/DisplayClusterMedia/Private/Input/DisplayClusterMediaInputViewport.h
@@ -16,7 +16,7 @@ class FDisplayClusterMediaInputViewport
 	: public FDisplayClusterMediaInputBase
 {
 public:
-	FDisplayClusterMediaInputViewport(const FString& MediaId, const FString& ClusterNodeId, const FString& ViewportId, UMediaSource* MediaSource, UMediaPlayer* MediaPlayer, UMediaTexture* MediaTexture);
+	FDisplayClusterMediaInputViewport(const FString& MediaId, const FString& ClusterNodeId, const FString& ViewportId, UMediaSource* MediaSource);
 
 public:
 	virtual bool Play() override;
@@ -28,7 +28,7 @@ public:
 	}
 
 private:
-	void PreFrameRender_RenderThread(FRHICommandListImmediate& RHICmdList, const IDisplayClusterViewportManagerProxy* ViewportManagerProxy, FViewport* Viewport);
+	void PostCrossGpuTransfer_RenderThread(FRHICommandListImmediate& RHICmdList, const IDisplayClusterViewportManagerProxy* ViewportManagerProxy, FViewport* Viewport);
 
 private:
 	const FString ViewportId;
diff --git a/Engine/Plugins/Runtime/nDisplay/Source/DisplayClusterMedia/Private/Synchronization/DisplayClusterFrameQueue.cpp b/Engine/Plugins/Runtime/nDisplay/Source/DisplayClusterMedia/Private/Synchronization/DisplayClusterFrameQueue.cpp
new file mode 100644
--- /dev/null
+++ b/Engine/Plugins/Runtime/nDisplay/Source/DisplayClusterMedia/Private/Synchronization/DisplayClusterFrameQueue.cpp
@@ -0,0 +1,304 @@
+// Copyright Epic Games, Inc. All Rights Reserved.
+
+#include "Synchronization/DisplayClusterFrameQueue.h"
+
+#include "DisplayClusterConfigurationTypes.h"
+#include "DisplayClusterMediaLog.h"
+#include "DisplayClusterRootActor.h"
+
+#include "IDisplayCluster.h"
+#include "IDisplayClusterCallbacks.h"
+#include "Game/IDisplayClusterGameManager.h"
+#include "Render/IDisplayClusterRenderManager.h"
+#include "Render/Viewport/IDisplayClusterViewportManager.h"
+#include "Render/Viewport/IDisplayClusterViewportManagerProxy.h"
+#include "Render/Viewport/IDisplayClusterViewportProxy.h"
+#include "ShaderParameters/DisplayClusterShaderParameters_ICVFX.h"
+#include "ShaderParameters/DisplayClusterShaderParameters_WarpBlend.h"
+
+#include "Engine/Engine.h"
+#include "Engine/GameViewportClient.h"
+
+
+void FDisplayClusterFrameQueue::Init()
+{
+	// Subscribe for events
+	if (IDisplayCluster::Get().GetOperationMode() == EDisplayClusterOperationMode::Cluster)
+	{
+		if (GEngine && GEngine->GameViewport)
+		{
+			GEngine->GameViewport->OnBeginDraw().AddRaw(this, &FDisplayClusterFrameQueue::HandleBeginDraw);
+			GEngine->GameViewport->OnEndDraw().AddRaw(this, &FDisplayClusterFrameQueue::HandleEndDraw);
+		}
+
+		IDisplayCluster::Get().GetCallbacks().OnDisplayClusterProcessLatency_RenderThread().AddRaw(this, &FDisplayClusterFrameQueue::HandleProcessLatency_RenderThread);
+		IDisplayCluster::Get().GetCallbacks().OnDisplayClusterPreProcessIcvfx_RenderThread().AddRaw(this, &FDisplayClusterFrameQueue::HandlePreProcessIcvfx_RenderThread);
+	}
+}
+
+void FDisplayClusterFrameQueue::Release()
+{
+	// Unsubscribe from events
+	if (IDisplayCluster::Get().GetOperationMode() == EDisplayClusterOperationMode::Cluster)
+	{
+		if (GEngine && GEngine->GameViewport)
+		{
+			GEngine->GameViewport->OnBeginDraw().RemoveAll(this);
+			GEngine->GameViewport->OnEndDraw().RemoveAll(this);
+		}
+
+		IDisplayCluster::Get().GetCallbacks().OnDisplayClusterProcessLatency_RenderThread().RemoveAll(this);
+		IDisplayCluster::Get().GetCallbacks().OnDisplayClusterPreProcessIcvfx_RenderThread().RemoveAll(this);
+	}
+
+	// We need to properly release render thread data
+	ENQUEUE_RENDER_COMMAND(FDisplayClusterFrameQueue_Release)(
+		[this](FRHICommandListImmediate& RHICmdList)
+		{
+			CleanQueue_RenderThread();
+		});
+}
+
+void FDisplayClusterFrameQueue::HandleBeginDraw()
+{
+	// Update latency value before rendering new frame
+	if (ADisplayClusterRootActor* RootActor = IDisplayCluster::Get().GetGameMgr()->GetRootActor())
+	{
+		if (const UDisplayClusterConfigurationData* const ConfigData = RootActor->GetConfigData())
+		{
+			// Get latency from config
+			const int32 NewLatency = FMath::Clamp(ConfigData->MediaSettings.Latency, 0, 10);
+			UE_LOG(LogDisplayClusterMedia, VeryVerbose, TEXT("New latency is %u"), NewLatency);
+
+			// Update latency data on RT
+			ENQUEUE_RENDER_COMMAND(FDisplayClusterFrameQueue_SetLatency)(
+				[this, NewLatency](FRHICommandListImmediate& RHICmdList)
+				{
+					SetLatency_RenderThread(NewLatency);
+				});
+		}
+	}
+
+	// Process pending frames adding
+	ENQUEUE_RENDER_COMMAND(FDisplayClusterFrameQueue_ProcessAddingPendingFrames)(
+		[this](FRHICommandListImmediate& RHICmdList)
+		{
+			ProcessAddingPendingFrames_RenderThread();
+		});
+}
+
+void FDisplayClusterFrameQueue::HandleEndDraw()
+{
+	// We need to update head and tail indices before rendering next frame
+	ENQUEUE_RENDER_COMMAND(FDisplayClusterFrameQueue_StepQueueIndices)(
+		[this](FRHICommandListImmediate& RHICmdList)
+		{
+			StepQueueIndices_RenderThread();
+		});
+}
+
+void FDisplayClusterFrameQueue::HandleProcessLatency_RenderThread(FRHICommandListImmediate& RHICmdList, const IDisplayClusterViewportManagerProxy* ViewportManagerProxy, FViewport* Viewport)
+{
+	if (Frames.Num() > 0)
+	{
+		// At this point, all the views for current frame have been rendered already
+		if (IDisplayClusterRenderManager* RenderMgr = IDisplayCluster::Get().GetRenderMgr())
+		{
+			if (IDisplayClusterViewportManager* ViewportMgr = RenderMgr->GetViewportManager())
+			{
+				if (IDisplayClusterViewportManagerProxy* ViewportMgrProxy = ViewportMgr->GetProxy())
+				{
+					TArrayView<IDisplayClusterViewportProxy*> ViewportProxies = ViewportMgrProxy->GetViewports_RenderThread();
+					for (IDisplayClusterViewportProxy* ViewportProxy : ViewportProxies)
+					{
+						// Handle textures of all viewports that don't have any media input attached. Those don't get rendered.
+						if (!ViewportProxy->GetRenderSettings_RenderThread().bSkipSceneRenderingButLeaveResourcesAvailable)
+						{
+							// Get viewport texture
+							TArray<FRHITexture*> Textures;
+							ViewportProxy->GetResources_RenderThread(EDisplayClusterViewportResourceType::InternalRenderTargetResource, Textures);
+							check(Textures.Num() > 0);
+
+							const FString ViewportId(ViewportProxy->GetId());
+
+							// Enqueue this view (head)
+							Frames[IdxHead].SaveView(RHICmdList, ViewportId, Textures[0]);
+							// Dequeue from the past (tail)
+							Frames[IdxTail].LoadView(RHICmdList, ViewportId, Textures[0]);
+						}
+					}
+				}
+			}
+		}
+	}
+}
+
+void FDisplayClusterFrameQueue::HandlePreProcessIcvfx_RenderThread(FRHICommandListImmediate& RHICmdList, const IDisplayClusterViewportProxy* ViewportProxy, FDisplayClusterShaderParameters_WarpBlend& WarpBlendParameters, FDisplayClusterShaderParameters_ICVFX& ICVFXParameters)
+{
+	// Substitute shader data
+	if (Frames.Num() > 0)
+	{
+		const FString ViewportId(ViewportProxy->GetId());
+
+		Frames[IdxHead].SaveData(RHICmdList, ViewportId, WarpBlendParameters, ICVFXParameters);
+		Frames[IdxTail].LoadData(RHICmdList, ViewportId, WarpBlendParameters, ICVFXParameters);
+	}
+}
+
+void FDisplayClusterFrameQueue::SetLatency_RenderThread(int32 NewLatency)
+{
+	// Nothing to do if requested latency is the same as previously (no changes since last frame)
+	if (LastRequestedLatency != NewLatency && NewLatency >= 0)
+	{
+		// Check if latency just was disabled
+		if (NewLatency == 0)
+		{
+			// Clear the queue
+			CleanQueue_RenderThread();
+		}
+		// Check if new latency is higher
+		else if (NewLatency > LastRequestedLatency)
+		{
+			FramesPendingToAdd = (NewLatency + 1) - Frames.Num();
+		}
+		// Check if new latency is lower
+		else if (NewLatency < LastRequestedLatency)
+		{
+			// Depending on the current queue size, we need to find out if the queue needs to be resized
+			const int32 TargetQueueSize = NewLatency + 1;
+
+			// Yeah, we need to remove unnecessary elements from the queue
+			if (TargetQueueSize < Frames.Num())
+			{
+				RemoveFromQueue_RenderThread(Frames.Num() - TargetQueueSize);
+				FramesPendingToAdd = 0;
+			}
+			// No need to remove anything from the queue, just recalculate new amount of pending frames
+			else
+			{
+				FramesPendingToAdd = TargetQueueSize - Frames.Num();
+			}
+		}
+
+		// Finally, remember latency requested
+		LastRequestedLatency = NewLatency;
+	}
+}
+
+void FDisplayClusterFrameQueue::ProcessAddingPendingFrames_RenderThread()
+{
+	if (FramesPendingToAdd > 0)
+	{
+		// Increase queue size step by step, 1 item per frame. This allows to avoid
+		// situations where we have empty queue items.
+		FramesAdded = 1;
+		AddToQueue_RenderThread(FramesAdded);
+	}
+	else
+	{
+		// Queue length remains unchanged
+		FramesAdded = 0;
+	}
+}
+
+void FDisplayClusterFrameQueue::AddToQueue_RenderThread(int32 FramesAmount)
+{
+	checkSlow(FramesAmount > 0);
+
+	// Duplicate the head frame to increase the queue
+	for (int32 Idx = 0; Idx < FramesAmount; ++Idx)
+	{
+		// In case the queue is empty, just instantiate new item
+		if (Frames.IsEmpty())
+		{
+			Frames.AddDefaulted();
+
+			// Reset indices
+			IdxHead = 0;
+			IdxTail = 0;
+		}
+		// Otherwise duplicate head item to prevent 'black' frames in the future
+		else
+		{
+			// For some reason, the following doesn't work (there will be check fail inside Insert)
+			// >>>> Frames.Insert(Frames[IdxHead], IdxHead)
+			//
+			// To achieve the same avoiding unnecessary data copying, I use the approach where data copying
+			// is performed in copy constructor, and the instance is pushed to the container using move constructor.
+			FDisplayClusterFrameQueueItem ClonedItem(Frames[IdxHead]);
+			Frames.Insert(MoveTemp(ClonedItem), IdxHead);
+
+			// Recalculate tail index
+			if (IdxTail > IdxHead)
+			{
+				++IdxTail;
+			}
+		}
+	}
+}
+
+void FDisplayClusterFrameQueue::RemoveFromQueue_RenderThread(int32 FramesAmout)
+{
+	checkSlow(FramesAmout < Frames.Num());
+
+	// We need to remove specified amount of frames starting from the tail
+	for (int32 Idx = 0; Idx < FramesAmout; ++Idx)
+	{
+		// Remove the item
+		Frames.RemoveAt(IdxTail);
+
+		// Recalculate head index
+		if (IdxHead > IdxTail)
+		{
+			--IdxHead;
+		}
+
+		// Recalculate tail index
+		if (IdxTail >= Frames.Num())
+		{
+			IdxTail = 0;
+		}
+	}
+}
+
+void FDisplayClusterFrameQueue::CleanQueue_RenderThread()
+{
+	// Clean up the container
+	Frames.Empty();
+
+	// Reset internals
+	IdxHead = 0;
+	IdxTail = 0;
+	FramesPendingToAdd = 0;
+}
+
+void FDisplayClusterFrameQueue::StepQueueIndices_RenderThread()
+{
+	// Nothing to do if latency is zero
+	if (Frames.IsEmpty())
+	{
+		UE_LOG(LogDisplayClusterMedia, VeryVerbose, TEXT("Queue indices weren't updated, latency=0"));
+		return;
+	}
+
+	// Head index
+	if (++IdxHead >= Frames.Num())
+	{
+		IdxHead = 0;
+	}
+
+	// Tail index doesn't move while we the queue being increased
+	if (FramesAdded == 0)
+	{
+		if (++IdxTail >= Frames.Num())
+		{
+			IdxTail = 0;
+		}
+	}
+	// Otherwise update pending amount
+	else
+	{
+		FramesPendingToAdd -= FramesAdded;
+		checkSlow(FramesPendingToAdd >= 0);
+	}
+}
diff --git a/Engine/Plugins/Runtime/nDisplay/Source/DisplayClusterMedia/Private/Synchronization/DisplayClusterFrameQueue.h b/Engine/Plugins/Runtime/nDisplay/Source/DisplayClusterMedia/Private/Synchronization/DisplayClusterFrameQueue.h
new file mode 100644
--- /dev/null
+++ b/Engine/Plugins/Runtime/nDisplay/Source/DisplayClusterMedia/Private/Synchronization/DisplayClusterFrameQueue.h
@@ -0,0 +1,76 @@
+// Copyright Epic Games, Inc. All Rights Reserved.
+
+#pragma once
+
+#include "CoreMinimal.h"
+
+#include "Synchronization/DisplayClusterFrameQueueItem.h"
+
+class FDisplayClusterMediaInputViewport;
+class FDisplayClusterShaderParameters_ICVFX;
+class FViewport;
+class IDisplayClusterViewportManagerProxy;
+class IDisplayClusterViewportProxy;
+class RHICmdList;
+struct FDisplayClusterShaderParameters_WarpBlend;
+
+/**
+ * Latency queue class
+ *
+ * Implements artificial latency feature. It's responsible for caching visual and shaders data to be used later.
+ *
+ * All the data that comes out of the rendering pipeline for current frame is cached to the HEAD of the queue. The final
+ * output is built from the TAIL view data. Every queue item contains all necessary information for every nD viewport.
+ *
+ * Note: this queue needs to be an intermediate part of the nD rendering pipeline. It should be moved there
+ * at some point during refactoring.
+ */
+class FDisplayClusterFrameQueue
+{
+public:
+	// Initializes latency feature
+	void Init();
+	// Releases latency feature
+	void Release();
+
+private:
+	// Handles BeginDraw event
+	void HandleBeginDraw();
+	// Handles EndDraw event
+	void HandleEndDraw();
+	// Handles latency processing event
+	void HandleProcessLatency_RenderThread(FRHICommandListImmediate& RHICmdList, const IDisplayClusterViewportManagerProxy* ViewportManagerProxy, FViewport* Viewport);
+	// Handles ICVFX pre-processing event
+	void HandlePreProcessIcvfx_RenderThread(FRHICommandListImmediate& RHICmdList, const IDisplayClusterViewportProxy* ViewportProxy, FDisplayClusterShaderParameters_WarpBlend& WarpBlendParameters, FDisplayClusterShaderParameters_ICVFX& ICVFXParameters);
+
+private:
+	// Changes current latency
+	void SetLatency_RenderThread(int32 NewLatency);
+	// Auxiliary method to process queue resize
+	void ProcessAddingPendingFrames_RenderThread();
+	// Auxiliary method to increase queue length by amount of frames
+	void AddToQueue_RenderThread(int32 FramesAmount);
+	// Auxiliary method to decrease queue length by amount of frames
+	void RemoveFromQueue_RenderThread(int32 FramesAmout);
+	// Cleans the queue container
+	void CleanQueue_RenderThread();
+	// Updates head and tail indices every frame
+	void StepQueueIndices_RenderThread();
+
+private:
+	// Frames buffer
+	TArray<FDisplayClusterFrameQueueItem> Frames;
+
+	// Latency value requested last time
+	int32 LastRequestedLatency = 0;
+
+	// Amount of frames pending to add to the latency queue
+	int32 FramesPendingToAdd = 0;
+	// Amount of frames that were added at current frame
+	int32 FramesAdded = 0;
+
+	// Head queue index
+	int32 IdxHead = 0;
+	// Tail queue index
+	int32 IdxTail = 0;
+};
diff --git a/Engine/Plugins/Runtime/nDisplay/Source/DisplayClusterMedia/Private/Synchronization/DisplayClusterFrameQueueItem.cpp b/Engine/Plugins/Runtime/nDisplay/Source/DisplayClusterMedia/Private/Synchronization/DisplayClusterFrameQueueItem.cpp
new file mode 100644
--- /dev/null
+++ b/Engine/Plugins/Runtime/nDisplay/Source/DisplayClusterMedia/Private/Synchronization/DisplayClusterFrameQueueItem.cpp
@@ -0,0 +1,152 @@
+// Copyright Epic Games, Inc. All Rights Reserved.
+
+#include "Synchronization/DisplayClusterFrameQueueItem.h"
+
+#include "DisplayClusterConfigurationTypes.h"
+#include "DisplayClusterMediaLog.h"
+#include "DisplayClusterRootActor.h"
+
+#include "IDisplayCluster.h"
+#include "Render/IDisplayClusterRenderManager.h"
+#include "Render/Viewport/IDisplayClusterViewportManager.h"
+#include "Render/Viewport/IDisplayClusterViewportManagerProxy.h"
+
+#include "RHICommandList.h"
+
+
+FDisplayClusterFrameQueueItem::FDisplayClusterFrameQueueItem(const FDisplayClusterFrameQueueItem& Other)
+{
+	for (const TPair<FString, FDisplayClusterFrameQueueItemView>& View : Other.Views)
+	{
+		FDisplayClusterFrameQueueItemView& NewItem = this->Views.Emplace(View.Key, View.Value);
+		if (View.Value.Texture.IsValid())
+		{
+			NewItem.Texture = CreateTexture(View.Value.Texture);
+			FRHICommandListImmediate& RHICmdList = GetImmediateCommandList_ForRenderCommand();
+			RHICmdList.CopyTexture(View.Value.Texture.GetReference(), NewItem.Texture, FRHICopyTextureInfo());
+		}
+	}
+}
+
+void FDisplayClusterFrameQueueItem::SaveView(FRHICommandListImmediate& RHICmdList, const FString& ViewportId, FRHITexture* Texture)
+{
+	checkSlow(Texture);
+
+	// Find proper view and copy texture data
+	FDisplayClusterFrameQueueItemView* View = Views.Find(ViewportId);
+	if (!View)
+	{
+		View = &Views.Emplace(ViewportId);
+	}
+
+	// Create texture if not yet available. Or re-create if the source texture has been updated (re-sized, other format, etc.)
+	if (!View->Texture.IsValid() ||
+		Texture->GetFormat() != View->Texture.GetReference()->GetFormat() ||
+		Texture->GetDesc().Extent != View->Texture.GetReference()->GetDesc().Extent)
+	{
+		View->Texture = CreateTexture(Texture);
+		check(View->Texture.IsValid());
+	}
+
+	// Copy texture data
+	if (View->Texture.IsValid())
+	{
+		RHICmdList.CopyTexture(Texture, View->Texture.GetReference(), FRHICopyTextureInfo());
+	}
+}
+
+void FDisplayClusterFrameQueueItem::LoadView(FRHICommandListImmediate& RHICmdList, const FString& ViewportId, FRHITexture* Texture)
+{
+	checkSlow(Texture);
+
+	// Find proper view and copy texture data
+	const FDisplayClusterFrameQueueItemView* const View = Views.Find(ViewportId);
+	if (View && View->Texture.IsValid())
+	{
+		RHICmdList.CopyTexture(View->Texture.GetReference(), Texture, FRHICopyTextureInfo());
+	}
+}
+
+void FDisplayClusterFrameQueueItem::SaveData(FRHICommandListImmediate& RHICmdList, const FString& ViewportId, FDisplayClusterShaderParameters_WarpBlend& WarpBlendParameters, FDisplayClusterShaderParameters_ICVFX& ICVFXParameters)
+{
+	if (FDisplayClusterFrameQueueItemView* View = Views.Find(ViewportId))
+	{
+		// Warp/blend - cherry pick necessary parameters only
+		View->WarpBlendData.bRenderAlphaChannel = WarpBlendParameters.bRenderAlphaChannel;
+		View->WarpBlendData.Context = WarpBlendParameters.Context;
+
+		// ICVFX - cherry pick necessary parameters only
+		View->IcvfxData.Cameras = ICVFXParameters.Cameras;
+	}
+}
+
+void FDisplayClusterFrameQueueItem::LoadData(FRHICommandListImmediate& RHICmdList, const FString& ViewportId, FDisplayClusterShaderParameters_WarpBlend& WarpBlendParameters, FDisplayClusterShaderParameters_ICVFX& ICVFXParameters)
+{
+	if (const FDisplayClusterFrameQueueItemView* const View = Views.Find(ViewportId))
+	{
+		// Warp/blend - cherry pick necessary parameters only
+		WarpBlendParameters.bRenderAlphaChannel = View->WarpBlendData.bRenderAlphaChannel;
+		WarpBlendParameters.Context = View->WarpBlendData.Context;
+
+		// ICVFX - cherry pick necessary parameters only
+		for (const FDisplayClusterShaderParameters_ICVFX::FCameraSettings& CameraSettings : View->IcvfxData.Cameras)
+		{
+			// Look up for proper camera settings object. Here we use ViewportId as a key to compare.
+			FDisplayClusterShaderParameters_ICVFX::FCameraSettings* TargetCameraSettings = ICVFXParameters.Cameras.FindByPredicate([CameraSettings](const FDisplayClusterShaderParameters_ICVFX::FCameraSettings& Item)
+				{
+					return Item.Resource.ViewportId.Equals(CameraSettings.Resource.ViewportId, ESearchCase::IgnoreCase);
+				});
+
+			// Update camera settings
+			if (TargetCameraSettings)
+			{
+				TargetCameraSettings->SoftEdge                 = CameraSettings.SoftEdge;
+				TargetCameraSettings->InnerCameraBorderColor      = CameraSettings.InnerCameraBorderColor;
+				TargetCameraSettings->InnerCameraBorderThickness  = CameraSettings.InnerCameraBorderThickness;
+				TargetCameraSettings->InnerCameraFrameAspectRatio = CameraSettings.InnerCameraFrameAspectRatio;
+				TargetCameraSettings->Local2WorldTransform     = CameraSettings.Local2WorldTransform;
+				TargetCameraSettings->CameraViewRotation       = CameraSettings.CameraViewRotation;
+				TargetCameraSettings->CameraViewLocation       = CameraSettings.CameraViewLocation;
+				TargetCameraSettings->CameraPrjMatrix          = CameraSettings.CameraPrjMatrix;
+				TargetCameraSettings->ChromakeySource          = CameraSettings.ChromakeySource;
+				TargetCameraSettings->ChromakeyColor           = CameraSettings.ChromakeyColor;
+				TargetCameraSettings->ChromakeyMarkersColor    = CameraSettings.ChromakeyMarkersColor;
+				TargetCameraSettings->ChromakeyMarkersScale    = CameraSettings.ChromakeyMarkersScale;
+				TargetCameraSettings->ChromakeyMarkersDistance = CameraSettings.ChromakeyMarkersDistance;
+				TargetCameraSettings->ChromakeyMarkersOffset   = CameraSettings.ChromakeyMarkersOffset;
+				TargetCameraSettings->RenderOrder              = CameraSettings.RenderOrder;
+			}
+		}
+	}
+}
+
+FTextureRHIRef FDisplayClusterFrameQueueItem::CreateTexture(FRHITexture* ReferenceTexture)
+{
+	if (ReferenceTexture)
+	{
+		// Use original format and size
+		const int32 SizeX = ReferenceTexture->GetDesc().Extent.X;
+		const int32 SizeY = ReferenceTexture->GetDesc().Extent.Y;
+		const EPixelFormat Format = ReferenceTexture->GetFormat();
+
+		// Prepare description
+		FRHITextureCreateDesc Desc =
+			FRHITextureCreateDesc::Create2D(TEXT("DisplayClusterFrameQueueCacheTexture"), SizeX, SizeY, Format)
+			.SetClearValue(FClearValueBinding::Black)
+			.SetNumMips(1)
+			.SetFlags(ETextureCreateFlags::Dynamic)
+			.AddFlags(ETextureCreateFlags::MultiGPUGraphIgnore)
+			.SetInitialState(ERHIAccess::SRVMask);
+
+		// Leave original flags, but make sure it's ResolveTargetable but not RenderTargetable
+		ETextureCreateFlags Flags = ReferenceTexture->GetFlags();
+		Flags &= ~ETextureCreateFlags::RenderTargetable;
+		Flags |= ETextureCreateFlags::ResolveTargetable;
+		Desc.SetFlags(Flags);
+
+		// Create texture
+		return RHICreateTexture(Desc);
+	}
+
+	return nullptr;
+}
diff --git a/Engine/Plugins/Runtime/nDisplay/Source/DisplayClusterMedia/Private/Synchronization/DisplayClusterFrameQueueItem.h b/Engine/Plugins/Runtime/nDisplay/Source/DisplayClusterMedia/Private/Synchronization/DisplayClusterFrameQueueItem.h
new file mode 100644
--- /dev/null
+++ b/Engine/Plugins/Runtime/nDisplay/Source/DisplayClusterMedia/Private/Synchronization/DisplayClusterFrameQueueItem.h
@@ -0,0 +1,64 @@
+// Copyright Epic Games, Inc. All Rights Reserved.
+
+#pragma once
+
+#include "CoreMinimal.h"
+
+#include "ShaderParameters/DisplayClusterShaderParameters_ICVFX.h"
+#include "ShaderParameters/DisplayClusterShaderParameters_WarpBlend.h"
+
+class FRHICommandListImmediate;
+class FRHITexture;
+
+
+/**
+ * Viewport data cache to be used in the latency queue
+ */
+struct FDisplayClusterFrameQueueItemView
+{
+	/** Cached view texture of a viewoprt */
+	FTextureRHIRef Texture = nullptr;
+
+	/** Cached warp/blend shader data */
+	FDisplayClusterShaderParameters_WarpBlend WarpBlendData;
+
+	/** Cached ICVFX shader data */
+	FDisplayClusterShaderParameters_ICVFX IcvfxData;
+};
+
+
+/**
+ * Latency queue item class
+ *
+ * It's responsible for caching all required data for single frame.
+ * 
+ * Note: this one is render thread only. It's never used on the main thread or any other thread.
+ */
+class FDisplayClusterFrameQueueItem
+{
+public:
+	FDisplayClusterFrameQueueItem() = default;
+	FDisplayClusterFrameQueueItem(const FDisplayClusterFrameQueueItem& Other);
+	FDisplayClusterFrameQueueItem(FDisplayClusterFrameQueueItem&& Other) = default;
+
+	~FDisplayClusterFrameQueueItem() = default;
+
+public:
+	// Save texture of a viewport
+	void SaveView(FRHICommandListImmediate& RHICmdList, const FString& ViewportId, FRHITexture* Texture);
+	// Load texture of a viewport
+	void LoadView(FRHICommandListImmediate& RHICmdList, const FString& ViewportId, FRHITexture* Texture);
+
+	// Save shader data of a viewport
+	void SaveData(FRHICommandListImmediate& RHICmdList, const FString& ViewportId, FDisplayClusterShaderParameters_WarpBlend& WarpBlendParameters, FDisplayClusterShaderParameters_ICVFX& ICVFXParameters);
+	// Load shader data of a viewport
+	void LoadData(FRHICommandListImmediate& RHICmdList, const FString& ViewportId, FDisplayClusterShaderParameters_WarpBlend& WarpBlendParameters, FDisplayClusterShaderParameters_ICVFX& ICVFXParameters);
+
+private:
+	// Auxiliary function that creates a new texture based on the reference (same format, size)
+	FTextureRHIRef CreateTexture(FRHITexture* ReferenceTexture);
+
+private:
+	// Available views (viewports)
+	TMap<FString, FDisplayClusterFrameQueueItemView> Views;
+};
diff --git a/Engine/Plugins/Runtime/nDisplay/Source/DisplayClusterProjection/Private/Policy/MPCDI/DisplayClusterProjectionMPCDIPolicy.cpp b/Engine/Plugins/Runtime/nDisplay/Source/DisplayClusterProjection/Private/Policy/MPCDI/DisplayClusterProjectionMPCDIPolicy.cpp
--- a/Engine/Plugins/Runtime/nDisplay/Source/DisplayClusterProjection/Private/Policy/MPCDI/DisplayClusterProjectionMPCDIPolicy.cpp
+++ b/Engine/Plugins/Runtime/nDisplay/Source/DisplayClusterProjection/Private/Policy/MPCDI/DisplayClusterProjectionMPCDIPolicy.cpp
@@ -7,6 +7,7 @@
 #include "DisplayClusterProjectionStrings.h"
 
 #include "IDisplayCluster.h"
+#include "IDisplayClusterCallbacks.h"
 #include "Config/IDisplayClusterConfigManager.h"
 #include "Game/IDisplayClusterGameManager.h"
 
@@ -337,6 +338,11 @@ void FDisplayClusterProjectionMPCDIPolicy::ApplyWarpBlend_RenderThread(FRHIComma
 
 				WarpBlendParameters.bRenderAlphaChannel = InViewportProxy->GetRenderSettings_RenderThread().bWarpBlendRenderAlphaChannel;
 
+				// Before starting the whole ICVFX shaders pipeline, we need to provide ability to modify any ICVFX shader parameters.
+				// One of use cases is the latency queue that needs to substitute shader parameters.
+				IDisplayCluster::Get().GetCallbacks().OnDisplayClusterPreProcessIcvfx_RenderThread().Broadcast(RHICmdList, InViewportProxy, WarpBlendParameters, ShaderICVFX);
+
+				// Start ICVFX pipeline
 				if (!ShadersAPI.RenderWarpBlend_ICVFX(RHICmdList, WarpBlendParameters, ShaderICVFX))
 				{
 					if (!IsEditorOperationMode_RenderThread(InViewportProxy))
