nDisplay: Added several hotkeys to the nDisplay operator window and light card editor to activate common actions, including setting the view orientation (top, bottom, etc), opening the drawer, and framing the current selection

diff --git a/Engine/Plugins/Runtime/nDisplay/Source/DisplayClusterColorGrading/Private/DisplayClusterColorGradingCommands.cpp b/Engine/Plugins/Runtime/nDisplay/Source/DisplayClusterColorGrading/Private/DisplayClusterColorGradingCommands.cpp
--- a/Engine/Plugins/Runtime/nDisplay/Source/DisplayClusterColorGrading/Private/DisplayClusterColorGradingCommands.cpp
+++ b/Engine/Plugins/Runtime/nDisplay/Source/DisplayClusterColorGrading/Private/DisplayClusterColorGradingCommands.cpp
@@ -10,6 +10,8 @@ void FDisplayClusterColorGradingCommands::RegisterCommands()
 	UI_COMMAND(ContrastColorWheelVisibility, "Constrast", "Show or hide the constrast color wheel", EUserInterfaceActionType::ToggleButton, FInputChord());
 	UI_COMMAND(ColorWheelSliderOrientationHorizontal, "Right", "Puts the color wheel sliders to the right of the color wheel", EUserInterfaceActionType::RadioButton, FInputChord());
 	UI_COMMAND(ColorWheelSliderOrientationVertical, "Below", "Puts the color wheel sliders below the color wheel", EUserInterfaceActionType::RadioButton, FInputChord());
+
+	UI_COMMAND(OpenColorGradingDrawer, "Open Color Grading Drawer", "Opens the color grading drawer from the status bar", EUserInterfaceActionType::Button, FInputChord(EModifierKey::Control, EKeys::SpaceBar));
 }
 
 #undef LOCTEXT_NAMESPACE
diff --git a/Engine/Plugins/Runtime/nDisplay/Source/DisplayClusterColorGrading/Private/DisplayClusterColorGradingCommands.h b/Engine/Plugins/Runtime/nDisplay/Source/DisplayClusterColorGrading/Private/DisplayClusterColorGradingCommands.h
--- a/Engine/Plugins/Runtime/nDisplay/Source/DisplayClusterColorGrading/Private/DisplayClusterColorGradingCommands.h
+++ b/Engine/Plugins/Runtime/nDisplay/Source/DisplayClusterColorGrading/Private/DisplayClusterColorGradingCommands.h
@@ -23,4 +23,6 @@ public:
 	TSharedPtr<FUICommandInfo> ContrastColorWheelVisibility;
 	TSharedPtr<FUICommandInfo> ColorWheelSliderOrientationHorizontal;
 	TSharedPtr<FUICommandInfo> ColorWheelSliderOrientationVertical;
+
+	TSharedPtr<FUICommandInfo> OpenColorGradingDrawer;
 };
\ No newline at end of file
diff --git a/Engine/Plugins/Runtime/nDisplay/Source/DisplayClusterColorGrading/Private/Drawer/DisplayClusterColorGradingDrawerSingleton.cpp b/Engine/Plugins/Runtime/nDisplay/Source/DisplayClusterColorGrading/Private/Drawer/DisplayClusterColorGradingDrawerSingleton.cpp
--- a/Engine/Plugins/Runtime/nDisplay/Source/DisplayClusterColorGrading/Private/Drawer/DisplayClusterColorGradingDrawerSingleton.cpp
+++ b/Engine/Plugins/Runtime/nDisplay/Source/DisplayClusterColorGrading/Private/Drawer/DisplayClusterColorGradingDrawerSingleton.cpp
@@ -2,6 +2,7 @@
 
 #include "DisplayClusterColorGradingDrawerSingleton.h"
 
+#include "DisplayClusterColorGradingCommands.h"
 #include "DisplayClusterColorGradingStyle.h"
 
 #include "DisplayClusterRootActor.h"
@@ -36,6 +37,7 @@ FDisplayClusterColorGradingDrawerSingleton::FDisplayClusterColorGradingDrawerSin
 
 	IDisplayClusterOperator::Get().OnRegisterLayoutExtensions().AddRaw(this, &FDisplayClusterColorGradingDrawerSingleton::ExtendOperatorTabLayout);
 	IDisplayClusterOperator::Get().OnRegisterStatusBarExtensions().AddRaw(this, &FDisplayClusterColorGradingDrawerSingleton::ExtendOperatorStatusBar);
+	IDisplayClusterOperator::Get().OnAppendOperatorPanelCommands().AddRaw(this, &FDisplayClusterColorGradingDrawerSingleton::AppendOperatorPanelCommands);
 }
 
 FDisplayClusterColorGradingDrawerSingleton::~FDisplayClusterColorGradingDrawerSingleton()
@@ -144,6 +146,19 @@ void FDisplayClusterColorGradingDrawerSingleton::ExtendOperatorStatusBar(FDispla
 	StatusBarExtender.AddWidgetDrawer(ColorGradingDrawerConfig);
 }
 
+void FDisplayClusterColorGradingDrawerSingleton::AppendOperatorPanelCommands(TSharedRef<FUICommandList> OperatorPanelCommandList)
+{
+	OperatorPanelCommandList->MapAction(
+		FDisplayClusterColorGradingCommands::Get().OpenColorGradingDrawer,
+		FExecuteAction::CreateRaw(this, &FDisplayClusterColorGradingDrawerSingleton::OpenColorGradingDrawer)
+	);
+}
+
+void FDisplayClusterColorGradingDrawerSingleton::OpenColorGradingDrawer()
+{
+	IDisplayClusterOperator::Get().ToggleDrawer(ColorGradingDrawerId);
+}
+
 void FDisplayClusterColorGradingDrawerSingleton::SaveDrawerState(const TSharedPtr<SWidget>& DrawerContent)
 {
 	if (ColorGradingDrawer.IsValid())
diff --git a/Engine/Plugins/Runtime/nDisplay/Source/DisplayClusterColorGrading/Private/Drawer/DisplayClusterColorGradingDrawerSingleton.h b/Engine/Plugins/Runtime/nDisplay/Source/DisplayClusterColorGrading/Private/Drawer/DisplayClusterColorGradingDrawerSingleton.h
--- a/Engine/Plugins/Runtime/nDisplay/Source/DisplayClusterColorGrading/Private/Drawer/DisplayClusterColorGradingDrawerSingleton.h
+++ b/Engine/Plugins/Runtime/nDisplay/Source/DisplayClusterColorGrading/Private/Drawer/DisplayClusterColorGradingDrawerSingleton.h
@@ -9,6 +9,7 @@
 class FDisplayClusterOperatorStatusBarExtender;
 class FLayoutExtender;
 class FSpawnTabArgs;
+class FUICommandList;
 class SWidget;
 class SDockTab;
 
@@ -45,6 +46,12 @@ private:
 	/** Status bar extender delegate callback that registers the drawer spawner with the operator panel's status bar */
 	void ExtendOperatorStatusBar(FDisplayClusterOperatorStatusBarExtender& StatusBarExtender);
 
+	/** Delegate callback that appends the operator panel command list to add color grading drawer commands */
+	void AppendOperatorPanelCommands(TSharedRef<FUICommandList> OperatorPanelCommandList);
+
+	/** Opens the color grading drawer */
+	void OpenColorGradingDrawer();
+
 	/** Delegate callback when the drawer is closed to save the drawer state */
 	void SaveDrawerState(const TSharedPtr<SWidget>& DrawerContent);
 
diff --git a/Engine/Plugins/Runtime/nDisplay/Source/DisplayClusterLightCardEditor/Private/DisplayClusterLightCardEditorCommands.cpp b/Engine/Plugins/Runtime/nDisplay/Source/DisplayClusterLightCardEditor/Private/DisplayClusterLightCardEditorCommands.cpp
--- a/Engine/Plugins/Runtime/nDisplay/Source/DisplayClusterLightCardEditor/Private/DisplayClusterLightCardEditorCommands.cpp
+++ b/Engine/Plugins/Runtime/nDisplay/Source/DisplayClusterLightCardEditor/Private/DisplayClusterLightCardEditorCommands.cpp
@@ -8,18 +8,19 @@
 void FDisplayClusterLightCardEditorCommands::RegisterCommands()
 {
 	UI_COMMAND(ResetCamera, "Reset Camera", "Resets the camera to focus on the mesh", EUserInterfaceActionType::Button, FInputChord());
+	UI_COMMAND(FrameSelection, "Frame Selection", "Frames the selected items in the viewport", EUserInterfaceActionType::Button, FInputChord(EKeys::F));
 
 	UI_COMMAND(PerspectiveProjection, "Perspective", "A perspective projection from the stage's view origin", EUserInterfaceActionType::Button, FInputChord());
 	UI_COMMAND(OrthographicProjection, "Orthographic", "An orthographic projection from the stage's view origin", EUserInterfaceActionType::Button, FInputChord());
 	UI_COMMAND(AzimuthalProjection, "Azimuthal", "An azimuthal hemisphere projection from the stage's view origin", EUserInterfaceActionType::Button, FInputChord());
 	UI_COMMAND(UVProjection, "UV", "A UV projection that flattens the stage's meshes based on their UV coordinates", EUserInterfaceActionType::Button, FInputChord());
 
-	UI_COMMAND(ViewOrientationTop, "Top", "Orient the view to look at the top of the stage", EUserInterfaceActionType::Button, FInputChord());
-	UI_COMMAND(ViewOrientationBottom, "Bottom", "Orient the view look at the bottom of the stage", EUserInterfaceActionType::Button, FInputChord());
-	UI_COMMAND(ViewOrientationLeft, "Left", "Orient the view to look at the left of the stage", EUserInterfaceActionType::Button, FInputChord());
-	UI_COMMAND(ViewOrientationRight, "Right", "Orient the view to look at the right of the stage", EUserInterfaceActionType::Button, FInputChord());
-	UI_COMMAND(ViewOrientationFront, "Front", "Orient the view to look at the front of the stage", EUserInterfaceActionType::Button, FInputChord());
-	UI_COMMAND(ViewOrientationBack, "Back", "Orient the view to look at the back of the stage", EUserInterfaceActionType::Button, FInputChord());
+	UI_COMMAND(ViewOrientationTop, "Top", "Orient the view to look at the top of the stage", EUserInterfaceActionType::Button, FInputChord(EModifierKey::Alt, EKeys::J));
+	UI_COMMAND(ViewOrientationBottom, "Bottom", "Orient the view look at the bottom of the stage", EUserInterfaceActionType::Button, FInputChord(EModifierKey::Alt | EModifierKey::Shift, EKeys::J));
+	UI_COMMAND(ViewOrientationLeft, "Left", "Orient the view to look at the left of the stage", EUserInterfaceActionType::Button, FInputChord(EModifierKey::Alt, EKeys::K));
+	UI_COMMAND(ViewOrientationRight, "Right", "Orient the view to look at the right of the stage", EUserInterfaceActionType::Button, FInputChord(EModifierKey::Alt | EModifierKey::Shift, EKeys::K));
+	UI_COMMAND(ViewOrientationFront, "Front", "Orient the view to look at the front of the stage", EUserInterfaceActionType::Button, FInputChord(EModifierKey::Alt, EKeys::H));
+	UI_COMMAND(ViewOrientationBack, "Back", "Orient the view to look at the back of the stage", EUserInterfaceActionType::Button, FInputChord(EModifierKey::Alt | EModifierKey::Shift, EKeys::H));
 
 #if PLATFORM_MAC
 	UI_COMMAND(CycleEditorWidgetCoordinateSystem, "Cycle Transform Coordinate System", "Cycles the transform gizmo coordinate systems between cartesian and spherical coordinates", EUserInterfaceActionType::Button, FInputChord(EModifierKey::Command, EKeys::Tilde));
diff --git a/Engine/Plugins/Runtime/nDisplay/Source/DisplayClusterLightCardEditor/Private/DisplayClusterLightCardEditorCommands.h b/Engine/Plugins/Runtime/nDisplay/Source/DisplayClusterLightCardEditor/Private/DisplayClusterLightCardEditorCommands.h
--- a/Engine/Plugins/Runtime/nDisplay/Source/DisplayClusterLightCardEditor/Private/DisplayClusterLightCardEditorCommands.h
+++ b/Engine/Plugins/Runtime/nDisplay/Source/DisplayClusterLightCardEditor/Private/DisplayClusterLightCardEditorCommands.h
@@ -24,6 +24,7 @@ public:
 public:
 	// Viewport commands
 	TSharedPtr<FUICommandInfo> ResetCamera;
+	TSharedPtr<FUICommandInfo> FrameSelection;
 
 	TSharedPtr<FUICommandInfo> PerspectiveProjection;
 	TSharedPtr<FUICommandInfo> OrthographicProjection;
diff --git a/Engine/Plugins/Runtime/nDisplay/Source/DisplayClusterLightCardEditor/Private/Outliner/SDisplayClusterLightCardOutliner.cpp b/Engine/Plugins/Runtime/nDisplay/Source/DisplayClusterLightCardEditor/Private/Outliner/SDisplayClusterLightCardOutliner.cpp
--- a/Engine/Plugins/Runtime/nDisplay/Source/DisplayClusterLightCardEditor/Private/Outliner/SDisplayClusterLightCardOutliner.cpp
+++ b/Engine/Plugins/Runtime/nDisplay/Source/DisplayClusterLightCardEditor/Private/Outliner/SDisplayClusterLightCardOutliner.cpp
@@ -341,20 +341,29 @@ void SDisplayClusterLightCardOutliner::RegisterContextMenu(FName& InName, FToolM
 				Section.AddMenuEntry(FDisplayClusterLightCardEditorCommands::Get().SaveLightCardTemplate);
 			}
 			
-			FToolMenuSection& Section = InMenu->AddSection("LightCardEditSection", LOCTEXT("LightCardEditSectionName", "Edit"));
-
-			if (bAddFullEditMenu)
+			// View section
 			{
-				Section.AddMenuEntry(FGenericCommands::Get().Cut);
-				Section.AddMenuEntry(FGenericCommands::Get().Copy);
+				FToolMenuSection& Section = InMenu->AddSection("LightCardViewSection", LOCTEXT("LightCardviewSectionName", "View"));
+				Section.AddMenuEntry(FDisplayClusterLightCardEditorCommands::Get().FrameSelection);
 			}
-			
-			Section.AddMenuEntry(FGenericCommands::Get().Paste);
-			
-			if (bAddFullEditMenu)
-			{	
-				Section.AddMenuEntry(FGenericCommands::Get().Duplicate);
-				Section.AddMenuEntry(FGenericCommands::Get().Delete);
+
+			// Edit section
+			{
+				FToolMenuSection& Section = InMenu->AddSection("LightCardEditSection", LOCTEXT("LightCardEditSectionName", "Edit"));
+
+				if (bAddFullEditMenu)
+				{
+					Section.AddMenuEntry(FGenericCommands::Get().Cut);
+					Section.AddMenuEntry(FGenericCommands::Get().Copy);
+				}
+
+				Section.AddMenuEntry(FGenericCommands::Get().Paste);
+
+				if (bAddFullEditMenu)
+				{
+					Section.AddMenuEntry(FGenericCommands::Get().Duplicate);
+					Section.AddMenuEntry(FGenericCommands::Get().Delete);
+				}
 			}
 		}));
 	}
diff --git a/Engine/Plugins/Runtime/nDisplay/Source/DisplayClusterLightCardEditor/Private/Viewport/DisplayClusterLightCardEditorViewportClient.cpp b/Engine/Plugins/Runtime/nDisplay/Source/DisplayClusterLightCardEditor/Private/Viewport/DisplayClusterLightCardEditorViewportClient.cpp
--- a/Engine/Plugins/Runtime/nDisplay/Source/DisplayClusterLightCardEditor/Private/Viewport/DisplayClusterLightCardEditorViewportClient.cpp
+++ b/Engine/Plugins/Runtime/nDisplay/Source/DisplayClusterLightCardEditor/Private/Viewport/DisplayClusterLightCardEditorViewportClient.cpp
@@ -1090,13 +1090,22 @@ void FDisplayClusterLightCardEditorViewportClient::ProcessClick(FSceneView& View
 			}
 			else if (!bMultiSelect)
 			{
-				SelectActor(nullptr);
+
+				// Unless a right click is being performed, clear the selection if no geometry was clicked
+				if (!bIsRightClickSelection)
+				{
+					SelectActor(nullptr);
+				}
 			}
 		}
 	}
 	else
 	{
-		SelectActor(nullptr);
+		// Unless a right click is being performed, clear the selection if no geometry was clicked
+		if (!bIsRightClickSelection)
+		{
+			SelectActor(nullptr);
+		}
 	}
 	
 	PropagateActorSelection();
@@ -1632,6 +1641,11 @@ void FDisplayClusterLightCardEditorViewportClient::SelectActors(const TArray<AAc
 	}
 }
 
+bool FDisplayClusterLightCardEditorViewportClient::HasSelection() const
+{
+	return SelectedActors.Num() > 0;
+}
+
 void FDisplayClusterLightCardEditorViewportClient::SetEditorWidgetMode(FDisplayClusterLightCardEditorWidget::EWidgetMode InWidgetMode)
 {
 	// Force the coordinate system back to spherical if the widget mode is set to anything besides translate
@@ -1725,6 +1739,107 @@ void FDisplayClusterLightCardEditorViewportClient::ResetCamera(bool bLocationOnl
 	ResetCameraProjectionTransforms();
 }
 
+void FDisplayClusterLightCardEditorViewportClient::FrameSelection()
+{
+	TArray<FVector> SelectionPositions;
+	FVector AveragePosition = FVector::ZeroVector;
+
+	if (ProjectionMode == EDisplayClusterMeshProjectionType::UV)
+	{
+		FVector2D AverageUVCoords = FVector2D::ZeroVector;
+
+		const float UVPlaneSize = ADisplayClusterLightCardActor::UVPlaneDefaultSize;
+		const float UVPlaneDistance = ADisplayClusterLightCardActor::UVPlaneDefaultDistance;
+
+		for (const FDisplayClusterWeakStageActorPtr& Actor : SelectedActors)
+		{
+			if (Actor.IsValid() && Actor->IsUVActor())
+			{
+				const FVector2D UVCoords = Actor->GetUVCoordinates();
+				const FVector ProxyPosition = FVector(UVPlaneDistance, UVPlaneSize * (UVCoords.X - 0.5f), UVPlaneSize * (0.5f - UVCoords.Y));
+				AverageUVCoords += Actor->GetUVCoordinates();
+				SelectionPositions.Add(ProxyPosition);
+			}
+		}
+
+		if (SelectionPositions.Num())
+		{
+			AverageUVCoords /= SelectionPositions.Num();
+			AveragePosition = FVector(UVPlaneDistance, UVPlaneSize *(AverageUVCoords.X - 0.5f), UVPlaneSize * (0.5f - AverageUVCoords.Y));
+
+			const FVector ViewLocation = FVector(0.0f, AveragePosition.Y, AveragePosition.Z);
+			SetViewLocation(ViewLocation);
+		}
+	}
+	else
+	{
+		for (const FDisplayClusterWeakStageActorPtr& Actor : SelectedActors)
+		{
+			if (Actor.IsValid() && !Actor->IsUVActor())
+			{
+				const FVector ActorPosition = Actor->GetStageActorTransform().GetLocation() - GetViewLocation();
+				AveragePosition += ActorPosition;
+				SelectionPositions.Add(ActorPosition);
+			}
+		}
+
+		if (SelectionPositions.Num())
+		{
+			AveragePosition /= SelectionPositions.Num();
+
+			SetViewRotation(AveragePosition.GetSafeNormal().ToOrientationRotator());
+		}
+	}
+
+	if (SelectionPositions.Num() > 1)
+	{
+		// In the case where more than one actor is selected, the view's FOV may need to be increased so that all selected items are in view
+
+		float MaxTheta = 0;
+		if (RenderViewportType == ELevelViewportType::LVT_Perspective)
+		{
+			// For perspective projection, find the largest angle from the averaged position, and if the current FOV is smaller than that, set the FOV to be equal to that angle
+			const FVector NormalizedAveragePos = AveragePosition.GetSafeNormal();
+			for (const FVector& SelectionPosition : SelectionPositions)
+			{
+				float Theta = FMath::RadiansToDegrees(FMath::Acos(SelectionPosition.GetSafeNormal() | NormalizedAveragePos));
+				MaxTheta = FMath::Clamp(MaxTheta, Theta, 90.0f);
+			}
+		}
+		else
+		{
+			// For orthographic projection, rotate each selected actor position into view space, and compute the bounding box of all positions.
+			// From this, compute the FOV needed to match those bounds (as per GetSceneViewInitOptions, the ortho projection size is equal
+			// to 0.5 * tan(0.5 * FOV) * ViewportSize / DPI)
+			const FMatrix ViewBasis = FRotationMatrix::MakeFromX(AveragePosition.GetSafeNormal());
+
+			FBox FrameBounds(EForceInit::ForceInit);
+			for (const FVector& SelectionPosition : SelectionPositions)
+			{
+				const FVector ViewPosition = ViewBasis.TransformVector(SelectionPosition);
+				FrameBounds += ViewPosition;
+			}
+
+			const FVector2D FrameSize(FrameBounds.GetSize().Y, FrameBounds.GetSize().Z);
+			const FVector2D ViewportSize = Viewport->GetSizeXY();
+			const float DPIScale = GetDPIScale();
+
+			const float ThetaX = FMath::Atan(2 * DPIScale * FrameSize.X / ViewportSize.X);
+			const float ThetaY = FMath::Atan(2 * DPIScale * FrameSize.Y / ViewportSize.Y);
+			MaxTheta = FMath::RadiansToDegrees(FMath::Max(ThetaX, ThetaY));
+		}
+
+		// Only change the FOV if the FOV needed to frame the selection is larger than the current FOV
+		const float HalfFOV = 0.5f * GetProjectionModeFOV(ProjectionMode);
+		if (MaxTheta > HalfFOV)
+		{
+			// Round up the new FOV to the nearest scroll increment, to give the frame a comfortable amount of extra space
+			const float NewFOV = FMath::DivideAndRoundUp(2.0f * MaxTheta, FOVScrollIncrement) * FOVScrollIncrement;
+			SetProjectionModeFOV(ProjectionMode, FMath::Clamp(NewFOV, CameraController->GetConfig().MinimumAllowedFOV, CameraController->GetConfig().MaximumAllowedFOV));
+		}
+	}
+}
+
 void FDisplayClusterLightCardEditorViewportClient::CycleCoordinateSystem()
 {
 	if (EditorWidgetCoordinateSystem == FDisplayClusterLightCardEditorHelper::ECoordinateSystem::Cartesian)
diff --git a/Engine/Plugins/Runtime/nDisplay/Source/DisplayClusterLightCardEditor/Private/Viewport/DisplayClusterLightCardEditorViewportClient.h b/Engine/Plugins/Runtime/nDisplay/Source/DisplayClusterLightCardEditor/Private/Viewport/DisplayClusterLightCardEditorViewportClient.h
--- a/Engine/Plugins/Runtime/nDisplay/Source/DisplayClusterLightCardEditor/Private/Viewport/DisplayClusterLightCardEditorViewportClient.h
+++ b/Engine/Plugins/Runtime/nDisplay/Source/DisplayClusterLightCardEditor/Private/Viewport/DisplayClusterLightCardEditorViewportClient.h
@@ -94,6 +94,9 @@ public:
 	/** Selects the actor proxies that correspond to the specified actors */
 	void SelectActors(const TArray<AActor*>& ActorsToSelect);
 
+	/** Gets whether any items are selected in the viewport */
+	bool HasSelection() const;
+
 	FDisplayClusterLightCardEditorWidget::EWidgetMode GetEditorWidgetMode() const { return EditorWidget->GetWidgetMode(); }
 	void SetEditorWidgetMode(FDisplayClusterLightCardEditorWidget::EWidgetMode InWidgetMode);
 
@@ -115,6 +118,9 @@ public:
 	/** Resets the camera to the initial rotation / position */
 	void ResetCamera(bool bLocationOnly = false);
 
+	/** Frames the currently selected items in the viewport */
+	void FrameSelection();
+
 	/** Cycles the editor widget's coordinate system between cartesian and spherical */
 	void CycleCoordinateSystem();
 
diff --git a/Engine/Plugins/Runtime/nDisplay/Source/DisplayClusterLightCardEditor/Private/Viewport/DisplayClusterLightcardEditorViewport.cpp b/Engine/Plugins/Runtime/nDisplay/Source/DisplayClusterLightCardEditor/Private/Viewport/DisplayClusterLightcardEditorViewport.cpp
--- a/Engine/Plugins/Runtime/nDisplay/Source/DisplayClusterLightCardEditor/Private/Viewport/DisplayClusterLightcardEditorViewport.cpp
+++ b/Engine/Plugins/Runtime/nDisplay/Source/DisplayClusterLightCardEditor/Private/Viewport/DisplayClusterLightcardEditorViewport.cpp
@@ -456,6 +456,16 @@ void SDisplayClusterLightCardEditorViewport::BindCommands()
 			FExecuteAction::CreateSP(this, &SDisplayClusterLightCardEditorViewport::CycleEditorWidgetMode),
 			FCanExecuteAction()
 		);
+
+		// Remove commands that are being replaced by new specific commands for the nDisplay operator
+		CommandList->UnmapAction(Commands.Top);
+		CommandList->UnmapAction(Commands.Bottom);
+		CommandList->UnmapAction(Commands.Left);
+		CommandList->UnmapAction(Commands.Right);
+		CommandList->UnmapAction(Commands.Front);
+		CommandList->UnmapAction(Commands.Back);
+		CommandList->UnmapAction(Commands.FocusViewportToSelection);
+		CommandList->UnmapAction(Commands.FocusAllViewportsToSelection);
 	}
 
 	{
@@ -513,6 +523,12 @@ void SDisplayClusterLightCardEditorViewport::BindCommands()
 			Commands.ResetCamera,
 			FExecuteAction::CreateSP(ViewportClient.Get(), &FDisplayClusterLightCardEditorViewportClient::ResetCamera, false));
 
+		CommandList->MapAction(
+			Commands.FrameSelection,
+			FExecuteAction::CreateSP(ViewportClient.Get(), &FDisplayClusterLightCardEditorViewportClient::FrameSelection),
+			FCanExecuteAction::CreateSP(ViewportClient.Get(), &FDisplayClusterLightCardEditorViewportClient::HasSelection)
+		);
+
 		CommandList->MapAction(
 			Commands.CycleEditorWidgetCoordinateSystem,
 			FExecuteAction::CreateSP(ViewportClient.Get(), &FDisplayClusterLightCardEditorViewportClient::CycleCoordinateSystem));
@@ -660,6 +676,12 @@ TSharedRef<SWidget> SDisplayClusterLightCardEditorViewport::MakeContextMenu()
 	MenuBuilder.AddMenuEntry(FDisplayClusterLightCardEditorCommands::Get().RemoveLightCard);
 	MenuBuilder.AddMenuEntry(FDisplayClusterLightCardEditorCommands::Get().SaveLightCardTemplate);
 	
+	MenuBuilder.BeginSection("View", LOCTEXT("ViewSection", "View"));
+	{
+		MenuBuilder.AddMenuEntry(FDisplayClusterLightCardEditorCommands::Get().FrameSelection);
+	}
+	MenuBuilder.EndSection();
+
 	MenuBuilder.BeginSection("Edit", LOCTEXT("EditSection", "Edit"));
 	{
 		MenuBuilder.AddMenuEntry(FGenericCommands::Get().Cut);
diff --git a/Engine/Plugins/Runtime/nDisplay/Source/DisplayClusterOperator/Private/DisplayClusterOperatorModule.cpp b/Engine/Plugins/Runtime/nDisplay/Source/DisplayClusterOperator/Private/DisplayClusterOperatorModule.cpp
--- a/Engine/Plugins/Runtime/nDisplay/Source/DisplayClusterOperator/Private/DisplayClusterOperatorModule.cpp
+++ b/Engine/Plugins/Runtime/nDisplay/Source/DisplayClusterOperator/Private/DisplayClusterOperatorModule.cpp
@@ -83,6 +83,14 @@ void FDisplayClusterOperatorModule::ShowDetailsForObjects(const TArray<UObject*>
 	DetailObjectsChanged.Broadcast(Objects);
 }
 
+void FDisplayClusterOperatorModule::ToggleDrawer(const FName DrawerId)
+{
+	if (ActiveOperatorPanel.IsValid())
+	{
+		ActiveOperatorPanel.Pin()->ToggleDrawer(DrawerId);
+	}
+}
+
 void FDisplayClusterOperatorModule::ForceDismissDrawers()
 {
 	if (ActiveOperatorPanel.IsValid())
diff --git a/Engine/Plugins/Runtime/nDisplay/Source/DisplayClusterOperator/Private/DisplayClusterOperatorModule.h b/Engine/Plugins/Runtime/nDisplay/Source/DisplayClusterOperator/Private/DisplayClusterOperatorModule.h
--- a/Engine/Plugins/Runtime/nDisplay/Source/DisplayClusterOperator/Private/DisplayClusterOperatorModule.h
+++ b/Engine/Plugins/Runtime/nDisplay/Source/DisplayClusterOperator/Private/DisplayClusterOperatorModule.h
@@ -29,6 +29,7 @@ public:
 	virtual TSharedRef<IDisplayClusterOperatorViewModel> GetOperatorViewModel() override;
 	virtual FOnRegisterLayoutExtensions& OnRegisterLayoutExtensions() override { return RegisterLayoutExtensions; }
 	virtual FOnRegisterStatusBarExtensions& OnRegisterStatusBarExtensions() override { return RegisterStatusBarExtensions; }
+	virtual FOnAppendOperatorPanelCommands& OnAppendOperatorPanelCommands() override { return AppendOperatorPanelCommands; }
 	virtual FOnDetailObjectsChanged& OnDetailObjectsChanged() override { return DetailObjectsChanged; }
 
 	virtual FName GetPrimaryOperatorExtensionId() override;
@@ -38,6 +39,7 @@ public:
 	virtual void GetRootActorLevelInstances(TArray<ADisplayClusterRootActor*>& OutRootActorInstances) override;
 	virtual void ShowDetailsForObject(UObject* Object) override;
 	virtual void ShowDetailsForObjects(const TArray<UObject*>& Objects) override;
+	virtual void ToggleDrawer(const FName DrawerId) override;
 	virtual void ForceDismissDrawers() override;
 	//~ End IDisplayClusterOperator interface
 
@@ -51,6 +53,7 @@ private:
 private:
 	FOnRegisterLayoutExtensions RegisterLayoutExtensions;
 	FOnRegisterStatusBarExtensions RegisterStatusBarExtensions;
+	FOnAppendOperatorPanelCommands AppendOperatorPanelCommands;
 	FOnDetailObjectsChanged DetailObjectsChanged;
 
 	TWeakPtr<SDisplayClusterOperatorPanel> ActiveOperatorPanel;
diff --git a/Engine/Plugins/Runtime/nDisplay/Source/DisplayClusterOperator/Private/SDisplayClusterOperatorPanel.cpp b/Engine/Plugins/Runtime/nDisplay/Source/DisplayClusterOperator/Private/SDisplayClusterOperatorPanel.cpp
--- a/Engine/Plugins/Runtime/nDisplay/Source/DisplayClusterOperator/Private/SDisplayClusterOperatorPanel.cpp
+++ b/Engine/Plugins/Runtime/nDisplay/Source/DisplayClusterOperator/Private/SDisplayClusterOperatorPanel.cpp
@@ -8,11 +8,13 @@
 #include "SDisplayClusterOperatorStatusBar.h"
 #include "DisplayClusterOperatorStatusBarExtender.h"
 
+#include "Editor.h"
 #include "Styling/AppStyle.h"
 #include "Framework/Docking/TabManager.h"
 #include "Framework/Docking/LayoutExtender.h"
 #include "Framework/MultiBox/MultiBoxBuilder.h"
 #include "SKismetInspector.h"
+#include "TimerManager.h"
 #include "Widgets/Docking/SDockTab.h"
 
 #include "Components/ActorComponent.h"
@@ -31,6 +33,9 @@ SDisplayClusterOperatorPanel::~SDisplayClusterOperatorPanel()
 
 void SDisplayClusterOperatorPanel::Construct(const FArguments& InArgs, const TSharedRef<FTabManager>& InTabManager, const TSharedPtr<SWindow>& WindowOwner)
 {
+	CommandList = MakeShareable(new FUICommandList);
+	BindCommands();
+
 	TabManager = InTabManager;
 	TSharedRef<FWorkspaceItem> AppMenuGroup = TabManager->AddLocalWorkspaceMenuCategory(LOCTEXT("OperatorMenuGroupName", "In-Camera VFX"));
 	TabManager->SetAllowWindowMenuBar(true);
@@ -131,6 +136,41 @@ void SDisplayClusterOperatorPanel::Construct(const FArguments& InArgs, const TSh
 	}
 }
 
+FReply SDisplayClusterOperatorPanel::OnKeyDown(const FGeometry& MyGeometry, const FKeyEvent& InKeyEvent)
+{
+	if (CommandList.IsValid() && CommandList->ProcessCommandBindings(InKeyEvent))
+	{
+		return FReply::Handled();
+	}
+
+	return FReply::Unhandled();
+}
+
+void SDisplayClusterOperatorPanel::ToggleDrawer(const FName DrawerId)
+{
+	if (StatusBar.IsValid())
+	{
+		bool bWasDismissed = false;
+
+		if (StatusBar->IsDrawerOpened(DrawerId))
+		{
+			StatusBar->DismissDrawer(nullptr);
+			bWasDismissed = true;
+		}
+
+		if (!bWasDismissed)
+		{
+			GEditor->GetTimerManager()->SetTimerForNextTick(FTimerDelegate::CreateLambda([this, DrawerId]()
+			{
+				if (StatusBar.IsValid())
+				{
+					StatusBar->OpenDrawer(DrawerId);
+				}
+			}));
+		}
+	}
+}
+
 void SDisplayClusterOperatorPanel::ForceDismissDrawers()
 {
 	if (StatusBar.IsValid())
@@ -139,6 +179,14 @@ void SDisplayClusterOperatorPanel::ForceDismissDrawers()
 	}
 }
 
+void SDisplayClusterOperatorPanel::BindCommands()
+{
+	if (CommandList.IsValid())
+	{
+		IDisplayClusterOperator::Get().OnAppendOperatorPanelCommands().Broadcast(CommandList.ToSharedRef());
+	}
+}
+
 TSharedRef<SDockTab> SDisplayClusterOperatorPanel::SpawnToolbarTab(const FSpawnTabArgs& Args)
 {
 	return SNew(SDockTab)
diff --git a/Engine/Plugins/Runtime/nDisplay/Source/DisplayClusterOperator/Private/SDisplayClusterOperatorPanel.h b/Engine/Plugins/Runtime/nDisplay/Source/DisplayClusterOperator/Private/SDisplayClusterOperatorPanel.h
--- a/Engine/Plugins/Runtime/nDisplay/Source/DisplayClusterOperator/Private/SDisplayClusterOperatorPanel.h
+++ b/Engine/Plugins/Runtime/nDisplay/Source/DisplayClusterOperator/Private/SDisplayClusterOperatorPanel.h
@@ -9,6 +9,7 @@ class FTabManager;
 class FLayoutExtender;
 class FSpawnTabArgs;
 class FMenuBuilder;
+class FUICommandList;
 class SBox;
 class SDockTab;
 class SWindow;
@@ -39,10 +40,20 @@ public:
 
 	void Construct(const FArguments& InArgs, const TSharedRef<FTabManager>& InTabManager, const TSharedPtr<SWindow>& WindowOwner);
 
+	//~ SWidget interface
+	virtual FReply OnKeyDown(const FGeometry& MyGeometry, const FKeyEvent& InKeyEvent);
+	//~ End SWidget interface
+
+	/** Toggles the state of a drawer with the specified ID, closing the drawer if it is open, and opening the drawer if it is closed */
+	void ToggleDrawer(const FName DrawerId);
+
 	/** Forces the operator panel to dismiss any open drawers */
 	void ForceDismissDrawers();
 
 private:
+	/** Binds any commands to the operator panel's command list, as well as allowing any external modules to register their own commands */
+	void BindCommands();
+
 	/** Creates a tab with the operator toolbar in it */
 	TSharedRef<SDockTab> SpawnToolbarTab(const FSpawnTabArgs& Args);
 
@@ -71,6 +82,9 @@ private:
 	/** A reference to the operator panel's status bar widget */
 	TSharedPtr<SDisplayClusterOperatorStatusBar> StatusBar;
 
+	/** A command list for processing commands for the window */
+	TSharedPtr<FUICommandList> CommandList;
+
 	/** The delegate handle for the operator module's OnDetailObjectsChanged event */
 	FDelegateHandle DetailObjectsChangedHandle;
 };
\ No newline at end of file
diff --git a/Engine/Plugins/Runtime/nDisplay/Source/DisplayClusterOperator/Private/SDisplayClusterOperatorStatusBar.cpp b/Engine/Plugins/Runtime/nDisplay/Source/DisplayClusterOperator/Private/SDisplayClusterOperatorStatusBar.cpp
--- a/Engine/Plugins/Runtime/nDisplay/Source/DisplayClusterOperator/Private/SDisplayClusterOperatorStatusBar.cpp
+++ b/Engine/Plugins/Runtime/nDisplay/Source/DisplayClusterOperator/Private/SDisplayClusterOperatorStatusBar.cpp
@@ -25,6 +25,16 @@ void SDisplayClusterOperatorStatusBar::RegisterDrawer(FWidgetDrawerConfig&& Draw
 	WidgetDrawer->RegisterDrawer(MoveTemp(Drawer), SlotIndex);
 }
 
+bool SDisplayClusterOperatorStatusBar::IsDrawerOpened(const FName DrawerId) const
+{
+	return WidgetDrawer->IsDrawerOpened(DrawerId);
+}
+
+void SDisplayClusterOperatorStatusBar::OpenDrawer(const FName DrawerId)
+{
+	WidgetDrawer->OpenDrawer(DrawerId);
+}
+
 bool SDisplayClusterOperatorStatusBar::DismissDrawer(const TSharedPtr<SWidget>& NewlyFocusedWidget)
 {
 	return WidgetDrawer->DismissDrawer(NewlyFocusedWidget);
diff --git a/Engine/Plugins/Runtime/nDisplay/Source/DisplayClusterOperator/Private/SDisplayClusterOperatorStatusBar.h b/Engine/Plugins/Runtime/nDisplay/Source/DisplayClusterOperator/Private/SDisplayClusterOperatorStatusBar.h
--- a/Engine/Plugins/Runtime/nDisplay/Source/DisplayClusterOperator/Private/SDisplayClusterOperatorStatusBar.h
+++ b/Engine/Plugins/Runtime/nDisplay/Source/DisplayClusterOperator/Private/SDisplayClusterOperatorStatusBar.h
@@ -28,6 +28,12 @@ public:
 	/** Registers a new drawer with this status bar. Registering will add a button to open and close the drawer */
 	void RegisterDrawer(FWidgetDrawerConfig&& Drawer, int32 SlotIndex = INDEX_NONE);
 
+	/** Gets whether a drawer with the specified ID is currently open */
+	bool IsDrawerOpened(const FName DrawerId) const;
+
+	/** Opens a drawer with the specified ID */
+	void OpenDrawer(const FName DrawerId);
+
 	/**
 	 * Dismisses an open drawer with an animation.  The drawer contents are removed once the animation is complete
 	 *
diff --git a/Engine/Plugins/Runtime/nDisplay/Source/DisplayClusterOperator/Public/IDisplayClusterOperator.h b/Engine/Plugins/Runtime/nDisplay/Source/DisplayClusterOperator/Public/IDisplayClusterOperator.h
--- a/Engine/Plugins/Runtime/nDisplay/Source/DisplayClusterOperator/Public/IDisplayClusterOperator.h
+++ b/Engine/Plugins/Runtime/nDisplay/Source/DisplayClusterOperator/Public/IDisplayClusterOperator.h
@@ -9,13 +9,10 @@
 class IDisplayClusterOperatorViewModel;
 class FLayoutExtender;
 class FExtensibilityManager;
+class FUICommandList;
 class FDisplayClusterOperatorStatusBarExtender;
 class ADisplayClusterRootActor;
 
-DECLARE_EVENT_OneParam(IDisplayClusterOperator, FOnRegisterLayoutExtensions, FLayoutExtender&);
-DECLARE_EVENT_OneParam(IDisplayClusterOperator, FOnRegisterStatusBarExtensions, FDisplayClusterOperatorStatusBarExtender&);
-DECLARE_EVENT_OneParam(IDisplayClusterOperator, FOnDetailObjectsChanged, const TArray<UObject*>&);
-
 /**
  * Display Cluster Operator module interface
  */
@@ -24,6 +21,11 @@ class IDisplayClusterOperator : public IModuleInterface
 public:
 	static constexpr const TCHAR* ModuleName = TEXT("DisplayClusterOperator");
 
+	DECLARE_MULTICAST_DELEGATE_OneParam(FOnRegisterLayoutExtensions, FLayoutExtender&);
+	DECLARE_MULTICAST_DELEGATE_OneParam(FOnRegisterStatusBarExtensions, FDisplayClusterOperatorStatusBarExtender&);
+	DECLARE_MULTICAST_DELEGATE_OneParam(FOnAppendOperatorPanelCommands, TSharedRef<FUICommandList>)
+	DECLARE_MULTICAST_DELEGATE_OneParam(FOnDetailObjectsChanged, const TArray<UObject*>&);
+
 public:
 	virtual ~IDisplayClusterOperator() = default;
 
@@ -57,6 +59,9 @@ public:
 	/** Gets the event handler that is raised when the operator panel processes extensions to its status bar */
 	virtual FOnRegisterStatusBarExtensions& OnRegisterStatusBarExtensions() = 0;
 
+	/** Gets the event handler that is raised when the operator panel binds commands to its command list */
+	virtual FOnAppendOperatorPanelCommands& OnAppendOperatorPanelCommands() = 0;
+
 	/** Gets the event handler that is raised when the objects being displayed in the operator's details panel are changed */
 	virtual FOnDetailObjectsChanged& OnDetailObjectsChanged() = 0;
 
@@ -81,6 +86,9 @@ public:
 	/** Displays the properties of the specified object in the operator's details panel */
 	virtual void ShowDetailsForObjects(const TArray<UObject*>& Objects) = 0;
 
+	/** Toggles the state of a drawer with the specified ID, closing the drawer if it is open, and opening the drawer if it is closed */
+	virtual void ToggleDrawer(const FName DrawerId) = 0;
+
 	/** Forces the operator panel to dismiss any open drawers */
 	virtual void ForceDismissDrawers() = 0;
 };
\ No newline at end of file
