nDisplay: Changed the label of the azimuthal projection to "Dome" projection, and set the projection mode options to be radio buttons

diff --git a/Engine/Plugins/Runtime/nDisplay/Source/DisplayClusterLightCardEditor/Private/DisplayClusterLightCardEditorCommands.cpp b/Engine/Plugins/Runtime/nDisplay/Source/DisplayClusterLightCardEditor/Private/DisplayClusterLightCardEditorCommands.cpp
--- a/Engine/Plugins/Runtime/nDisplay/Source/DisplayClusterLightCardEditor/Private/DisplayClusterLightCardEditorCommands.cpp
+++ b/Engine/Plugins/Runtime/nDisplay/Source/DisplayClusterLightCardEditor/Private/DisplayClusterLightCardEditorCommands.cpp
@@ -10,10 +10,10 @@ void FDisplayClusterLightCardEditorCommands::RegisterCommands()
 	UI_COMMAND(ResetCamera, "Reset Camera", "Resets the camera to focus on the mesh", EUserInterfaceActionType::Button, FInputChord());
 	UI_COMMAND(FrameSelection, "Frame Selection", "Frames the selected items in the viewport", EUserInterfaceActionType::Button, FInputChord(EKeys::F));
 
-	UI_COMMAND(PerspectiveProjection, "Perspective", "A perspective projection from the stage's view origin", EUserInterfaceActionType::Button, FInputChord());
-	UI_COMMAND(OrthographicProjection, "Orthographic", "An orthographic projection from the stage's view origin", EUserInterfaceActionType::Button, FInputChord());
-	UI_COMMAND(AzimuthalProjection, "Azimuthal", "An azimuthal hemisphere projection from the stage's view origin", EUserInterfaceActionType::Button, FInputChord());
-	UI_COMMAND(UVProjection, "UV", "A UV projection that flattens the stage's meshes based on their UV coordinates", EUserInterfaceActionType::Button, FInputChord());
+	UI_COMMAND(PerspectiveProjection, "Perspective", "A perspective projection from the stage's view origin", EUserInterfaceActionType::RadioButton, FInputChord());
+	UI_COMMAND(OrthographicProjection, "Orthographic", "An orthographic projection from the stage's view origin", EUserInterfaceActionType::RadioButton, FInputChord());
+	UI_COMMAND(AzimuthalProjection, "Dome", "A hemispherical projection from the stage's view origin", EUserInterfaceActionType::RadioButton, FInputChord());
+	UI_COMMAND(UVProjection, "UV", "A UV projection that flattens the stage's meshes based on their UV coordinates", EUserInterfaceActionType::RadioButton, FInputChord());
 
 	UI_COMMAND(ViewOrientationTop, "Top", "Orient the view to look at the top of the stage", EUserInterfaceActionType::Button, FInputChord(EModifierKey::Alt, EKeys::J));
 	UI_COMMAND(ViewOrientationBottom, "Bottom", "Orient the view look at the bottom of the stage", EUserInterfaceActionType::Button, FInputChord(EModifierKey::Alt | EModifierKey::Shift, EKeys::J));
diff --git a/Engine/Plugins/Runtime/nDisplay/Source/DisplayClusterLightCardEditor/Private/Viewport/DisplayClusterLightcardEditorViewport.cpp b/Engine/Plugins/Runtime/nDisplay/Source/DisplayClusterLightCardEditor/Private/Viewport/DisplayClusterLightcardEditorViewport.cpp
--- a/Engine/Plugins/Runtime/nDisplay/Source/DisplayClusterLightCardEditor/Private/Viewport/DisplayClusterLightcardEditorViewport.cpp
+++ b/Engine/Plugins/Runtime/nDisplay/Source/DisplayClusterLightCardEditor/Private/Viewport/DisplayClusterLightcardEditorViewport.cpp
@@ -135,7 +135,7 @@ public:
 				break;
 
 			case EDisplayClusterMeshProjectionType::Azimuthal:
-				Label = LOCTEXT("ProjectionwMenuTitle_Azimuthal", "Azimuthal");
+				Label = LOCTEXT("ProjectionwMenuTitle_Azimuthal", "Dome");
 				break;
 
 			case EDisplayClusterMeshProjectionType::UV:
