nDisplay: Added indicator to UV cards in the editor only.

diff --git a/Engine/Plugins/Runtime/nDisplay/Source/DisplayCluster/Private/Game/EngineClasses/Scene/DisplayClusterLightCardActor.cpp b/Engine/Plugins/Runtime/nDisplay/Source/DisplayCluster/Private/Game/EngineClasses/Scene/DisplayClusterLightCardActor.cpp
--- a/Engine/Plugins/Runtime/nDisplay/Source/DisplayCluster/Private/Game/EngineClasses/Scene/DisplayClusterLightCardActor.cpp
+++ b/Engine/Plugins/Runtime/nDisplay/Source/DisplayCluster/Private/Game/EngineClasses/Scene/DisplayClusterLightCardActor.cpp
@@ -64,13 +64,30 @@ ADisplayClusterLightCardActor::ADisplayClusterLightCardActor(const FObjectInitia
 	LightCardTransformerComponent = CreateDefaultSubobject<USceneComponent>(TEXT("LightCardTransformer"));
 	LightCardTransformerComponent->AttachToComponent(MainSpringArmComponent, FAttachmentTransformRules::KeepRelativeTransform);
 
-	static ConstructorHelpers::FObjectFinder<UStaticMesh> PlaneObj(TEXT("/nDisplay/LightCard/SM_LightCardPlane"));
+	{
+		static ConstructorHelpers::FObjectFinder<UStaticMesh> PlaneObj(TEXT("/nDisplay/LightCard/SM_LightCardPlane"));
+
+		LightCardComponent = CreateDefaultSubobject<UStaticMeshComponent>(TEXT("LightCard"));
+		LightCardComponent->AttachToComponent(LightCardTransformerComponent, FAttachmentTransformRules::KeepRelativeTransform);
+		LightCardComponent->SetCollisionProfileName(TEXT("OverlapAllDynamic"));
+		LightCardComponent->Mobility = EComponentMobility::Movable;
+		LightCardComponent->SetStaticMesh(PlaneObj.Object);
+	}
+
+#if WITH_EDITOR
+	{
+		static ConstructorHelpers::FObjectFinder<UStaticMesh> UVIndicatorObj(TEXT("/nDisplay/LightCard/SM_UVIndicator"));
+
+		UVIndicatorComponent = CreateDefaultSubobject<UStaticMeshComponent>(TEXT("UVIndicator"));
+		UVIndicatorComponent->AttachToComponent(LightCardComponent, FAttachmentTransformRules::KeepRelativeTransform);
+		UVIndicatorComponent->SetCollisionEnabled(ECollisionEnabled::NoCollision);
+		UVIndicatorComponent->Mobility = EComponentMobility::Movable;
+		UVIndicatorComponent->SetStaticMesh(UVIndicatorObj.Object);
 
-	LightCardComponent = CreateDefaultSubobject<UStaticMeshComponent>(TEXT("LightCard"));
-	LightCardComponent->AttachToComponent(LightCardTransformerComponent, FAttachmentTransformRules::KeepRelativeTransform);
-	LightCardComponent->SetCollisionProfileName(TEXT("OverlapAllDynamic"));
-	LightCardComponent->Mobility = EComponentMobility::Movable;
-	LightCardComponent->SetStaticMesh(PlaneObj.Object);
+		UVIndicatorComponent->SetRelativeLocation(FVector(0, 0, 1)); // Slightly in front of the parent
+		UVIndicatorComponent->SetVisibility(false);
+	}
+#endif // WITH_EDITOR
 
 	UpdateStageActorTransform();
 
@@ -113,6 +130,7 @@ void ADisplayClusterLightCardActor::Tick(float DeltaSeconds)
 	UpdateStageActorTransform();
 	UpdateLightCardMaterialInstance();
 	UpdateLightCardVisibility();
+	UpdateUVIndicator();
 }
 
 #if WITH_EDITOR
@@ -176,6 +194,18 @@ FName ADisplayClusterLightCardActor::GetCustomIconName() const
 
 #endif
 
+void ADisplayClusterLightCardActor::GetLightCardMeshComponents(TArray<UStaticMeshComponent*>& MeshComponents) const
+{
+	MeshComponents.Add(LightCardComponent.Get());
+
+#if WITH_EDITOR
+	if (GIsEditor && bIsUVLightCard)
+	{
+		MeshComponents.Add(UVIndicatorComponent.Get());
+	}
+#endif // WITH_EDITOR
+}
+
 UStaticMesh* ADisplayClusterLightCardActor::GetStaticMesh() const
 {
 	return LightCardComponent->GetStaticMesh();
@@ -332,6 +362,7 @@ void ADisplayClusterLightCardActor::UpdateLightCardVisibility()
 	{
 		// Only render the UV light card if this is a proxy actor
 		const bool bShouldBeVisible = IsProxy();
+
 		if (LightCardComponent->IsVisible() != bShouldBeVisible)
 		{
 			LightCardComponent->SetVisibility(bShouldBeVisible);
@@ -339,6 +370,59 @@ void ADisplayClusterLightCardActor::UpdateLightCardVisibility()
 	}
 }
 
+
+void ADisplayClusterLightCardActor::UpdateUVIndicator()
+{
+#if WITH_EDITOR
+	if (GIsEditor && bIsUVLightCard && IsProxy())
+	{
+		if (!UVIndicatorComponent->IsVisible())
+		{
+			UVIndicatorComponent->SetVisibility(true);
+		}
+	}
+	else if (UVIndicatorComponent->IsVisible())
+	{
+		UVIndicatorComponent->SetVisibility(false);
+	}
+
+	if (GIsEditor && bIsUVLightCard)
+	{
+		// Keep 1:1 aspect ratio
+		if (const USceneComponent* UVParentComponent = UVIndicatorComponent->GetAttachParent())
+		{
+			const FVector ParentWorldScale = UVParentComponent->GetComponentScale();
+
+			const double ParentScaleAbsX = FMath::Abs(ParentWorldScale.X);
+			const double ParentScaleAbsY = FMath::Abs(ParentWorldScale.Y);
+
+			const double ParentScaleMin = FMath::Min(ParentScaleAbsX, ParentScaleAbsY);
+
+			if (ParentScaleMin > KINDA_SMALL_NUMBER)
+			{
+				constexpr double UVScale = 0.25;
+				FVector ScaleFactor(1, 1, 1);
+
+				if (ParentScaleAbsX > ParentScaleAbsY)
+				{
+					ScaleFactor.X = ParentScaleAbsY / ParentScaleAbsX;
+				}
+				else
+				{
+					ScaleFactor.Y = ParentScaleAbsX / ParentScaleAbsY;
+				}
+
+				// Don't let the uv indicator be overly big if the light card itself is big.
+				constexpr double MaxAllowedParentScale = 0.5;
+				const double ClampScaleFactor = FMath::Clamp(MaxAllowedParentScale / ParentScaleMin, 0, 1);
+
+				UVIndicatorComponent->SetRelativeScale3D(UVScale * ClampScaleFactor * ScaleFactor);
+			}
+		}
+	}
+#endif // WITH_EDITOR
+}
+
 void ADisplayClusterLightCardActor::ShowLightCardLabel(bool bValue, float ScaleValue, ADisplayClusterRootActor* InRootActor)
 {
 #if WITH_EDITOR
diff --git a/Engine/Plugins/Runtime/nDisplay/Source/DisplayCluster/Private/Render/Viewport/LightCard/DisplayClusterViewportLightCardManager.cpp b/Engine/Plugins/Runtime/nDisplay/Source/DisplayCluster/Private/Render/Viewport/LightCard/DisplayClusterViewportLightCardManager.cpp
--- a/Engine/Plugins/Runtime/nDisplay/Source/DisplayCluster/Private/Render/Viewport/LightCard/DisplayClusterViewportLightCardManager.cpp
+++ b/Engine/Plugins/Runtime/nDisplay/Source/DisplayCluster/Private/Render/Viewport/LightCard/DisplayClusterViewportLightCardManager.cpp
@@ -118,6 +118,8 @@ void FDisplayClusterViewportLightCardManager::RenderFrame()
 			/** A list of primitive components that have been added to the preview scene for rendering in the current frame */
 			TArray<UPrimitiveComponent*> LoadedPrimitiveComponents;
 
+			TArray<UStaticMeshComponent*> LightCardMeshComponents;
+
 			bool bLoadedPrimitives = false;
 			for (ADisplayClusterLightCardActor* LightCard : UVLightCards)
 			{
@@ -126,13 +128,18 @@ void FDisplayClusterViewportLightCardManager::RenderFrame()
 					continue;
 				}
 
-				UStaticMeshComponent* LightCardMeshComp = LightCard->GetLightCardMeshComponent();
-				if (LightCardMeshComp && LightCardMeshComp->SceneProxy == nullptr)
+				LightCardMeshComponents.Empty(LightCardMeshComponents.Num());
+				LightCard->GetLightCardMeshComponents(LightCardMeshComponents);
+
+				for (UStaticMeshComponent* LightCardMeshComp : LightCardMeshComponents)
 				{
-					SceneInterface->AddPrimitive(LightCardMeshComp);
-					LoadedPrimitiveComponents.Add(LightCardMeshComp);
+					if (LightCardMeshComp && LightCardMeshComp->SceneProxy == nullptr)
+					{
+						SceneInterface->AddPrimitive(LightCardMeshComp);
+						LoadedPrimitiveComponents.Add(LightCardMeshComp);
 
-					bLoadedPrimitives = true;
+						bLoadedPrimitives = true;
+					}
 				}
 			}
 
diff --git a/Engine/Plugins/Runtime/nDisplay/Source/DisplayCluster/Public/DisplayClusterLightCardActor.h b/Engine/Plugins/Runtime/nDisplay/Source/DisplayCluster/Public/DisplayClusterLightCardActor.h
--- a/Engine/Plugins/Runtime/nDisplay/Source/DisplayCluster/Public/DisplayClusterLightCardActor.h
+++ b/Engine/Plugins/Runtime/nDisplay/Source/DisplayCluster/Public/DisplayClusterLightCardActor.h
@@ -73,8 +73,8 @@ public:
 	virtual FName GetCustomIconName() const override;
 #endif
 
-	/** Gets the light card mesh component */
-	UStaticMeshComponent* GetLightCardMeshComponent() const { return LightCardComponent.Get(); }
+	/** Gets the light card mesh components */
+	void GetLightCardMeshComponents(TArray<UStaticMeshComponent*>& MeshComponents) const;
 
 	/** Returns the current static mesh used by this light card */
 	UStaticMesh* GetStaticMesh() const;
@@ -91,6 +91,9 @@ public:
 	/** Updates the light card actor visibility */
 	void UpdateLightCardVisibility();
 
+	/** Updates the UV Indicator */
+	void UpdateUVIndicator();
+
 	/** Show or hide the light card label  */
 	void ShowLightCardLabel(bool bValue, float ScaleValue, ADisplayClusterRootActor* InRootActor);
 
@@ -263,6 +266,11 @@ protected:
 	UPROPERTY(transient, BlueprintReadOnly, Category = "Default")
 	TObjectPtr<UDisplayClusterLabelComponent> LabelComponent;
 
+#if WITH_EDITORONLY_DATA
+	UPROPERTY(transient, BlueprintReadOnly, Category = "Default")
+	TObjectPtr<UStaticMeshComponent> UVIndicatorComponent;
+#endif
+
 	/** The current owner of the light card */
 	TWeakObjectPtr<ADisplayClusterRootActor> RootActorOwner;
 
