nDisplay: GetMeshComponents now returns MeshComponents instead of static mesh components.

Note: Preflights were not available at the time.

diff --git a/Engine/Plugins/Runtime/nDisplay/Source/DisplayCluster/Private/Game/EngineClasses/Scene/DisplayClusterLabelComponent.cpp b/Engine/Plugins/Runtime/nDisplay/Source/DisplayCluster/Private/Game/EngineClasses/Scene/DisplayClusterLabelComponent.cpp
--- a/Engine/Plugins/Runtime/nDisplay/Source/DisplayCluster/Private/Game/EngineClasses/Scene/DisplayClusterLabelComponent.cpp
+++ b/Engine/Plugins/Runtime/nDisplay/Source/DisplayCluster/Private/Game/EngineClasses/Scene/DisplayClusterLabelComponent.cpp
@@ -63,6 +63,11 @@ void UDisplayClusterLabelComponent::OnRegister()
 	WidgetComponent->SetWidgetClass(WidgetClass.LoadSynchronous());
 }
 
+UDisplayClusterWidgetComponent* UDisplayClusterLabelComponent::GetWidgetComponent() const
+{
+	return WidgetComponent.Get();
+}
+
 #if WITH_EDITOR
 
 void UDisplayClusterLabelComponent::PostEditChangeProperty(FPropertyChangedEvent& PropertyChangedEvent)
diff --git a/Engine/Plugins/Runtime/nDisplay/Source/DisplayCluster/Private/Game/EngineClasses/Scene/DisplayClusterLightCardActor.cpp b/Engine/Plugins/Runtime/nDisplay/Source/DisplayCluster/Private/Game/EngineClasses/Scene/DisplayClusterLightCardActor.cpp
--- a/Engine/Plugins/Runtime/nDisplay/Source/DisplayCluster/Private/Game/EngineClasses/Scene/DisplayClusterLightCardActor.cpp
+++ b/Engine/Plugins/Runtime/nDisplay/Source/DisplayCluster/Private/Game/EngineClasses/Scene/DisplayClusterLightCardActor.cpp
@@ -194,7 +194,7 @@ FName ADisplayClusterLightCardActor::GetCustomIconName() const
 
 #endif
 
-void ADisplayClusterLightCardActor::GetLightCardMeshComponents(TArray<UStaticMeshComponent*>& MeshComponents) const
+void ADisplayClusterLightCardActor::GetLightCardMeshComponents(TArray<UMeshComponent*>& MeshComponents) const
 {
 	MeshComponents.Add(LightCardComponent.Get());
 
diff --git a/Engine/Plugins/Runtime/nDisplay/Source/DisplayCluster/Private/Render/Viewport/LightCard/DisplayClusterViewportLightCardManager.cpp b/Engine/Plugins/Runtime/nDisplay/Source/DisplayCluster/Private/Render/Viewport/LightCard/DisplayClusterViewportLightCardManager.cpp
--- a/Engine/Plugins/Runtime/nDisplay/Source/DisplayCluster/Private/Render/Viewport/LightCard/DisplayClusterViewportLightCardManager.cpp
+++ b/Engine/Plugins/Runtime/nDisplay/Source/DisplayCluster/Private/Render/Viewport/LightCard/DisplayClusterViewportLightCardManager.cpp
@@ -118,7 +118,7 @@ void FDisplayClusterViewportLightCardManager::RenderFrame()
 			/** A list of primitive components that have been added to the preview scene for rendering in the current frame */
 			TArray<UPrimitiveComponent*> LoadedPrimitiveComponents;
 
-			TArray<UStaticMeshComponent*> LightCardMeshComponents;
+			TArray<UMeshComponent*> LightCardMeshComponents;
 
 			bool bLoadedPrimitives = false;
 			for (ADisplayClusterLightCardActor* LightCard : UVLightCards)
@@ -131,7 +131,7 @@ void FDisplayClusterViewportLightCardManager::RenderFrame()
 				LightCardMeshComponents.Empty(LightCardMeshComponents.Num());
 				LightCard->GetLightCardMeshComponents(LightCardMeshComponents);
 
-				for (UStaticMeshComponent* LightCardMeshComp : LightCardMeshComponents)
+				for (UMeshComponent* LightCardMeshComp : LightCardMeshComponents)
 				{
 					if (LightCardMeshComp && LightCardMeshComp->SceneProxy == nullptr)
 					{
diff --git a/Engine/Plugins/Runtime/nDisplay/Source/DisplayCluster/Public/Components/DisplayClusterLabelComponent.h b/Engine/Plugins/Runtime/nDisplay/Source/DisplayCluster/Public/Components/DisplayClusterLabelComponent.h
--- a/Engine/Plugins/Runtime/nDisplay/Source/DisplayCluster/Public/Components/DisplayClusterLabelComponent.h
+++ b/Engine/Plugins/Runtime/nDisplay/Source/DisplayCluster/Public/Components/DisplayClusterLabelComponent.h
@@ -7,8 +7,9 @@
 
 #include "DisplayClusterLabelComponent.generated.h"
 
-class UDisplayClusterLabelWidget;
 class ADisplayClusterRootActor;
+class UDisplayClusterLabelWidget;
+class UDisplayClusterWidgetComponent;
 class UTextRenderComponent;
 class UUserWidget;
 class UWidgetComponent;
@@ -35,8 +36,12 @@ public:
 	
 	virtual void TickComponent(float DeltaTime, ELevelTick TickType,
 	                           FActorComponentTickFunction* ThisTickFunction) override;
+
 	virtual void OnRegister() override;
 
+	/** Returns the widget component */
+	UDisplayClusterWidgetComponent* GetWidgetComponent() const;
+
 #if WITH_EDITOR
 public:
 	virtual void PostEditChangeProperty(FPropertyChangedEvent& PropertyChangedEvent) override;
@@ -56,7 +61,7 @@ private:
 	
 	/** The widget component to display for this label */
 	UPROPERTY(VisibleAnywhere, Category=LabelText)
-	TObjectPtr<class UDisplayClusterWidgetComponent> WidgetComponent;
+	TObjectPtr<UDisplayClusterWidgetComponent> WidgetComponent;
 
 	/** The widget class to apply to the widget component */
 	UPROPERTY(EditAnywhere, Category=LabelText)
diff --git a/Engine/Plugins/Runtime/nDisplay/Source/DisplayCluster/Public/DisplayClusterLightCardActor.h b/Engine/Plugins/Runtime/nDisplay/Source/DisplayCluster/Public/DisplayClusterLightCardActor.h
--- a/Engine/Plugins/Runtime/nDisplay/Source/DisplayCluster/Public/DisplayClusterLightCardActor.h
+++ b/Engine/Plugins/Runtime/nDisplay/Source/DisplayCluster/Public/DisplayClusterLightCardActor.h
@@ -11,6 +11,7 @@
 class ADisplayClusterRootActor;
 class UActorComponent;
 class UDisplayClusterLabelComponent;
+class UMeshComponent;
 class USceneComponent;
 class USpringArmComponent;
 class UStaticMeshComponent;
@@ -74,7 +75,7 @@ public:
 #endif
 
 	/** Gets the light card mesh components */
-	void GetLightCardMeshComponents(TArray<UStaticMeshComponent*>& MeshComponents) const;
+	void GetLightCardMeshComponents(TArray<UMeshComponent*>& MeshComponents) const;
 
 	/** Returns the current static mesh used by this light card */
 	UStaticMesh* GetStaticMesh() const;
