nDisplay: Bubbled up translucency sort order of lightcard mesh component to light card actor.

Note: Preflights were unavailable.

diff --git a/Engine/Plugins/Runtime/nDisplay/Source/DisplayCluster/Public/DisplayClusterLightCardActor.h b/Engine/Plugins/Runtime/nDisplay/Source/DisplayCluster/Public/DisplayClusterLightCardActor.h
--- a/Engine/Plugins/Runtime/nDisplay/Source/DisplayCluster/Public/DisplayClusterLightCardActor.h
+++ b/Engine/Plugins/Runtime/nDisplay/Source/DisplayCluster/Public/DisplayClusterLightCardActor.h
@@ -257,18 +257,18 @@ protected:
 	UPROPERTY(EditDefaultsOnly, BlueprintReadOnly, Category = "Default")
 	TObjectPtr<USceneComponent> LightCardTransformerComponent;
 
-	UPROPERTY(EditDefaultsOnly, BlueprintReadOnly, Category = "Default")
+	UPROPERTY(EditAnywhere, BlueprintReadOnly, Category = "Default")
 	TObjectPtr<UStaticMeshComponent> LightCardComponent;
 
 	/** Components added by the IDisplayLightCardActorExtender */
 	UPROPERTY(VisibleAnywhere, Category = "Default")
 	TMap<FName, TObjectPtr<UActorComponent>> ExtenderNameToComponentMap;
 
-	UPROPERTY(transient, BlueprintReadOnly, Category = "Default")
+	UPROPERTY(Transient, BlueprintReadOnly, Category = "Default")
 	TObjectPtr<UDisplayClusterLabelComponent> LabelComponent;
 
 #if WITH_EDITORONLY_DATA
-	UPROPERTY(transient, BlueprintReadOnly, Category = "Default")
+	UPROPERTY(Transient, BlueprintReadOnly, Category = "Default")
 	TObjectPtr<UStaticMeshComponent> UVIndicatorComponent;
 #endif
 
diff --git a/Engine/Plugins/Runtime/nDisplay/Source/DisplayClusterLightCardEditor/Private/DetailCustomizations/DisplayClusterLightCardActorDetails.cpp b/Engine/Plugins/Runtime/nDisplay/Source/DisplayClusterLightCardEditor/Private/DetailCustomizations/DisplayClusterLightCardActorDetails.cpp
--- a/Engine/Plugins/Runtime/nDisplay/Source/DisplayClusterLightCardEditor/Private/DetailCustomizations/DisplayClusterLightCardActorDetails.cpp
+++ b/Engine/Plugins/Runtime/nDisplay/Source/DisplayClusterLightCardEditor/Private/DetailCustomizations/DisplayClusterLightCardActorDetails.cpp
@@ -17,85 +17,103 @@ TSharedRef<IDetailCustomization> FDisplayClusterLightCardActorDetails::MakeInsta
 }
 
 void FDisplayClusterLightCardActorDetails::CustomizeDetails(IDetailLayoutBuilder& InLayoutBuilder)
-{		
-	// Add the detail views of components added for Light Card Actor Extenders
-	const TSharedRef<IPropertyHandle> ExtenderNameToComponentMapHandle = InLayoutBuilder.GetProperty(ADisplayClusterLightCardActor::GetExtenderNameToComponentMapMemberName());
-	ExtenderNameToComponentMapHandle->MarkHiddenByCustomization();
+{
+	// Expose translucency sort priority property of the light card component
+	{
+		TSharedRef<IPropertyHandle> LightCardComponentPropertyHandle = InLayoutBuilder.GetProperty(TEXT("LightCardComponent"));
+		check(LightCardComponentPropertyHandle->IsValidHandle());
 
-	TArray<const void*> RawDatas;
-	ExtenderNameToComponentMapHandle->AccessRawData(RawDatas);
-	if (!RawDatas.IsEmpty())
+		TSharedPtr<IPropertyHandle> TranslucencySortPriorityPropertyHandle = LightCardComponentPropertyHandle->GetChildHandle(TEXT("TranslucencySortPriority"));
+		check(TranslucencySortPriorityPropertyHandle->IsValidHandle());
+
+		InLayoutBuilder
+			.EditCategory(TEXT("Rendering"))
+			.AddProperty(TranslucencySortPriorityPropertyHandle);
+	}
+
+	// Add the detail views of components added for Light Card Actor Extenders
 	{
-		// Remember the native categories, so extender categories can be ordered bottom-most later
-		TArray<FName> NativeCategories;
-		InLayoutBuilder.GetCategoryNames(NativeCategories);
-
-		IModularFeatures& ModularFeatures = IModularFeatures::Get();
-		const TArray<IDisplayClusterLightCardActorExtender*> Extenders = ModularFeatures.GetModularFeatureImplementations<IDisplayClusterLightCardActorExtender>(IDisplayClusterLightCardActorExtender::ModularFeatureName);
-
-		FName Category;
-		bool bShouldShowSubcategories = false;
-		TArray<UObject*> Components;
-		TArray<FName> ExtenderCategories;
-		for (const void* RawData : RawDatas)
+		const TSharedRef<IPropertyHandle> ExtenderNameToComponentMapHandle = InLayoutBuilder.GetProperty(ADisplayClusterLightCardActor::GetExtenderNameToComponentMapMemberName());
+		ExtenderNameToComponentMapHandle->MarkHiddenByCustomization();
+
+		TArray<const void*> RawDatas;
+		ExtenderNameToComponentMapHandle->AccessRawData(RawDatas);
+
+		if (!RawDatas.IsEmpty())
 		{
-			TMap<FName, UActorComponent*>* ExtenderNameToComponentMapPtr = (TMap<FName, UActorComponent*>*)(RawData);
-			check(ExtenderNameToComponentMapPtr);
+			// Remember the native categories, so extender categories can be ordered bottom-most later
+			TArray<FName> NativeCategories;
+			InLayoutBuilder.GetCategoryNames(NativeCategories);
+
+			IModularFeatures& ModularFeatures = IModularFeatures::Get();
+			const TArray<IDisplayClusterLightCardActorExtender*> Extenders = ModularFeatures.GetModularFeatureImplementations<IDisplayClusterLightCardActorExtender>(IDisplayClusterLightCardActorExtender::ModularFeatureName);
 
-			for (const TTuple<FName, UActorComponent*>& ExtenderNameToComponentPair : *ExtenderNameToComponentMapPtr)
+			FName Category;
+			bool bShouldShowSubcategories = false;
+			TArray<UObject*> Components;
+			TArray<FName> ExtenderCategories;
+
+			for (const void* RawData : RawDatas)
 			{
-				if (!ensureMsgf(ExtenderNameToComponentPair.Value, TEXT("Trying to display component for Extender %s, but component is invalid."), *ExtenderNameToComponentPair.Key.ToString()))
-				{
-					continue;
-				}
+				TMap<FName, UActorComponent*>* ExtenderNameToComponentMapPtr = (TMap<FName, UActorComponent*>*)(RawData);
+				check(ExtenderNameToComponentMapPtr);
 
-				IDisplayClusterLightCardActorExtender* const* ExtenderPtr = 
-					Algo::FindByPredicate(Extenders, [ExtenderNameToComponentPair](const IDisplayClusterLightCardActorExtender* Extender)
-						{
-							return Extender->GetExtenderName() == ExtenderNameToComponentPair.Key;
-						});
-				if (!ensureMsgf(ExtenderPtr, TEXT("Cannot find Extender %s expected to exist for component %s."), *ExtenderNameToComponentPair.Key.ToString(), *ExtenderNameToComponentPair.Value->GetName()))
+				for (const TTuple<FName, UActorComponent*>& ExtenderNameToComponentPair : *ExtenderNameToComponentMapPtr)
 				{
-					continue;
-				}
+					if (!ensureMsgf(ExtenderNameToComponentPair.Value, TEXT("Trying to display component for Extender %s, but component is invalid."), *ExtenderNameToComponentPair.Key.ToString()))
+					{
+						continue;
+					}
 
-				bShouldShowSubcategories = (*ExtenderPtr)->ShouldShowSubcategories();
-				Components.Add(ExtenderNameToComponentPair.Value);
+					IDisplayClusterLightCardActorExtender* const* ExtenderPtr =
+						Algo::FindByPredicate(Extenders, [ExtenderNameToComponentPair](const IDisplayClusterLightCardActorExtender* Extender)
+							{
+								return Extender->GetExtenderName() == ExtenderNameToComponentPair.Key;
+							});
 
-				Category = (*ExtenderPtr)->GetCategory();
-				if (!NativeCategories.Contains(Category))
-				{
-					ExtenderCategories.AddUnique(Category);
+					if (!ensureMsgf(ExtenderPtr, TEXT("Cannot find Extender %s expected to exist for component %s."), *ExtenderNameToComponentPair.Key.ToString(), *ExtenderNameToComponentPair.Value->GetName()))
+					{
+						continue;
+					}
+
+					bShouldShowSubcategories = (*ExtenderPtr)->ShouldShowSubcategories();
+					Components.Add(ExtenderNameToComponentPair.Value);
+
+					Category = (*ExtenderPtr)->GetCategory();
+					if (!NativeCategories.Contains(Category))
+					{
+						ExtenderCategories.AddUnique(Category);
+					}
 				}
 			}
-		}
 
-		FAddPropertyParams AddPropertyParams;
-		AddPropertyParams.CreateCategoryNodes(bShouldShowSubcategories);
-		AddPropertyParams.HideRootObjectNode(true);
+			FAddPropertyParams AddPropertyParams;
+			AddPropertyParams.CreateCategoryNodes(bShouldShowSubcategories);
+			AddPropertyParams.HideRootObjectNode(true);
 
-		InLayoutBuilder
-			.EditCategory(Category)
-			.AddExternalObjects(Components, EPropertyLocation::Default, AddPropertyParams);
+			InLayoutBuilder
+				.EditCategory(Category)
+				.AddExternalObjects(Components, EPropertyLocation::Default, AddPropertyParams);
 
-		// Sort extender categories bottom-most
-		InLayoutBuilder.SortCategories([ExtenderCategories](const TMap<FName, IDetailCategoryBuilder*>& CategoryMap)
-			{
-				int32 NativeCategorySortOrder = 1;
-				int32 ExtenderCategorySortOrder = CategoryMap.Num();
-				for (const TTuple<FName, IDetailCategoryBuilder*>& CategoryPair : CategoryMap)
+			// Sort extender categories bottom-most
+			InLayoutBuilder.SortCategories([ExtenderCategories](const TMap<FName, IDetailCategoryBuilder*>& CategoryMap)
 				{
-					if (ExtenderCategories.Contains(CategoryPair.Key))
-					{
-						CategoryPair.Value->SetSortOrder(ExtenderCategorySortOrder);
-						ExtenderCategorySortOrder--;
-					}
-					else
+					int32 NativeCategorySortOrder = 1;
+					int32 ExtenderCategorySortOrder = CategoryMap.Num();
+					for (const TTuple<FName, IDetailCategoryBuilder*>& CategoryPair : CategoryMap)
 					{
-						CategoryPair.Value->SetSortOrder(NativeCategorySortOrder);
-						NativeCategorySortOrder++;
+						if (ExtenderCategories.Contains(CategoryPair.Key))
+						{
+							CategoryPair.Value->SetSortOrder(ExtenderCategorySortOrder);
+							ExtenderCategorySortOrder--;
+						}
+						else
+						{
+							CategoryPair.Value->SetSortOrder(NativeCategorySortOrder);
+							NativeCategorySortOrder++;
+						}
 					}
-				}
-			});
+				});
+		}
 	}
 }
