nDisplay: Add File/Edit/Window menus to the light card editor.

- Operator panel now supports registering IDisplayClusterOperatorApp delegates.
- Operator panel now supports registering menu extensions.
- Refactor SDisplayClusterLightCardEditor to FDisplayClusterLightCardEditor. This inherits the new IDisplayClusterOperatorApp interface.
  - Viewport is now its own tab. Light card editor now manages actor selections instead of relying on the outliner.
- "Add Existing Light Card" is now "Import Existing Content" and only available from the File Menu.

diff --git a/Engine/Plugins/Runtime/nDisplay/Source/DisplayClusterLightCardEditor/Private/DisplayClusterLightCardEditor.cpp b/Engine/Plugins/Runtime/nDisplay/Source/DisplayClusterLightCardEditor/Private/DisplayClusterLightCardEditor.cpp
new file mode 100644
--- /dev/null
+++ b/Engine/Plugins/Runtime/nDisplay/Source/DisplayClusterLightCardEditor/Private/DisplayClusterLightCardEditor.cpp
@@ -0,0 +1,1547 @@
+// Copyright Epic Games, Inc. All Rights Reserved.
+
+#include "DisplayClusterLightCardEditor.h"
+
+#include "DisplayClusterLightCardEditorCommands.h"
+#include "DisplayClusterLightCardEditorUtils.h"
+#include "DisplayClusterLightCardEditorStyle.h"
+#include "IDisplayClusterLightCardEditor.h"
+#include "LightCardTemplates/DisplayClusterLightCardTemplate.h"
+#include "LightCardTemplates/DisplayClusterLightCardTemplateHelpers.h"
+#include "LightCardTemplates/SDisplayClusterLightCardTemplateList.h"
+#include "Outliner/SDisplayClusterLightCardOutliner.h"
+#include "Settings/DisplayClusterLightCardEditorSettings.h"
+
+#include "Viewport/DisplayClusterLightcardEditorViewport.h"
+#include "Viewport/DisplayClusterLightCardEditorViewportClient.h"
+
+#include "IDisplayClusterOperator.h"
+#include "IDisplayClusterOperatorViewModel.h"
+
+#include "DisplayClusterConfigurationTypes.h"
+
+#include "DisplayClusterRootActor.h"
+#include "Blueprints/DisplayClusterBlueprintLib.h"
+#include "Components/DisplayClusterCameraComponent.h"
+
+#include "ContentBrowserModule.h"
+#include "FileHelpers.h"
+#include "IContentBrowserSingleton.h"
+#include "LevelEditorActions.h"
+#include "ObjectEditorUtils.h"
+#include "PropertyCustomizationHelpers.h"
+#include "ScopedTransaction.h"
+#include "Selection.h"
+#include "AssetRegistry/AssetRegistryModule.h"
+#include "Framework/Commands/GenericCommands.h"
+#include "Framework/Docking/LayoutExtender.h"
+#include "Framework/Docking/TabManager.h"
+#include "Framework/MultiBox/MultiBoxBuilder.h"
+#include "Framework/MultiBox/MultiBoxExtender.h"
+#include "Misc/FileHelper.h"
+#include "Misc/TransactionObjectEvent.h"
+#include "Styling/SlateIconFinder.h"
+#include "Toolkits/AssetEditorToolkit.h"
+#include "Widgets/Docking/SDockTab.h"
+#include "Widgets/Input/SNumericEntryBox.h"
+#include "Widgets/Workflow/SWizard.h"
+
+#define LOCTEXT_NAMESPACE "FDisplayClusterLightCardEditor"
+
+const FName FDisplayClusterLightCardEditor::ViewportTabName = TEXT("DisplayClusterLightCardViewportTab");
+const FName FDisplayClusterLightCardEditor::OutlinerTabName = TEXT("DisplayClusterLightCardOutlinerTab");
+
+TSharedRef<IDisplayClusterOperatorApp> FDisplayClusterLightCardEditor::MakeInstance(
+	TSharedRef<IDisplayClusterOperatorViewModel> InViewModel)
+{
+	TSharedRef<FDisplayClusterLightCardEditor> Instance = MakeShared<FDisplayClusterLightCardEditor>();
+	Instance->Initialize(InViewModel); // Initialize necessary to perform any SP binding after SP construction
+	return Instance;
+}
+
+FDisplayClusterLightCardEditor::~FDisplayClusterLightCardEditor()
+{
+	UnregisterTabSpawners();
+	
+	IDisplayClusterOperator::Get().GetOperatorViewModel()->OnActiveRootActorChanged().Remove(ActiveRootActorChangedHandle);
+	FCoreUObjectDelegates::OnObjectPropertyChanged.RemoveAll(this);
+
+	if (GEngine != nullptr)
+	{
+		GEngine->OnLevelActorDeleted().RemoveAll(this);
+	}
+
+	if (OnObjectTransactedHandle.IsValid())
+	{
+		FCoreUObjectDelegates::OnObjectTransacted.Remove(OnObjectTransactedHandle);
+	}
+
+	RemoveCompileDelegates();
+}
+
+void FDisplayClusterLightCardEditor::PostUndo(bool bSuccess)
+{
+	FEditorUndoClient::PostUndo(bSuccess);
+	RefreshLabels();
+}
+
+void FDisplayClusterLightCardEditor::PostRedo(bool bSuccess)
+{
+	FEditorUndoClient::PostRedo(bSuccess);
+	RefreshLabels();
+}
+
+void FDisplayClusterLightCardEditor::Initialize(TSharedRef<IDisplayClusterOperatorViewModel> InViewModel)
+{
+	OperatorViewModel = InViewModel;
+
+	ActiveRootActorChangedHandle = InViewModel->OnActiveRootActorChanged().AddSP(this, &FDisplayClusterLightCardEditor::OnActiveRootActorChanged);
+	if (GEngine != nullptr)
+	{
+		GEngine->OnLevelActorDeleted().AddSP(this, &FDisplayClusterLightCardEditor::OnLevelActorDeleted);
+	}
+	
+	OnObjectTransactedHandle = FCoreUObjectDelegates::OnObjectTransacted.AddSP(this, &FDisplayClusterLightCardEditor::OnObjectTransacted);
+
+	BindCommands();
+	RegisterTabSpawners();
+	RegisterToolbarExtensions();
+	RegisterMenuExtensions();
+
+	BindCompileDelegates();
+
+	RefreshLabels();
+
+	GEditor->RegisterForUndo(this);
+}
+
+TArray<AActor*> FDisplayClusterLightCardEditor::FindAllManagedActors() const
+{
+	TArray<AActor*> ManagedActors;
+	if (ActiveRootActor.IsValid())
+	{
+		// First populate light cards
+		TSet<ADisplayClusterLightCardActor*> RootActorLightCardActors;
+		UDisplayClusterBlueprintLib::FindLightCardsForRootActor(ActiveRootActor.Get(), RootActorLightCardActors);
+
+		for (TObjectPtr<ADisplayClusterLightCardActor> LightCardActor : RootActorLightCardActors)
+		{
+			ManagedActors.Add(LightCardActor);
+		}
+		
+		// Find other actors
+		if (UWorld* World = ActiveRootActor->GetWorld())
+		{
+			for (const TWeakObjectPtr<AActor> WeakActor : TActorRange<AActor>(World))
+			{
+				if (WeakActor.IsValid() && !WeakActor->IsA<ADisplayClusterLightCardActor>() &&
+					UE::DisplayClusterLightCardEditorUtils::IsManagedActor(WeakActor.Get()))
+				{
+					ManagedActors.Add(WeakActor.Get());
+				}
+			}
+		}
+	}
+
+	return ManagedActors;
+}
+
+void FDisplayClusterLightCardEditor::SelectActors(const TArray<AActor*>& ActorsToSelect)
+{
+	if (LightCardOutliner)
+	{
+		LightCardOutliner->SelectActors(ActorsToSelect);
+	}
+
+	if (ActorsToSelect.Num() > 0)
+	{
+		SelectedActors = FObjectEditorUtils::GetTypedWeakObjectPtrs<AActor>(static_cast<TArray<UObject*>>(ActorsToSelect));
+	}
+	else
+	{
+		SelectedActors.Empty();
+	}
+	
+	IDisplayClusterOperator::Get().ShowDetailsForObjects(*reinterpret_cast<const TArray<UObject*>*>(&ActorsToSelect));
+}
+
+void FDisplayClusterLightCardEditor::SelectActorProxies(const TArray<AActor*>& ActorsToSelect)
+{
+	if (ViewportView)
+	{
+		ViewportView->GetLightCardEditorViewportClient()->SelectActors(ActorsToSelect);
+	}
+}
+
+void FDisplayClusterLightCardEditor::GetSelectedActors(TArray<AActor*>& OutSelectedActors)
+{
+	OutSelectedActors = GetSelectedActorsAs<AActor>();
+}
+
+void FDisplayClusterLightCardEditor::CenterActorInView(AActor* Actor)
+{
+	check(Actor);
+	if (ViewportView)
+	{
+		ViewportView->GetLightCardEditorViewportClient()->CenterActorInView(Actor);
+	}
+}
+
+AActor* FDisplayClusterLightCardEditor::SpawnActor(TSubclassOf<AActor> InActorClass, const FName& InActorName)
+{
+	if (!ActiveRootActor.IsValid())
+	{
+		return nullptr;
+	}
+
+	FScopedTransaction Transaction(LOCTEXT("SpawnActorTransactionMessage", "Spawn Actor"));
+	
+	check(InActorClass);
+
+	const FVector SpawnLocation = ActiveRootActor->GetDefaultCamera()->GetComponentLocation();
+	FRotator SpawnRotation = ActiveRootActor->GetDefaultCamera()->GetComponentRotation();
+	SpawnRotation.Yaw -= 180.f;
+
+	FActorSpawnParameters SpawnParameters;
+	SpawnParameters.bNoFail = true;
+	SpawnParameters.SpawnCollisionHandlingOverride = ESpawnActorCollisionHandlingMethod::AdjustIfPossibleButAlwaysSpawn;
+	SpawnParameters.Name = InActorName;
+	SpawnParameters.NameMode = FActorSpawnParameters::ESpawnActorNameMode::Requested;
+	SpawnParameters.OverrideLevel = ActiveRootActor->GetWorld()->GetCurrentLevel();
+
+	AActor* NewActor = CastChecked<AActor>(
+		ActiveRootActor->GetWorld()->SpawnActor(InActorClass,
+			&SpawnLocation, &SpawnRotation, MoveTemp(SpawnParameters)));
+
+	NewActor->SetActorLabel(NewActor->GetName());
+
+	if (ADisplayClusterLightCardActor* NewLightCard = Cast<ADisplayClusterLightCardActor>(NewActor))
+	{
+		if (ViewportView.IsValid())
+		{
+			const EDisplayClusterMeshProjectionType ProjectionMode = ViewportView->GetLightCardEditorViewportClient()->GetProjectionMode();
+			if (ProjectionMode == EDisplayClusterMeshProjectionType::UV)
+			{
+				NewLightCard->bIsUVLightCard = true;
+				NewLightCard->Scale /= 4;
+				NewLightCard->Feathering = 0.05; // Just enough to avoid jagged look on UV lightcards.
+			}
+		}
+		const TArray<ADisplayClusterLightCardActor*> LightCards { NewLightCard } ;
+		AddLightCardsToActor(LightCards);
+	}
+	else
+	{
+		RefreshPreviewActors(EDisplayClusterLightCardEditorProxyType::LightCards);
+	}
+	
+	SelectActors({ NewActor });
+
+	return NewActor;
+}
+
+ADisplayClusterLightCardActor* FDisplayClusterLightCardEditor::SpawnLightCardFromTemplate(
+	const UDisplayClusterLightCardTemplate* InTemplate, ULevel* InLevel, bool bIsPreview)
+{
+	if (!ActiveRootActor.IsValid())
+	{
+		return nullptr;
+	}
+
+	FScopedTransaction Transaction(LOCTEXT("SpawnLightCardFromTemplateTransactionMessage", "Spawn Light Card from Template"));
+	
+	check(InTemplate && InTemplate->LightCardActor);
+
+	ULevel* Level = InLevel ? InLevel : ActiveRootActor->GetWorld()->GetCurrentLevel();
+	
+	FName UniqueName = *InTemplate->GetName().Replace(TEXT("Template"), TEXT(""));
+	if (StaticFindObjectFast(InTemplate->LightCardActor->GetClass(), Level, UniqueName))
+	{
+		UniqueName = MakeUniqueObjectName(Level, InTemplate->LightCardActor->GetClass(), UniqueName);
+	}
+	
+	// Duplicate, don't copy properties or spawn from a template. Doing so will copy component data incorrectly,
+	// specifically the static mesh override textures. They will be parented to the template, not the level instance
+	// and prevent the map from saving.
+	ADisplayClusterLightCardActor* NewLightCard = CastChecked<ADisplayClusterLightCardActor>(StaticDuplicateObject(InTemplate->LightCardActor.Get(), Level, UniqueName));
+	Level->AddLoadedActor(NewLightCard);
+	
+	if (!bIsPreview)
+	{
+		NewLightCard->SetActorLabel(NewLightCard->GetName());
+	
+		const TArray<ADisplayClusterLightCardActor*> LightCards { NewLightCard } ;
+		AddLightCardsToActor(LightCards);
+
+		FDisplayClusterLightCardEditorRecentItem RecentlyPlacedItem;
+		RecentlyPlacedItem.ObjectPath = InTemplate;
+		RecentlyPlacedItem.ItemType = FDisplayClusterLightCardEditorRecentItem::Type_LightCardTemplate;
+		
+		AddRecentlyPlacedItem(MoveTemp(RecentlyPlacedItem));
+		SelectActors({NewLightCard});
+	}
+	
+	return NewLightCard;
+}
+
+void FDisplayClusterLightCardEditor::AddNewLightCard()
+{
+	check(ActiveRootActor.IsValid());
+
+	FScopedTransaction Transaction(LOCTEXT("AddNewLightCardTransactionMessage", "Add New Light Card"));
+
+	// When adding a new lightcard, usually the desired location is in the middle of the viewport
+	if (ADisplayClusterLightCardActor* NewLightCard = SpawnActorAs<ADisplayClusterLightCardActor>(TEXT("LightCard")))
+	{
+		CenterActorInView(NewLightCard);
+
+		FDisplayClusterLightCardEditorRecentItem RecentlyPlacedItem;
+		RecentlyPlacedItem.ObjectPath = NewLightCard->GetClass();
+		RecentlyPlacedItem.ItemType = FDisplayClusterLightCardEditorRecentItem::Type_LightCard;
+		
+		AddRecentlyPlacedItem(MoveTemp(RecentlyPlacedItem));
+		SelectActors({NewLightCard});
+	}
+}
+
+void FDisplayClusterLightCardEditor::AddExistingLightCard()
+{
+	TSharedPtr<SWindow> PickerWindow;
+	TWeakObjectPtr<ADisplayClusterLightCardActor> SelectedActorPtr;
+	bool bFinished = false;
+	
+	const TSharedRef<SWidget> ActorPicker = PropertyCustomizationHelpers::MakeActorPickerWithMenu(
+		nullptr,
+		false,
+		FOnShouldFilterActor::CreateLambda([&](const AActor* const InActor) -> bool // ActorFilter
+		{
+			const bool IsAllowed = InActor != nullptr && !InActor->IsChildActor() && InActor->IsA<ADisplayClusterLightCardActor>() &&
+				!InActor->GetClass()->HasAnyClassFlags(CLASS_Interface)	&& !InActor->IsA<ADisplayClusterRootActor>();
+			
+			return IsAllowed;
+		}),
+		FOnActorSelected::CreateLambda([&](AActor* InActor) -> void // OnSet
+		{
+			SelectedActorPtr = Cast<ADisplayClusterLightCardActor>(InActor);
+		}),
+		FSimpleDelegate::CreateLambda([&]() -> void // OnClose
+		{
+		}),
+		FSimpleDelegate::CreateLambda([&]() -> void // OnUseSelected
+		{
+			if (ADisplayClusterLightCardActor* Selection = Cast<ADisplayClusterLightCardActor>(GEditor->GetSelectedActors()->GetTop(ADisplayClusterLightCardActor::StaticClass())))
+			{
+				SelectedActorPtr = Selection;
+			}
+		}));
+	
+	PickerWindow = SNew(SWindow)
+	.Title(LOCTEXT("AddExistingLightCard", "Select an existing Light Card actor"))
+	.ClientSize(FVector2D(500, 525))
+	.SupportsMinimize(false) .SupportsMaximize(false)
+	[
+		SNew(SBorder)
+		.BorderImage(FAppStyle::GetBrush("Menu.Background"))
+		[
+			SNew(SWizard)
+			.FinishButtonText(LOCTEXT("FinishAddingExistingLightCard", "Add Actor"))
+			.OnCanceled(FSimpleDelegate::CreateLambda([&]()
+			{
+				if (PickerWindow.IsValid())
+				{
+					PickerWindow->RequestDestroyWindow();
+				}
+			}))
+			.OnFinished(FSimpleDelegate::CreateLambda([&]()
+			{
+				bFinished = true;
+				if (PickerWindow.IsValid())
+				{
+					PickerWindow->RequestDestroyWindow();
+				}
+			}))
+			.CanFinish(TAttribute<bool>::CreateLambda([&]()
+			{
+				return SelectedActorPtr.IsValid();
+			}))
+			.ShowPageList(false)
+			+SWizard::Page()
+			.CanShow(true)
+			[
+				SNew(SBorder)
+				.VAlign(VAlign_Fill)
+				.BorderImage(FAppStyle::GetBrush("Menu.Background"))
+				[
+					SNew(SVerticalBox)
+					+SVerticalBox::Slot()
+					.AutoHeight()
+					.VAlign(VAlign_Fill)
+					.HAlign(HAlign_Fill)
+					[
+						ActorPicker
+					]
+					+SVerticalBox::Slot()
+					.AutoHeight()
+					.VAlign(VAlign_Bottom)
+					.Padding(0.f, 8.f)
+					[
+						SNew(STextBlock)
+						.TextStyle(FAppStyle::Get(), "NormalText.Important")
+						.Text_Lambda([&]
+						{
+							const FString Result = FString::Printf(TEXT("Selected Actor: %s"),
+								SelectedActorPtr.IsValid() ? *SelectedActorPtr->GetActorLabel() : TEXT(""));
+							return FText::FromString(Result);
+						})
+					]
+				]
+			]
+		]
+	];
+
+	GEditor->EditorAddModalWindow(PickerWindow.ToSharedRef());
+	if (bFinished && SelectedActorPtr.IsValid())
+	{
+		FScopedTransaction Transaction(LOCTEXT("AddExistingLightCardTransactionMessage", "Add Existing Light Card"));
+
+		TArray<ADisplayClusterLightCardActor*> LightCards { SelectedActorPtr.Get() };
+		AddLightCardsToActor(LightCards);
+	}
+
+	PickerWindow.Reset();
+	SelectedActorPtr.Reset();
+}
+
+void FDisplayClusterLightCardEditor::AddNewFlag()
+{
+	check(ActiveRootActor.IsValid());
+
+	FScopedTransaction Transaction(LOCTEXT("AddNewFlagTransactionMessage", "Add New Flag"));
+
+	// When adding a new lightcard, usually the desired location is in the middle of the viewport
+	if (ADisplayClusterLightCardActor* NewLightCard = SpawnActorAs<ADisplayClusterLightCardActor>(TEXT("Flag")))
+	{
+		NewLightCard->Color = FLinearColor(0.f, 0.f, 0.f);
+		
+		CenterActorInView(NewLightCard);
+
+		FDisplayClusterLightCardEditorRecentItem RecentlyPlacedItem;
+		RecentlyPlacedItem.ObjectPath = NewLightCard->GetClass();
+		RecentlyPlacedItem.ItemType = FDisplayClusterLightCardEditorRecentItem::Type_Flag;
+		
+		AddRecentlyPlacedItem(MoveTemp(RecentlyPlacedItem));
+		SelectActors({NewLightCard});
+	}
+}
+
+void FDisplayClusterLightCardEditor::AddNewDynamic(UClass* InClass)
+{
+	check(ActiveRootActor.IsValid());
+
+	FScopedTransaction Transaction(LOCTEXT("AddNewActorMessage", "Add New Actor"));
+	
+	if (AActor* NewActor = SpawnActor(InClass))
+	{
+		CenterActorInView(NewActor);
+
+		FDisplayClusterLightCardEditorRecentItem RecentlyPlacedItem;
+		RecentlyPlacedItem.ObjectPath = NewActor->GetClass();
+		RecentlyPlacedItem.ItemType = FDisplayClusterLightCardEditorRecentItem::Type_Dynamic;
+		
+		AddRecentlyPlacedItem(MoveTemp(RecentlyPlacedItem));
+		SelectActors({NewActor});
+	}
+}
+
+void FDisplayClusterLightCardEditor::AddLightCardsToActor(const TArray<ADisplayClusterLightCardActor*>& LightCards)
+{
+	if (ActiveRootActor.IsValid())
+	{
+		UDisplayClusterConfigurationData* ConfigData = ActiveRootActor->GetConfigData();
+		ConfigData->Modify();
+		FDisplayClusterConfigurationICVFX_VisibilityList& RootActorLightCards = ConfigData->StageSettings.Lightcard.ShowOnlyList;
+
+		for (ADisplayClusterLightCardActor* LightCard : LightCards)
+		{
+			check(LightCard);
+
+			if (!RootActorLightCards.Actors.ContainsByPredicate([&](const TSoftObjectPtr<AActor>& Actor)
+				{
+					// Don't add if a loaded actor is already present.
+					return Actor.Get() == LightCard;
+				}))
+			{
+				LightCard->ShowLightCardLabel(ShouldShowLightCardLabels(), *GetLightCardLabelScale(), ActiveRootActor.Get());
+				
+				const TSoftObjectPtr<AActor> LightCardSoftObject(LightCard);
+
+				// Remove any exact paths to this actor. It's possible invalid actors are present if a light card
+				// was force deleted from a level.
+				RootActorLightCards.Actors.RemoveAll([&](const TSoftObjectPtr<AActor>& Actor)
+					{
+						return Actor == LightCardSoftObject;
+					});
+
+				RootActorLightCards.Actors.Add(LightCard);
+			}
+		}
+
+		RefreshPreviewActors(EDisplayClusterLightCardEditorProxyType::LightCards);
+	}
+}
+
+bool FDisplayClusterLightCardEditor::CanAddLightCard() const
+{
+	return ActiveRootActor.IsValid() && ActiveRootActor->GetWorld() != nullptr;
+}
+
+void FDisplayClusterLightCardEditor::CutSelectedActors()
+{
+	CopySelectedActors();
+
+	RemoveActors(SelectedActors, /*bDeleteActors*/ true);
+}
+
+bool FDisplayClusterLightCardEditor::CanCutSelectedActors()
+{
+	return CanCopySelectedActors() && CanRemoveSelectedActors();
+}
+
+void FDisplayClusterLightCardEditor::CopySelectedActors()
+{
+	UWorld* EditorWorld = GEditor->GetEditorWorldContext().World();
+
+	const bool bNoteSelectionChange = false;
+	const bool bDeselectBSPSurfs = true;
+	const bool bWarnAboutManyActors = false;
+	GEditor->SelectNone(bNoteSelectionChange, bDeselectBSPSurfs, bWarnAboutManyActors);
+
+	for (const TWeakObjectPtr<AActor>& LightCard : SelectedActors)
+	{
+		const bool bInSelected = true;
+		const bool bNotify = false;
+		const bool bSelectEvenIfHidden = true;
+		GEditor->SelectActor(LightCard.Get(), true, bNotify, bSelectEvenIfHidden);
+	}
+
+	const bool bShouldCut = false;
+	const bool bIsMove = false;
+	const bool bWarnAboutReferences = false;
+	GEditor->CopySelectedActorsToClipboard(EditorWorld, bShouldCut, bIsMove, bWarnAboutReferences);
+}
+
+bool FDisplayClusterLightCardEditor::CanCopySelectedActors() const
+{
+	return SelectedActors.Num() > 0;
+}
+
+void FDisplayClusterLightCardEditor::PasteActors(bool bOffsetActorPosition)
+{
+	UWorld* EditorWorld = GEditor->GetEditorWorldContext().World();
+	GEditor->PasteSelectedActorsFromClipboard(EditorWorld, LOCTEXT("PasteActorsTransactionMessage", "Paste Actors"), EPasteTo::PT_OriginalLocation);
+
+	TArray<AActor*> PastedActors;
+	TArray<ADisplayClusterLightCardActor*> PastedLightCards;
+	for (FSelectionIterator It(GEditor->GetSelectedActorIterator()); It; ++It)
+	{
+		if (AActor* Actor = Cast<AActor>(*It))
+		{
+			if (!Actor->Implements<UDisplayClusterStageActor>())
+			{
+				continue;
+			}
+
+			IDisplayClusterStageActor* StageActor = CastChecked<IDisplayClusterStageActor>(Actor);
+			PastedActors.Add(Actor);
+			if (bOffsetActorPosition)
+			{
+				// If the light card should be offset from its pasted location, offset its longitude and latitude by a number of
+				// degrees equal to an arc length of 10 units (arc length = angle in radians * radius)
+				const float AngleOffset = FMath::RadiansToDegrees(10.0f / FMath::Max(StageActor->GetDistanceFromCenter(), 1.0f));
+				StageActor->SetLatitude(StageActor->GetLatitude() - AngleOffset);
+				StageActor->SetLongitude(StageActor->GetLongitude() + AngleOffset);
+			}
+			
+			if (ADisplayClusterLightCardActor* LightCard = Cast<ADisplayClusterLightCardActor>(Actor))
+			{
+				PastedLightCards.Add(LightCard);
+			}
+		}
+	}
+
+	AddLightCardsToActor(PastedLightCards);
+	SelectActors(PastedActors);
+}
+
+bool FDisplayClusterLightCardEditor::CanPasteActors() const
+{
+	UWorld* EditorWorld = GEditor->GetEditorWorldContext().World();
+	return GEditor->CanPasteSelectedActorsFromClipboard(EditorWorld);
+}
+
+void FDisplayClusterLightCardEditor::DuplicateSelectedActors()
+{
+	CopySelectedActors();
+
+	const bool bOffsetActorPosition = true;
+	PasteActors(bOffsetActorPosition);
+}
+
+bool FDisplayClusterLightCardEditor::CanDuplicateSelectedActors() const
+{
+	return CanCopySelectedActors();
+}
+
+void FDisplayClusterLightCardEditor::RemoveSelectedActors(bool bDeleteLightCardActor)
+{
+	RemoveActors(SelectedActors, bDeleteLightCardActor);
+}
+
+void FDisplayClusterLightCardEditor::RemoveActors(
+	const TArray<TWeakObjectPtr<AActor>>& InActorsToRemove, bool bDeleteActors)
+{
+	FScopedTransaction Transaction(LOCTEXT("RemoveActorTransactionMessage", "Remove Light Card(s)"));
+
+	USelection* EdSelectionManager = GEditor->GetSelectedActors();
+	UWorld* WorldToUse = nullptr;
+
+	if (bDeleteActors)
+	{
+		EdSelectionManager->BeginBatchSelectOperation();
+		EdSelectionManager->Modify();
+		EdSelectionManager->DeselectAll();
+	}
+
+	if (ActiveRootActor.IsValid())
+	{
+		UDisplayClusterConfigurationData* ConfigData = ActiveRootActor->GetConfigData();
+		ConfigData->Modify();
+
+		FDisplayClusterConfigurationICVFX_VisibilityList& RootActorLightCards = ConfigData->StageSettings.Lightcard.ShowOnlyList;
+
+		for (const TWeakObjectPtr<AActor>& Actor : InActorsToRemove)
+		{
+			if (ADisplayClusterLightCardActor* LightCard = Cast<ADisplayClusterLightCardActor>(Actor.Get()))
+			{
+				RootActorLightCards.Actors.RemoveAll([&](const TSoftObjectPtr<AActor>& InActor)
+					{
+						return InActor.Get() == LightCard;
+					});
+
+				if (!bDeleteActors)
+				{
+					LightCard->ShowLightCardLabel(false, *GetLightCardLabelScale(), ActiveRootActor.Get());
+				}
+			}
+
+			if (bDeleteActors)
+			{
+				WorldToUse = Actor->GetWorld();
+				GEditor->SelectActor(Actor.Get(), /*bSelect =*/true, /*bNotifyForActor =*/false, /*bSelectEvenIfHidden =*/true);
+			}
+		}
+	}
+
+	if (bDeleteActors)
+	{
+		EdSelectionManager->EndBatchSelectOperation();
+
+		if (WorldToUse)
+		{
+			GEditor->edactDeleteSelected(WorldToUse);
+		}
+	}
+
+	RefreshPreviewActors(EDisplayClusterLightCardEditorProxyType::LightCards);
+}
+
+bool FDisplayClusterLightCardEditor::CanRemoveSelectedActors() const
+{
+	for (const TWeakObjectPtr<AActor>& Actor : SelectedActors)
+	{
+		if (Actor.IsValid())
+		{
+			return true;
+		}
+	}
+	
+	return false;
+}
+
+bool FDisplayClusterLightCardEditor::CanRemoveSelectedLightCardFromActor() const
+{
+	for (const TWeakObjectPtr<AActor>& Actor : SelectedActors)
+	{
+		if (Actor.IsValid() && Actor->IsA<ADisplayClusterLightCardActor>())
+		{
+			return true;
+		}
+	}
+	
+	return false;
+}
+
+void FDisplayClusterLightCardEditor::CreateLightCardTemplate()
+{
+	TArray<ADisplayClusterLightCardActor*> SelectedLightCardActors = GetSelectedActorsAs<ADisplayClusterLightCardActor>();
+
+	check(SelectedLightCardActors.Num() == 1);
+	const ADisplayClusterLightCardActor* LightCardActor = SelectedLightCardActors[0];
+
+	const UDisplayClusterLightCardEditorProjectSettings* Settings = GetDefault<UDisplayClusterLightCardEditorProjectSettings>();
+	FString DefaultPath = Settings->LightCardTemplateDefaultPath.Path;
+	if (DefaultPath.IsEmpty())
+	{
+		DefaultPath = TEXT("/Game");
+	}
+	
+	FString DefaultAssetName = LightCardActor->GetActorLabel();
+	if (!DefaultAssetName.EndsWith(TEXT("Template")))
+	{
+		DefaultAssetName += TEXT("Template");
+	}
+
+	FString PackageName;
+	bool FilenameValid = false;
+	while (!FilenameValid)
+	{
+		FSaveAssetDialogConfig SaveAssetDialogConfig;
+		{
+			SaveAssetDialogConfig.DefaultPath = MoveTemp(DefaultPath);
+			SaveAssetDialogConfig.DefaultAssetName = MoveTemp(DefaultAssetName);
+			SaveAssetDialogConfig.AssetClassNames.Add(UDisplayClusterLightCardTemplate::StaticClass()->GetClassPathName());
+			SaveAssetDialogConfig.ExistingAssetPolicy = ESaveAssetDialogExistingAssetPolicy::AllowButWarn;
+			SaveAssetDialogConfig.DialogTitleOverride = LOCTEXT("SaveAssetDialogTitle", "Save Asset As");
+		}
+
+		const FContentBrowserModule& ContentBrowserModule = FModuleManager::LoadModuleChecked<FContentBrowserModule>("ContentBrowser");
+		const FString SaveObjectPath = ContentBrowserModule.Get().CreateModalSaveAssetDialog(SaveAssetDialogConfig);
+
+		if (SaveObjectPath.IsEmpty())
+		{
+			return;
+		}
+
+		PackageName = FPackageName::ObjectPathToPackageName(SaveObjectPath);
+		
+		FText OutError;
+		FilenameValid = FFileHelper::IsFilenameValidForSaving(PackageName, OutError);
+	}
+
+	const FString NewAssetName = FPackageName::GetLongPackageAssetName(PackageName);
+	UPackage* AssetPackage = CreatePackage(*PackageName);
+	check(AssetPackage);
+
+	const EObjectFlags Flags = RF_Public | RF_Standalone;
+	UDisplayClusterLightCardTemplate* Template = NewObject<UDisplayClusterLightCardTemplate>(AssetPackage, FName(*NewAssetName), Flags);
+	Template->LightCardActor = CastChecked<ADisplayClusterLightCardActor>(StaticDuplicateObject(LightCardActor, Template));
+
+	FAssetRegistryModule::AssetCreated(Template);
+
+	AssetPackage->MarkPackageDirty();
+
+	FEditorFileUtils::PromptForCheckoutAndSave({AssetPackage}, true, false);
+}
+
+bool FDisplayClusterLightCardEditor::CanCreateLightCardTemplate() const
+{
+	const TArray<ADisplayClusterLightCardActor*> SelectedLightCards = GetSelectedActorsAs<ADisplayClusterLightCardActor>();
+	return SelectedLightCards.Num() == 1;
+}
+
+void FDisplayClusterLightCardEditor::ToggleLightCardLabels()
+{
+	ShowLightCardLabels(!ShouldShowLightCardLabels());
+}
+
+void FDisplayClusterLightCardEditor::ShowLightCardLabels(bool bVisible)
+{
+	if (!GetActiveRootActor().IsValid())
+	{
+		return;
+	}
+	
+	FScopedTransaction Transaction(LOCTEXT("ToggleLightCardLabelsTransactionMessage", "Toggle Light Card Labels"), !GIsTransacting);
+	
+	IDisplayClusterLightCardEditor& LightCardEditorModule = FModuleManager::GetModuleChecked<IDisplayClusterLightCardEditor>(IDisplayClusterLightCardEditor::ModuleName);
+
+	IDisplayClusterLightCardEditor::FLabelArgs Args;
+	Args.bVisible = bVisible;
+	Args.Scale = *GetLightCardLabelScale();
+	Args.RootActor = GetActiveRootActor().Get();
+	
+	LightCardEditorModule.ShowLabels(MoveTemp(Args));
+
+	RefreshPreviewActors(EDisplayClusterLightCardEditorProxyType::LightCards);
+}
+
+bool FDisplayClusterLightCardEditor::ShouldShowLightCardLabels() const
+{
+	return GetDefault<UDisplayClusterLightCardEditorProjectSettings>()->bDisplayLightCardLabels;
+}
+
+TOptional<float> FDisplayClusterLightCardEditor::GetLightCardLabelScale() const
+{
+	return GetDefault<UDisplayClusterLightCardEditorProjectSettings>()->LightCardLabelScale;
+}
+
+void FDisplayClusterLightCardEditor::SetLightCardLabelScale(float NewValue)
+{
+	if (!GetActiveRootActor().IsValid())
+	{
+		return;
+	}
+	
+	FScopedTransaction Transaction(LOCTEXT("ScaleLightCardLabelsTransactionMessage", "Scale Light Card Labels"), !GIsTransacting);
+	
+	IDisplayClusterLightCardEditor& LightCardEditorModule = FModuleManager::GetModuleChecked<IDisplayClusterLightCardEditor>(IDisplayClusterLightCardEditor::ModuleName);
+
+	IDisplayClusterLightCardEditor::FLabelArgs Args;
+	Args.bVisible = ShouldShowLightCardLabels();
+	Args.Scale = NewValue;
+	Args.RootActor = GetActiveRootActor().Get();
+	
+	LightCardEditorModule.ShowLabels(MoveTemp(Args));
+}
+
+void FDisplayClusterLightCardEditor::ShowIcons(bool bVisible)
+{
+	UDisplayClusterLightCardEditorSettings* Settings = GetMutableDefault<UDisplayClusterLightCardEditorSettings>();
+	Settings->bDisplayIcons = bVisible;
+	Settings->PostEditChange();
+	Settings->SaveConfig();
+}
+
+bool FDisplayClusterLightCardEditor::ShouldShowIcons() const
+{
+	return GetDefault<UDisplayClusterLightCardEditorSettings>()->bDisplayIcons;
+}
+
+TOptional<float> FDisplayClusterLightCardEditor::GetIconScale() const
+{
+	return GetDefault<UDisplayClusterLightCardEditorSettings>()->IconScale;
+}
+
+void FDisplayClusterLightCardEditor::SetIconScale(float NewValue)
+{
+	UDisplayClusterLightCardEditorSettings* Settings = GetMutableDefault<UDisplayClusterLightCardEditorSettings>();
+	Settings->IconScale = NewValue;
+	Settings->PostEditChange();
+	Settings->SaveConfig();
+}
+
+void FDisplayClusterLightCardEditor::OnActiveRootActorChanged(ADisplayClusterRootActor* NewRootActor)
+{
+	if (NewRootActor == ActiveRootActor.Get())
+	{
+		return;
+	}
+	
+	RemoveCompileDelegates();
+	
+	// The new root actor pointer could be null, indicating that it was deleted or the user didn't select a valid root actor
+	ActiveRootActor = NewRootActor;
+	if (LightCardOutliner)
+	{
+		LightCardOutliner->SetRootActor(NewRootActor);
+	}
+	
+	if (ViewportView)
+	{
+		ViewportView->SetRootActor(NewRootActor);
+	}
+	
+	BindCompileDelegates();
+
+	FCoreUObjectDelegates::OnObjectPropertyChanged.RemoveAll(this);
+	FCoreUObjectDelegates::OnObjectPropertyChanged.AddSP(this, &FDisplayClusterLightCardEditor::OnActorPropertyChanged);
+
+	RefreshLabels();
+}
+
+void FDisplayClusterLightCardEditor::RegisterTabSpawners()
+{
+	IDisplayClusterOperator::Get().OnRegisterLayoutExtensions().AddStatic(&FDisplayClusterLightCardEditor::RegisterLayoutExtension);
+
+	check(OperatorViewModel.IsValid());
+	const TSharedRef<FWorkspaceItem> AppMenuGroup = OperatorViewModel.Pin()->GetWorkspaceMenuGroup().ToSharedRef();
+
+	const TSharedPtr<FTabManager> TabManager = OperatorViewModel.Pin()->GetTabManager();
+	
+	TabManager->RegisterTabSpawner(ViewportTabName, FOnSpawnTab::CreateSP(this, &FDisplayClusterLightCardEditor::SpawnViewportTab))
+	.SetDisplayName(LOCTEXT("ViewportTab_DisplayName", "Viewport"))
+	.SetTooltipText(LOCTEXT("ViewportTab_Tooltip", "Light Card Editor Viewport"))
+	.SetGroup(AppMenuGroup)
+	.SetIcon(FSlateIcon(FAppStyle::Get().GetStyleSetName(), "LevelEditor.Tabs.Viewports"));
+	
+	TabManager->RegisterTabSpawner(OutlinerTabName, FOnSpawnTab::CreateSP(this, &FDisplayClusterLightCardEditor::SpawnOutlinerTab))
+	.SetDisplayName(LOCTEXT("OutlinerTab_DisplayName", "Outliner"))
+	.SetTooltipText(LOCTEXT("OutlinerTab_Tooltip", "Light Card Editor Outliner"))
+	.SetGroup(AppMenuGroup)
+	.SetIcon(FSlateIcon(FAppStyle::GetAppStyleSetName(), "LevelEditor.Tabs.Outliner"));
+}
+
+void FDisplayClusterLightCardEditor::UnregisterTabSpawners()
+{
+	if (OperatorViewModel.IsValid())
+	{
+		if (const TSharedPtr<FTabManager> OperatorPanelTabManager = OperatorViewModel.Pin()->GetTabManager())
+		{
+			OperatorPanelTabManager->UnregisterTabSpawner(ViewportTabName);
+			OperatorPanelTabManager->UnregisterTabSpawner(OutlinerTabName);
+		}
+	}
+}
+
+void FDisplayClusterLightCardEditor::RegisterLayoutExtension(FLayoutExtender& InExtender)
+{
+	const FTabManager::FTab ViewportTab(FTabId(ViewportTabName, ETabIdFlags::SaveLayout), ETabState::OpenedTab);
+	InExtender.ExtendStack(IDisplayClusterOperator::Get().GetPrimaryOperatorExtensionId(), ELayoutExtensionPosition::After, ViewportTab);
+	
+	FTabManager::FTab OutlinerTab(FTabId(OutlinerTabName, ETabIdFlags::SaveLayout), ETabState::OpenedTab);
+	InExtender.ExtendLayout(IDisplayClusterOperator::Get().GetDetailsTabId(), ELayoutExtensionPosition::Above, OutlinerTab);
+}
+
+TSharedRef<SDockTab> FDisplayClusterLightCardEditor::SpawnViewportTab(const FSpawnTabArgs& SpawnTabArgs)
+{
+	return SNew(SDockTab)
+	.TabRole(ETabRole::DocumentTab)
+	[
+		CreateViewportWidget()
+	];
+}
+
+TSharedRef<SDockTab> FDisplayClusterLightCardEditor::SpawnOutlinerTab(const FSpawnTabArgs& SpawnTabArgs)
+{
+	return SNew(SDockTab)
+	.TabRole(ETabRole::DocumentTab)
+	[
+		CreateLightCardOutlinerWidget()
+	];
+}
+
+TSharedRef<SWidget> FDisplayClusterLightCardEditor::CreateLightCardOutlinerWidget()
+{
+	return SAssignNew(LightCardOutliner, SDisplayClusterLightCardOutliner, SharedThis(this), CommandList);
+}
+
+TSharedRef<SWidget> FDisplayClusterLightCardEditor::CreateViewportWidget()
+{
+	return SAssignNew(ViewportView, SDisplayClusterLightCardEditorViewport, SharedThis(this), CommandList);
+}
+
+TSharedRef<SWidget> FDisplayClusterLightCardEditor::GeneratePlaceActorsMenu()
+{
+	const bool bInShouldCloseWindowAfterMenuSelection = true;
+
+	FMenuBuilder MenuBuilder(bInShouldCloseWindowAfterMenuSelection, CommandList);
+	MenuBuilder.BeginSection("PlaceActors", LOCTEXT("PlaceActorsMenuHeader", "Place Actors"));
+	{
+		FSlateIcon LightCardIcon = FSlateIconFinder::FindIconForClass(ADisplayClusterLightCardActor::StaticClass());
+
+		MenuBuilder.AddMenuEntry(FDisplayClusterLightCardEditorCommands::Get().AddNewFlag,
+			NAME_None, TAttribute<FText>(), TAttribute<FText>(), LightCardIcon);
+		MenuBuilder.AddMenuEntry(FDisplayClusterLightCardEditorCommands::Get().AddNewLightCard,
+			NAME_None, TAttribute<FText>(), TAttribute<FText>(), LightCardIcon);
+		
+		TSet<UClass*> StageActorClasses = UE::DisplayClusterLightCardEditorUtils::GetAllStageActorClasses();
+		for (UClass* Class : StageActorClasses)
+		{
+			if (Class == ADisplayClusterLightCardActor::StaticClass())
+			{
+				// Added manually already
+				continue;
+			}
+			
+			FString Label = FString::Printf(TEXT("Add %s"), *Class->GetDisplayNameText().ToString());
+			FSlateIcon StageActorIcon = FSlateIconFinder::FindIconForClass(Class);
+			MenuBuilder.AddMenuEntry(FText::FromString(Label), LOCTEXT("AddStageActorHeader", "Add a stage actor to the scene"), StageActorIcon,
+				FUIAction(FExecuteAction::CreateRaw(this, &FDisplayClusterLightCardEditor::AddNewDynamic, Class)));
+		}
+		
+		MenuBuilder.AddSubMenu(LOCTEXT("AllTemplatesLabel", "All Templates"),
+		LOCTEXT("AllTemplatesTooltip", "Select a template"),
+		FNewMenuDelegate::CreateSP(this, &FDisplayClusterLightCardEditor::GenerateTemplateSubMenu),
+		false,
+		FSlateIcon(),
+		true);
+	}
+	MenuBuilder.EndSection();
+	MenuBuilder.BeginSection("Favorites", LOCTEXT("FavoritesMenuHeader", "Favorites"));
+	{
+		TArray<UDisplayClusterLightCardTemplate*> LightCardTemplates =
+		UE::DisplayClusterLightCardTemplateHelpers::GetLightCardTemplates(/*bFavoritesOnly*/ true);
+
+		TemplateBrushes.Empty();
+		
+		for (UDisplayClusterLightCardTemplate* Template : LightCardTemplates)
+		{
+			// Create a brush if this template is using a custom texture.
+			if (Template->LightCardActor != nullptr && Template->LightCardActor->Texture.Get() != nullptr)
+			{
+				TSharedPtr<FSlateBrush> SlateBrush = MakeShared<FSlateBrush>();
+
+				SlateBrush->SetResourceObject(Template->LightCardActor->Texture.Get());
+				SlateBrush->ImageSize = FVector2D(16.f, 16.f);
+				TemplateBrushes.Add(Template, SlateBrush);
+			}
+
+			TWeakObjectPtr<UDisplayClusterLightCardTemplate> TemplateWeakPtr = MakeWeakObjectPtr(Template);
+
+			const TSharedPtr<SWidget> TemplateWidget = SNew(SHorizontalBox)
+			+SHorizontalBox::Slot()
+			.AutoWidth()
+			.VAlign(VAlign_Center)
+			[
+				SNew(SBox)
+				.HeightOverride(16.f)
+				.WidthOverride(16.f)
+				[
+					SNew(SImage)
+					.ColorAndOpacity(FSlateColor::UseForeground())
+					.Image(this, &FDisplayClusterLightCardEditor::GetLightCardTemplateIcon, TemplateWeakPtr)
+				]
+			]
+			+SHorizontalBox::Slot()
+			.AutoWidth()
+			.VAlign(VAlign_Center)
+			.Padding(4.f, 0.f)
+			[
+				SNew(STextBlock)
+				.Text(FText::FromString(Template->GetName()))
+			];
+
+			FMenuEntryParams EntryParams;
+			EntryParams.EntryWidget = TemplateWidget;
+			EntryParams.LabelOverride = FText::FromString(Template->GetName());
+			EntryParams.ToolTipOverride = LOCTEXT("FavoriteEntryTooltip", "Spawn this favorite");
+			EntryParams.DirectActions.ExecuteAction = FExecuteAction::CreateLambda([this, TemplateWeakPtr]()
+			{
+				if (TemplateWeakPtr.IsValid())
+				{
+					SpawnLightCardFromTemplate(TemplateWeakPtr.Get());
+				}
+			});
+			
+			MenuBuilder.AddMenuEntry(EntryParams);
+		}
+	}
+	MenuBuilder.EndSection();
+
+	CleanupRecentlyPlacedItems();
+	
+	MenuBuilder.BeginSection("RecentlyPlaced", LOCTEXT("RecentlyPlacedMenuHeader", "Recently Placed"));
+	{
+		const UDisplayClusterLightCardEditorSettings* Settings = GetDefault<UDisplayClusterLightCardEditorSettings>();
+		for (const FDisplayClusterLightCardEditorRecentItem& RecentlyPlacedItem : Settings->RecentlyPlacedItems)
+		{
+			TWeakObjectPtr<UObject> ObjectWeakPtr = MakeWeakObjectPtr(RecentlyPlacedItem.ObjectPath.LoadSynchronous());
+			check(ObjectWeakPtr.IsValid());
+			
+			const TSharedPtr<SWidget> RecentItemWidget = SNew(SHorizontalBox)
+			+SHorizontalBox::Slot()
+			.AutoWidth()
+			.VAlign(VAlign_Center)
+			[
+				SNew(SBox)
+				.HeightOverride(16.f)
+				.WidthOverride(16.f)
+				[
+					SNew(SImage)
+					.ColorAndOpacity(FSlateColor::UseForeground())
+					.Image(RecentlyPlacedItem.GetSlateBrush())
+				]
+			]
+			+SHorizontalBox::Slot()
+			.AutoWidth()
+			.VAlign(VAlign_Center)
+			.Padding(4.f, 0.f)
+			[
+				SNew(STextBlock)
+				.Text(RecentlyPlacedItem.GetItemDisplayName())
+			];
+
+			FMenuEntryParams EntryParams;
+			EntryParams.EntryWidget = RecentItemWidget;
+			EntryParams.LabelOverride = FText::FromString(ObjectWeakPtr->GetName());
+			EntryParams.ToolTipOverride = LOCTEXT("RecentItemEntryTooltip", "Spawn this recently placed item");
+			EntryParams.DirectActions.ExecuteAction = FExecuteAction::CreateLambda([this, RecentlyPlacedItem, ObjectWeakPtr]()
+			{
+				if (ObjectWeakPtr.IsValid())
+				{
+					if (RecentlyPlacedItem.ItemType == FDisplayClusterLightCardEditorRecentItem::Type_LightCard)
+					{
+						AddNewLightCard();
+					}
+					else if (RecentlyPlacedItem.ItemType == FDisplayClusterLightCardEditorRecentItem::Type_Flag)
+					{
+						AddNewFlag();
+					}
+					else if (RecentlyPlacedItem.ItemType == FDisplayClusterLightCardEditorRecentItem::Type_LightCardTemplate)
+					{
+						if (const UDisplayClusterLightCardTemplate* Template = Cast<UDisplayClusterLightCardTemplate>(ObjectWeakPtr.Get()))
+						{
+							SpawnLightCardFromTemplate(Template);
+						}
+					}
+					else if (RecentlyPlacedItem.ItemType == FDisplayClusterLightCardEditorRecentItem::Type_Dynamic)
+					{
+						AddNewDynamic(Cast<UClass>(ObjectWeakPtr.Get()));
+					}
+				}
+			});
+			
+			MenuBuilder.AddMenuEntry(EntryParams);
+		}
+	}
+	MenuBuilder.EndSection();
+
+	return MenuBuilder.MakeWidget();
+}
+
+const FSlateBrush* FDisplayClusterLightCardEditor::GetLightCardTemplateIcon(const TWeakObjectPtr<UDisplayClusterLightCardTemplate> InTemplate) const
+{
+	if (!InTemplate.IsValid())
+	{
+		return nullptr;
+	}
+		
+	if (const TSharedPtr<FSlateBrush>* Brush = TemplateBrushes.Find(InTemplate))
+	{
+		return (*Brush).Get();
+	}
+	
+	return FSlateIconFinder::FindIconBrushForClass(InTemplate->GetClass());
+}
+
+void FDisplayClusterLightCardEditor::GenerateTemplateSubMenu(FMenuBuilder& InMenuBuilder)
+{
+	InMenuBuilder.BeginSection("Templates", LOCTEXT("TemplatesMenuHeader", "Templates"));
+	{
+		InMenuBuilder.AddWidget(CreateLightCardTemplateWidget(), FText::GetEmpty(), true);
+	}
+	InMenuBuilder.EndSection();
+}
+
+TSharedRef<SWidget> FDisplayClusterLightCardEditor::GenerateLabelsMenu()
+{
+	const bool bInShouldCloseWindowAfterMenuSelection = true;
+
+	FMenuBuilder MenuBuilder(bInShouldCloseWindowAfterMenuSelection, CommandList);
+
+	MenuBuilder.BeginSection("Labels", LOCTEXT("LabelsMenuHeader", "Labels"));
+	{
+		MenuBuilder.AddMenuEntry(FDisplayClusterLightCardEditorCommands::Get().ToggleLightCardLabels);
+	
+		MenuBuilder.AddWidget(
+		SNew(SBox)
+		.HAlign(HAlign_Right)
+		[
+			SNew(SBox)
+			.Padding(FMargin(8.0f, 0.0f, 0.0f, 0.0f))
+			.WidthOverride(100.0f)
+			[
+				SNew(SBorder)
+				.BorderImage(FAppStyle::Get().GetBrush("Menu.WidgetBorder"))
+				.Padding(FMargin(1.0f))
+				[
+					SNew(SBox)
+					.MinDesiredWidth(64)
+					[
+						SNew(SNumericEntryBox<float>)
+						.Value(this, &FDisplayClusterLightCardEditor::GetLightCardLabelScale)
+						.OnValueChanged(this, &FDisplayClusterLightCardEditor::SetLightCardLabelScale)
+						.MinValue(0)
+						.MaxValue(FLT_MAX)
+						.MinSliderValue(0)
+						.MaxSliderValue(10)
+						.AllowSpin(true)
+					]
+				]
+			]
+		],
+		LOCTEXT("LightCardLabelScale_Label", "Label Scale"));
+	}
+	MenuBuilder.EndSection();
+
+	MenuBuilder.BeginSection("Icons", LOCTEXT("IconsMenuHeader", "Icons"));
+	{
+		MenuBuilder.AddMenuEntry(FDisplayClusterLightCardEditorCommands::Get().ToggleIconVisibility,
+			NAME_None, LOCTEXT("IconVisibilityLabel", "Actor Icons"));
+	
+		MenuBuilder.AddWidget(
+		SNew(SBox)
+		.HAlign(HAlign_Right)
+		[
+			SNew(SBox)
+			.Padding(FMargin(8.0f, 0.0f, 0.0f, 0.0f))
+			.WidthOverride(100.0f)
+			[
+				SNew(SBorder)
+				.BorderImage(FAppStyle::Get().GetBrush("Menu.WidgetBorder"))
+				.Padding(FMargin(1.0f))
+				[
+					SNew(SBox)
+					.MinDesiredWidth(64)
+					[
+						SNew(SNumericEntryBox<float>)
+						.Value(this, &FDisplayClusterLightCardEditor::GetIconScale)
+						.OnValueChanged(this, &FDisplayClusterLightCardEditor::SetIconScale)
+						.MinValue(0)
+						.MaxValue(FLT_MAX)
+						.MinSliderValue(0)
+						.MaxSliderValue(4)
+						.AllowSpin(true)
+					]
+				]
+			]
+		],
+		LOCTEXT("LightCardIconScale_Label", "Icon Scale"));
+	}
+	MenuBuilder.EndSection();
+	
+	return MenuBuilder.MakeWidget();
+}
+
+TSharedRef<SWidget> FDisplayClusterLightCardEditor::CreateLightCardTemplateWidget()
+{
+	return SNew(SDisplayClusterLightCardTemplateList, SharedThis(this))
+	.HideHeader(true)
+	.SpawnOnSelection(true);
+}
+
+void FDisplayClusterLightCardEditor::BindCommands()
+{
+	CommandList = MakeShareable(new FUICommandList);
+
+	CommandList->MapAction(
+		FDisplayClusterLightCardEditorCommands::Get().AddNewLightCard,
+		FExecuteAction::CreateSP(this, &FDisplayClusterLightCardEditor::AddNewLightCard),
+		FCanExecuteAction::CreateSP(this, &FDisplayClusterLightCardEditor::CanAddLightCard));
+
+	CommandList->MapAction(
+		FDisplayClusterLightCardEditorCommands::Get().AddNewFlag,
+		FExecuteAction::CreateSP(this, &FDisplayClusterLightCardEditor::AddNewFlag),
+		FCanExecuteAction::CreateSP(this, &FDisplayClusterLightCardEditor::CanAddLightCard));
+	
+	CommandList->MapAction(
+		FDisplayClusterLightCardEditorCommands::Get().AddExistingLightCard,
+		FExecuteAction::CreateSP(this, &FDisplayClusterLightCardEditor::AddExistingLightCard),
+		FCanExecuteAction::CreateSP(this, &FDisplayClusterLightCardEditor::CanAddLightCard));
+
+	CommandList->MapAction(
+		FDisplayClusterLightCardEditorCommands::Get().RemoveLightCard,
+		FExecuteAction::CreateSP(this, &FDisplayClusterLightCardEditor::RemoveSelectedActors, false),
+		FCanExecuteAction::CreateSP(this, &FDisplayClusterLightCardEditor::CanRemoveSelectedLightCardFromActor));
+
+	CommandList->MapAction(
+		FGenericCommands::Get().Cut,
+		FExecuteAction::CreateSP(this, &FDisplayClusterLightCardEditor::CutSelectedActors),
+		FCanExecuteAction::CreateSP(this, &FDisplayClusterLightCardEditor::CanCutSelectedActors));
+
+	CommandList->MapAction(
+		FGenericCommands::Get().Copy,
+		FExecuteAction::CreateSP(this, &FDisplayClusterLightCardEditor::CopySelectedActors),
+		FCanExecuteAction::CreateSP(this, &FDisplayClusterLightCardEditor::CanCopySelectedActors));
+
+	CommandList->MapAction(
+		FGenericCommands::Get().Paste,
+		FExecuteAction::CreateSP(this, &FDisplayClusterLightCardEditor::PasteActors, false),
+		FCanExecuteAction::CreateSP(this, &FDisplayClusterLightCardEditor::CanPasteActors));
+
+	CommandList->MapAction(
+		FGenericCommands::Get().Duplicate,
+		FExecuteAction::CreateSP(this, &FDisplayClusterLightCardEditor::DuplicateSelectedActors),
+		FCanExecuteAction::CreateSP(this, &FDisplayClusterLightCardEditor::CanDuplicateSelectedActors));
+
+	CommandList->MapAction(
+		FGenericCommands::Get().Delete,
+		FExecuteAction::CreateSP(this, &FDisplayClusterLightCardEditor::RemoveSelectedActors, true),
+		FCanExecuteAction::CreateSP(this, &FDisplayClusterLightCardEditor::CanRemoveSelectedActors));
+
+	CommandList->MapAction(
+		FDisplayClusterLightCardEditorCommands::Get().SaveLightCardTemplate,
+		FExecuteAction::CreateSP(this, &FDisplayClusterLightCardEditor::CreateLightCardTemplate),
+		FCanExecuteAction::CreateSP(this, &FDisplayClusterLightCardEditor::CanCreateLightCardTemplate));
+
+	CommandList->MapAction(
+		FDisplayClusterLightCardEditorCommands::Get().ToggleLightCardLabels,
+		FExecuteAction::CreateSP(this, &FDisplayClusterLightCardEditor::ToggleLightCardLabels),
+		FCanExecuteAction(),
+		FIsActionChecked::CreateSP(this, &FDisplayClusterLightCardEditor::ShouldShowLightCardLabels));
+}
+
+void FDisplayClusterLightCardEditor::RegisterToolbarExtensions()
+{
+	const TSharedPtr<FExtender> ToolbarExtender = MakeShareable(new FExtender);
+	ToolbarExtender->AddToolBarExtension("General", EExtensionHook::After, nullptr,
+		FToolBarExtensionDelegate::CreateSP(this, &FDisplayClusterLightCardEditor::ExtendToolbar));
+	
+	IDisplayClusterOperator::Get().GetOperatorToolBarExtensibilityManager()->AddExtender(ToolbarExtender);
+}
+
+void FDisplayClusterLightCardEditor::RegisterMenuExtensions()
+{
+	const TSharedPtr<FExtender> ToolbarExtender = MakeShareable(new FExtender);
+	ToolbarExtender->AddMenuExtension("FileOpen", EExtensionHook::Before, CommandList,
+		FMenuExtensionDelegate::CreateSP(this, &FDisplayClusterLightCardEditor::ExtendFileMenu));
+
+	ToolbarExtender->AddMenuExtension("EditHistory", EExtensionHook::After, CommandList,
+	FMenuExtensionDelegate::CreateSP(this, &FDisplayClusterLightCardEditor::ExtendEditMenu));
+	
+	IDisplayClusterOperator::Get().GetOperatorMenuExtensibilityManager()->AddExtender(ToolbarExtender);
+}
+
+void FDisplayClusterLightCardEditor::ExtendToolbar(FToolBarBuilder& ToolbarBuilder)
+{
+	ToolbarBuilder.AddSeparator();
+	{
+		const FUIAction DefaultAction;
+		ToolbarBuilder.AddComboButton(
+			DefaultAction,
+			FOnGetContent::CreateSP(this, &FDisplayClusterLightCardEditor::GeneratePlaceActorsMenu),
+			TAttribute<FText>(),
+			LOCTEXT("PlaceActors_ToolTip", "Place actors in the scene"),
+			FSlateIcon(FAppStyle::GetAppStyleSetName(), "Icons.PlaceActors"));
+	}
+	
+	ToolbarBuilder.AddSeparator();
+	{
+		const FUIAction DefaultAction;
+		ToolbarBuilder.AddComboButton(
+			DefaultAction,
+			FOnGetContent::CreateSP(this, &FDisplayClusterLightCardEditor::GenerateLabelsMenu),
+			TAttribute<FText>(),
+			LOCTEXT("Labels_ToolTip", "Configure options for labels"),
+			FSlateIcon(FDisplayClusterLightCardEditorStyle::Get().GetStyleSetName(), "DisplayClusterLightCardEditor.Labels"));
+	}
+}
+
+void FDisplayClusterLightCardEditor::ExtendFileMenu(FMenuBuilder& MenuBuilder)
+{
+	MenuBuilder.BeginSection("Import", LOCTEXT("ImportHeading", "Import"));
+	const FSlateIcon LightCardIcon = FSlateIconFinder::FindIconForClass(ADisplayClusterLightCardActor::StaticClass());
+	MenuBuilder.AddMenuEntry(FDisplayClusterLightCardEditorCommands::Get().AddExistingLightCard,
+		NAME_None, TAttribute<FText>(), TAttribute<FText>(), LightCardIcon);
+	MenuBuilder.EndSection();
+}
+
+void FDisplayClusterLightCardEditor::ExtendEditMenu(FMenuBuilder& MenuBuilder)
+{
+	MenuBuilder.BeginSection("Edit", LOCTEXT("EditHeading", "Edit"));
+	MenuBuilder.AddMenuEntry(FGenericCommands::Get().Cut);
+	MenuBuilder.AddMenuEntry(FGenericCommands::Get().Copy);
+	MenuBuilder.AddMenuEntry(FGenericCommands::Get().Paste);
+	MenuBuilder.AddMenuEntry(FGenericCommands::Get().Duplicate);
+	MenuBuilder.AddMenuEntry(FGenericCommands::Get().Delete);
+	MenuBuilder.EndSection();
+}
+
+void FDisplayClusterLightCardEditor::RefreshPreviewActors(EDisplayClusterLightCardEditorProxyType ProxyType)
+{
+	RemoveCompileDelegates();
+	
+	if (ADisplayClusterRootActor* RootActor = GetActiveRootActor().Get())
+	{
+		if (LightCardOutliner.IsValid())
+		{
+			LightCardOutliner->SetRootActor(RootActor);
+		}
+			
+		if (ViewportView.IsValid())
+		{
+			const bool bForce = true;
+			ViewportView->GetLightCardEditorViewportClient()->UpdatePreviewActor(RootActor, bForce, ProxyType);
+		}
+	}
+	
+	BindCompileDelegates();
+}
+
+void FDisplayClusterLightCardEditor::RefreshLabels()
+{
+	ShowLightCardLabels(ShouldShowLightCardLabels());
+}
+
+bool FDisplayClusterLightCardEditor::IsOurObject(UObject* InObject,
+	EDisplayClusterLightCardEditorProxyType& OutProxyType) const
+{
+	auto IsOurActor = [InObject] (UObject* ObjectToCompare) -> bool
+	{
+		if (ObjectToCompare)
+		{
+			if (InObject == ObjectToCompare)
+			{
+				return true;
+			}
+
+			if (const UObject* RootActorOuter = InObject->GetTypedOuter(ObjectToCompare->GetClass()))
+			{
+				return RootActorOuter == ObjectToCompare;
+			}
+		}
+
+		return false;
+	};
+
+	EDisplayClusterLightCardEditorProxyType ProxyType = EDisplayClusterLightCardEditorProxyType::All;
+	
+	bool bIsOurActor = IsOurActor(GetActiveRootActor().Get());
+	if (!bIsOurActor && LightCardOutliner.IsValid())
+	{
+		for (const TSharedPtr<SDisplayClusterLightCardOutliner::FStageActorTreeItem>& LightCard : LightCardOutliner->GetStageActorTreeItems())
+		{
+			bIsOurActor = IsOurActor(LightCard->Actor.Get());
+			if (bIsOurActor)
+			{
+				ProxyType = EDisplayClusterLightCardEditorProxyType::LightCards;
+				break;
+			}
+		}
+	}
+
+	OutProxyType = ProxyType;
+	return bIsOurActor;
+}
+
+void FDisplayClusterLightCardEditor::BindCompileDelegates()
+{
+	if (LightCardOutliner.IsValid())
+	{
+		for (const TSharedPtr<SDisplayClusterLightCardOutliner::FStageActorTreeItem>& LightCardActor : LightCardOutliner->GetStageActorTreeItems())
+		{
+			if (LightCardActor.IsValid() && LightCardActor->Actor.IsValid())
+			{
+				if (UBlueprint* Blueprint = UBlueprint::GetBlueprintFromClass(LightCardActor->Actor->GetClass()))
+				{
+					Blueprint->OnCompiled().AddSP(this, &FDisplayClusterLightCardEditor::OnBlueprintCompiled);
+				}
+			}
+		}
+	}
+}
+
+void FDisplayClusterLightCardEditor::RemoveCompileDelegates()
+{
+	if (LightCardOutliner.IsValid())
+	{
+		for (const TSharedPtr<SDisplayClusterLightCardOutliner::FStageActorTreeItem>& LightCardActor : LightCardOutliner->GetStageActorTreeItems())
+		{
+			if (LightCardActor.IsValid() && LightCardActor->Actor.IsValid())
+			{
+				if (UBlueprint* Blueprint = UBlueprint::GetBlueprintFromClass(LightCardActor->Actor->GetClass()))
+				{
+					Blueprint->OnCompiled().RemoveAll(this);
+				}
+			}
+		}
+	}
+}
+
+void FDisplayClusterLightCardEditor::OnActorPropertyChanged(UObject* ObjectBeingModified,
+                                                            FPropertyChangedEvent& PropertyChangedEvent)
+{
+	EDisplayClusterLightCardEditorProxyType ProxyType;
+	if (IsOurObject(ObjectBeingModified, ProxyType))
+	{
+		const FName PropertyName = PropertyChangedEvent.GetPropertyName();
+		const bool bTransformationChanged =
+				(PropertyName == USceneComponent::GetRelativeLocationPropertyName() ||
+					PropertyName == USceneComponent::GetRelativeRotationPropertyName() ||
+					PropertyName == USceneComponent::GetRelativeScale3DPropertyName());
+		if (PropertyChangedEvent.ChangeType == EPropertyChangeType::Interactive ||
+			(bTransformationChanged && !ObjectBeingModified->IsA<ADisplayClusterRootActor>()))
+		{
+			// Real-time & efficient update when dragging a slider.
+			if (ViewportView.IsValid())
+			{
+				ViewportView->GetLightCardEditorViewportClient()->UpdateProxyTransforms();
+			}
+		}
+		else
+		{
+			// Full destroy and refresh.
+			RefreshPreviewActors(ProxyType);
+		}
+	}
+}
+
+void FDisplayClusterLightCardEditor::OnLevelActorDeleted(AActor* Actor)
+{
+	if (LightCardOutliner.IsValid() &&
+		LightCardOutliner->GetStageActorTreeItems().ContainsByPredicate([Actor](const TSharedPtr<SDisplayClusterLightCardOutliner::FStageActorTreeItem>& LightCardTreeItem)
+	{
+		return LightCardTreeItem.IsValid() && Actor == LightCardTreeItem->Actor.Get();
+	}))
+	{
+		if (Actor && Actor->GetClass()->HasAnyClassFlags(CLASS_NewerVersionExists))
+		{
+			// When a blueprint class is regenerated instances are deleted and replaced.
+			// In this case the OnCompiled() delegate will fire and refresh the actor.
+			return;
+		}
+		
+		if (ViewportView.IsValid())
+		{
+			ViewportView->GetLightCardEditorViewportClient()->GetWorld()->GetTimerManager().SetTimerForNextTick([=]()
+			{
+				// Schedule for next tick so available selections are properly updated once the
+				// actor is fully deleted.
+				RefreshPreviewActors(EDisplayClusterLightCardEditorProxyType::LightCards);
+			});
+		}
+	}
+}
+
+void FDisplayClusterLightCardEditor::OnBlueprintCompiled(UBlueprint* Blueprint)
+{
+	// Right now only LightCard blueprints are handled here.
+	RefreshPreviewActors(EDisplayClusterLightCardEditorProxyType::LightCards);
+}
+
+void FDisplayClusterLightCardEditor::OnObjectTransacted(UObject* Object,
+	const FTransactionObjectEvent& TransactionObjectEvent)
+{
+	if (TransactionObjectEvent.GetEventType() == ETransactionObjectEventType::UndoRedo)
+	{
+		// Always refresh on undo because the light card actor may not inherit our C++ class
+		// so we can't easily distinguish it. This supports the case where the user deletes
+		// a LightCard actor from the level manually then undoes it.
+		RefreshPreviewActors();
+	}
+}
+
+void FDisplayClusterLightCardEditor::AddRecentlyPlacedItem(const FDisplayClusterLightCardEditorRecentItem& InItem)
+{
+	UDisplayClusterLightCardEditorSettings* Settings = GetMutableDefault<UDisplayClusterLightCardEditorSettings>();
+	
+	Settings->RecentlyPlacedItems.RemoveAll([&](const FDisplayClusterLightCardEditorRecentItem& Compare)
+	{
+		return Compare == InItem;
+	});
+	Settings->RecentlyPlacedItems.Insert(InItem, 0);
+
+	const int32 MaxRecentlyPlaced = 5;
+	
+	if (Settings->RecentlyPlacedItems.Num() > MaxRecentlyPlaced)
+	{
+		Settings->RecentlyPlacedItems.RemoveAt(MaxRecentlyPlaced - 1, Settings->RecentlyPlacedItems.Num() - MaxRecentlyPlaced);
+	}
+
+	Settings->PostEditChange();
+	Settings->SaveConfig();
+}
+
+void FDisplayClusterLightCardEditor::CleanupRecentlyPlacedItems()
+{
+	UDisplayClusterLightCardEditorSettings* Settings = GetMutableDefault<UDisplayClusterLightCardEditorSettings>();
+	Settings->RecentlyPlacedItems.RemoveAll([&](const FDisplayClusterLightCardEditorRecentItem& Compare)
+	{
+		return !Compare.ObjectPath.LoadSynchronous();
+	});
+
+	Settings->PostEditChange();
+	Settings->SaveConfig();
+}
+
+#undef LOCTEXT_NAMESPACE
diff --git a/Engine/Plugins/Runtime/nDisplay/Source/DisplayClusterLightCardEditor/Private/DisplayClusterLightCardEditor.h b/Engine/Plugins/Runtime/nDisplay/Source/DisplayClusterLightCardEditor/Private/DisplayClusterLightCardEditor.h
new file mode 100644
--- /dev/null
+++ b/Engine/Plugins/Runtime/nDisplay/Source/DisplayClusterLightCardEditor/Private/DisplayClusterLightCardEditor.h
@@ -0,0 +1,322 @@
+// Copyright Epic Games, Inc. All Rights Reserved.
+
+#pragma once
+
+#include "CoreMinimal.h"
+#include "Widgets/DeclarativeSyntaxSupport.h"
+
+#include "DisplayClusterLightCardEditorProxyType.h"
+
+#include "EditorUndoClient.h"
+#include "IDisplayClusterOperatorApp.h"
+
+class FLayoutExtender;
+class FMenuBuilder;
+class FSpawnTabArgs;
+class FTabManager;
+class FToolBarBuilder;
+class FUICommandList;
+class SDockTab;
+class SDisplayClusterLightCardOutliner;
+class SDisplayClusterLightCardTemplateList;
+class SDisplayClusterLightCardEditorViewport;
+class ADisplayClusterRootActor;
+class ADisplayClusterLightCardActor;
+class UDisplayClusterLightCardTemplate;
+class IDisplayClusterOperatorViewModel;
+struct FSlateBrush;
+
+struct FDisplayClusterLightCardEditorRecentItem;
+
+/** A panel that can be spawned in a tab that contains all the UI elements that make up the 2D light cards editor */
+class FDisplayClusterLightCardEditor : public IDisplayClusterOperatorApp, public FEditorUndoClient
+{
+public:
+	/** The name of the tab that the viewport lives in */
+	static const FName ViewportTabName;
+	
+	/** The name of the tab that the outliner lives in */
+	static const FName OutlinerTabName;
+
+	/** Construct an instance of the light card editor */
+	static TSharedRef<IDisplayClusterOperatorApp> MakeInstance(TSharedRef<IDisplayClusterOperatorViewModel> InViewModel);
+	
+	~FDisplayClusterLightCardEditor();
+
+	// FEditorUndoClient
+	virtual void PostUndo(bool bSuccess) override;
+	virtual void PostRedo(bool bSuccess) override;
+	// ~FEditorUndoClient
+
+	/** Initialize the light card editor instance */
+	void Initialize(TSharedRef<IDisplayClusterOperatorViewModel> InViewModel);
+
+	/** Return the command list for the light card editor */
+	TSharedPtr<FUICommandList> GetCommandList() const { return CommandList; }
+	
+	/** The current active root actor for this light card editor */
+	TWeakObjectPtr<ADisplayClusterRootActor> GetActiveRootActor() const { return ActiveRootActor; }
+
+	/** Iterate the world actors looking for all actors this editor is managing */
+	TArray<AActor*> FindAllManagedActors() const;
+
+	/** Selects the specified actors in the outliner and details panel */
+	void SelectActors(const TArray<AActor*>& ActorsToSelect);
+
+	template<typename T>
+	TArray<T*> GetSelectedActorsAs() const
+	{
+		TArray<T*> OutArray;
+		Algo::TransformIf(SelectedActors, OutArray, [](const TWeakObjectPtr<AActor>& InItem)
+		{
+			return InItem.IsValid() && InItem->IsA<T>();
+		},
+		[](const TWeakObjectPtr<AActor>& InItem)
+		{
+			return CastChecked<T>(InItem.Get());
+		});
+		return OutArray;
+	}
+	
+	/** Gets the actors that are selected in the outliner */
+	void GetSelectedActors(TArray<AActor*>& OutSelectedActors);
+
+	/** Selects the actor proxies that correspond to the specified actors */
+	void SelectActorProxies(const TArray<AActor*>& ActorsToSelect);
+
+	/** Places the given light card in the middle of the viewport */
+	void CenterActorInView(AActor* Actor);
+
+	/** Spawns a new actor and adds it to the root actor if it is a light card */
+	AActor* SpawnActor(TSubclassOf<AActor> InActorClass, const FName& InActorName = NAME_None);
+
+	template<typename T>
+	T* SpawnActorAs(const FName& InActorName = NAME_None)
+	{
+		return Cast<T>(SpawnActor(T::StaticClass(), InActorName));
+	}
+	
+	/** Spawns a new light card from a light card template */
+	ADisplayClusterLightCardActor* SpawnLightCardFromTemplate(const UDisplayClusterLightCardTemplate* InTemplate, ULevel* InLevel = nullptr, bool bIsPreview = false);
+
+	/** Adds a new light card to the root actor and centers it in the viewport */
+	void AddNewLightCard();
+
+	/** Select an existing Light Card from a menu */
+	void AddExistingLightCard();
+
+	/** Adds a new light card configured as a flag */
+	void AddNewFlag();
+
+	/** Add a new actor dynamically based on a class */
+	void AddNewDynamic(UClass* InClass);
+
+	/** Adds the given Light Card to the root actor */
+	void AddLightCardsToActor(const TArray<ADisplayClusterLightCardActor*>& LightCards);
+
+	/** If a Light Card can currently be added */
+	bool CanAddLightCard() const;
+
+	/** Copies any selected actors to the clipboard, and then deletes them */
+	void CutSelectedActors();
+
+	/** Determines if there are selected actors that can be cut */
+	bool CanCutSelectedActors();
+
+	/** Copies any selected actors to the clipboard */
+	void CopySelectedActors();
+
+	/** Determines if there are selected actors that can be copied */
+	bool CanCopySelectedActors() const;
+
+	/** Pastes any actors in the clipboard to the current root actor */
+	void PasteActors(bool bOffsetActorPosition);
+
+	/** Determines if there are any actors that can be pasted from the clipboard */
+	bool CanPasteActors() const;
+
+	/** Copies any selected actors to the clipboard and then pastes them */
+	void DuplicateSelectedActors();
+
+	/** Determines if there are selected actors that can be duplicated */
+	bool CanDuplicateSelectedActors() const;
+
+	/**
+	 * Remove the light card from the actor
+	 *@param bDeleteLightCardActor Delete the actor from the level
+	 */
+	void RemoveSelectedActors(bool bDeleteLightCardActor);
+
+	/**
+	 * Remove the given light cards from the actor
+	 * @param InActorsToRemove Actors which should be removed
+	 * @param bDeleteActors If the light cards should be deleted from the level
+	 */
+	void RemoveActors(const TArray<TWeakObjectPtr<AActor>>& InActorsToRemove, bool bDeleteActors);
+	
+	/** If the selected actors can be removed */
+	bool CanRemoveSelectedActors() const;
+
+	/** If the selected light card can be removed from the actor */
+	bool CanRemoveSelectedLightCardFromActor() const;
+
+	/** Creates a template of the selected light card */
+	void CreateLightCardTemplate();
+
+	/** Determines if there are selected light cards that can have templates created */
+	bool CanCreateLightCardTemplate() const;
+
+	/** Toggles the visible status of labels */
+	void ToggleLightCardLabels();
+	
+	/** Display or hide labels for all included light cards */
+	void ShowLightCardLabels(bool bVisible);
+
+	/** If light card labels are currently toggled */
+	bool ShouldShowLightCardLabels() const;
+
+	/** Return the current light card label scale */
+	TOptional<float> GetLightCardLabelScale() const;
+
+	/** Update light card label scale */
+	void SetLightCardLabelScale(float NewValue);
+
+	/** Set all icon visibility */
+	void ShowIcons(bool bVisible);
+
+	/** If all icons should be visible */
+	bool ShouldShowIcons() const;
+
+	/** Return the current icon scale */
+	TOptional<float> GetIconScale() const;
+
+	/** Update all icons scale */
+	void SetIconScale(float NewValue);
+	
+private:
+	/** Raised when the active Display cluster root actor has been changed in the operator panel */
+	void OnActiveRootActorChanged(ADisplayClusterRootActor* NewRootActor);
+	
+	/** Registers the light card editor with the global tab manager and adds it to the operator panel's extension tab stack */
+	void RegisterTabSpawners();
+
+	/** Unregisters the light card editor from the global tab manager */
+	void UnregisterTabSpawners();
+	
+	/** Registers the light card editor tab with the operator panel using a layout extension */
+	static void RegisterLayoutExtension(FLayoutExtender& InExtender);
+
+	/** Spawns in the viewport tab */
+	TSharedRef<SDockTab> SpawnViewportTab(const FSpawnTabArgs& SpawnTabArgs);
+	
+	/** Spawns in the outliner tab */
+	TSharedRef<SDockTab> SpawnOutlinerTab(const FSpawnTabArgs& SpawnTabArgs);
+	
+	/** Creates the widget used to show the list of light cards associated with the active root actor */
+	TSharedRef<SWidget> CreateLightCardOutlinerWidget();
+
+	/** Create the widget for selecting light card templates */
+	TSharedRef<SWidget> CreateLightCardTemplateWidget();
+	
+	/** Create the 3d viewport widget */
+	TSharedRef<SWidget> CreateViewportWidget();
+
+	/** Generate the place actors drop down menu */
+	TSharedRef<SWidget> GeneratePlaceActorsMenu();
+
+	/** Return the correct template icon to use */
+	const FSlateBrush* GetLightCardTemplateIcon(const TWeakObjectPtr<UDisplayClusterLightCardTemplate> InTemplate) const;
+
+	/** Generate the All Templates sub menu */
+	void GenerateTemplateSubMenu(FMenuBuilder& InMenuBuilder);
+	
+	/** Generate the labels dropdown menu */
+	TSharedRef<SWidget> GenerateLabelsMenu();
+
+	/** Creates the editor's command list and binds commands to it */
+	void BindCommands();
+
+	/** Register any extensions for the toolbar */
+	void RegisterToolbarExtensions();
+
+	/** Register any extensions for the menu bar */
+	void RegisterMenuExtensions();
+
+	/** Extend the toolbar for this local instance */
+	void ExtendToolbar(FToolBarBuilder& ToolbarBuilder);
+
+	/** Extend the main file menu */
+	void ExtendFileMenu(FMenuBuilder& MenuBuilder);
+	
+	/** Extend the main edit menu */
+	void ExtendEditMenu(FMenuBuilder& MenuBuilder);
+	
+	/** Refresh all preview actors */
+	void RefreshPreviewActors(EDisplayClusterLightCardEditorProxyType ProxyType = EDisplayClusterLightCardEditorProxyType::All);
+
+	/** Reapply the current label settings */
+	void RefreshLabels();
+	
+	/**
+	 * Check if an object is managed by us
+	 * @param InObject The object to compare
+	 * @param OutProxyType The type of the object
+	 * @return True if our object, false if not
+	 */
+	bool IsOurObject(UObject* InObject, EDisplayClusterLightCardEditorProxyType& OutProxyType) const;
+	
+	/** Bind delegates to when a BP compiles */
+	void BindCompileDelegates();
+
+	/** Remove compile delegates from a BP */
+	void RemoveCompileDelegates();
+
+	/** When a property on the actor has changed */
+	void OnActorPropertyChanged(UObject* ObjectBeingModified, FPropertyChangedEvent& PropertyChangedEvent);
+
+	/** Raised when the user deletes an actor from the level */
+	void OnLevelActorDeleted(AActor* Actor);
+
+	/** Raised when a supported blueprint is compiled */
+	void OnBlueprintCompiled(UBlueprint* Blueprint);
+
+	/** Raised when any object is transacted */
+	void OnObjectTransacted(UObject* Object, const FTransactionObjectEvent& TransactionObjectEvent);
+
+	/** Add an item to recently placed list. Reorders matching items to the top of the list and trims the list */
+	void AddRecentlyPlacedItem(const FDisplayClusterLightCardEditorRecentItem& InItem);
+
+	/** Make sure no invalid entries are present */
+	void CleanupRecentlyPlacedItems();
+	
+private:
+	/** The light card outliner widget */
+	TSharedPtr<SDisplayClusterLightCardOutliner> LightCardOutliner;
+
+	/** Templates to their icon brushes */
+	TMap<TWeakObjectPtr<UDisplayClusterLightCardTemplate>, TSharedPtr<FSlateBrush>> TemplateBrushes;
+	
+	/** The 3d viewport */
+	TSharedPtr<SDisplayClusterLightCardEditorViewport> ViewportView;
+
+	/** The command list for editor commands */
+	TSharedPtr<FUICommandList> CommandList;
+
+	/** Stores the mouse position when the context menu was opened */
+	TOptional<FIntPoint> CachedContextMenuMousePos;
+
+	/** A reference to the root actor that is currently being operated on */
+	TWeakObjectPtr<ADisplayClusterRootActor> ActiveRootActor;
+
+	/** The view model for the operator panel */
+	TWeakPtr<IDisplayClusterOperatorViewModel> OperatorViewModel;
+
+	/** Selected actors */
+	TArray<TWeakObjectPtr<AActor>> SelectedActors;
+
+	/** Delegate handle for the OnActiveRootActorChanged delegate */
+	FDelegateHandle ActiveRootActorChangedHandle;
+
+	/** Delegate handle for when an object is transacted */
+	FDelegateHandle OnObjectTransactedHandle;
+};
\ No newline at end of file
diff --git a/Engine/Plugins/Runtime/nDisplay/Source/DisplayClusterLightCardEditor/Private/DisplayClusterLightCardEditorCommands.cpp b/Engine/Plugins/Runtime/nDisplay/Source/DisplayClusterLightCardEditor/Private/DisplayClusterLightCardEditorCommands.cpp
--- a/Engine/Plugins/Runtime/nDisplay/Source/DisplayClusterLightCardEditor/Private/DisplayClusterLightCardEditorCommands.cpp
+++ b/Engine/Plugins/Runtime/nDisplay/Source/DisplayClusterLightCardEditor/Private/DisplayClusterLightCardEditorCommands.cpp
@@ -33,7 +33,7 @@ void FDisplayClusterLightCardEditorCommands::RegisterCommands()
 
 	UI_COMMAND(AddNewLightCard, "Add New Light Card", "Add and assign a new Light Card to the actor", EUserInterfaceActionType::Button, FInputChord());
 	UI_COMMAND(AddNewFlag, "Add Flag", "Add and assign a new light control flag", EUserInterfaceActionType::Button, FInputChord());
-	UI_COMMAND(AddExistingLightCard, "Add Existing Light Card", "Add an existing Light Card to the actor", EUserInterfaceActionType::Button, FInputChord());
+	UI_COMMAND(AddExistingLightCard, "Import Existing Content", "Add an existing Light Card to the actor", EUserInterfaceActionType::Button, FInputChord());
 	UI_COMMAND(RemoveLightCard, "Remove from Actor", "Remove the Light Card from the actor but do not delete it", EUserInterfaceActionType::Button, FInputChord());
 	UI_COMMAND(PasteHere, "Paste Here", "Paste clipboard contents at the click location", EUserInterfaceActionType::Button, FInputChord());
 
diff --git a/Engine/Plugins/Runtime/nDisplay/Source/DisplayClusterLightCardEditor/Private/DisplayClusterLightCardEditorModule.cpp b/Engine/Plugins/Runtime/nDisplay/Source/DisplayClusterLightCardEditor/Private/DisplayClusterLightCardEditorModule.cpp
--- a/Engine/Plugins/Runtime/nDisplay/Source/DisplayClusterLightCardEditor/Private/DisplayClusterLightCardEditorModule.cpp
+++ b/Engine/Plugins/Runtime/nDisplay/Source/DisplayClusterLightCardEditor/Private/DisplayClusterLightCardEditorModule.cpp
@@ -6,29 +6,43 @@
 #include "DisplayClusterLightCardActor.h"
 #include "DisplayClusterLightCardEditorCommands.h"
 #include "DisplayClusterRootActor.h"
-#include "SDisplayClusterLightCardEditor.h"
+#include "DisplayClusterLightCardEditor.h"
 #include "DetailCustomizations/DisplayClusterLightCardActorDetails.h"
 #include "Settings/DisplayClusterLightCardEditorSettings.h"
 
+#include "IDisplayClusterOperator.h"
+
 #include "ISettingsModule.h"
-#include "ScopedTransaction.h"
 
 #define LOCTEXT_NAMESPACE "DisplayClusterLightCardEditor"
 
 void FDisplayClusterLightCardEditorModule::StartupModule()
 {
-	RegisterTabSpawners();
 	RegisterSettings();
 	RegisterDetailCustomizations();
 	FDisplayClusterLightCardEditorCommands::Register();
+
+	if (FModuleManager::Get().IsModuleLoaded(IDisplayClusterOperator::ModuleName))
+	{
+		RegisterOperatorApp();
+	}
+	else
+	{
+		ModuleChangedHandle = FModuleManager::Get().OnModulesChanged().AddRaw(this, &FDisplayClusterLightCardEditorModule::HandleModuleChanged);
+	}
 }
 
 void FDisplayClusterLightCardEditorModule::ShutdownModule()
 {
-	UnregisterTabSpawners();
 	UnregisterSettings();
 	UnregisterDetailCustomizations();
+	UnregisterOperatorApp();
 
+	if (ModuleChangedHandle.IsValid())
+	{
+		FModuleManager::Get().OnModulesChanged().Remove(ModuleChangedHandle);
+	}
+	
 	FDisplayClusterLightCardEditorCommands::Unregister();
 }
 
@@ -56,16 +70,6 @@ void FDisplayClusterLightCardEditorModule::ShowLabels(const FLabelArgs& InArgs)
 	}
 }
 
-void FDisplayClusterLightCardEditorModule::RegisterTabSpawners()
-{
-	SDisplayClusterLightCardEditor::RegisterTabSpawner();
-}
-
-void FDisplayClusterLightCardEditorModule::UnregisterTabSpawners()
-{
-	SDisplayClusterLightCardEditor::UnregisterTabSpawner();
-}
-
 void FDisplayClusterLightCardEditorModule::RegisterSettings()
 {
 	if (ISettingsModule* SettingsModule = FModuleManager::GetModulePtr<ISettingsModule>("Settings"))
@@ -106,6 +110,29 @@ void FDisplayClusterLightCardEditorModule::UnregisterDetailCustomizations()
 	}
 }
 
+void FDisplayClusterLightCardEditorModule::RegisterOperatorApp()
+{
+	OperatorAppHandle = IDisplayClusterOperator::Get().RegisterApp(IDisplayClusterOperator::FOnGetAppInstance::CreateStatic(&FDisplayClusterLightCardEditor::MakeInstance));
+}
+
+void FDisplayClusterLightCardEditorModule::UnregisterOperatorApp()
+{
+	if (IDisplayClusterOperator* OperatorModule = FModuleManager::GetModulePtr<IDisplayClusterOperator>(IDisplayClusterOperator::ModuleName))
+	{
+		OperatorModule->UnregisterApp(OperatorAppHandle);
+	}
+}
+
+void FDisplayClusterLightCardEditorModule::HandleModuleChanged(FName InModuleName, EModuleChangeReason InChangeReason)
+{
+	if (InModuleName == IDisplayClusterOperator::ModuleName && InChangeReason == EModuleChangeReason::ModuleLoaded)
+	{
+		RegisterOperatorApp();
+		FModuleManager::Get().OnModulesChanged().Remove(ModuleChangedHandle);
+		ModuleChangedHandle.Reset();
+	}
+}
+
 IMPLEMENT_MODULE(FDisplayClusterLightCardEditorModule, DisplayClusterLightCardEditor);
 
 #undef LOCTEXT_NAMESPACE
diff --git a/Engine/Plugins/Runtime/nDisplay/Source/DisplayClusterLightCardEditor/Private/DisplayClusterLightCardEditorModule.h b/Engine/Plugins/Runtime/nDisplay/Source/DisplayClusterLightCardEditor/Private/DisplayClusterLightCardEditorModule.h
--- a/Engine/Plugins/Runtime/nDisplay/Source/DisplayClusterLightCardEditor/Private/DisplayClusterLightCardEditorModule.h
+++ b/Engine/Plugins/Runtime/nDisplay/Source/DisplayClusterLightCardEditor/Private/DisplayClusterLightCardEditorModule.h
@@ -3,7 +3,6 @@
 #pragma once
 
 #include "IDisplayClusterLightCardEditor.h"
-#include "UObject/WeakObjectPtr.h"
 
 /**
  * Display Cluster editor module
@@ -19,12 +18,18 @@ public:
 	virtual void ShowLabels(const FLabelArgs& InArgs) override;
 	
 private:
-	void RegisterTabSpawners();
-	void UnregisterTabSpawners();
-
 	void RegisterSettings();
 	void UnregisterSettings();
 
 	void RegisterDetailCustomizations();
 	void UnregisterDetailCustomizations();
+
+	void RegisterOperatorApp();
+	void UnregisterOperatorApp();
+	
+	void HandleModuleChanged(FName InModuleName, EModuleChangeReason InChangeReason);
+	
+private:
+	FDelegateHandle ModuleChangedHandle;
+	FDelegateHandle OperatorAppHandle;
 };
diff --git a/Engine/Plugins/Runtime/nDisplay/Source/DisplayClusterLightCardEditor/Private/LightCardTemplates/SDisplayClusterLightCardTemplateList.cpp b/Engine/Plugins/Runtime/nDisplay/Source/DisplayClusterLightCardEditor/Private/LightCardTemplates/SDisplayClusterLightCardTemplateList.cpp
--- a/Engine/Plugins/Runtime/nDisplay/Source/DisplayClusterLightCardEditor/Private/LightCardTemplates/SDisplayClusterLightCardTemplateList.cpp
+++ b/Engine/Plugins/Runtime/nDisplay/Source/DisplayClusterLightCardEditor/Private/LightCardTemplates/SDisplayClusterLightCardTemplateList.cpp
@@ -5,7 +5,7 @@
 #include "DisplayClusterLightCardTemplate.h"
 #include "DisplayClusterLightCardTemplateDragDropOp.h"
 #include "DisplayClusterLightCardTemplateHelpers.h"
-#include "SDisplayClusterLightCardEditor.h"
+#include "DisplayClusterLightCardEditor.h"
 
 #include "AssetRegistry/AssetRegistryModule.h"
 #include "Engine/Texture.h"
@@ -182,7 +182,7 @@ SDisplayClusterLightCardTemplateList::~SDisplayClusterLightCardTemplateList()
 	AssetRegistry.OnFilesLoaded().RemoveAll(this);
 }
 
-void SDisplayClusterLightCardTemplateList::Construct(const FArguments& InArgs, TSharedPtr<SDisplayClusterLightCardEditor> InLightCardEditor)
+void SDisplayClusterLightCardTemplateList::Construct(const FArguments& InArgs, TSharedPtr<FDisplayClusterLightCardEditor> InLightCardEditor)
 {
 	LightCardEditorPtr = InLightCardEditor;
 
diff --git a/Engine/Plugins/Runtime/nDisplay/Source/DisplayClusterLightCardEditor/Private/LightCardTemplates/SDisplayClusterLightCardTemplateList.h b/Engine/Plugins/Runtime/nDisplay/Source/DisplayClusterLightCardEditor/Private/LightCardTemplates/SDisplayClusterLightCardTemplateList.h
--- a/Engine/Plugins/Runtime/nDisplay/Source/DisplayClusterLightCardEditor/Private/LightCardTemplates/SDisplayClusterLightCardTemplateList.h
+++ b/Engine/Plugins/Runtime/nDisplay/Source/DisplayClusterLightCardEditor/Private/LightCardTemplates/SDisplayClusterLightCardTemplateList.h
@@ -6,7 +6,7 @@
 #include "Widgets/DeclarativeSyntaxSupport.h"
 #include "Widgets/SCompoundWidget.h"
 
-class SDisplayClusterLightCardEditor;
+class FDisplayClusterLightCardEditor;
 class UDisplayClusterLightCardTemplate;
 class ITableRow;
 class STableViewBase;
@@ -38,7 +38,7 @@ public:
 
 	virtual ~SDisplayClusterLightCardTemplateList() override;
 
-	void Construct(const FArguments& InArgs, TSharedPtr<SDisplayClusterLightCardEditor> InLightCardEditor);
+	void Construct(const FArguments& InArgs, TSharedPtr<FDisplayClusterLightCardEditor> InLightCardEditor);
 
 	/** Find available templates and refresh the list */
 	void RefreshTemplateList();
@@ -65,7 +65,7 @@ private:
 	
 private:
 	/** Pointer to the light card editor that owns this widget */
-	TWeakPtr<SDisplayClusterLightCardEditor> LightCardEditorPtr;
+	TWeakPtr<FDisplayClusterLightCardEditor> LightCardEditorPtr;
 
 	/** A list of all light card tree items to be displayed in the tree view */
 	TArray<TSharedPtr<FLightCardTemplateTreeItem>> LightCardTemplateTree;
diff --git a/Engine/Plugins/Runtime/nDisplay/Source/DisplayClusterLightCardEditor/Private/Outliner/SDisplayClusterLightCardOutliner.cpp b/Engine/Plugins/Runtime/nDisplay/Source/DisplayClusterLightCardEditor/Private/Outliner/SDisplayClusterLightCardOutliner.cpp
--- a/Engine/Plugins/Runtime/nDisplay/Source/DisplayClusterLightCardEditor/Private/Outliner/SDisplayClusterLightCardOutliner.cpp
+++ b/Engine/Plugins/Runtime/nDisplay/Source/DisplayClusterLightCardEditor/Private/Outliner/SDisplayClusterLightCardOutliner.cpp
@@ -4,7 +4,7 @@
 
 #include "DisplayClusterLightCardEditorCommands.h"
 #include "DisplayClusterLightCardOutlinerMode.h"
-#include "SDisplayClusterLightCardEditor.h"
+#include "DisplayClusterLightCardEditor.h"
 
 #include "IDisplayClusterOperator.h"
 
@@ -37,27 +37,37 @@ SDisplayClusterLightCardOutliner::~SDisplayClusterLightCardOutliner()
 	}
 }
 
-void SDisplayClusterLightCardOutliner::Construct(const FArguments& InArgs, TSharedPtr<SDisplayClusterLightCardEditor> InLightCardEditor, TSharedPtr<FUICommandList> InCommandList)
+void SDisplayClusterLightCardOutliner::Construct(const FArguments& InArgs, TSharedPtr<FDisplayClusterLightCardEditor> InLightCardEditor, TSharedPtr<FUICommandList> InCommandList)
 {
 	LightCardEditorPtr = InLightCardEditor;
-	CommandList = InCommandList;
+	check(LightCardEditorPtr.IsValid());
 	
+	CommandList = InCommandList;
+
+	RootActor = InLightCardEditor->GetActiveRootActor();
 	CreateWorldOutliner();
+
+	SelectActors(LightCardEditorPtr.Pin()->GetSelectedActorsAs<AActor>());
 }
 
 void SDisplayClusterLightCardOutliner::SetRootActor(ADisplayClusterRootActor* NewRootActor)
 {
 	TArray<AActor*> PreviouslySelectedActors;
-	GetSelectedActors(PreviouslySelectedActors);
+	if (LightCardEditorPtr.IsValid())
+	{
+		PreviouslySelectedActors = LightCardEditorPtr.Pin()->GetSelectedActorsAs<AActor>();
+	}
 	
 	RootActor = NewRootActor;
-	FillActorList();
 
 	CreateWorldOutliner();
 	
 	// Select previously selected light cards. This fixes an issue where the details panel may clear in some cases
 	// and also supports maintaining the selection of shared light cards across different root actors.
-	SelectActors(MoveTemp(PreviouslySelectedActors));
+	if (LightCardEditorPtr.IsValid())
+	{
+		LightCardEditorPtr.Pin()->SelectActors(MoveTemp(PreviouslySelectedActors));
+	}
 }
 
 void SDisplayClusterLightCardOutliner::GetSelectedActors(TArray<AActor*>& OutSelectedActors) const
@@ -75,7 +85,6 @@ void SDisplayClusterLightCardOutliner::GetSelectedActors(TArray<AActor*>& OutSel
 
 void SDisplayClusterLightCardOutliner::SelectActors(const TArray<AActor*>& ActorsToSelect)
 {
-	TArray<AActor*> ActualSelectedActors;
 	TArray<FSceneOutlinerTreeItemPtr> SelectedTreeItems;
 	CachedSelectedActors.Reset();
 	
@@ -86,15 +95,12 @@ void SDisplayClusterLightCardOutliner::SelectActors(const TArray<AActor*>& Actor
 			if (FSceneOutlinerTreeItemPtr OutlinerItem = SceneOutliner->GetTreeItem(TreeItem->Actor.Get()))
 			{
 				SelectedTreeItems.Add(OutlinerItem);
-				ActualSelectedActors.Add(TreeItem->Actor.Get());
 				CachedSelectedActors.Add(TreeItem->Actor);
 			}
 		}
 	}
 	
 	SceneOutliner->SetItemSelection(SelectedTreeItems, true);
-	
-	IDisplayClusterOperator::Get().ShowDetailsForObjects(*reinterpret_cast<TArray<UObject*>*>(&ActualSelectedActors));
 }
 
 void SDisplayClusterLightCardOutliner::RestoreCachedSelection()
@@ -107,12 +113,17 @@ void SDisplayClusterLightCardOutliner::RestoreCachedSelection()
 			ActorsToSelect.Add(CachedActor.Get());
 		}
 	}
-	
-	SelectActors(MoveTemp(ActorsToSelect));
+
+	if (LightCardEditorPtr.IsValid())
+	{
+		LightCardEditorPtr.Pin()->SelectActors(MoveTemp(ActorsToSelect));
+	}
 }
 
 void SDisplayClusterLightCardOutliner::CreateWorldOutliner()
 {
+	FillActorList();
+	
 	auto OutlinerFilterPredicate = [this](const AActor* InActor)
 	{
 		return TrackedActors.Contains(InActor);
@@ -247,10 +258,9 @@ void SDisplayClusterLightCardOutliner::OnOutlinerSelectionChanged(FSceneOutliner
 	TArray<AActor*> SelectedLightCards;
 	GetSelectedActors(SelectedLightCards);
 
-	IDisplayClusterOperator::Get().ShowDetailsForObjects(*reinterpret_cast<TArray<UObject*>*>(&SelectedLightCards));
-
 	if (LightCardEditorPtr.IsValid())
 	{
+		LightCardEditorPtr.Pin()->SelectActors(SelectedLightCards);
 		LightCardEditorPtr.Pin()->SelectActorProxies(SelectedLightCards);
 	}
 }
diff --git a/Engine/Plugins/Runtime/nDisplay/Source/DisplayClusterLightCardEditor/Private/Outliner/SDisplayClusterLightCardOutliner.h b/Engine/Plugins/Runtime/nDisplay/Source/DisplayClusterLightCardEditor/Private/Outliner/SDisplayClusterLightCardOutliner.h
--- a/Engine/Plugins/Runtime/nDisplay/Source/DisplayClusterLightCardEditor/Private/Outliner/SDisplayClusterLightCardOutliner.h
+++ b/Engine/Plugins/Runtime/nDisplay/Source/DisplayClusterLightCardEditor/Private/Outliner/SDisplayClusterLightCardOutliner.h
@@ -11,7 +11,7 @@ class ADisplayClusterLightCardActor;
 class ITableRow;
 class FExtender;
 class FUICommandList;
-class SDisplayClusterLightCardEditor;
+class FDisplayClusterLightCardEditor;
 class STableViewBase;
 
 class SSceneOutliner;
@@ -38,7 +38,7 @@ public:
 
 	virtual ~SDisplayClusterLightCardOutliner() override;
 
-	void Construct(const FArguments& InArgs, TSharedPtr<SDisplayClusterLightCardEditor> InLightCardEditor, TSharedPtr<FUICommandList> InCommandList);
+	void Construct(const FArguments& InArgs, TSharedPtr<FDisplayClusterLightCardEditor> InLightCardEditor, TSharedPtr<FUICommandList> InCommandList);
 
 	void SetRootActor(ADisplayClusterRootActor* NewRootActor);
 
@@ -47,24 +47,7 @@ public:
 	/** Gets a list of light card actors that are currently selected in the list */
 	void GetSelectedActors(TArray<AActor*>& OutSelectedActors) const;
 
-	template<typename T>
-	TArray<T*> GetSelectedActorsAs() const
-	{
-		TArray<AActor*> SelectedActors;
-		GetSelectedActors(SelectedActors);
-		TArray<T*> OutArray;
-		Algo::TransformIf(SelectedActors, OutArray, [](AActor* InItem)
-		{
-			return InItem && InItem->IsA<T>();
-		},
-		[](AActor* InItem)
-		{
-			return CastChecked<T>(InItem);
-		});
-		return OutArray;
-	}
-
-	/** Select light cards in the outliner and display their details */
+	/** Select light cards in the outliner */
 	void SelectActors(const TArray<AActor*>& ActorsToSelect);
 
 	/** Restores the last valid cached selection to the outliner */
@@ -96,7 +79,7 @@ private:
 
 private:
 	/** Pointer to the light card editor that owns this widget */
-	TWeakPtr<SDisplayClusterLightCardEditor> LightCardEditorPtr;
+	TWeakPtr<FDisplayClusterLightCardEditor> LightCardEditorPtr;
 
 	/** A hierarchical list of stage actor tree items to be displayed in the tree view */
 	TArray<TSharedPtr<FStageActorTreeItem>> StageActorTree;
diff --git a/Engine/Plugins/Runtime/nDisplay/Source/DisplayClusterLightCardEditor/Private/SDisplayClusterLightCardEditor.cpp b/Engine/Plugins/Runtime/nDisplay/Source/DisplayClusterLightCardEditor/Private/SDisplayClusterLightCardEditor.cpp
deleted file mode 100644
diff --git a/Engine/Plugins/Runtime/nDisplay/Source/DisplayClusterLightCardEditor/Private/SDisplayClusterLightCardEditor.h b/Engine/Plugins/Runtime/nDisplay/Source/DisplayClusterLightCardEditor/Private/SDisplayClusterLightCardEditor.h
deleted file mode 100644
diff --git a/Engine/Plugins/Runtime/nDisplay/Source/DisplayClusterLightCardEditor/Private/Viewport/DisplayClusterLightCardEditorViewportClient.cpp b/Engine/Plugins/Runtime/nDisplay/Source/DisplayClusterLightCardEditor/Private/Viewport/DisplayClusterLightCardEditorViewportClient.cpp
--- a/Engine/Plugins/Runtime/nDisplay/Source/DisplayClusterLightCardEditor/Private/Viewport/DisplayClusterLightCardEditorViewportClient.cpp
+++ b/Engine/Plugins/Runtime/nDisplay/Source/DisplayClusterLightCardEditor/Private/Viewport/DisplayClusterLightCardEditorViewportClient.cpp
@@ -17,7 +17,7 @@
 #include "DisplayClusterLightCardEditorLog.h"
 #include "DisplayClusterLightCardEditorUtils.h"
 #include "IDisplayClusterScenePreview.h"
-#include "SDisplayClusterLightCardEditor.h"
+#include "DisplayClusterLightCardEditor.h"
 #include "Settings/DisplayClusterLightCardEditorSettings.h"
 
 #include "AudioDevice.h"
diff --git a/Engine/Plugins/Runtime/nDisplay/Source/DisplayClusterLightCardEditor/Private/Viewport/DisplayClusterLightCardEditorViewportClient.h b/Engine/Plugins/Runtime/nDisplay/Source/DisplayClusterLightCardEditor/Private/Viewport/DisplayClusterLightCardEditorViewportClient.h
--- a/Engine/Plugins/Runtime/nDisplay/Source/DisplayClusterLightCardEditor/Private/Viewport/DisplayClusterLightCardEditorViewportClient.h
+++ b/Engine/Plugins/Runtime/nDisplay/Source/DisplayClusterLightCardEditor/Private/Viewport/DisplayClusterLightCardEditorViewportClient.h
@@ -15,7 +15,7 @@
 #include "StageActor/DisplayClusterWeakStageActorPtr.h"
 
 class ADisplayClusterRootActor;
-class SDisplayClusterLightCardEditor;
+class FDisplayClusterLightCardEditor;
 class SDisplayClusterLightCardEditorViewport;
 class FScopedTransaction;
 class UDisplayClusterConfigurationViewport;
@@ -256,7 +256,7 @@ private:
 private:
 	TWeakPtr<FSceneViewport> SceneViewportPtr;
 	TWeakPtr<SDisplayClusterLightCardEditorViewport> LightCardEditorViewportPtr;
-	TWeakPtr<SDisplayClusterLightCardEditor> LightCardEditorPtr;
+	TWeakPtr<FDisplayClusterLightCardEditor> LightCardEditorPtr;
 	TWeakObjectPtr<ADisplayClusterRootActor> RootActorProxy;
 	TWeakObjectPtr<ADisplayClusterRootActor> RootActorLevelInstance;
 
diff --git a/Engine/Plugins/Runtime/nDisplay/Source/DisplayClusterLightCardEditor/Private/Viewport/DisplayClusterLightcardEditorViewport.cpp b/Engine/Plugins/Runtime/nDisplay/Source/DisplayClusterLightCardEditor/Private/Viewport/DisplayClusterLightcardEditorViewport.cpp
--- a/Engine/Plugins/Runtime/nDisplay/Source/DisplayClusterLightCardEditor/Private/Viewport/DisplayClusterLightcardEditorViewport.cpp
+++ b/Engine/Plugins/Runtime/nDisplay/Source/DisplayClusterLightCardEditor/Private/Viewport/DisplayClusterLightcardEditorViewport.cpp
@@ -5,7 +5,7 @@
 #include "DisplayClusterLightCardEditorCommands.h"
 #include "DisplayClusterLightCardEditorStyle.h"
 #include "DisplayClusterLightCardEditorViewportClient.h"
-#include "SDisplayClusterLightCardEditor.h"
+#include "DisplayClusterLightCardEditor.h"
 #include "LightCardTemplates/DisplayClusterLightCardTemplate.h"
 #include "LightCardTemplates/DisplayClusterLightCardTemplateDragDropOp.h"
 
@@ -315,7 +315,7 @@ const FVector SDisplayClusterLightCardEditorViewport::ViewDirectionRight = FVect
 const FVector SDisplayClusterLightCardEditorViewport::ViewDirectionFront = FVector(1.0f, 0.0f, 0.0f);
 const FVector SDisplayClusterLightCardEditorViewport::ViewDirectionBack = FVector(-1.0f, 0.0f, 0.0f);
 
-void SDisplayClusterLightCardEditorViewport::Construct(const FArguments& InArgs, TSharedPtr<SDisplayClusterLightCardEditor> InLightCardEditor, TSharedPtr<class FUICommandList> InCommandList)
+void SDisplayClusterLightCardEditorViewport::Construct(const FArguments& InArgs, TSharedPtr<FDisplayClusterLightCardEditor> InLightCardEditor, TSharedPtr<class FUICommandList> InCommandList)
 {
 	LightCardEditorPtr = InLightCardEditor;
 	
@@ -327,6 +327,8 @@ void SDisplayClusterLightCardEditorViewport::Construct(const FArguments& InArgs,
 	{
 		CommandList->Append(InCommandList.ToSharedRef());
 	}
+
+	SetRootActor(LightCardEditorPtr.Pin()->GetActiveRootActor().Get());
 }
 
 SDisplayClusterLightCardEditorViewport::~SDisplayClusterLightCardEditorViewport()
diff --git a/Engine/Plugins/Runtime/nDisplay/Source/DisplayClusterLightCardEditor/Private/Viewport/DisplayClusterLightcardEditorViewport.h b/Engine/Plugins/Runtime/nDisplay/Source/DisplayClusterLightCardEditor/Private/Viewport/DisplayClusterLightcardEditorViewport.h
--- a/Engine/Plugins/Runtime/nDisplay/Source/DisplayClusterLightCardEditor/Private/Viewport/DisplayClusterLightcardEditorViewport.h
+++ b/Engine/Plugins/Runtime/nDisplay/Source/DisplayClusterLightCardEditor/Private/Viewport/DisplayClusterLightcardEditorViewport.h
@@ -10,7 +10,7 @@
 #include "DisplayClusterMeshProjectionRenderer.h"
 #include "DisplayClusterLightCardEditorWidget.h"
 
-class SDisplayClusterLightCardEditor;
+class FDisplayClusterLightCardEditor;
 class ADisplayClusterRootActor;
 class FDisplayClusterLightCardEditorViewportClient;
 
@@ -31,7 +31,7 @@ public:
 	SLATE_BEGIN_ARGS(SDisplayClusterLightCardEditorViewport) {}
 	SLATE_END_ARGS()
 
-	void Construct(const FArguments& InArgs, TSharedPtr<SDisplayClusterLightCardEditor> InLightCardEditor, TSharedPtr<class FUICommandList> InCommandList);
+	void Construct(const FArguments& InArgs, TSharedPtr<FDisplayClusterLightCardEditor> InLightCardEditor, TSharedPtr<class FUICommandList> InCommandList);
 	~SDisplayClusterLightCardEditorViewport();
 
 	// ICommonEditorViewportToolbarInfoProvider
@@ -44,7 +44,7 @@ public:
 	virtual FReply OnKeyDown(const FGeometry& MyGeometry, const FKeyEvent& InKeyEvent) override;
 	// ~SWidget
 
-	TWeakPtr<SDisplayClusterLightCardEditor> GetLightCardEditor() { return LightCardEditorPtr;}
+	TWeakPtr<FDisplayClusterLightCardEditor> GetLightCardEditor() { return LightCardEditorPtr;}
 
 	void SetRootActor(ADisplayClusterRootActor* NewRootActor);
 	
@@ -116,7 +116,7 @@ private:
 	
 	/** Level viewport client */
 	TSharedPtr<FDisplayClusterLightCardEditorViewportClient> ViewportClient;
-	TWeakPtr<SDisplayClusterLightCardEditor> LightCardEditorPtr;
+	TWeakPtr<FDisplayClusterLightCardEditor> LightCardEditorPtr;
 
 	/** The cached mouse position used for the Paste Here command */
 	FVector2D PasteHerePos;
diff --git a/Engine/Plugins/Runtime/nDisplay/Source/DisplayClusterOperator/DisplayClusterOperator.Build.cs b/Engine/Plugins/Runtime/nDisplay/Source/DisplayClusterOperator/DisplayClusterOperator.Build.cs
--- a/Engine/Plugins/Runtime/nDisplay/Source/DisplayClusterOperator/DisplayClusterOperator.Build.cs
+++ b/Engine/Plugins/Runtime/nDisplay/Source/DisplayClusterOperator/DisplayClusterOperator.Build.cs
@@ -20,6 +20,7 @@ public class DisplayClusterOperator : ModuleRules
 				"Projects",
 				"Slate",
 				"SlateCore",
+				"ToolMenus",
 				"ToolWidgets",
 				"UnrealEd",
 				"WorkspaceMenuStructure",
diff --git a/Engine/Plugins/Runtime/nDisplay/Source/DisplayClusterOperator/Private/DisplayClusterOperatorModule.cpp b/Engine/Plugins/Runtime/nDisplay/Source/DisplayClusterOperator/Private/DisplayClusterOperatorModule.cpp
--- a/Engine/Plugins/Runtime/nDisplay/Source/DisplayClusterOperator/Private/DisplayClusterOperatorModule.cpp
+++ b/Engine/Plugins/Runtime/nDisplay/Source/DisplayClusterOperator/Private/DisplayClusterOperatorModule.cpp
@@ -24,6 +24,7 @@ void FDisplayClusterOperatorModule::StartupModule()
 {
 	OperatorViewModel = MakeShared<FDisplayClusterOperatorViewModel>();
 	OperatorToolBarExtensibilityManager = MakeShared<FExtensibilityManager>();
+	OperatorMenuExtensibilityManager = MakeShared<FExtensibilityManager>();
 
 	FDisplayClusterOperatorStyle::Get();
 	
@@ -35,6 +36,19 @@ void FDisplayClusterOperatorModule::ShutdownModule()
 	UnregisterTabSpawners();
 }
 
+FDelegateHandle FDisplayClusterOperatorModule::RegisterApp(const FOnGetAppInstance& InGetAppInstanceDelegate)
+{
+	const FDelegateHandle Handle = InGetAppInstanceDelegate.GetHandle();
+	RegisteredApps.Add(Handle, InGetAppInstanceDelegate);
+	return Handle;
+}
+
+void FDisplayClusterOperatorModule::UnregisterApp(const FDelegateHandle& InHandle)
+{
+	AppUnregisteredEvent.Broadcast(InHandle);
+	RegisteredApps.Remove(InHandle);
+}
+
 TSharedRef<IDisplayClusterOperatorViewModel> FDisplayClusterOperatorModule::GetOperatorViewModel()
 {
 	return OperatorViewModel.ToSharedRef();
@@ -117,7 +131,7 @@ void FDisplayClusterOperatorModule::UnregisterTabSpawners()
 TSharedRef<SDockTab> FDisplayClusterOperatorModule::SpawnOperatorPanelTab(const FSpawnTabArgs& SpawnTabArgs)
 {
 	const TSharedRef<SDockTab> MajorTab = SNew(SDockTab)
-		.TabRole(ETabRole::NomadTab)
+		.TabRole(ETabRole::MajorTab)
 		.IconColor(FDisplayClusterOperatorStyle::Get().GetColor(TEXT("OperatorPanel.AssetColor")))
 		.OnTabClosed_Raw(this, &FDisplayClusterOperatorModule::OnOperatorPanelTabClosed);
 
diff --git a/Engine/Plugins/Runtime/nDisplay/Source/DisplayClusterOperator/Private/DisplayClusterOperatorModule.h b/Engine/Plugins/Runtime/nDisplay/Source/DisplayClusterOperator/Private/DisplayClusterOperatorModule.h
--- a/Engine/Plugins/Runtime/nDisplay/Source/DisplayClusterOperator/Private/DisplayClusterOperatorModule.h
+++ b/Engine/Plugins/Runtime/nDisplay/Source/DisplayClusterOperator/Private/DisplayClusterOperatorModule.h
@@ -16,6 +16,8 @@ class FDisplayClusterOperatorModule :
 	public IDisplayClusterOperator
 {
 public:
+	DECLARE_MULTICAST_DELEGATE_OneParam(FOnAppUnregistered, const FDelegateHandle&)
+	
 	/** The name of the tab that the operator panel lives in */
 	static const FName OperatorPanelTabName;
 
@@ -26,6 +28,8 @@ public:
 	//~ End IModuleInterface interface
 
 	//~ IDisplayClusterOperator interface
+	virtual FDelegateHandle RegisterApp(const FOnGetAppInstance& InGetAppInstanceDelegate) override;
+	virtual void UnregisterApp(const FDelegateHandle& InHandle) override;
 	virtual TSharedRef<IDisplayClusterOperatorViewModel> GetOperatorViewModel() override;
 	virtual FOnRegisterLayoutExtensions& OnRegisterLayoutExtensions() override { return RegisterLayoutExtensions; }
 	virtual FOnRegisterStatusBarExtensions& OnRegisterStatusBarExtensions() override { return RegisterStatusBarExtensions; }
@@ -36,6 +40,7 @@ public:
 	virtual FName GetAuxilliaryOperatorExtensionId() override;
 	virtual FName GetDetailsTabId() override;
 	virtual TSharedPtr<FExtensibilityManager> GetOperatorToolBarExtensibilityManager() override { return OperatorToolBarExtensibilityManager; }
+	virtual TSharedPtr<FExtensibilityManager> GetOperatorMenuExtensibilityManager() override { return OperatorMenuExtensibilityManager; }
 	virtual void GetRootActorLevelInstances(TArray<ADisplayClusterRootActor*>& OutRootActorInstances) override;
 	virtual void ShowDetailsForObject(UObject* Object) override;
 	virtual void ShowDetailsForObjects(const TArray<UObject*>& Objects) override;
@@ -43,6 +48,12 @@ public:
 	virtual void ForceDismissDrawers() override;
 	//~ End IDisplayClusterOperator interface
 
+	/** Retrieve currently registered app delegates */
+	const TMap<FDelegateHandle, FOnGetAppInstance>& GetRegisteredApps() const { return RegisteredApps; }
+
+	/** Broadcast when an app is unregistered */
+	FOnAppUnregistered& OnAppUnregistered() { return AppUnregisteredEvent; }
+	
 private:
 	void RegisterTabSpawners();
 	void UnregisterTabSpawners();
@@ -59,4 +70,8 @@ private:
 	TWeakPtr<SDisplayClusterOperatorPanel> ActiveOperatorPanel;
 	TSharedPtr<FDisplayClusterOperatorViewModel> OperatorViewModel;
 	TSharedPtr<FExtensibilityManager> OperatorToolBarExtensibilityManager;
+	TSharedPtr<FExtensibilityManager> OperatorMenuExtensibilityManager;
+
+	TMap<FDelegateHandle, FOnGetAppInstance> RegisteredApps;
+	FOnAppUnregistered AppUnregisteredEvent;
 };
diff --git a/Engine/Plugins/Runtime/nDisplay/Source/DisplayClusterOperator/Private/DisplayClusterOperatorViewModel.cpp b/Engine/Plugins/Runtime/nDisplay/Source/DisplayClusterOperator/Private/DisplayClusterOperatorViewModel.cpp
--- a/Engine/Plugins/Runtime/nDisplay/Source/DisplayClusterOperator/Private/DisplayClusterOperatorViewModel.cpp
+++ b/Engine/Plugins/Runtime/nDisplay/Source/DisplayClusterOperator/Private/DisplayClusterOperatorViewModel.cpp
@@ -6,6 +6,8 @@
 
 #include "Framework/Docking/TabManager.h"
 
+#define LOCTEXT_NAMESPACE "DisplayClusterOperatorViewModel"
+
 void FDisplayClusterOperatorViewModel::SetRootActor(ADisplayClusterRootActor* InRootActor)
 {
 	RootActor = InRootActor;
@@ -15,10 +17,13 @@ void FDisplayClusterOperatorViewModel::SetRootActor(ADisplayClusterRootActor* In
 TSharedRef<FTabManager> FDisplayClusterOperatorViewModel::CreateTabManager(const TSharedRef<SDockTab>& MajorTabOwner)
 {
 	TabManager = FGlobalTabmanager::Get()->NewTabManager(MajorTabOwner);
+	WorkspaceItem = TabManager->AddLocalWorkspaceMenuCategory(LOCTEXT("OperatorMenuGroupName", "In-Camera VFX"));
 	return TabManager.ToSharedRef();
 }
 
 void FDisplayClusterOperatorViewModel::ResetTabManager()
 {
 	TabManager.Reset();
-}
\ No newline at end of file
+}
+
+#undef LOCTEXT_NAMESPACE
\ No newline at end of file
diff --git a/Engine/Plugins/Runtime/nDisplay/Source/DisplayClusterOperator/Private/DisplayClusterOperatorViewModel.h b/Engine/Plugins/Runtime/nDisplay/Source/DisplayClusterOperator/Private/DisplayClusterOperatorViewModel.h
--- a/Engine/Plugins/Runtime/nDisplay/Source/DisplayClusterOperator/Private/DisplayClusterOperatorViewModel.h
+++ b/Engine/Plugins/Runtime/nDisplay/Source/DisplayClusterOperator/Private/DisplayClusterOperatorViewModel.h
@@ -22,6 +22,8 @@ public:
 	virtual FOnActiveRootActorChanged& OnActiveRootActorChanged() override { return RootActorChanged; }
 
 	virtual TSharedPtr<FTabManager> GetTabManager() const override { return TabManager; }
+
+	virtual TSharedPtr<FWorkspaceItem> GetWorkspaceMenuGroup() const override { return WorkspaceItem; }
 	//~ End IDisplayClusterOperatorViewModel interface
 
 	TSharedRef<FTabManager> CreateTabManager(const TSharedRef<SDockTab>& MajorTabOwner);
@@ -30,6 +32,7 @@ public:
 private:
 	TWeakObjectPtr<ADisplayClusterRootActor> RootActor;
 	TSharedPtr<FTabManager> TabManager;
+	TSharedPtr<FWorkspaceItem> WorkspaceItem;
 
 	FOnActiveRootActorChanged RootActorChanged;
 };
\ No newline at end of file
diff --git a/Engine/Plugins/Runtime/nDisplay/Source/DisplayClusterOperator/Private/SDisplayClusterOperatorPanel.cpp b/Engine/Plugins/Runtime/nDisplay/Source/DisplayClusterOperator/Private/SDisplayClusterOperatorPanel.cpp
--- a/Engine/Plugins/Runtime/nDisplay/Source/DisplayClusterOperator/Private/SDisplayClusterOperatorPanel.cpp
+++ b/Engine/Plugins/Runtime/nDisplay/Source/DisplayClusterOperator/Private/SDisplayClusterOperatorPanel.cpp
@@ -2,66 +2,81 @@
 
 #include "SDisplayClusterOperatorPanel.h"
 
+#include "DisplayClusterOperatorModule.h"
+#include "DisplayClusterOperatorStatusBarExtender.h"
 #include "IDisplayClusterOperator.h"
 #include "IDisplayClusterOperatorViewModel.h"
-#include "SDisplayClusterOperatorToolbar.h"
 #include "SDisplayClusterOperatorStatusBar.h"
-#include "DisplayClusterOperatorStatusBarExtender.h"
+#include "SDisplayClusterOperatorToolbar.h"
 
 #include "Editor.h"
-#include "Styling/AppStyle.h"
-#include "Framework/Docking/TabManager.h"
-#include "Framework/Docking/LayoutExtender.h"
-#include "Framework/MultiBox/MultiBoxBuilder.h"
+#include "LevelEditor.h"
+#include "LevelEditorActions.h"
 #include "SKismetInspector.h"
 #include "TimerManager.h"
+#include "ToolMenus.h"
+#include "Framework/Docking/LayoutExtender.h"
+#include "Framework/Docking/TabManager.h"
+#include "Framework/MultiBox/MultiBoxBuilder.h"
+#include "Interfaces/IMainFrameModule.h"
+#include "Styling/AppStyle.h"
 #include "Widgets/Docking/SDockTab.h"
 
-#include "Components/ActorComponent.h"
-
 #define LOCTEXT_NAMESPACE "SDisplayClusterOperatorPanel"
 
-const FName SDisplayClusterOperatorPanel::ToolbarTabId = TEXT("OperatorToolbar");
 const FName SDisplayClusterOperatorPanel::DetailsTabId = TEXT("OperatorDetails");
 const FName SDisplayClusterOperatorPanel::PrimaryTabExtensionId = TEXT("PrimaryOperatorTabStack");
 const FName SDisplayClusterOperatorPanel::AuxilliaryTabExtensionId = TEXT("AuxilliaryOperatorTabStack");
 
 SDisplayClusterOperatorPanel::~SDisplayClusterOperatorPanel()
 {
-	IDisplayClusterOperator::Get().OnDetailObjectsChanged().Remove(DetailObjectsChangedHandle);
+	FDisplayClusterOperatorModule& OperatorModule = FModuleManager::GetModuleChecked<FDisplayClusterOperatorModule>(IDisplayClusterOperator::ModuleName);
+	OperatorModule.OnAppUnregistered().RemoveAll(this);
+	OperatorModule.OnDetailObjectsChanged().Remove(DetailObjectsChangedHandle);
 }
 
 void SDisplayClusterOperatorPanel::Construct(const FArguments& InArgs, const TSharedRef<FTabManager>& InTabManager, const TSharedPtr<SWindow>& WindowOwner)
 {
 	CommandList = MakeShareable(new FUICommandList);
 	BindCommands();
-
+	
 	TabManager = InTabManager;
-	TSharedRef<FWorkspaceItem> AppMenuGroup = TabManager->AddLocalWorkspaceMenuCategory(LOCTEXT("OperatorMenuGroupName", "In-Camera VFX"));
+	const auto& PersistLayout = [](const TSharedRef<FTabManager::FLayout>& LayoutToSave)
+	{
+		FLayoutSaveRestore::SaveToConfig(GEditorLayoutIni, LayoutToSave);
+	};
+	TabManager->SetOnPersistLayout(FTabManager::FOnPersistLayout::CreateLambda(PersistLayout));
 	TabManager->SetAllowWindowMenuBar(true);
 
-	TabManager->RegisterTabSpawner(ToolbarTabId, FOnSpawnTab::CreateSP(this, &SDisplayClusterOperatorPanel::SpawnToolbarTab))
-		.SetDisplayName(LOCTEXT("ToolbarTabTitle", "Toolbar"))
-		.SetIcon(FSlateIcon(FAppStyle::GetAppStyleSetName(), "Icons.Settings"))
-		.SetGroup(AppMenuGroup);
+	const TSharedRef<IDisplayClusterOperatorViewModel> ViewModel = IDisplayClusterOperator::Get().GetOperatorViewModel();
+	const TSharedRef<FWorkspaceItem> AppMenuGroup = ViewModel->GetWorkspaceMenuGroup().ToSharedRef();
 
 	TabManager->RegisterTabSpawner(DetailsTabId, FOnSpawnTab::CreateSP(this, &SDisplayClusterOperatorPanel::SpawnDetailsTab))
 		.SetDisplayName(LOCTEXT("DetailsTabTitle", "Details"))
 		.SetIcon(FSlateIcon(FAppStyle::GetAppStyleSetName(), "LevelEditor.Tabs.Details"))
 		.SetGroup(AppMenuGroup);
 
-	const TSharedRef<FTabManager::FLayout> Layout = FTabManager::NewLayout("nDisplayOperatorLayout")
+	// Load any operator apps
+	{
+		FDisplayClusterOperatorModule& OperatorModule = FModuleManager::GetModuleChecked<FDisplayClusterOperatorModule>(IDisplayClusterOperator::ModuleName);
+		OperatorModule.OnAppUnregistered().AddSP(this, &SDisplayClusterOperatorPanel::OnOperatorAppUnregistered);
+		for (const TTuple<FDelegateHandle, IDisplayClusterOperator::FOnGetAppInstance>&
+			 RegisteredAppKeyVal : OperatorModule.GetRegisteredApps())
+		{
+			if (RegisteredAppKeyVal.Value.IsBound())
+			{
+				TSharedRef<IDisplayClusterOperatorApp> AppInstance = RegisteredAppKeyVal.Value.Execute(ViewModel);
+				AppInstances.Add(RegisteredAppKeyVal.Key, AppInstance);
+			}
+		}
+	}
+	
+	TSharedRef<FTabManager::FLayout> Layout = FTabManager::NewLayout("nDisplayOperatorLayout_v1")
 		->AddArea
 		(
 			FTabManager::NewPrimaryArea()
 			->SetOrientation(EOrientation::Orient_Vertical)
 			->Split
-			(
-				FTabManager::NewStack()
-					->AddTab(ToolbarTabId, ETabState::OpenedTab)
-					->SetHideTabWell(true)
-			)
-			->Split
 			(
 				FTabManager::NewSplitter()
 					->SetOrientation(Orient_Horizontal)
@@ -73,6 +88,7 @@ void SDisplayClusterOperatorPanel::Construct(const FArguments& InArgs, const TSh
 						->Split
 						(
 							FTabManager::NewStack()
+								->SetHideTabWell(true)
 								->SetExtensionId(PrimaryTabExtensionId)
 						)
 						->Split
@@ -97,16 +113,50 @@ void SDisplayClusterOperatorPanel::Construct(const FArguments& InArgs, const TSh
 			)
 		);
 
+	Layout = FLayoutSaveRestore::LoadFromConfig(GEditorLayoutIni, Layout);
+	
+	static FName MenuName = "DisplayClusterOperator.MainMenu";
+	const FLevelEditorModule& LevelEditor = FModuleManager::GetModuleChecked<FLevelEditorModule>(TEXT("LevelEditor"));
+
+	// Fill toolbar
+	{
+		UToolMenus* ToolMenus = UToolMenus::Get();
+		const FName ParentMenuName("MainFrame.MainMenu");
+		
+		if (!ToolMenus->IsMenuRegistered(MenuName))
+		{
+			ToolMenus->RegisterMenu(MenuName, ParentMenuName);
+
+			// File menu
+			{
+				UToolMenu* SubMenu = ToolMenus->RegisterMenu(*(MenuName.ToString() + TEXT(".File")), *(ParentMenuName.ToString() + TEXT(".File")));
+				const FToolMenuInsert InsertPos(NAME_None, EToolMenuInsertType::First);
+				FToolMenuSection& Section = SubMenu->FindOrAddSection("FileSave");
+
+				// Add generic save level options. If we ever add other types of apps that don't rely on the level
+				// instance being saved we may want to refactor this so the save option is added per app instead.
+				const FLevelEditorCommands& Commands = LevelEditor.GetLevelEditorCommands();
+				Section.AddMenuEntry(Commands.Save).InsertPosition = InsertPos;
+				Section.AddMenuEntry(Commands.SaveAs).InsertPosition = InsertPos;
+			}
+		}
+	};
+	
 	LayoutExtender = MakeShared<FLayoutExtender>();
 	IDisplayClusterOperator::Get().OnRegisterLayoutExtensions().Broadcast(*LayoutExtender);
 	Layout->ProcessExtensions(*LayoutExtender);
 
 	DetailObjectsChangedHandle = IDisplayClusterOperator::Get().OnDetailObjectsChanged().AddSP(this, &SDisplayClusterOperatorPanel::DisplayObjectsInDetailsPanel);
-
+	
 	ChildSlot
 	[
 		SNew(SVerticalBox)
-
+		+SVerticalBox::Slot()
+		.AutoHeight()
+		[
+			SAssignNew(ToolbarContainer, SBox)
+		]
+		
 		+SVerticalBox::Slot()
 		.Padding(4.0f, 2.f, 4.f, 2.f)
 		.FillHeight(1.0f)
@@ -122,6 +172,15 @@ void SDisplayClusterOperatorPanel::Construct(const FArguments& InArgs, const TSh
 		]
 	];
 
+	TArray<TSharedPtr<FExtender>> MenuExtenders;
+	const TSharedPtr<FExtender> MenuExtender = IDisplayClusterOperator::Get().GetOperatorMenuExtensibilityManager()->GetAllExtenders();
+
+	CommandList->Append(LevelEditor.GetGlobalLevelEditorActions());
+	FToolMenuContext ToolMenuContext(CommandList, MenuExtender);
+	
+	const IMainFrameModule& MainFrameModule = FModuleManager::LoadModuleChecked<IMainFrameModule>("MainFrame");
+	MainFrameModule.MakeMainMenu(TabManager, MenuName, ToolMenuContext);
+	
 	// Allow any external modules to register extensions to the operator panel's status bar
 	FDisplayClusterOperatorStatusBarExtender StatusBarExtender;
 	IDisplayClusterOperator::Get().OnRegisterStatusBarExtensions().Broadcast(StatusBarExtender);
@@ -132,7 +191,7 @@ void SDisplayClusterOperatorPanel::Construct(const FArguments& InArgs, const TSh
 		// Create toolbar after tab has been restored so child windows can register their toolbar extensions
 		ToolbarContainer->SetContent(
 		SAssignNew(Toolbar, SDisplayClusterOperatorToolbar)
-		.CommandList(TSharedPtr<FUICommandList>()));
+		.CommandList(CommandList));
 	}
 }
 
@@ -187,16 +246,6 @@ void SDisplayClusterOperatorPanel::BindCommands()
 	}
 }
 
-TSharedRef<SDockTab> SDisplayClusterOperatorPanel::SpawnToolbarTab(const FSpawnTabArgs& Args)
-{
-	return SNew(SDockTab)
-		.ShouldAutosize(true)
-		.TabRole(ETabRole::PanelTab)
-		[
-			SAssignNew(ToolbarContainer, SBox)
-		];
-}
-
 TSharedRef<SDockTab> SDisplayClusterOperatorPanel::SpawnDetailsTab(const FSpawnTabArgs& Args)
 {
 	SAssignNew(DetailsView, SKismetInspector)
@@ -219,4 +268,9 @@ void SDisplayClusterOperatorPanel::DisplayObjectsInDetailsPanel(const TArray<UOb
 	}
 }
 
+void SDisplayClusterOperatorPanel::OnOperatorAppUnregistered(const FDelegateHandle& InHandle)
+{
+	AppInstances.Remove(InHandle);
+}
+
 #undef LOCTEXT_NAMESPACE
diff --git a/Engine/Plugins/Runtime/nDisplay/Source/DisplayClusterOperator/Private/SDisplayClusterOperatorPanel.h b/Engine/Plugins/Runtime/nDisplay/Source/DisplayClusterOperator/Private/SDisplayClusterOperatorPanel.h
--- a/Engine/Plugins/Runtime/nDisplay/Source/DisplayClusterOperator/Private/SDisplayClusterOperatorPanel.h
+++ b/Engine/Plugins/Runtime/nDisplay/Source/DisplayClusterOperator/Private/SDisplayClusterOperatorPanel.h
@@ -15,14 +15,12 @@ class SDockTab;
 class SWindow;
 class SDisplayClusterOperatorToolbar;
 class SDisplayClusterOperatorStatusBar;
+class IDisplayClusterOperatorApp;
 
 /** The nDisplay operator panel that allows users to edit root actor instances with a variety of tools */
 class SDisplayClusterOperatorPanel : public SCompoundWidget
 {
 public:
-	/** The tab ID that the operator toolbar lives in */
-	static const FName ToolbarTabId;
-
 	/** The tab ID that the details panel lives in */
 	static const FName DetailsTabId;
 
@@ -54,15 +52,15 @@ private:
 	/** Binds any commands to the operator panel's command list, as well as allowing any external modules to register their own commands */
 	void BindCommands();
 
-	/** Creates a tab with the operator toolbar in it */
-	TSharedRef<SDockTab> SpawnToolbarTab(const FSpawnTabArgs& Args);
-
 	/** Creates a tab with the details panel in it */
 	TSharedRef<SDockTab> SpawnDetailsTab(const FSpawnTabArgs& Args);
 
 	/** Raised when something wants to display a list of objects in the operator's details panel */
 	void DisplayObjectsInDetailsPanel(const TArray<UObject*>& Objects);
 
+	/** Called when app is unregistered */
+	void OnOperatorAppUnregistered(const FDelegateHandle& InHandle);
+	
 private:
 	/** Holds the tab manager that manages the operator's tabs. */
 	TSharedPtr<FTabManager> TabManager;
@@ -87,4 +85,7 @@ private:
 
 	/** The delegate handle for the operator module's OnDetailObjectsChanged event */
 	FDelegateHandle DetailObjectsChangedHandle;
+
+	/** Loaded app instances */
+	TMap<FDelegateHandle, TSharedRef<IDisplayClusterOperatorApp>> AppInstances;
 };
\ No newline at end of file
diff --git a/Engine/Plugins/Runtime/nDisplay/Source/DisplayClusterOperator/Private/SDisplayClusterOperatorToolbar.cpp b/Engine/Plugins/Runtime/nDisplay/Source/DisplayClusterOperator/Private/SDisplayClusterOperatorToolbar.cpp
--- a/Engine/Plugins/Runtime/nDisplay/Source/DisplayClusterOperator/Private/SDisplayClusterOperatorToolbar.cpp
+++ b/Engine/Plugins/Runtime/nDisplay/Source/DisplayClusterOperator/Private/SDisplayClusterOperatorToolbar.cpp
@@ -8,6 +8,7 @@
 
 #include "Styling/AppStyle.h"
 #include "LevelEditor.h"
+#include "LevelEditorActions.h"
 #include "Framework/MultiBox/MultiBoxBuilder.h"
 #include "Toolkits/AssetEditorToolkit.h"
 #include "Widgets/Layout/SBorder.h"
@@ -75,6 +76,13 @@ void SDisplayClusterOperatorToolbar::Construct(const FArguments& InArgs)
 	
 	ToolBarBuilder.BeginSection("General");
 	{
+		// Add a generic save level option. If we ever add other types of apps that don't rely on the level
+		// instance being saved we may want to refactor this so the save option is added per app instead.
+		const FLevelEditorModule& LevelEditor = FModuleManager::GetModuleChecked<FLevelEditorModule>(TEXT("LevelEditor"));
+		ToolBarBuilder.AddToolBarButton(LevelEditor.GetLevelEditorCommands().Save);
+		
+		ToolBarBuilder.AddSeparator();
+		
 		ToolBarBuilder.AddWidget(RootActorSelectionBox.ToSharedRef());
 	}
 	ToolBarBuilder.EndSection();
@@ -174,7 +182,7 @@ void SDisplayClusterOperatorToolbar::OnRootActorChanged(TSharedPtr<FString> Item
 		if (UBlueprint* Blueprint = UBlueprint::GetBlueprintFromClass(SelectedRootActor->GetClass()))
 		{
 			Blueprint->OnCompiled().RemoveAll(this);
-			Blueprint->OnCompiled().AddRaw(this, &SDisplayClusterOperatorToolbar::OnBlueprintCompiled);
+			Blueprint->OnCompiled().AddSP(this, &SDisplayClusterOperatorToolbar::OnBlueprintCompiled);
 		}
 	}
 
diff --git a/Engine/Plugins/Runtime/nDisplay/Source/DisplayClusterOperator/Public/IDisplayClusterOperator.h b/Engine/Plugins/Runtime/nDisplay/Source/DisplayClusterOperator/Public/IDisplayClusterOperator.h
--- a/Engine/Plugins/Runtime/nDisplay/Source/DisplayClusterOperator/Public/IDisplayClusterOperator.h
+++ b/Engine/Plugins/Runtime/nDisplay/Source/DisplayClusterOperator/Public/IDisplayClusterOperator.h
@@ -6,8 +6,10 @@
 #include "Modules/ModuleInterface.h"
 #include "Modules/ModuleManager.h"
 
+class IDisplayClusterOperatorApp;
 class IDisplayClusterOperatorViewModel;
 class FLayoutExtender;
+class FTabManager;
 class FExtensibilityManager;
 class FUICommandList;
 class FDisplayClusterOperatorStatusBarExtender;
@@ -25,7 +27,8 @@ public:
 	DECLARE_MULTICAST_DELEGATE_OneParam(FOnRegisterStatusBarExtensions, FDisplayClusterOperatorStatusBarExtender&);
 	DECLARE_MULTICAST_DELEGATE_OneParam(FOnAppendOperatorPanelCommands, TSharedRef<FUICommandList>)
 	DECLARE_MULTICAST_DELEGATE_OneParam(FOnDetailObjectsChanged, const TArray<UObject*>&);
-
+	DECLARE_DELEGATE_RetVal_OneParam(TSharedRef<IDisplayClusterOperatorApp>, FOnGetAppInstance, TSharedRef<IDisplayClusterOperatorViewModel>);
+	
 public:
 	virtual ~IDisplayClusterOperator() = default;
 
@@ -50,6 +53,12 @@ public:
 		return FModuleManager::Get().IsModuleLoaded(ModuleName);
 	}
 
+	/** Register an app with the operator panel */
+	virtual FDelegateHandle RegisterApp(const FOnGetAppInstance& InGetAppInstanceDelegate) = 0;
+
+	/** Unregisters an app from the operator panel */
+	virtual void UnregisterApp(const FDelegateHandle& InHandle) = 0;
+	
 	/** Gets the operator panel's view model, which stores any state used by the operator panel */
 	virtual TSharedRef<IDisplayClusterOperatorViewModel> GetOperatorViewModel() = 0;
 
@@ -77,6 +86,9 @@ public:
 	/** Gets the extensibility manager for the operator panel's toolbar */
 	virtual TSharedPtr<FExtensibilityManager> GetOperatorToolBarExtensibilityManager() = 0;
 
+	/** Gets the extensibility manager for the operator panel's menu bar */
+	virtual TSharedPtr<FExtensibilityManager> GetOperatorMenuExtensibilityManager() = 0;
+	
 	/** Gets a list of all nDisplay root actor instances that are on the currently loaded level */
 	virtual void GetRootActorLevelInstances(TArray<ADisplayClusterRootActor*>& OutRootActorInstances) = 0;
 
diff --git a/Engine/Plugins/Runtime/nDisplay/Source/DisplayClusterOperator/Public/IDisplayClusterOperatorApp.h b/Engine/Plugins/Runtime/nDisplay/Source/DisplayClusterOperator/Public/IDisplayClusterOperatorApp.h
new file mode 100644
--- /dev/null
+++ b/Engine/Plugins/Runtime/nDisplay/Source/DisplayClusterOperator/Public/IDisplayClusterOperatorApp.h
@@ -0,0 +1,11 @@
+// Copyright Epic Games, Inc. All Rights Reserved.
+
+#pragma once
+
+#include "Templates/SharedPointer.h"
+
+class IDisplayClusterOperatorApp : public TSharedFromThis<IDisplayClusterOperatorApp>
+{
+public:
+	virtual ~IDisplayClusterOperatorApp() = default;
+};
\ No newline at end of file
diff --git a/Engine/Plugins/Runtime/nDisplay/Source/DisplayClusterOperator/Public/IDisplayClusterOperatorViewModel.h b/Engine/Plugins/Runtime/nDisplay/Source/DisplayClusterOperator/Public/IDisplayClusterOperatorViewModel.h
--- a/Engine/Plugins/Runtime/nDisplay/Source/DisplayClusterOperator/Public/IDisplayClusterOperatorViewModel.h
+++ b/Engine/Plugins/Runtime/nDisplay/Source/DisplayClusterOperator/Public/IDisplayClusterOperatorViewModel.h
@@ -6,6 +6,7 @@
 
 class ADisplayClusterRootActor;
 class FTabManager;
+class FWorkspaceItem;
 
 DECLARE_EVENT_OneParam(IDisplayClusterOperatorViewModel, FOnActiveRootActorChanged, ADisplayClusterRootActor*);
 
@@ -27,4 +28,7 @@ public:
 
 	/** Gets the tab manager of the active operator panel, if there is an open operator panel */
 	virtual TSharedPtr<FTabManager> GetTabManager() const = 0;
+
+	/** Gets the registered workspace menu group */
+	virtual TSharedPtr<FWorkspaceItem> GetWorkspaceMenuGroup() const = 0;
 };
\ No newline at end of file
