nDisplay: Added and moved the Name property of the per-viewport and per-node color grading groups to the top of the details panel

diff --git a/Engine/Plugins/Runtime/nDisplay/Source/DisplayClusterConfigurator/Private/Views/Details/Cluster/DisplayClusterConfiguratorNodeSelectionCustomization.cpp b/Engine/Plugins/Runtime/nDisplay/Source/DisplayClusterConfigurator/Private/Views/Details/Cluster/DisplayClusterConfiguratorNodeSelectionCustomization.cpp
--- a/Engine/Plugins/Runtime/nDisplay/Source/DisplayClusterConfigurator/Private/Views/Details/Cluster/DisplayClusterConfiguratorNodeSelectionCustomization.cpp
+++ b/Engine/Plugins/Runtime/nDisplay/Source/DisplayClusterConfigurator/Private/Views/Details/Cluster/DisplayClusterConfiguratorNodeSelectionCustomization.cpp
@@ -318,9 +318,9 @@ void FDisplayClusterConfiguratorPerViewportColorGradingCustomization::SetHeader(
 
 void FDisplayClusterConfiguratorPerViewportColorGradingCustomization::SetChildren(const TSharedRef<IPropertyHandle>& InPropertyHandle, IDetailChildrenBuilder& InChildBuilder, IPropertyTypeCustomizationUtils& CustomizationUtils)
 {
+	TSharedPtr<IPropertyHandle> NameHandle = GET_CHILD_HANDLE(FDisplayClusterConfigurationViewport_PerViewportColorGrading, Name);
 	TSharedPtr<IPropertyHandle> IsEnabledHandle = GET_CHILD_HANDLE(FDisplayClusterConfigurationViewport_PerViewportColorGrading, bIsEnabled);
 	TSharedPtr<IPropertyHandle> IsEntireClusterPostProcessHandle = GET_CHILD_HANDLE(FDisplayClusterConfigurationViewport_PerViewportColorGrading, bIsEntireClusterEnabled);
-	TSharedPtr<IPropertyHandle> NameHandle = GET_CHILD_HANDLE(FDisplayClusterConfigurationViewport_PerViewportColorGrading, Name);
 	TSharedPtr<IPropertyHandle> PostProcessSettingsHandle = GET_CHILD_HANDLE(FDisplayClusterConfigurationViewport_PerViewportColorGrading, ColorGradingSettings);
 	TSharedPtr<IPropertyHandle> ArrayHandle = GET_CHILD_HANDLE(FDisplayClusterConfigurationViewport_PerViewportColorGrading, ApplyPostProcessToObjects);
 
@@ -331,9 +331,9 @@ void FDisplayClusterConfiguratorPerViewportColorGradingCustomization::SetChildre
 		return bCond1;
 	});
 
+	InChildBuilder.AddProperty(NameHandle.ToSharedRef());
 	InChildBuilder.AddProperty(IsEnabledHandle.ToSharedRef());
 	InChildBuilder.AddProperty(IsEntireClusterPostProcessHandle.ToSharedRef());
-	InChildBuilder.AddProperty(NameHandle.ToSharedRef());
 	InChildBuilder.AddProperty(PostProcessSettingsHandle.ToSharedRef()).EditCondition(IsEnabledEditCondition, nullptr);
 
 	ArrayHandle->SetInstanceMetaData(FDisplayClusterConfiguratorNodeSelection::NAME_ElementToolTip, ArrayHandle->GetPropertyDisplayName().ToString());
@@ -366,6 +366,7 @@ void FDisplayClusterConfiguratorPerNodeColorGradingCustomization::SetHeader(cons
 
 void FDisplayClusterConfiguratorPerNodeColorGradingCustomization::SetChildren(const TSharedRef<IPropertyHandle>& InPropertyHandle, IDetailChildrenBuilder& InChildBuilder, IPropertyTypeCustomizationUtils& CustomizationUtils)
 {
+	TSharedPtr<IPropertyHandle> NameHandle = GET_CHILD_HANDLE(FDisplayClusterConfigurationViewport_PerNodeColorGrading, Name);
 	TSharedPtr<IPropertyHandle> IsEnabledHandle = GET_CHILD_HANDLE(FDisplayClusterConfigurationViewport_PerNodeColorGrading, bIsEnabled);
 	TSharedPtr<IPropertyHandle> IsEntireClusterPostProcessHandle = GET_CHILD_HANDLE(FDisplayClusterConfigurationViewport_PerNodeColorGrading, bEntireClusterColorGrading);
 	TSharedPtr<IPropertyHandle> IsAllNodesPostProcessHandle = GET_CHILD_HANDLE(FDisplayClusterConfigurationViewport_PerNodeColorGrading, bAllNodesColorGrading);
@@ -379,6 +380,7 @@ void FDisplayClusterConfiguratorPerNodeColorGradingCustomization::SetChildren(co
 		return bCond1;
 	});
 
+	InChildBuilder.AddProperty(NameHandle.ToSharedRef());
 	InChildBuilder.AddProperty(IsEnabledHandle.ToSharedRef());
 	InChildBuilder.AddProperty(IsEntireClusterPostProcessHandle.ToSharedRef());
 	InChildBuilder.AddProperty(IsAllNodesPostProcessHandle.ToSharedRef());
