nDisplay: Fixed crash caused by removing the sublevel an nDisplay root actor was in when the ICVFX window was open

diff --git a/Engine/Plugins/Runtime/nDisplay/Source/DisplayClusterLightCardEditor/Private/DisplayClusterLightCardEditor.cpp b/Engine/Plugins/Runtime/nDisplay/Source/DisplayClusterLightCardEditor/Private/DisplayClusterLightCardEditor.cpp
--- a/Engine/Plugins/Runtime/nDisplay/Source/DisplayClusterLightCardEditor/Private/DisplayClusterLightCardEditor.cpp
+++ b/Engine/Plugins/Runtime/nDisplay/Source/DisplayClusterLightCardEditor/Private/DisplayClusterLightCardEditor.cpp
@@ -832,7 +832,7 @@ void FDisplayClusterLightCardEditor::SetIconScale(float NewValue)
 
 void FDisplayClusterLightCardEditor::OnActiveRootActorChanged(ADisplayClusterRootActor* NewRootActor)
 {
-	if (NewRootActor == ActiveRootActor.Get())
+	if (NewRootActor == ActiveRootActor.GetEvenIfUnreachable())
 	{
 		return;
 	}
diff --git a/Engine/Plugins/Runtime/nDisplay/Source/DisplayClusterLightCardEditor/Private/Viewport/DisplayClusterLightCardEditorViewportClient.cpp b/Engine/Plugins/Runtime/nDisplay/Source/DisplayClusterLightCardEditor/Private/Viewport/DisplayClusterLightCardEditorViewportClient.cpp
--- a/Engine/Plugins/Runtime/nDisplay/Source/DisplayClusterLightCardEditor/Private/Viewport/DisplayClusterLightCardEditorViewportClient.cpp
+++ b/Engine/Plugins/Runtime/nDisplay/Source/DisplayClusterLightCardEditor/Private/Viewport/DisplayClusterLightCardEditorViewportClient.cpp
@@ -1353,7 +1353,7 @@ void FDisplayClusterLightCardEditorViewportClient::UpdatePreviewActor(ADisplayCl
 {
 	DECLARE_SCOPE_CYCLE_COUNTER(TEXT("FDisplayClusterLightCardEditorViewportClient::UpdatePreviewActor"), STAT_UpdatePreviewActor, STATGROUP_NDisplayLightCardEditor);
 	
-	if ((!bForce && RootActor == RootActorLevelInstance.Get()) ||
+	if ((!bForce && RootActor == RootActorLevelInstance.GetEvenIfUnreachable()) ||
 		(ProxyTypesRefreshing.Contains(ProxyType) || ProxyTypesRefreshing.Contains(EDisplayClusterLightCardEditorProxyType::All)))
 	{
 		return;
diff --git a/Engine/Plugins/Runtime/nDisplay/Source/DisplayClusterOperator/Private/DisplayClusterOperatorViewModel.cpp b/Engine/Plugins/Runtime/nDisplay/Source/DisplayClusterOperator/Private/DisplayClusterOperatorViewModel.cpp
--- a/Engine/Plugins/Runtime/nDisplay/Source/DisplayClusterOperator/Private/DisplayClusterOperatorViewModel.cpp
+++ b/Engine/Plugins/Runtime/nDisplay/Source/DisplayClusterOperator/Private/DisplayClusterOperatorViewModel.cpp
@@ -8,6 +8,16 @@
 
 #define LOCTEXT_NAMESPACE "DisplayClusterOperatorViewModel"
 
+ADisplayClusterRootActor* FDisplayClusterOperatorViewModel::GetRootActor(bool bEvenIfPendingKill) const
+{
+	if (bEvenIfPendingKill)
+	{
+		return RootActor.IsValid(bEvenIfPendingKill) ? RootActor.GetEvenIfUnreachable() : nullptr;
+	}
+
+	return RootActor.IsValid() ? RootActor.Get() : nullptr;
+}
+
 void FDisplayClusterOperatorViewModel::SetRootActor(ADisplayClusterRootActor* InRootActor)
 {
 	RootActor = InRootActor;
diff --git a/Engine/Plugins/Runtime/nDisplay/Source/DisplayClusterOperator/Private/DisplayClusterOperatorViewModel.h b/Engine/Plugins/Runtime/nDisplay/Source/DisplayClusterOperator/Private/DisplayClusterOperatorViewModel.h
--- a/Engine/Plugins/Runtime/nDisplay/Source/DisplayClusterOperator/Private/DisplayClusterOperatorViewModel.h
+++ b/Engine/Plugins/Runtime/nDisplay/Source/DisplayClusterOperator/Private/DisplayClusterOperatorViewModel.h
@@ -16,8 +16,8 @@ public:
 	virtual ~FDisplayClusterOperatorViewModel() = default;
 
 	//~ IDisplayClusterOperatorViewModel interface
-	virtual bool HasRootActor() const override { return RootActor.IsValid(); }
-	virtual ADisplayClusterRootActor* GetRootActor() const override { return RootActor.IsValid() ? RootActor.Get() : nullptr; }
+	virtual bool HasRootActor(bool bEvenIfPendingKill = false) const override { return RootActor.IsValid(bEvenIfPendingKill); }
+	virtual ADisplayClusterRootActor* GetRootActor(bool bEvenIfPendingKill = false) const override;
 	virtual void SetRootActor(ADisplayClusterRootActor* InRootActor) override;
 	virtual FOnActiveRootActorChanged& OnActiveRootActorChanged() override { return RootActorChanged; }
 
diff --git a/Engine/Plugins/Runtime/nDisplay/Source/DisplayClusterOperator/Private/SDisplayClusterOperatorToolbar.cpp b/Engine/Plugins/Runtime/nDisplay/Source/DisplayClusterOperator/Private/SDisplayClusterOperatorToolbar.cpp
--- a/Engine/Plugins/Runtime/nDisplay/Source/DisplayClusterOperator/Private/SDisplayClusterOperatorToolbar.cpp
+++ b/Engine/Plugins/Runtime/nDisplay/Source/DisplayClusterOperator/Private/SDisplayClusterOperatorToolbar.cpp
@@ -9,6 +9,7 @@
 #include "Styling/AppStyle.h"
 #include "LevelEditor.h"
 #include "LevelEditorActions.h"
+#include "EditorSupportDelegates.h"
 #include "Framework/MultiBox/MultiBoxBuilder.h"
 #include "Toolkits/AssetEditorToolkit.h"
 #include "Widgets/Layout/SBorder.h"
@@ -29,6 +30,8 @@ SDisplayClusterOperatorToolbar::~SDisplayClusterOperatorToolbar()
 		LevelEditor.OnMapChanged().Remove(MapChangedHandle);
 	}
 	
+	FEditorSupportDelegates::PrepareToCleanseEditorObject.RemoveAll(this);
+
 	if (ADisplayClusterRootActor* ActiveRootActor = ViewModel->GetRootActor())
 	{
 		if (UBlueprint* Blueprint = UBlueprint::GetBlueprintFromClass(ActiveRootActor->GetClass()))
@@ -109,6 +112,8 @@ void SDisplayClusterOperatorToolbar::Construct(const FArguments& InArgs)
 
 	FLevelEditorModule& LevelEditor = FModuleManager::LoadModuleChecked<FLevelEditorModule>("LevelEditor");
 	MapChangedHandle = LevelEditor.OnMapChanged().AddRaw(this, &SDisplayClusterOperatorToolbar::HandleMapChanged);
+
+	FEditorSupportDelegates::PrepareToCleanseEditorObject.AddSP(this, &SDisplayClusterOperatorToolbar::OnPrepareToCleanseEditorObject);
 }
 
 TSharedPtr<FString> SDisplayClusterOperatorToolbar::FillRootActorList(const FString& InitiallySelectedRootActor)
@@ -254,4 +259,18 @@ void SDisplayClusterOperatorToolbar::HandleMapChanged(UWorld* InWorld, EMapChang
 	}
 }
 
+void SDisplayClusterOperatorToolbar::OnPrepareToCleanseEditorObject(UObject* Object)
+{
+	if (ULevel* Level = Cast<ULevel>(Object))
+	{
+		const bool bEvenIfPendingKill = true;
+		if (Level->Actors.Contains(ViewModel->GetRootActor(bEvenIfPendingKill)))
+		{
+			// A level that contains the current root actor is being purged, so clear the root actor to ensure
+			// no UI elements continue to reference the root actor
+			ClearSelectedRootActor();
+		}
+	}
+}
+
 #undef LOCTEXT_NAMESPACE
diff --git a/Engine/Plugins/Runtime/nDisplay/Source/DisplayClusterOperator/Private/SDisplayClusterOperatorToolbar.h b/Engine/Plugins/Runtime/nDisplay/Source/DisplayClusterOperator/Private/SDisplayClusterOperatorToolbar.h
--- a/Engine/Plugins/Runtime/nDisplay/Source/DisplayClusterOperator/Private/SDisplayClusterOperatorToolbar.h
+++ b/Engine/Plugins/Runtime/nDisplay/Source/DisplayClusterOperator/Private/SDisplayClusterOperatorToolbar.h
@@ -61,6 +61,9 @@ private:
 	/** When the map changes */
 	void HandleMapChanged(UWorld* InWorld, EMapChangeType InMapChangeType);
 
+	/** Raised when an editor object is going to be cleased, usually used when a sublevel is about to be removed */
+	void OnPrepareToCleanseEditorObject(UObject* Object);
+
 private:
 	/** A reference the the operator panel's view model, which stores the operator panel state */
 	TSharedPtr<IDisplayClusterOperatorViewModel> ViewModel;
diff --git a/Engine/Plugins/Runtime/nDisplay/Source/DisplayClusterOperator/Public/IDisplayClusterOperatorViewModel.h b/Engine/Plugins/Runtime/nDisplay/Source/DisplayClusterOperator/Public/IDisplayClusterOperatorViewModel.h
--- a/Engine/Plugins/Runtime/nDisplay/Source/DisplayClusterOperator/Public/IDisplayClusterOperatorViewModel.h
+++ b/Engine/Plugins/Runtime/nDisplay/Source/DisplayClusterOperator/Public/IDisplayClusterOperatorViewModel.h
@@ -15,10 +15,10 @@ class DISPLAYCLUSTEROPERATOR_API IDisplayClusterOperatorViewModel
 {
 public:
 	/** Gets whether the view model has been populated with a valid root actor */
-	virtual bool HasRootActor() const = 0;
+	virtual bool HasRootActor(bool bEvenIfPendingKill = false) const = 0;
 
 	/** Gets the root actor that is actively being edited by the operator panel */
-	virtual ADisplayClusterRootActor* GetRootActor() const = 0;
+	virtual ADisplayClusterRootActor* GetRootActor(bool bEvenIfPendingKill = false) const = 0;
 
 	/** Sets the root actor that is actively being edited by the operator panel */
 	virtual void SetRootActor(ADisplayClusterRootActor* InRootActor) = 0;
