nDisplay: Fixed issue where color grading drawer was not automatically selecting root actor after a level change or when the active root actor changed. Also fixed issue where the drawer would not return to the details view state properly when refreshed.

diff --git a/Engine/Plugins/Runtime/nDisplay/Source/DisplayClusterColorGrading/Private/Drawer/DisplayClusterColorGradingDrawerSingleton.cpp b/Engine/Plugins/Runtime/nDisplay/Source/DisplayClusterColorGrading/Private/Drawer/DisplayClusterColorGradingDrawerSingleton.cpp
--- a/Engine/Plugins/Runtime/nDisplay/Source/DisplayClusterColorGrading/Private/Drawer/DisplayClusterColorGradingDrawerSingleton.cpp
+++ b/Engine/Plugins/Runtime/nDisplay/Source/DisplayClusterColorGrading/Private/Drawer/DisplayClusterColorGradingDrawerSingleton.cpp
@@ -38,10 +38,17 @@ FDisplayClusterColorGradingDrawerSingleton::FDisplayClusterColorGradingDrawerSin
 	IDisplayClusterOperator::Get().OnRegisterLayoutExtensions().AddRaw(this, &FDisplayClusterColorGradingDrawerSingleton::ExtendOperatorTabLayout);
 	IDisplayClusterOperator::Get().OnRegisterStatusBarExtensions().AddRaw(this, &FDisplayClusterColorGradingDrawerSingleton::ExtendOperatorStatusBar);
 	IDisplayClusterOperator::Get().OnAppendOperatorPanelCommands().AddRaw(this, &FDisplayClusterColorGradingDrawerSingleton::AppendOperatorPanelCommands);
+
+	IDisplayClusterOperator::Get().GetOperatorViewModel()->OnActiveRootActorChanged().AddRaw(this, &FDisplayClusterColorGradingDrawerSingleton::OnActiveRootActorChanged);
 }
 
 FDisplayClusterColorGradingDrawerSingleton::~FDisplayClusterColorGradingDrawerSingleton()
 {
+	IDisplayClusterOperator::Get().OnRegisterLayoutExtensions().RemoveAll(this);
+	IDisplayClusterOperator::Get().OnRegisterStatusBarExtensions().RemoveAll(this);
+	IDisplayClusterOperator::Get().OnAppendOperatorPanelCommands().RemoveAll(this);
+	IDisplayClusterOperator::Get().GetOperatorViewModel()->OnActiveRootActorChanged().RemoveAll(this);
+
 	if (FSlateApplication::IsInitialized())
 	{
 		FGlobalTabmanager::Get()->UnregisterNomadTabSpawner(ColorGradingDrawerTab);
@@ -171,4 +178,10 @@ void FDisplayClusterColorGradingDrawerSingleton::SaveDrawerState(const TSharedPt
 	}
 }
 
+void FDisplayClusterColorGradingDrawerSingleton::OnActiveRootActorChanged(ADisplayClusterRootActor* NewRootActor)
+{
+	// Clear the previous drawer state when the active root actor is changed, since it is most likely invalid
+	PreviousDrawerState.Reset();
+}
+
 #undef LOCTEXT_NAMESPACE
\ No newline at end of file
diff --git a/Engine/Plugins/Runtime/nDisplay/Source/DisplayClusterColorGrading/Private/Drawer/DisplayClusterColorGradingDrawerSingleton.h b/Engine/Plugins/Runtime/nDisplay/Source/DisplayClusterColorGrading/Private/Drawer/DisplayClusterColorGradingDrawerSingleton.h
--- a/Engine/Plugins/Runtime/nDisplay/Source/DisplayClusterColorGrading/Private/Drawer/DisplayClusterColorGradingDrawerSingleton.h
+++ b/Engine/Plugins/Runtime/nDisplay/Source/DisplayClusterColorGrading/Private/Drawer/DisplayClusterColorGradingDrawerSingleton.h
@@ -55,6 +55,9 @@ private:
 	/** Delegate callback when the drawer is closed to save the drawer state */
 	void SaveDrawerState(const TSharedPtr<SWidget>& DrawerContent);
 
+	/** Delegate callback that is raised when the active root actor of the operator panel has changed */
+	void OnActiveRootActorChanged(ADisplayClusterRootActor* NewRootActor);
+
 private:
 	/** A weak pointer to the active color grading drawer that is open */
 	TWeakPtr<SDisplayClusterColorGradingDrawer> ColorGradingDrawer;
diff --git a/Engine/Plugins/Runtime/nDisplay/Source/DisplayClusterColorGrading/Private/Drawer/SDisplayClusterColorGradingDrawer.cpp b/Engine/Plugins/Runtime/nDisplay/Source/DisplayClusterColorGrading/Private/Drawer/SDisplayClusterColorGradingDrawer.cpp
--- a/Engine/Plugins/Runtime/nDisplay/Source/DisplayClusterColorGrading/Private/Drawer/SDisplayClusterColorGradingDrawer.cpp
+++ b/Engine/Plugins/Runtime/nDisplay/Source/DisplayClusterColorGrading/Private/Drawer/SDisplayClusterColorGradingDrawer.cpp
@@ -374,6 +374,10 @@ void SDisplayClusterColorGradingDrawer::Refresh(bool bPreserveDrawerState)
 	{
 		SetDrawerState(DrawerState);
 	}
+	else
+	{
+		SetDrawerStateToDefault();
+	}
 }
 
 void SDisplayClusterColorGradingDrawer::Tick(const FGeometry& AllottedGeometry, const double InCurrentTime, const float InDeltaTime)
@@ -470,8 +474,6 @@ FDisplayClusterColorGradingDrawerState SDisplayClusterColorGradingDrawer::GetDra
 
 void SDisplayClusterColorGradingDrawer::SetDrawerState(const FDisplayClusterColorGradingDrawerState& InDrawerState)
 {
-	CurrentDrawerMode = InDrawerState.DrawerMode;
-
 	TArray<FDisplayClusterColorGradingListItemRef> LevelItemsToSelect;
 	TArray<FDisplayClusterColorGradingListItemRef> RootActorItemsToSelect;
 	for (const TWeakObjectPtr<UObject>& SelectedObject : InDrawerState.SelectedObjects)
@@ -521,6 +523,8 @@ void SDisplayClusterColorGradingDrawer::SetDrawerState(const FDisplayClusterColo
 	ColorGradingDataModel->SetSelectedColorGradingGroup(InDrawerState.SelectedColorGradingGroup);
 	ColorGradingDataModel->SetSelectedColorGradingElement(InDrawerState.SelectedColorGradingElement);
 
+	CurrentDrawerMode = InDrawerState.DrawerMode;
+
 	if (ColorWheelPanel.IsValid())
 	{
 		ColorWheelPanel->SetDrawerState(InDrawerState);
@@ -905,7 +909,7 @@ void SDisplayClusterColorGradingDrawer::OnBlueprintCompiled(UBlueprint* Blueprin
 
 void SDisplayClusterColorGradingDrawer::OnActiveRootActorChanged(ADisplayClusterRootActor* NewRootActor)
 {
-	const bool bPreserveDrawerState = true;
+	const bool bPreserveDrawerState = false;
 	Refresh(bPreserveDrawerState);
 }
 
