nDisplay: Add Rename option to light card editor outliner.

Fix not being able to select a folder in the outliner. Regression from 22193094.
Fix folders not having full edit option support.

diff --git a/Engine/Plugins/Runtime/nDisplay/Source/DisplayClusterLightCardEditor/Private/DisplayClusterLightCardEditor.cpp b/Engine/Plugins/Runtime/nDisplay/Source/DisplayClusterLightCardEditor/Private/DisplayClusterLightCardEditor.cpp
--- a/Engine/Plugins/Runtime/nDisplay/Source/DisplayClusterLightCardEditor/Private/DisplayClusterLightCardEditor.cpp
+++ b/Engine/Plugins/Runtime/nDisplay/Source/DisplayClusterLightCardEditor/Private/DisplayClusterLightCardEditor.cpp
@@ -32,7 +32,9 @@
 #include "PropertyCustomizationHelpers.h"
 #include "ScopedTransaction.h"
 #include "Selection.h"
+#include "UnrealEdGlobals.h"
 #include "AssetRegistry/AssetRegistryModule.h"
+#include "Editor/UnrealEdEngine.h"
 #include "Framework/Commands/GenericCommands.h"
 #include "Framework/Docking/LayoutExtender.h"
 #include "Framework/Docking/TabManager.h"
@@ -497,18 +499,42 @@ bool FDisplayClusterLightCardEditor::CanAddLightCard() const
 
 void FDisplayClusterLightCardEditor::CutSelectedActors()
 {
-	CopySelectedActors();
-
-	RemoveActors(SelectedActors, /*bDeleteActors*/ true);
+	CopySelectedActors(/* bShouldCut */ true);
 }
 
 bool FDisplayClusterLightCardEditor::CanCutSelectedActors()
 {
+	if (LightCardOutliner.IsValid() && LightCardOutliner->CanCutSelectedFolder())
+	{
+		// If only a folder is selected -- this isn't handled outside of the outliner
+		return true;
+	}
+	
 	return CanCopySelectedActors() && CanRemoveSelectedActors();
 }
 
-void FDisplayClusterLightCardEditor::CopySelectedActors()
+void FDisplayClusterLightCardEditor::CopySelectedActors(bool bShouldCut)
 {
+	const FText TransactionText = bShouldCut ? LOCTEXT("CutItemsTransactionMessage", "Cut Selected Items") : LOCTEXT("CopyItemsTransactionMessage", "Copy Selected Items");
+	FScopedTransaction Transaction(TransactionText);
+
+	if (DoOutlinerFoldersNeedEditorDelegates())
+	{
+		// We still must fire the delegates in case folders are selected
+		if (bShouldCut)
+		{
+			FEditorDelegates::OnEditCutActorsBegin.Broadcast();
+			FEditorDelegates::OnEditCutActorsEnd.Broadcast();
+		}
+		else
+		{
+			FEditorDelegates::OnEditCopyActorsBegin.Broadcast();
+			FEditorDelegates::OnEditCopyActorsEnd.Broadcast();
+		}
+		
+		return;
+	}
+	
 	UWorld* EditorWorld = GEditor->GetEditorWorldContext().World();
 
 	const bool bNoteSelectionChange = false;
@@ -521,25 +547,60 @@ void FDisplayClusterLightCardEditor::CopySelectedActors()
 		const bool bInSelected = true;
 		const bool bNotify = false;
 		const bool bSelectEvenIfHidden = true;
-		GEditor->SelectActor(LightCard.Get(), true, bNotify, bSelectEvenIfHidden);
+		GEditor->SelectActor(LightCard.Get(), bInSelected, bNotify, bSelectEvenIfHidden);
 	}
 
-	const bool bShouldCut = false;
 	const bool bIsMove = false;
 	const bool bWarnAboutReferences = false;
+
+	if (bShouldCut)
+	{
+		FEditorDelegates::OnEditCutActorsBegin.Broadcast();
+	}
+	else
+	{
+		FEditorDelegates::OnEditCopyActorsBegin.Broadcast();
+	}
+	
 	GEditor->CopySelectedActorsToClipboard(EditorWorld, bShouldCut, bIsMove, bWarnAboutReferences);
+
+	if (bShouldCut)
+	{
+		FEditorDelegates::OnEditCutActorsEnd.Broadcast();
+	}
+	else
+	{
+		FEditorDelegates::OnEditCopyActorsEnd.Broadcast();
+	}
 }
 
 bool FDisplayClusterLightCardEditor::CanCopySelectedActors() const
 {
+	if (LightCardOutliner.IsValid() && LightCardOutliner->CanCopySelectedFolder())
+	{
+		// True if only a folder is selected -- this isn't handled outside of the outliner
+		return true;
+	}
+	
 	return SelectedActors.Num() > 0;
 }
 
-void FDisplayClusterLightCardEditor::PasteActors(bool bOffsetActorPosition)
+void FDisplayClusterLightCardEditor::PasteActors()
 {
-	UWorld* EditorWorld = GEditor->GetEditorWorldContext().World();
-	GEditor->PasteSelectedActorsFromClipboard(EditorWorld, LOCTEXT("PasteActorsTransactionMessage", "Paste Actors"), EPasteTo::PT_OriginalLocation);
+	if (DoOutlinerFoldersNeedEditorDelegates() && LightCardOutliner->CanPasteSelectedFolder())
+	{
+		// We still must fire the delegates in case folders are selected
+		FEditorDelegates::OnEditPasteActorsBegin.Broadcast();
+		FEditorDelegates::OnEditPasteActorsEnd.Broadcast();
+		return;
+	}
 
+	UWorld* EditorWorld = GEditor->GetEditorWorldContext().World();
+	
+	FEditorDelegates::OnEditPasteActorsBegin.Broadcast();
+	GEditor->PasteSelectedActorsFromClipboard(EditorWorld, LOCTEXT("PasteItemsTransactionMessage", "Paste Items"), EPasteTo::PT_OriginalLocation);
+	FEditorDelegates::OnEditPasteActorsEnd.Broadcast();
+	
 	TArray<AActor*> PastedActors;
 	TArray<ADisplayClusterLightCardActor*> PastedLightCards;
 	for (FSelectionIterator It(GEditor->GetSelectedActorIterator()); It; ++It)
@@ -551,16 +612,7 @@ void FDisplayClusterLightCardEditor::PasteActors(bool bOffsetActorPosition)
 				continue;
 			}
 
-			IDisplayClusterStageActor* StageActor = CastChecked<IDisplayClusterStageActor>(Actor);
 			PastedActors.Add(Actor);
-			if (bOffsetActorPosition)
-			{
-				// If the light card should be offset from its pasted location, offset its longitude and latitude by a number of
-				// degrees equal to an arc length of 10 units (arc length = angle in radians * radius)
-				const float AngleOffset = FMath::RadiansToDegrees(10.0f / FMath::Max(StageActor->GetDistanceFromCenter(), 1.0f));
-				StageActor->SetLatitude(StageActor->GetLatitude() - AngleOffset);
-				StageActor->SetLongitude(StageActor->GetLongitude() + AngleOffset);
-			}
 			
 			if (ADisplayClusterLightCardActor* LightCard = Cast<ADisplayClusterLightCardActor>(Actor))
 			{
@@ -575,16 +627,60 @@ void FDisplayClusterLightCardEditor::PasteActors(bool bOffsetActorPosition)
 
 bool FDisplayClusterLightCardEditor::CanPasteActors() const
 {
+	if (LightCardOutliner.IsValid() && LightCardOutliner->CanPasteSelectedFolder())
+	{
+		// If only a folder is copied -- this isn't handled outside of the outliner
+		return true;
+	}
+	
 	UWorld* EditorWorld = GEditor->GetEditorWorldContext().World();
 	return GEditor->CanPasteSelectedActorsFromClipboard(EditorWorld);
 }
 
 void FDisplayClusterLightCardEditor::DuplicateSelectedActors()
 {
-	CopySelectedActors();
+	FScopedTransaction Transaction(LOCTEXT("DuplicateItemsTransactionMessage", "Duplicate Selected Items"));
 
-	const bool bOffsetActorPosition = true;
-	PasteActors(bOffsetActorPosition);
+	if (DoOutlinerFoldersNeedEditorDelegates())
+	{
+		// We still must fire the delegates in case folders are selected
+		FEditorDelegates::OnDuplicateActorsBegin.Broadcast();
+		FEditorDelegates::OnDuplicateActorsEnd.Broadcast();
+		return;
+	}
+	
+	FEditorDelegates::OnDuplicateActorsBegin.Broadcast();
+
+	const UWorld* EditorWorld = GEditor->GetEditorWorldContext().World();
+
+	TArray<AActor*> NewActors;
+	TArray<ADisplayClusterLightCardActor*> DuplicatedLightCards;
+	GUnrealEd->DuplicateActors(GetSelectedActorsAs<AActor>(), NewActors, EditorWorld->GetCurrentLevel(), FVector());
+
+	for (AActor* Actor : NewActors)
+	{
+		if (!Actor->Implements<UDisplayClusterStageActor>())
+		{
+			continue;
+		}
+		
+		IDisplayClusterStageActor* StageActor = CastChecked<IDisplayClusterStageActor>(Actor);
+		// If the light card should be offset from its pasted location, offset its longitude and latitude by a number of
+		// degrees equal to an arc length of 10 units (arc length = angle in radians * radius)
+		const float AngleOffset = FMath::RadiansToDegrees(10.0f / FMath::Max(StageActor->GetDistanceFromCenter(), 1.0f));
+		StageActor->SetLatitude(StageActor->GetLatitude() - AngleOffset);
+		StageActor->SetLongitude(StageActor->GetLongitude() + AngleOffset);
+
+		if (ADisplayClusterLightCardActor* LightCard = Cast<ADisplayClusterLightCardActor>(Actor))
+		{
+			DuplicatedLightCards.Add(LightCard);
+		}
+	}
+	
+	FEditorDelegates::OnDuplicateActorsEnd.Broadcast();
+	
+	AddLightCardsToActor(DuplicatedLightCards);
+	SelectActors(NewActors);
 }
 
 bool FDisplayClusterLightCardEditor::CanDuplicateSelectedActors() const
@@ -592,6 +688,24 @@ bool FDisplayClusterLightCardEditor::CanDuplicateSelectedActors() const
 	return CanCopySelectedActors();
 }
 
+void FDisplayClusterLightCardEditor::RenameSelectedItem()
+{
+	if (LightCardOutliner.IsValid())
+	{
+		return LightCardOutliner->RenameSelectedItem();
+	}
+}
+
+bool FDisplayClusterLightCardEditor::CanRenameSelectedItem() const
+{
+	if (LightCardOutliner.IsValid())
+	{
+		return LightCardOutliner->CanRenameSelectedItem();
+	}
+
+	return false;
+}
+
 void FDisplayClusterLightCardEditor::RemoveSelectedActors(bool bDeleteLightCardActor)
 {
 	RemoveActors(SelectedActors, bDeleteLightCardActor);
@@ -600,7 +714,20 @@ void FDisplayClusterLightCardEditor::RemoveSelectedActors(bool bDeleteLightCardA
 void FDisplayClusterLightCardEditor::RemoveActors(
 	const TArray<TWeakObjectPtr<AActor>>& InActorsToRemove, bool bDeleteActors)
 {
-	FScopedTransaction Transaction(LOCTEXT("RemoveActorTransactionMessage", "Remove Light Card(s)"));
+	if (DoOutlinerFoldersNeedEditorDelegates())
+	{
+		// We still must fire the delegates in case folders are selected
+		FEditorDelegates::OnDeleteActorsBegin.Broadcast();
+		FEditorDelegates::OnDeleteActorsEnd.Broadcast();
+		return;
+	}
+	
+	if (InActorsToRemove.Num() == 0)
+	{
+		return;
+	}
+
+	FScopedTransaction Transaction(LOCTEXT("RemoveItemsTransactionMessage", "Remove Item(s)"));
 
 	USelection* EdSelectionManager = GEditor->GetSelectedActors();
 	UWorld* WorldToUse = nullptr;
@@ -648,7 +775,9 @@ void FDisplayClusterLightCardEditor::RemoveActors(
 
 		if (WorldToUse)
 		{
+			FEditorDelegates::OnDeleteActorsBegin.Broadcast();
 			GEditor->edactDeleteSelected(WorldToUse);
+			FEditorDelegates::OnDeleteActorsEnd.Broadcast();
 		}
 	}
 
@@ -657,6 +786,12 @@ void FDisplayClusterLightCardEditor::RemoveActors(
 
 bool FDisplayClusterLightCardEditor::CanRemoveSelectedActors() const
 {
+	if (LightCardOutliner.IsValid() && LightCardOutliner->CanDeleteSelectedFolder())
+	{
+		// If only a folder is selected -- this isn't handled outside of the outliner
+		return true;
+	}
+	
 	for (const TWeakObjectPtr<AActor>& Actor : SelectedActors)
 	{
 		if (Actor.IsValid())
@@ -1241,12 +1376,12 @@ void FDisplayClusterLightCardEditor::BindCommands()
 
 	CommandList->MapAction(
 		FGenericCommands::Get().Copy,
-		FExecuteAction::CreateSP(this, &FDisplayClusterLightCardEditor::CopySelectedActors),
+		FExecuteAction::CreateSP(this, &FDisplayClusterLightCardEditor::CopySelectedActors, false),
 		FCanExecuteAction::CreateSP(this, &FDisplayClusterLightCardEditor::CanCopySelectedActors));
 
 	CommandList->MapAction(
 		FGenericCommands::Get().Paste,
-		FExecuteAction::CreateSP(this, &FDisplayClusterLightCardEditor::PasteActors, false),
+		FExecuteAction::CreateSP(this, &FDisplayClusterLightCardEditor::PasteActors),
 		FCanExecuteAction::CreateSP(this, &FDisplayClusterLightCardEditor::CanPasteActors));
 
 	CommandList->MapAction(
@@ -1254,6 +1389,11 @@ void FDisplayClusterLightCardEditor::BindCommands()
 		FExecuteAction::CreateSP(this, &FDisplayClusterLightCardEditor::DuplicateSelectedActors),
 		FCanExecuteAction::CreateSP(this, &FDisplayClusterLightCardEditor::CanDuplicateSelectedActors));
 
+	CommandList->MapAction(
+		FGenericCommands::Get().Rename,
+		FExecuteAction::CreateSP(this, &FDisplayClusterLightCardEditor::RenameSelectedItem),
+		FCanExecuteAction::CreateSP(this, &FDisplayClusterLightCardEditor::CanRenameSelectedItem));
+	
 	CommandList->MapAction(
 		FGenericCommands::Get().Delete,
 		FExecuteAction::CreateSP(this, &FDisplayClusterLightCardEditor::RemoveSelectedActors, true),
@@ -1544,4 +1684,9 @@ void FDisplayClusterLightCardEditor::CleanupRecentlyPlacedItems()
 	Settings->SaveConfig();
 }
 
+bool FDisplayClusterLightCardEditor::DoOutlinerFoldersNeedEditorDelegates() const
+{
+	return LightCardOutliner.IsValid() && SelectedActors.Num() == 0;
+}
+
 #undef LOCTEXT_NAMESPACE
diff --git a/Engine/Plugins/Runtime/nDisplay/Source/DisplayClusterLightCardEditor/Private/DisplayClusterLightCardEditor.h b/Engine/Plugins/Runtime/nDisplay/Source/DisplayClusterLightCardEditor/Private/DisplayClusterLightCardEditor.h
--- a/Engine/Plugins/Runtime/nDisplay/Source/DisplayClusterLightCardEditor/Private/DisplayClusterLightCardEditor.h
+++ b/Engine/Plugins/Runtime/nDisplay/Source/DisplayClusterLightCardEditor/Private/DisplayClusterLightCardEditor.h
@@ -124,13 +124,13 @@ public:
 	bool CanCutSelectedActors();
 
 	/** Copies any selected actors to the clipboard */
-	void CopySelectedActors();
+	void CopySelectedActors(bool bShouldCut = false);
 
 	/** Determines if there are selected actors that can be copied */
 	bool CanCopySelectedActors() const;
 
 	/** Pastes any actors in the clipboard to the current root actor */
-	void PasteActors(bool bOffsetActorPosition);
+	void PasteActors();
 
 	/** Determines if there are any actors that can be pasted from the clipboard */
 	bool CanPasteActors() const;
@@ -141,6 +141,12 @@ public:
 	/** Determines if there are selected actors that can be duplicated */
 	bool CanDuplicateSelectedActors() const;
 
+	/** Rename the selected item. Requires the outliner */
+	void RenameSelectedItem();
+
+	/** If the selected item can be renamed */
+	bool CanRenameSelectedItem() const;
+	
 	/**
 	 * Remove the light card from the actor
 	 *@param bDeleteLightCardActor Delete the actor from the level
@@ -288,6 +294,14 @@ private:
 
 	/** Make sure no invalid entries are present */
 	void CleanupRecentlyPlacedItems();
+
+	/**
+	 * Checks if it is possible only one or more folders - and no actors - are selected in the outliner
+	 * They require editor delegates fired when executing some generic commands
+	 *
+	 * @return true if the outliner is valid and no actor is selected
+	 */
+	bool DoOutlinerFoldersNeedEditorDelegates() const;
 	
 private:
 	/** The light card outliner widget */
diff --git a/Engine/Plugins/Runtime/nDisplay/Source/DisplayClusterLightCardEditor/Private/Outliner/DisplayClusterLightCardOutlinerMode.cpp b/Engine/Plugins/Runtime/nDisplay/Source/DisplayClusterLightCardEditor/Private/Outliner/DisplayClusterLightCardOutlinerMode.cpp
--- a/Engine/Plugins/Runtime/nDisplay/Source/DisplayClusterLightCardEditor/Private/Outliner/DisplayClusterLightCardOutlinerMode.cpp
+++ b/Engine/Plugins/Runtime/nDisplay/Source/DisplayClusterLightCardEditor/Private/Outliner/DisplayClusterLightCardOutlinerMode.cpp
@@ -110,6 +110,30 @@ FDisplayClusterLightCardOutlinerMode::FDisplayClusterLightCardOutlinerMode(SScen
 	/* bHideComponents */ true, /* bHideLevelInstanceHierarchy */ true, /* bHideUnloadedActors */ true, /* bHideEmptyFolders */ true)),
 	LightCardOutliner(InLightCardOutliner)
 {
+	FEditorDelegates::OnEditCutActorsBegin.AddRaw(this, &FDisplayClusterLightCardOutlinerMode::OnEditCutActorsBegin);
+	FEditorDelegates::OnEditCutActorsEnd.AddRaw(this, &FDisplayClusterLightCardOutlinerMode::OnEditCutActorsEnd);
+	FEditorDelegates::OnEditCopyActorsBegin.AddRaw(this, &FDisplayClusterLightCardOutlinerMode::OnEditCopyActorsBegin);
+	FEditorDelegates::OnEditCopyActorsEnd.AddRaw(this, &FDisplayClusterLightCardOutlinerMode::OnEditCopyActorsEnd);
+	FEditorDelegates::OnEditPasteActorsBegin.AddRaw(this, &FDisplayClusterLightCardOutlinerMode::OnEditPasteActorsBegin);
+	FEditorDelegates::OnEditPasteActorsEnd.AddRaw(this, &FDisplayClusterLightCardOutlinerMode::OnEditPasteActorsEnd);
+	FEditorDelegates::OnDuplicateActorsBegin.AddRaw(this, &FDisplayClusterLightCardOutlinerMode::OnDuplicateActorsBegin);
+	FEditorDelegates::OnDuplicateActorsEnd.AddRaw(this, &FDisplayClusterLightCardOutlinerMode::OnDuplicateActorsEnd);
+	FEditorDelegates::OnDeleteActorsBegin.AddRaw(this, &FDisplayClusterLightCardOutlinerMode::OnDeleteActorsBegin);
+	FEditorDelegates::OnDeleteActorsEnd.AddRaw(this, &FDisplayClusterLightCardOutlinerMode::OnDeleteActorsEnd);
+}
+
+FDisplayClusterLightCardOutlinerMode::~FDisplayClusterLightCardOutlinerMode()
+{
+	FEditorDelegates::OnEditCutActorsBegin.RemoveAll(this);
+	FEditorDelegates::OnEditCutActorsEnd.RemoveAll(this);
+	FEditorDelegates::OnEditCopyActorsBegin.RemoveAll(this);
+	FEditorDelegates::OnEditCopyActorsEnd.RemoveAll(this);
+	FEditorDelegates::OnEditPasteActorsBegin.RemoveAll(this);
+	FEditorDelegates::OnEditPasteActorsEnd.RemoveAll(this);
+	FEditorDelegates::OnDuplicateActorsBegin.RemoveAll(this);
+	FEditorDelegates::OnDuplicateActorsEnd.RemoveAll(this);
+	FEditorDelegates::OnDeleteActorsBegin.RemoveAll(this);
+	FEditorDelegates::OnDeleteActorsEnd.RemoveAll(this);
 }
 
 void FDisplayClusterLightCardOutlinerMode::SynchronizeSelection()
@@ -300,8 +324,11 @@ FReply FDisplayClusterLightCardOutlinerMode::OnKeyDown(const FKeyEvent& InKeyEve
 
 bool FDisplayClusterLightCardOutlinerMode::CanDelete() const
 {
-	// Handled by command in light card editor or in OnKeyDown for folders
-	return false;
+	// Also handled in OnKeyDown for folders
+	
+	const FSceneOutlinerItemSelection ItemSelection = SceneOutliner->GetSelection();
+	const uint32 NumberOfFolders = ItemSelection.Num<FFolderTreeItem>();
+	return (NumberOfFolders > 0 && NumberOfFolders == ItemSelection.Num());
 }
 
 bool FDisplayClusterLightCardOutlinerMode::CanRename() const
@@ -313,26 +340,26 @@ bool FDisplayClusterLightCardOutlinerMode::CanRename() const
 
 bool FDisplayClusterLightCardOutlinerMode::CanRenameItem(const ISceneOutlinerTreeItem& Item) const
 {
-	// Can only rename actor and folder items when in actor browsing mode
 	return (Item.IsValid() && (Item.IsA<FActorTreeItem>() || Item.IsA<FFolderTreeItem>()));
 }
 
 bool FDisplayClusterLightCardOutlinerMode::CanCut() const
 {
-	// Handled by command in light card editor
-	return false;
+	const FSceneOutlinerItemSelection ItemSelection = SceneOutliner->GetSelection();
+	const uint32 NumberOfFolders = ItemSelection.Num<FFolderTreeItem>();
+	return (NumberOfFolders > 0 && NumberOfFolders == ItemSelection.Num());
 }
 
 bool FDisplayClusterLightCardOutlinerMode::CanCopy() const
 {
-	// Handled by command in light card editor
-	return false;
+	const FSceneOutlinerItemSelection ItemSelection = SceneOutliner->GetSelection();
+	const uint32 NumberOfFolders = ItemSelection.Num<FFolderTreeItem>();
+	return (NumberOfFolders > 0 && NumberOfFolders == ItemSelection.Num());
 }
 
 bool FDisplayClusterLightCardOutlinerMode::CanPaste() const
 {
-	// Handled by command in light card editor
-	return false;
+	return CanPasteFoldersOnlyFromClipboard();
 }
 
 FFolder FDisplayClusterLightCardOutlinerMode::CreateNewFolder()
@@ -1005,6 +1032,62 @@ void FDisplayClusterLightCardOutlinerMode::OnDrop(ISceneOutlinerTreeItem& DropTa
 	}
 }
 
+void FDisplayClusterLightCardOutlinerMode::OnEditCutActorsBegin()
+{
+	SceneOutliner->CopyFoldersBegin();
+	SceneOutliner->DeleteFoldersBegin();
+}
+
+void FDisplayClusterLightCardOutlinerMode::OnEditCutActorsEnd()
+{
+	SceneOutliner->CopyFoldersEnd();
+	SceneOutliner->DeleteFoldersEnd();
+}
+
+void FDisplayClusterLightCardOutlinerMode::OnEditCopyActorsBegin()
+{
+	SceneOutliner->CopyFoldersBegin();
+}
+
+void FDisplayClusterLightCardOutlinerMode::OnEditCopyActorsEnd()
+{
+	SceneOutliner->CopyFoldersEnd();
+}
+
+void FDisplayClusterLightCardOutlinerMode::OnEditPasteActorsBegin()
+{
+	const TArray<FName> FolderPaths = SceneOutliner->GetClipboardPasteFolders();
+	SceneOutliner->PasteFoldersBegin(FolderPaths);
+}
+
+void FDisplayClusterLightCardOutlinerMode::OnEditPasteActorsEnd()
+{
+	SceneOutliner->PasteFoldersEnd();
+}
+
+void FDisplayClusterLightCardOutlinerMode::OnDuplicateActorsBegin()
+{
+	FFolder::FRootObject CommonRootObject;
+	TArray<FName> SelectedFolderPaths;
+	FFolder::GetFolderPathsAndCommonRootObject(SceneOutliner->GetSelection().GetData<FFolder>(FFolderPathSelector()), SelectedFolderPaths, CommonRootObject);
+	SceneOutliner->PasteFoldersBegin(SelectedFolderPaths);
+}
+
+void FDisplayClusterLightCardOutlinerMode::OnDuplicateActorsEnd()
+{
+	SceneOutliner->PasteFoldersEnd();
+}
+
+void FDisplayClusterLightCardOutlinerMode::OnDeleteActorsBegin()
+{
+	SceneOutliner->DeleteFoldersBegin();
+}
+
+void FDisplayClusterLightCardOutlinerMode::OnDeleteActorsEnd()
+{
+	SceneOutliner->DeleteFoldersEnd();
+}
+
 void FDisplayClusterLightCardOutlinerMode::RegisterContextMenu()
 {
 	UToolMenus* ToolMenus = UToolMenus::Get();
diff --git a/Engine/Plugins/Runtime/nDisplay/Source/DisplayClusterLightCardEditor/Private/Outliner/DisplayClusterLightCardOutlinerMode.h b/Engine/Plugins/Runtime/nDisplay/Source/DisplayClusterLightCardEditor/Private/Outliner/DisplayClusterLightCardOutlinerMode.h
--- a/Engine/Plugins/Runtime/nDisplay/Source/DisplayClusterLightCardEditor/Private/Outliner/DisplayClusterLightCardOutlinerMode.h
+++ b/Engine/Plugins/Runtime/nDisplay/Source/DisplayClusterLightCardEditor/Private/Outliner/DisplayClusterLightCardOutlinerMode.h
@@ -18,6 +18,8 @@ public:
 	FDisplayClusterLightCardOutlinerMode(SSceneOutliner* InSceneOutliner, TWeakPtr<SDisplayClusterLightCardOutliner> InLightCardOutliner,
 		TWeakObjectPtr<UWorld> InSpecifiedWorldToDisplay = nullptr);
 
+	virtual ~FDisplayClusterLightCardOutlinerMode() override;
+	
 	virtual void SynchronizeSelection() override;
 	virtual bool IsInteractive() const override { return true; }
 	virtual FCreateSceneOutlinerMode CreateFolderPickerMode(const FFolder::FRootObject& InRootObject = FFolder::GetInvalidRootObject()) const override;
@@ -47,6 +49,36 @@ public:
 	virtual FSceneOutlinerDragValidationInfo ValidateDrop(const ISceneOutlinerTreeItem& DropTarget, const FSceneOutlinerDragDropPayload& Payload) const override;
 	virtual void OnDrop(ISceneOutlinerTreeItem& DropTarget, const FSceneOutlinerDragDropPayload& Payload, const FSceneOutlinerDragValidationInfo& ValidationInfo) const override;
 
+	/** Called by engine when edit cut actors begins */
+	void OnEditCutActorsBegin();
+
+	/** Called by engine when edit cut actors ends */
+	void OnEditCutActorsEnd();
+
+	/** Called by engine when edit copy actors begins */
+	void OnEditCopyActorsBegin();
+
+	/** Called by engine when edit copy actors ends */
+	void OnEditCopyActorsEnd();
+
+	/** Called by engine when edit paste actors begins */
+	void OnEditPasteActorsBegin();
+
+	/** Called by engine when edit paste actors ends */
+	void OnEditPasteActorsEnd();
+
+	/** Called by engine when edit duplicate actors begins */
+	void OnDuplicateActorsBegin();
+
+	/** Called by engine when edit duplicate actors ends */
+	void OnDuplicateActorsEnd();
+
+	/** Called by engine when edit delete actors begins */
+	void OnDeleteActorsBegin();
+
+	/** Called by engine when edit delete actors ends */
+	void OnDeleteActorsEnd();
+	
 private:
 	static void RegisterContextMenu();
 	bool CanPasteFoldersOnlyFromClipboard() const;
diff --git a/Engine/Plugins/Runtime/nDisplay/Source/DisplayClusterLightCardEditor/Private/Outliner/SDisplayClusterLightCardOutliner.cpp b/Engine/Plugins/Runtime/nDisplay/Source/DisplayClusterLightCardEditor/Private/Outliner/SDisplayClusterLightCardOutliner.cpp
--- a/Engine/Plugins/Runtime/nDisplay/Source/DisplayClusterLightCardEditor/Private/Outliner/SDisplayClusterLightCardOutliner.cpp
+++ b/Engine/Plugins/Runtime/nDisplay/Source/DisplayClusterLightCardEditor/Private/Outliner/SDisplayClusterLightCardOutliner.cpp
@@ -85,8 +85,14 @@ void SDisplayClusterLightCardOutliner::GetSelectedActors(TArray<AActor*>& OutSel
 
 void SDisplayClusterLightCardOutliner::SelectActors(const TArray<AActor*>& ActorsToSelect)
 {
+	if (bIsOutlinerChangingSelection)
+	{
+		// Selection change initiated from outliner
+		return;
+	}
+	
 	TArray<FSceneOutlinerTreeItemPtr> SelectedTreeItems;
-	CachedSelectedActors.Reset();
+	CachedOutlinerItems.Reset();
 	
 	for (const TSharedPtr<FStageActorTreeItem>& TreeItem : StageActorTreeItems)
 	{
@@ -95,22 +101,34 @@ void SDisplayClusterLightCardOutliner::SelectActors(const TArray<AActor*>& Actor
 			if (FSceneOutlinerTreeItemPtr OutlinerItem = SceneOutliner->GetTreeItem(TreeItem->Actor.Get()))
 			{
 				SelectedTreeItems.Add(OutlinerItem);
-				CachedSelectedActors.Add(TreeItem->Actor);
 			}
 		}
 	}
-	
+
+	CachedOutlinerItems = SelectedTreeItems;
 	SceneOutliner->SetItemSelection(SelectedTreeItems, true);
 }
 
 void SDisplayClusterLightCardOutliner::RestoreCachedSelection()
 {
 	TArray<AActor*> ActorsToSelect;
-	for (const TWeakObjectPtr<AActor>& CachedActor : CachedSelectedActors)
+	TArray<TSharedPtr<ISceneOutlinerTreeItem>> AdditionalTreeItemsToSelect;
+	
+	for (TSharedPtr<ISceneOutlinerTreeItem> CachedItem : CachedOutlinerItems)
 	{
-		if (CachedActor.IsValid())
+		if (CachedItem.IsValid())
 		{
-			ActorsToSelect.Add(CachedActor.Get());
+			if (const FActorTreeItem* ActorItem = CachedItem->CastTo<FActorTreeItem>())
+			{
+				if (ActorItem->Actor.IsValid())
+				{
+					ActorsToSelect.Add(ActorItem->Actor.Get());
+				}
+			}
+			else
+			{
+				AdditionalTreeItemsToSelect.Add(CachedItem);
+			}
 		}
 	}
 
@@ -118,6 +136,82 @@ void SDisplayClusterLightCardOutliner::RestoreCachedSelection()
 	{
 		LightCardEditorPtr.Pin()->SelectActors(MoveTemp(ActorsToSelect));
 	}
+
+	if (AdditionalTreeItemsToSelect.Num() > 0)
+	{
+		SceneOutliner->AddToSelection(AdditionalTreeItemsToSelect);
+	}
+}
+
+bool SDisplayClusterLightCardOutliner::CanDeleteSelectedFolder() const
+{
+	if (SceneOutliner.IsValid())
+	{
+		const FDisplayClusterLightCardOutlinerMode* Mode = static_cast<const FDisplayClusterLightCardOutlinerMode*>(SceneOutliner->GetMode());
+		return Mode->CanDelete();
+	}
+	
+	return false;
+}
+
+bool SDisplayClusterLightCardOutliner::CanCutSelectedFolder() const
+{
+	if (SceneOutliner.IsValid())
+	{
+		const FDisplayClusterLightCardOutlinerMode* Mode = static_cast<const FDisplayClusterLightCardOutlinerMode*>(SceneOutliner->GetMode());
+		return Mode->CanCut();
+	}
+
+	return false;
+}
+
+bool SDisplayClusterLightCardOutliner::CanCopySelectedFolder() const
+{
+	if (SceneOutliner.IsValid())
+	{
+		const FDisplayClusterLightCardOutlinerMode* Mode = static_cast<const FDisplayClusterLightCardOutlinerMode*>(SceneOutliner->GetMode());
+		return Mode->CanCopy();
+	}
+
+	return false;
+}
+
+bool SDisplayClusterLightCardOutliner::CanPasteSelectedFolder() const
+{
+	if (SceneOutliner.IsValid())
+	{
+		const FDisplayClusterLightCardOutlinerMode* Mode = static_cast<const FDisplayClusterLightCardOutlinerMode*>(SceneOutliner->GetMode());
+		return Mode->CanPaste();
+	}
+
+	return false;
+}
+
+bool SDisplayClusterLightCardOutliner::CanRenameSelectedItem() const
+{
+	if (SceneOutliner.IsValid())
+	{
+		TArray<FSceneOutlinerTreeItemPtr> SelectedOutlinerItems = SceneOutliner->GetSelectedItems();
+		const FDisplayClusterLightCardOutlinerMode* Mode = static_cast<const FDisplayClusterLightCardOutlinerMode*>(SceneOutliner->GetMode());
+		if (SelectedOutlinerItems.Num() == 1 && SelectedOutlinerItems[0].IsValid())
+		{
+			return Mode->CanRenameItem(*SelectedOutlinerItems[0].Get());
+		}
+	}
+
+	return false;
+}
+
+void SDisplayClusterLightCardOutliner::RenameSelectedItem()
+{
+	if (SceneOutliner.IsValid())
+	{
+		TArray<FSceneOutlinerTreeItemPtr> SelectedItems = SceneOutliner->GetSelectedItems();
+		if (SelectedItems.Num() == 1)
+		{
+			SelectedItems[0]->RenameRequestEvent.ExecuteIfBound();
+		}
+	}
 }
 
 void SDisplayClusterLightCardOutliner::CreateWorldOutliner()
@@ -255,13 +349,18 @@ void SDisplayClusterLightCardOutliner::OnOutlinerSelectionChanged(FSceneOutliner
 {
 	MostRecentSelectedItem = TreeItem;
 
-	TArray<AActor*> SelectedLightCards;
-	GetSelectedActors(SelectedLightCards);
-
 	if (LightCardEditorPtr.IsValid())
 	{
+		TArray<AActor*> SelectedLightCards;
+    	GetSelectedActors(SelectedLightCards);
+	
+		bIsOutlinerChangingSelection = true;
 		LightCardEditorPtr.Pin()->SelectActors(SelectedLightCards);
 		LightCardEditorPtr.Pin()->SelectActorProxies(SelectedLightCards);
+		
+		CachedOutlinerItems = SceneOutliner->GetSelectedItems();
+		
+		bIsOutlinerChangingSelection = false;
 	}
 }
 
@@ -341,8 +440,6 @@ void SDisplayClusterLightCardOutliner::RegisterContextMenu(FName& InName, FToolM
 				}
 			}
 			
-			const bool bAddFullEditMenu = Actor && Actor->Implements<UDisplayClusterStageActor>();
-			
 			if (ADisplayClusterLightCardActor* LightCardActor = Cast<ADisplayClusterLightCardActor>(Actor))
 			{
 				FToolMenuSection& Section = InMenu->AddSection("LightCardDefaultSection", LOCTEXT("LightCardSectionName", "Light Cards"));
@@ -361,19 +458,12 @@ void SDisplayClusterLightCardOutliner::RegisterContextMenu(FName& InName, FToolM
 			{
 				FToolMenuSection& Section = InMenu->AddSection("LightCardEditSection", LOCTEXT("LightCardEditSectionName", "Edit"));
 
-				if (bAddFullEditMenu)
-				{
-					Section.AddMenuEntry(FGenericCommands::Get().Cut);
-					Section.AddMenuEntry(FGenericCommands::Get().Copy);
-				}
-
+				Section.AddMenuEntry(FGenericCommands::Get().Cut);
+				Section.AddMenuEntry(FGenericCommands::Get().Copy);
 				Section.AddMenuEntry(FGenericCommands::Get().Paste);
-
-				if (bAddFullEditMenu)
-				{
-					Section.AddMenuEntry(FGenericCommands::Get().Duplicate);
-					Section.AddMenuEntry(FGenericCommands::Get().Delete);
-				}
+				Section.AddMenuEntry(FGenericCommands::Get().Duplicate);
+				Section.AddMenuEntry(FGenericCommands::Get().Delete);
+				Section.AddMenuEntry(FGenericCommands::Get().Rename);
 			}
 		}));
 	}
diff --git a/Engine/Plugins/Runtime/nDisplay/Source/DisplayClusterLightCardEditor/Private/Outliner/SDisplayClusterLightCardOutliner.h b/Engine/Plugins/Runtime/nDisplay/Source/DisplayClusterLightCardEditor/Private/Outliner/SDisplayClusterLightCardOutliner.h
--- a/Engine/Plugins/Runtime/nDisplay/Source/DisplayClusterLightCardEditor/Private/Outliner/SDisplayClusterLightCardOutliner.h
+++ b/Engine/Plugins/Runtime/nDisplay/Source/DisplayClusterLightCardEditor/Private/Outliner/SDisplayClusterLightCardOutliner.h
@@ -53,6 +53,19 @@ public:
 	/** Restores the last valid cached selection to the outliner */
 	void RestoreCachedSelection();
 
+	// Outliner is responsible for renames on all items, all other edit options are for folders only
+	// Otherwise the light card editor is responsible for actor edits.
+	
+	bool CanDeleteSelectedFolder() const;
+	bool CanCutSelectedFolder() const;
+	bool CanCopySelectedFolder() const;
+	bool CanPasteSelectedFolder() const;
+
+	/** If either a folder or actor should be renamed */
+	bool CanRenameSelectedItem() const;
+	/** Rename the selected folder or actor */
+	void RenameSelectedItem();
+	
 private:
 	/** Creates a world outliner based on the root actor world */
 	void CreateWorldOutliner();
@@ -87,8 +100,8 @@ private:
 	/** All actors currently tracked by the scene outliner */
 	TMap<TObjectPtr<AActor>, TSharedPtr<FStageActorTreeItem>> TrackedActors;
 
-	/** Cached actors currently selected. Used for synchronizing with the outliner mode */
-	TArray<TWeakObjectPtr<AActor>> CachedSelectedActors;
+	/** Cached items currently selected. Used for synchronizing with the outliner mode */
+	TArray<FSceneOutlinerTreeItemPtr> CachedOutlinerItems;
 	
 	/** A flattened list of all the light card actors being displayed in the tree view */
 	TArray<TSharedPtr<FStageActorTreeItem>> StageActorTreeItems;
@@ -110,4 +123,7 @@ private:
 
 	/** Mapped commands for this list */
 	TSharedPtr<FUICommandList> CommandList;
+
+	/** If the outliner is responsible for changing the actor selection */
+	bool bIsOutlinerChangingSelection = false;
 };
\ No newline at end of file
diff --git a/Engine/Plugins/Runtime/nDisplay/Source/DisplayClusterLightCardEditor/Private/Viewport/DisplayClusterLightcardEditorViewport.cpp b/Engine/Plugins/Runtime/nDisplay/Source/DisplayClusterLightCardEditor/Private/Viewport/DisplayClusterLightcardEditorViewport.cpp
--- a/Engine/Plugins/Runtime/nDisplay/Source/DisplayClusterLightCardEditor/Private/Viewport/DisplayClusterLightcardEditorViewport.cpp
+++ b/Engine/Plugins/Runtime/nDisplay/Source/DisplayClusterLightCardEditor/Private/Viewport/DisplayClusterLightcardEditorViewport.cpp
@@ -783,8 +783,7 @@ void SDisplayClusterLightCardEditorViewport::PasteLightCardsHere()
 {
 	if (LightCardEditorPtr.IsValid())
 	{
-		const bool bOffsetActorPosition = false;
-		LightCardEditorPtr.Pin()->PasteActors(bOffsetActorPosition);
+		LightCardEditorPtr.Pin()->PasteActors();
 
 		if (ViewportClient.IsValid())
 		{
