nDisplay: Fix occasional crash when selecting components in the blueprint editor.

diff --git a/Engine/Plugins/Runtime/nDisplay/Source/DisplayClusterConfigurator/Private/Views/Viewport/DisplayClusterConfiguratorSCSEditorViewportClient.cpp b/Engine/Plugins/Runtime/nDisplay/Source/DisplayClusterConfigurator/Private/Views/Viewport/DisplayClusterConfiguratorSCSEditorViewportClient.cpp
--- a/Engine/Plugins/Runtime/nDisplay/Source/DisplayClusterConfigurator/Private/Views/Viewport/DisplayClusterConfiguratorSCSEditorViewportClient.cpp
+++ b/Engine/Plugins/Runtime/nDisplay/Source/DisplayClusterConfigurator/Private/Views/Viewport/DisplayClusterConfiguratorSCSEditorViewportClient.cpp
@@ -307,14 +307,20 @@ void FDisplayClusterConfiguratorSCSEditorViewportClient::Draw(const FSceneView*
 			{
 				FSubobjectEditorTreeNodePtrType SelectedNode = SelectedNodes[SelectionIndex];
 
-				const UActorComponent* Comp = SelectedNode.IsValid() ? SelectedNode->GetDataSource()->FindComponentInstanceInActor(PreviewActor) : nullptr;
-				if (Comp != nullptr && Comp->IsRegistered())
+				if (SelectedNode.IsValid())
 				{
-					// Try and find a visualizer
-					TSharedPtr<FComponentVisualizer> Visualizer = GUnrealEd->FindComponentVisualizer(Comp->GetClass());
-					if (Visualizer.IsValid())
+					if (const FSubobjectData* Data = SelectedNode->GetDataSource())
 					{
-						Visualizer->DrawVisualization(Comp, View, PDI);
+						const UActorComponent* Comp = Data->FindComponentInstanceInActor(PreviewActor);
+						if (Comp != nullptr && Comp->IsRegistered())
+						{
+							// Try and find a visualizer
+							TSharedPtr<FComponentVisualizer> Visualizer = GUnrealEd->FindComponentVisualizer(Comp->GetClass());
+							if (Visualizer.IsValid())
+							{
+								Visualizer->DrawVisualization(Comp, View, PDI);
+							}
+						}
 					}
 				}
 			}
@@ -353,27 +359,28 @@ void FDisplayClusterConfiguratorSCSEditorViewportClient::DrawCanvas(FViewport& I
 		const int32 HalfY = 0.5f * Viewport->GetSizeXY().Y;
 
 		const TArray<FSubobjectEditorTreeNodePtrType>& SelectedNodes = BlueprintEditorPtr.Pin()->GetSelectedSubobjectEditorTreeNodes();
-		if (bIsManipulating && SelectedNodes.Num() > 0)
+		if (bIsManipulating && SelectedNodes.Num() > 0 && SelectedNodes[0].IsValid())
 		{
-			FSubobjectData* Data = SelectedNodes[0]->GetDataSource();
-			const USceneComponent* SceneComp = Cast<USceneComponent>(Data->FindComponentInstanceInActor(PreviewActor));
-			if (SceneComp)
+			if (const FSubobjectData* Data = SelectedNodes[0]->GetDataSource())
 			{
-				const FVector WidgetLocation = GetWidgetLocation();
-				const FPlane Proj = View.Project(WidgetLocation);
-				if (Proj.W > 0.0f)
+				if (const USceneComponent* SceneComp = Cast<USceneComponent>(Data->FindComponentInstanceInActor(PreviewActor)))
 				{
-					const int32 XPos = HalfX + (HalfX * Proj.X);
-					const int32 YPos = HalfY + (HalfY * (Proj.Y * -1));
-					DrawAngles(
-						&Canvas,
-						XPos,
-						YPos,
-						GetCurrentWidgetAxis(),
-						GetWidgetMode(),
-						GetWidgetCoordSystem().Rotator(),
-						WidgetLocation,
-						UpdateViewportClientWindowDPIScale());
+					const FVector WidgetLocation = GetWidgetLocation();
+					const FPlane Proj = View.Project(WidgetLocation);
+					if (Proj.W > 0.0f)
+					{
+						const int32 XPos = HalfX + (HalfX * Proj.X);
+						const int32 YPos = HalfY + (HalfY * (Proj.Y * -1));
+						DrawAngles(
+							&Canvas,
+							XPos,
+							YPos,
+							GetCurrentWidgetAxis(),
+							GetWidgetMode(),
+							GetWidgetCoordSystem().Rotator(),
+							WidgetLocation,
+							UpdateViewportClientWindowDPIScale());
+					}
 				}
 			}
 		}
@@ -546,7 +553,11 @@ bool FDisplayClusterConfiguratorSCSEditorViewportClient::InputWidgetDelta(FViewp
 
 				for (const FSubobjectEditorTreeNodePtrType& SelectedNodePtr : SelectedNodes)
 				{
-					FSubobjectData* Data = SelectedNodePtr->GetDataSource();
+					const FSubobjectData* Data = SelectedNodePtr.IsValid() ? SelectedNodePtr->GetDataSource() : nullptr;
+					if (Data == nullptr)
+					{
+						continue;
+					}
 					// Don't allow editing of a root node, inherited SCS node or child node that also has a movable (non-root) parent node selected
 					const bool bCanEdit = GUnrealEd->ComponentVisManager.IsActive() ||
 						(!Data->IsRootComponent() && !IsMovableParentNodeSelected(SelectedNodePtr, SelectedNodes));
@@ -960,10 +971,16 @@ void FDisplayClusterConfiguratorSCSEditorViewportClient::FocusViewportToSelectio
 		if (SelectedNodes.Num() > 0)
 		{
 			// Use the last selected item for the widget location
-			const USceneComponent* SceneComp = Cast<USceneComponent>(SelectedNodes.Last()->GetDataSource()->FindComponentInstanceInActor(PreviewActor));
-			if (SceneComp)
+			const FSubobjectEditorTreeNodePtrType& LastNode = SelectedNodes.Last();
+			if (LastNode.IsValid())
 			{
-				FocusViewportOnBox(SceneComp->Bounds.GetBox());
+				if (const FSubobjectData* Data = LastNode->GetDataSource())
+				{
+					if (const USceneComponent* SceneComp = Cast<USceneComponent>(Data->FindComponentInstanceInActor(PreviewActor)))
+					{
+						FocusViewportOnBox(SceneComp->Bounds.GetBox());
+					}
+				}
 			}
 		}
 		else
@@ -1217,24 +1234,26 @@ void FDisplayClusterConfiguratorSCSEditorViewportClient::BeginTransaction(const
 			{
 				if (Node.IsValid())
 				{
-					FSubobjectData* Data = Node->GetDataSource();
-					if (USCS_Node* SCS_Node = const_cast<USCS_Node*>(Data->GetObject<USCS_Node>()))
+					if (const FSubobjectData* Data = Node->GetDataSource())
 					{
-						USimpleConstructionScript* SCS = SCS_Node->GetSCS();
-						UBlueprint* Blueprint = SCS ? SCS->GetBlueprint() : nullptr;
-						if (Blueprint == PreviewBlueprint)
+						if (USCS_Node* SCS_Node = const_cast<USCS_Node*>(Data->GetObject<USCS_Node>()))
 						{
-							SCS_Node->Modify();
+							USimpleConstructionScript* SCS = SCS_Node->GetSCS();
+							UBlueprint* Blueprint = SCS ? SCS->GetBlueprint() : nullptr;
+							if (Blueprint == PreviewBlueprint)
+							{
+								SCS_Node->Modify();
+							}
 						}
-					}
 
-					// Modify template, any instances will be reconstructed as part of PostUndo:
-					UActorComponent* ComponentTemplate = const_cast<UActorComponent*>(Data->GetObjectForBlueprint<UActorComponent>(PreviewBlueprint));
+						// Modify template, any instances will be reconstructed as part of PostUndo:
+						UActorComponent* ComponentTemplate = const_cast<UActorComponent*>(Data->GetObjectForBlueprint<UActorComponent>(PreviewBlueprint));
 					
-					if (ComponentTemplate != nullptr)
-					{
-						ComponentTemplate->SetFlags(RF_Transactional);
-						ComponentTemplate->Modify();
+						if (ComponentTemplate != nullptr)
+						{
+							ComponentTemplate->SetFlags(RF_Transactional);
+							ComponentTemplate->Modify();
+						}
 					}
 				}
 			}
