[EpicStageApp] UV lightcard mode in app

diff --git a/Engine/Plugins/Runtime/nDisplay/Source/DisplayClusterLightCardEditor/Private/Viewport/DisplayClusterLightCardEditorViewportClient.cpp b/Engine/Plugins/Runtime/nDisplay/Source/DisplayClusterLightCardEditor/Private/Viewport/DisplayClusterLightCardEditorViewportClient.cpp
--- a/Engine/Plugins/Runtime/nDisplay/Source/DisplayClusterLightCardEditor/Private/Viewport/DisplayClusterLightCardEditorViewportClient.cpp
+++ b/Engine/Plugins/Runtime/nDisplay/Source/DisplayClusterLightCardEditor/Private/Viewport/DisplayClusterLightCardEditorViewportClient.cpp
@@ -335,20 +335,10 @@ void FDisplayClusterLightCardEditorViewportClient::Draw(FViewport* InViewport, F
 		FDisplayClusterMeshProjectionRenderSettings RenderSettings;
 		RenderSettings.RenderType = EDisplayClusterMeshProjectionOutput::Color;
 		RenderSettings.EngineShowFlags = EngineShowFlags;
-		RenderSettings.ProjectionType = ProjectionMode;
-		RenderSettings.PrimitiveFilter.ShouldRenderPrimitiveDelegate = FDisplayClusterMeshProjectionPrimitiveFilter::FPrimitiveFilter::CreateSP(this, &FDisplayClusterLightCardEditorViewportClient::ShouldRenderPrimitive);
-		RenderSettings.PrimitiveFilter.ShouldApplyProjectionDelegate = FDisplayClusterMeshProjectionPrimitiveFilter::FPrimitiveFilter::CreateSP(this, &FDisplayClusterLightCardEditorViewportClient::ShouldApplyProjectionToPrimitive);
 
-		if (ProjectionMode == EDisplayClusterMeshProjectionType::UV)
-		{
-			RenderSettings.ProjectionTypeSettings.UVProjectionIndex = 1;
-			RenderSettings.ProjectionTypeSettings.UVProjectionPlaneSize = ADisplayClusterLightCardActor::UVPlaneDefaultSize;
-			RenderSettings.ProjectionTypeSettings.UVProjectionPlaneDistance = ADisplayClusterLightCardActor::UVPlaneDefaultDistance;
-
-			// Compute the UV plane offset to allow panning. Need to convert to view space, since the UV projection assumes all coordinates are in view space
-			const FMatrix WorldToViewTransform = FMatrix(FVector::ZAxisVector, FVector::XAxisVector, FVector::YAxisVector, FVector::ZeroVector);
-			RenderSettings.ProjectionTypeSettings.UVProjectionPlaneOffset = -WorldToViewTransform.TransformVector(GetViewTransform().GetLocation());
-		}
+		ProjectionHelper->ConfigureRenderProjectionSettings(RenderSettings, GetViewTransform().GetLocation());
+		RenderSettings.PrimitiveFilter.ShouldRenderPrimitiveDelegate = ProjectionHelper->CreateDefaultShouldRenderPrimitiveFilter();
+		RenderSettings.PrimitiveFilter.ShouldApplyProjectionDelegate = ProjectionHelper->CreateDefaultShouldApplyProjectionToPrimitiveFilter();
 
 		GetSceneViewInitOptions(RenderSettings.ViewInitOptions);
 
@@ -2228,63 +2218,24 @@ void FDisplayClusterLightCardEditorViewportClient::MoveActorsToPixel(const FIntP
 
 void FDisplayClusterLightCardEditorViewportClient::MoveSelectedUVActors(FViewport* InViewport, EAxisList::Type CurrentAxis)
 {
-	if (LastSelectedActor.IsValid() && LastSelectedActor->IsUVActor())
-	{
-		const FVector2D DeltaUV = GetUVActorTranslationDelta(InViewport, LastSelectedActor, CurrentAxis);
-		for (const FDisplayClusterWeakStageActorPtr& Actor : SelectedActors)
-		{
-			if (!Actor.IsValid() || !Actor->IsUVActor())
-			{
-				continue;
-			}
-
-			Actor->SetUVCoordinates(Actor->GetUVCoordinates() + DeltaUV);
-
-			PropagateActorTransform(Actor);
-		}
-	}
-}
-
-FVector2D FDisplayClusterLightCardEditorViewportClient::GetUVActorTranslationDelta(FViewport* InViewport, const FDisplayClusterWeakStageActorPtr& InActor, EAxisList::Type CurrentAxis)
-{
-	FIntPoint MousePos;
-	InViewport->GetMousePos(MousePos);
-
 	FSceneViewFamilyContext ViewFamily(FSceneViewFamily::ConstructionValues(
 		InViewport,
 		GetScene(),
 		EngineShowFlags)
 		.SetRealtimeUpdate(IsRealtime()));
-	FSceneView* View = CalcSceneView(&ViewFamily);
 
-	FVector Origin;
-	FVector Direction;
-	ProjectionHelper->PixelToWorld(*View, MousePos, Origin, Direction);
-
-	const float UVProjectionPlaneSize = ADisplayClusterLightCardActor::UVPlaneDefaultSize;
-	const float UVProjectionPlaneDistance = ADisplayClusterLightCardActor::UVPlaneDefaultDistance;
-
-	const FVector ViewOrigin = View->ViewMatrices.GetViewOrigin();
-	const FPlane UVProjectionPlane(ViewOrigin + FVector::ForwardVector * UVProjectionPlaneDistance, -FVector::ForwardVector);
-	const FVector PlaneIntersection = FMath::RayPlaneIntersection(Origin, Direction, UVProjectionPlane);
-
-	const FVector DesiredLocation = PlaneIntersection - DragWidgetOffset;
-	const FVector2D DesiredUVLocation = FVector2D(DesiredLocation.Y / UVProjectionPlaneSize + 0.5f, 0.5f - DesiredLocation.Z / UVProjectionPlaneSize);
+	FSceneView* View = CalcSceneView(&ViewFamily);
 
-	const FVector2D UVDelta = DesiredUVLocation - InActor->GetUVCoordinates();
+	FIntPoint MousePos;
+	InViewport->GetMousePos(MousePos);
 
-	FVector2D UVAxis = FVector2D::ZeroVector;
-	if (CurrentAxis & EAxisList::Type::X)
-	{
-		UVAxis += FVector2D(1.0, 0.0);
-	}
+	ProjectionHelper->DragUVActors(SelectedActors, MousePos, *View, DragWidgetOffset, CurrentAxis, LastSelectedActor);
 
-	if (CurrentAxis & EAxisList::Type::Y)
+	// Update the level instances
+	for (const FDisplayClusterWeakStageActorPtr& Actor : SelectedActors)
 	{
-		UVAxis += FVector2D(0.0, 1.0);
+		PropagateActorTransform(Actor);
 	}
-
-	return UVDelta * UVAxis;
 }
 
 void FDisplayClusterLightCardEditorViewportClient::ScaleSelectedActors(FViewport* InViewport, EAxisList::Type CurrentAxis)
@@ -2749,35 +2700,6 @@ void FDisplayClusterLightCardEditorViewportClient::ExitDrawingLightCardMode()
 	}
 }
 
-bool FDisplayClusterLightCardEditorViewportClient::ShouldRenderPrimitive(const UPrimitiveComponent* PrimitiveComponent)
-{
-	const bool bIsUVProjection = ProjectionMode == EDisplayClusterMeshProjectionType::UV;
-	if (ADisplayClusterLightCardActor* LightCard = Cast<ADisplayClusterLightCardActor>(PrimitiveComponent->GetOwner()))
-	{
-		// Only render the UV light cards when in UV projection mode, and only render non-UV light cards in any other projection mode
-		return bIsUVProjection ? LightCard->bIsUVLightCard : !LightCard->bIsUVLightCard;
-	}
-	else
-	{
-		return true;
-	}
-}
-
-bool FDisplayClusterLightCardEditorViewportClient::ShouldApplyProjectionToPrimitive(const UPrimitiveComponent* PrimitiveComponent)
-{
-	const bool bIsUVProjection = ProjectionMode == EDisplayClusterMeshProjectionType::UV;
-	if (ADisplayClusterLightCardActor* LightCard = Cast<ADisplayClusterLightCardActor>(PrimitiveComponent->GetOwner()))
-	{
-		// When in UV projection mode, don't render the UV light cards using the UV projection, render them linearly
-		if (bIsUVProjection && LightCard->bIsUVLightCard)
-		{
-			return false;
-		}
-	}
-
-	return true;
-}
-
 FDisplayClusterLightCardEditorViewportClient::FSpriteProxy FDisplayClusterLightCardEditorViewportClient::FSpriteProxy::
 FromBillboard(const UBillboardComponent* InBillboardComponent)
 {
diff --git a/Engine/Plugins/Runtime/nDisplay/Source/DisplayClusterLightCardEditor/Private/Viewport/DisplayClusterLightCardEditorViewportClient.h b/Engine/Plugins/Runtime/nDisplay/Source/DisplayClusterLightCardEditor/Private/Viewport/DisplayClusterLightCardEditorViewportClient.h
--- a/Engine/Plugins/Runtime/nDisplay/Source/DisplayClusterLightCardEditor/Private/Viewport/DisplayClusterLightCardEditorViewportClient.h
+++ b/Engine/Plugins/Runtime/nDisplay/Source/DisplayClusterLightCardEditor/Private/Viewport/DisplayClusterLightCardEditorViewportClient.h
@@ -196,9 +196,6 @@ private:
 	/** Moves the currently selected UV actors */
 	void MoveSelectedUVActors(FViewport* InViewport, EAxisList::Type CurrentAxis);
 
-	/** Determines the appropriate delta in UV coordinates needed to move the specified actor to the mouse's location */
-	FVector2D GetUVActorTranslationDelta(FViewport* InViewport, const FDisplayClusterWeakStageActorPtr& InActor, EAxisList::Type CurrentAxis);
-
 	/** Scales the currently selected actors */
 	void ScaleSelectedActors(FViewport* InViewport, EAxisList::Type CurrentAxis);
 
@@ -247,12 +244,6 @@ private:
 	/** Calculates the final distance from the origin of an actor, given its flush distance and a desired offset */
 	double CalculateFinalActorDistance(double FlushDistance, double DesiredOffsetFromFlush = 0.) const;
 
-	/** Callback passed into the mesh projection renderer to filter which primitives are drawn */
-	bool ShouldRenderPrimitive(const UPrimitiveComponent* PrimitiveComponent);
-
-	/** Callback passed into the mesh projection renderer to filter which primitives drawn with the current projection and which are drawn linearly */
-	bool ShouldApplyProjectionToPrimitive(const UPrimitiveComponent* PrimitiveComponent);
-
 private:
 	TWeakPtr<FSceneViewport> SceneViewportPtr;
 	TWeakPtr<SDisplayClusterLightCardEditorViewport> LightCardEditorViewportPtr;
diff --git a/Engine/Plugins/Runtime/nDisplay/Source/DisplayClusterScenePreview/Private/DisplayClusterLightCardEditorHelper.cpp b/Engine/Plugins/Runtime/nDisplay/Source/DisplayClusterScenePreview/Private/DisplayClusterLightCardEditorHelper.cpp
--- a/Engine/Plugins/Runtime/nDisplay/Source/DisplayClusterScenePreview/Private/DisplayClusterLightCardEditorHelper.cpp
+++ b/Engine/Plugins/Runtime/nDisplay/Source/DisplayClusterScenePreview/Private/DisplayClusterLightCardEditorHelper.cpp
@@ -518,24 +518,50 @@ void FDisplayClusterLightCardEditorHelper::DragActors(
 	const FSceneView& SceneView, ECoordinateSystem CoordinateSystem, const FVector& DragWidgetOffset, EAxisList::Type DragAxis,
 	FDisplayClusterWeakStageActorPtr PrimaryActor)
 {
-	if (Actors.IsEmpty())
+	if (Actors.IsEmpty() || !UpdateNormalMaps() || !UpdateRootActor())
 	{
 		return;
 	}
 
-	const ADisplayClusterRootActor* RootActor = UpdateRootActor();
-	if (!RootActor)
+	InternalDragActors(Actors, PixelPos, SceneView, CoordinateSystem, DragWidgetOffset, DragAxis, PrimaryActor);
+}
+
+void FDisplayClusterLightCardEditorHelper::DragUVActors(
+	const TArray<FDisplayClusterWeakStageActorPtr>& Actors, const FIntPoint& PixelPos, const FSceneView& SceneView,
+	const FVector& DragWidgetOffset, EAxisList::Type DragAxis, FDisplayClusterWeakStageActorPtr PrimaryActor)
+{
+	if (Actors.IsEmpty() || !UpdateNormalMaps() || !UpdateRootActor())
 	{
 		return;
 	}
 
-	if (!UpdateNormalMaps())
+	if (!PrimaryActor.IsValid())
+	{
+		PrimaryActor = Actors.Last();
+	}
+
+	if (!PrimaryActor.IsValid() || !PrimaryActor->IsUVActor())
 	{
 		return;
 	}
-	
-	// Maps are already ready, so we can move immediately
-	InternalDragActors(Actors, PixelPos, SceneView, CoordinateSystem, DragWidgetOffset, DragAxis, PrimaryActor);
+
+	const FVector2D DeltaUV = GetUVActorTranslationDelta(PixelPos, SceneView, PrimaryActor, DragAxis, DragWidgetOffset);
+	for (const FDisplayClusterWeakStageActorPtr& Actor : Actors)
+	{
+		if (!Actor.IsValid() || !Actor->IsUVActor())
+		{
+			continue;
+		}
+
+		Actor->SetUVCoordinates(Actor->GetUVCoordinates() + DeltaUV);
+
+#if WITH_EDITOR
+		if (!Actor->IsProxy())
+		{
+			PostEditChangePropertiesForMovedActor(Actor);
+		}
+#endif
+	}
 }
 
 void FDisplayClusterLightCardEditorHelper::VerifyAndFixActorOrigin(const FDisplayClusterWeakStageActorPtr& Actor)
@@ -700,12 +726,17 @@ void FDisplayClusterLightCardEditorHelper::PixelToWorld(const FSceneView& View,
 }
 
 bool FDisplayClusterLightCardEditorHelper::WorldToPixel(const FSceneView& View, const FVector& WorldPos, FVector2D& OutPixelPos) const
+{
+	return WorldToPixel(View, WorldPos, OutPixelPos, ProjectionMode);
+}
+
+bool FDisplayClusterLightCardEditorHelper::WorldToPixel(const FSceneView& View, const FVector& WorldPos, FVector2D& OutPixelPos, EDisplayClusterMeshProjectionType OverrideProjectionMode) const
 {
 	const FMatrix& ViewMatrix = View.ViewMatrices.GetViewMatrix();
 	const FMatrix& ProjMatrix = View.ViewMatrices.GetProjectionMatrix();
 
 	const FVector ViewPos = ViewMatrix.TransformPosition(WorldPos);
-	const FVector ProjectedViewPos = FDisplayClusterMeshProjectionRenderer::ProjectViewPosition(ViewPos, ProjectionMode);
+	const FVector ProjectedViewPos = FDisplayClusterMeshProjectionRenderer::ProjectViewPosition(ViewPos, OverrideProjectionMode);
 	const FVector4 ScreenPos = ProjMatrix.TransformFVector4(FVector4(ProjectedViewPos, 1));
 
 	return View.ScreenToPixel(ScreenPos, OutPixelPos);
@@ -842,6 +873,21 @@ void FDisplayClusterLightCardEditorHelper::GetSceneViewInitOptions(
 	OutViewInitOptions.FOV = InFOV;
 }
 
+void FDisplayClusterLightCardEditorHelper::ConfigureRenderProjectionSettings(FDisplayClusterMeshProjectionRenderSettings& OutRenderSettings, const FVector ViewLocation) const
+{
+	OutRenderSettings.ProjectionType = ProjectionMode;
+	if (ProjectionMode == EDisplayClusterMeshProjectionType::UV)
+	{
+		OutRenderSettings.ProjectionTypeSettings.UVProjectionIndex = 1;
+		OutRenderSettings.ProjectionTypeSettings.UVProjectionPlaneSize = ADisplayClusterLightCardActor::UVPlaneDefaultSize;
+		OutRenderSettings.ProjectionTypeSettings.UVProjectionPlaneDistance = ADisplayClusterLightCardActor::UVPlaneDefaultDistance;
+
+		// Compute the UV plane offset to allow panning. Need to convert to view space, since the UV projection assumes all coordinates are in view space
+		const FMatrix WorldToViewTransform = FMatrix(FVector::ZAxisVector, FVector::XAxisVector, FVector::YAxisVector, FVector::ZeroVector);
+		OutRenderSettings.ProjectionTypeSettings.UVProjectionPlaneOffset = -WorldToViewTransform.TransformVector(ViewLocation);
+	}
+}
+
 FDisplayClusterLightCardEditorHelper::FSphericalCoordinates FDisplayClusterLightCardEditorHelper::GetActorCoordinates(const FDisplayClusterWeakStageActorPtr& Actor)
 {
 	const FVector ActorLocation = Actor->GetStageActorTransform(true).GetTranslation();
@@ -859,6 +905,43 @@ FDisplayClusterLightCardEditorHelper::FSphericalCoordinates FDisplayClusterLight
 	return ActorSphericalCoords;
 }
 
+FDisplayClusterMeshProjectionPrimitiveFilter::FPrimitiveFilter FDisplayClusterLightCardEditorHelper::CreateDefaultShouldRenderPrimitiveFilter() const
+{
+	const bool bIsUVProjection = ProjectionMode == EDisplayClusterMeshProjectionType::UV;
+
+	// Create a lambda function so that it's safe to access even if this helper is destroyed
+	return FDisplayClusterMeshProjectionPrimitiveFilter::FPrimitiveFilter::CreateLambda([bIsUVProjection](const UPrimitiveComponent* PrimitiveComponent)
+		{
+			if (ADisplayClusterLightCardActor* LightCard = Cast<ADisplayClusterLightCardActor>(PrimitiveComponent->GetOwner()))
+			{
+				// Only render the UV light cards when in UV projection mode, and only render non-UV light cards in any other projection mode
+				return bIsUVProjection ? LightCard->bIsUVLightCard : !LightCard->bIsUVLightCard;
+			}
+
+			return true;
+		});
+}
+
+FDisplayClusterMeshProjectionPrimitiveFilter::FPrimitiveFilter FDisplayClusterLightCardEditorHelper::CreateDefaultShouldApplyProjectionToPrimitiveFilter() const
+{
+	const bool bIsUVProjection = ProjectionMode == EDisplayClusterMeshProjectionType::UV;
+
+	// Create a lambda function so that it's safe to access even if this helper is destroyed
+	return FDisplayClusterMeshProjectionPrimitiveFilter::FPrimitiveFilter::CreateLambda([bIsUVProjection](const UPrimitiveComponent* PrimitiveComponent)
+		{
+			if (ADisplayClusterLightCardActor* LightCard = Cast<ADisplayClusterLightCardActor>(PrimitiveComponent->GetOwner()))
+			{
+				// When in UV projection mode, don't render the UV light cards using the UV projection, render them linearly
+				if (bIsUVProjection && LightCard->bIsUVLightCard)
+				{
+					return false;
+				}
+			}
+
+			return true;
+		});
+}
+
 ADisplayClusterRootActor* FDisplayClusterLightCardEditorHelper::UpdateRootActor()
 {
 	ADisplayClusterRootActor* NewRootActor = IDisplayClusterScenePreview::Get().GetRendererRootActor(RendererId);
@@ -1061,6 +1144,43 @@ FDisplayClusterLightCardEditorHelper::FSphericalCoordinates FDisplayClusterLight
 	return DeltaCoords;
 }
 
+FVector2D FDisplayClusterLightCardEditorHelper::GetUVActorTranslationDelta(
+	const FIntPoint& PixelPos,
+	const FSceneView& View,
+	const FDisplayClusterWeakStageActorPtr& Actor,
+	EAxisList::Type DragAxis,
+	const FVector& DragWidgetOffset)
+{
+	FVector Origin;
+	FVector Direction;
+	PixelToWorld(View, PixelPos, Origin, Direction);
+
+	const float UVProjectionPlaneSize = ADisplayClusterLightCardActor::UVPlaneDefaultSize;
+	const float UVProjectionPlaneDistance = ADisplayClusterLightCardActor::UVPlaneDefaultDistance;
+
+	const FVector ViewOrigin = View.ViewMatrices.GetViewOrigin();
+	const FPlane UVProjectionPlane(ViewOrigin + FVector::ForwardVector * UVProjectionPlaneDistance, -FVector::ForwardVector);
+	const FVector PlaneIntersection = FMath::RayPlaneIntersection(Origin, Direction, UVProjectionPlane);
+	
+	const FVector DesiredLocation = PlaneIntersection - DragWidgetOffset;
+	const FVector2D DesiredUVLocation = FVector2D(DesiredLocation.Y / UVProjectionPlaneSize + 0.5f, 0.5f - DesiredLocation.Z / UVProjectionPlaneSize);
+
+	const FVector2D UVDelta = DesiredUVLocation - Actor->GetUVCoordinates();
+
+	FVector2D UVAxis = FVector2D::ZeroVector;
+	if (DragAxis & EAxisList::Type::X)
+	{
+		UVAxis += FVector2D(1.0, 0.0);
+	}
+
+	if (DragAxis & EAxisList::Type::Y)
+	{
+		UVAxis += FVector2D(0.0, 1.0);
+	}
+
+	return UVDelta * UVAxis;
+}
+
 void FDisplayClusterLightCardEditorHelper::InternalMoveActorTo(
 	const FDisplayClusterWeakStageActorPtr& Actor, const FSphericalCoordinates& Position, bool bIsFinalChange) const
 {
diff --git a/Engine/Plugins/Runtime/nDisplay/Source/DisplayClusterScenePreview/Public/DisplayClusterLightCardEditorHelper.h b/Engine/Plugins/Runtime/nDisplay/Source/DisplayClusterScenePreview/Public/DisplayClusterLightCardEditorHelper.h
--- a/Engine/Plugins/Runtime/nDisplay/Source/DisplayClusterScenePreview/Public/DisplayClusterLightCardEditorHelper.h
+++ b/Engine/Plugins/Runtime/nDisplay/Source/DisplayClusterScenePreview/Public/DisplayClusterLightCardEditorHelper.h
@@ -206,6 +206,19 @@ public:
 	DISPLAYCLUSTERSCENEPREVIEW_API void DragActors(const TArray<FDisplayClusterWeakStageActorPtr>& Actors, const FIntPoint& PixelPos, const FSceneView& SceneView,
 		ECoordinateSystem CoordinateSystem, const FVector& DragWidgetOffset, EAxisList::Type DragAxis, FDisplayClusterWeakStageActorPtr PrimaryActor = nullptr);
 	
+	/**
+	 * Moves specified UV actors to a coordinate in viewport space as if dragged by a translate widget.
+	 *
+	 * @param Actors The actors that we are moving
+	 * @param PixelPos The screen pixel position of the widget
+	 * @param SceneView The scene view used to convert from pixel position to 3D position
+	 * @param DragWidgetOffset The offset between the actual cursor position and the position of the widget when the drag action started
+	 * @param DragAxis The axis along which the widget is being dragged
+	 * @param PrimaryActor The actor used to calculate the translation/rotation delta. If not provided, the last entry in Actors will be used.
+	 */
+	DISPLAYCLUSTERSCENEPREVIEW_API void DragUVActors(const TArray<FDisplayClusterWeakStageActorPtr>& Actors, const FIntPoint& PixelPos, const FSceneView& SceneView,
+		const FVector& DragWidgetOffset, EAxisList::Type DragAxis, FDisplayClusterWeakStageActorPtr PrimaryActor = nullptr);
+
 	/** Ensures that the actor root component is at the same location as the projection/view origin */
 	DISPLAYCLUSTERSCENEPREVIEW_API void VerifyAndFixActorOrigin(const FDisplayClusterWeakStageActorPtr& Actor);
 
@@ -238,9 +251,12 @@ public:
 	/** Converts a pixel coordinate into a point and direction vector in world space */
 	DISPLAYCLUSTERSCENEPREVIEW_API void PixelToWorld(const FSceneView& View, const FIntPoint& PixelPos, FVector& OutOrigin, FVector& OutDirection) const;
 
-	/** Converts a world coordinate into a point in screen space, and returns true if the world position is on the screen */
+	/** Converts a world coordinate into a point in screen space, and returns true if the world position is on the screen. */
 	DISPLAYCLUSTERSCENEPREVIEW_API bool WorldToPixel(const FSceneView& View, const FVector& WorldPos, FVector2D& OutPixelPos) const;
 
+	/** Converts a world coordinate into a point in screen space using a specific projection mode, and returns true if the world position is on the screen. */
+	DISPLAYCLUSTERSCENEPREVIEW_API bool WorldToPixel(const FSceneView& View, const FVector& WorldPos, FVector2D& OutPixelPos, EDisplayClusterMeshProjectionType OverrideProjectionMode) const;
+
 	/**
 	 * Sets up FSceneViewInitOptions' basic settings in a consistent way for preview renders.
 	 * If RotationMatrix is not provided, one will be automatically generated based on the view rotation.
@@ -256,9 +272,28 @@ public:
 		const FMatrix* InRotationMatrix = nullptr,
 		float InDPIScale = 1.f);
 
+	/**
+	 * Sets up projection-related settings of a preview render in a consistent way. 
+	 * This includes the ProjectionType and its related ProjectionTypeSettings.
+	 * If provided, the UV projection's panning settings will updated to reflect the ViewLocation passed in.
+	 */
+	DISPLAYCLUSTERSCENEPREVIEW_API void ConfigureRenderProjectionSettings(FDisplayClusterMeshProjectionRenderSettings& OutRenderSettings, FVector ViewLocation = FVector::ZeroVector) const;
+
 	/** Gets the spherical coordinates of the specified actor. */
 	DISPLAYCLUSTERSCENEPREVIEW_API FSphericalCoordinates GetActorCoordinates(const FDisplayClusterWeakStageActorPtr& Actor);
 
+	/**
+	 * Creates the default ShouldRenderPrimitive filter that most FDisplayClusterMeshProjectionRenderSettings will use.
+	 * Filters actors so that only UV lightcards are shown in UV mode, and only non-UV light cards are shown otherwise.
+	 */
+	DISPLAYCLUSTERSCENEPREVIEW_API FDisplayClusterMeshProjectionPrimitiveFilter::FPrimitiveFilter CreateDefaultShouldRenderPrimitiveFilter() const;
+
+	/**
+	 * Creates the default ShouldApplyProjectionToPrimitive filter that most FDisplayClusterMeshProjectionRenderSettings will use.
+	 * Filters actors so that in UV projection mode, UV light cards are rendered linearly.
+	 */
+	DISPLAYCLUSTERSCENEPREVIEW_API FDisplayClusterMeshProjectionPrimitiveFilter::FPrimitiveFilter CreateDefaultShouldApplyProjectionToPrimitiveFilter() const;
+
 private:
 	/**
 	 * Update the cached pointer to the root actor from the associated preview renderer and, if the actor has changed, update related data.
@@ -283,6 +318,10 @@ private:
 	FSphericalCoordinates GetActorTranslationDelta(const FIntPoint& PixelPos, const FSceneView& View, const FDisplayClusterWeakStageActorPtr& Actor,
 		ECoordinateSystem CoordinateSystem, EAxisList::Type DragAxis, const FVector& DragWidgetOffset);
 
+	/** Determines the appropriate delta in UV coordinates needed to move the specified UV actor to the given position within the view. */
+	FVector2D GetUVActorTranslationDelta(const FIntPoint& PixelPos, const FSceneView& View, const FDisplayClusterWeakStageActorPtr& Actor,
+		EAxisList::Type DragAxis, const FVector& DragWidgetOffset);
+
 	/** Moves specified actor to desired coordinates immediately using the current normal maps. Requires valid normal maps. */
 	void InternalMoveActorTo(const FDisplayClusterWeakStageActorPtr& Actor, const FSphericalCoordinates& Position, bool bIsFinalChange) const;
 
diff --git a/Engine/Plugins/VirtualProduction/EpicStageApp/Source/EpicStageApp/Private/StageAppRequest.h b/Engine/Plugins/VirtualProduction/EpicStageApp/Source/EpicStageApp/Private/StageAppRequest.h
--- a/Engine/Plugins/VirtualProduction/EpicStageApp/Source/EpicStageApp/Private/StageAppRequest.h
+++ b/Engine/Plugins/VirtualProduction/EpicStageApp/Source/EpicStageApp/Private/StageAppRequest.h
@@ -21,7 +21,8 @@ enum class ERCWebSocketNDisplayPreviewRenderProjectionType : uint8
 {
 	Perspective,
 	Azimuthal,
-	Orthographic
+	Orthographic,
+	UV
 };
 
 /** Preview renderer settings exposed to WebSocket clients. */
diff --git a/Engine/Plugins/VirtualProduction/EpicStageApp/Source/EpicStageApp/Private/StageAppRouteHandler.cpp b/Engine/Plugins/VirtualProduction/EpicStageApp/Source/EpicStageApp/Private/StageAppRouteHandler.cpp
--- a/Engine/Plugins/VirtualProduction/EpicStageApp/Source/EpicStageApp/Private/StageAppRouteHandler.cpp
+++ b/Engine/Plugins/VirtualProduction/EpicStageApp/Source/EpicStageApp/Private/StageAppRouteHandler.cpp
@@ -63,6 +63,11 @@ namespace StageAppRouteHandlerUtils
 			bOutIsOrthographic = true;
 			break;
 
+		case ERCWebSocketNDisplayPreviewRenderProjectionType::UV:
+			OutProjectionType = EDisplayClusterMeshProjectionType::UV;
+			bOutIsOrthographic = true;
+			break;
+
 		default:
 			checkf(false, TEXT("Unknown projection type %d"), InProjectionType);
 		}
@@ -133,12 +138,17 @@ bool FStageAppRouteHandler::FPerRendererData::GetSceneViewInitOptions(FSceneView
 		return false;
 	}
 
+	// In UV mode, lock the location and rotation
+	const bool bIsUVMode = GetLightCardHelper().GetProjectionMode() == EDisplayClusterMeshProjectionType::UV;
+	const FVector Location = bIsUVMode ? FVector::ZeroVector : ViewOriginComponent->GetComponentLocation();
+	const FRotator Rotation = bIsUVMode ? FVector::ForwardVector.Rotation() : FRotator(RootActor->GetActorRotation().Quaternion() * PreviewSettings.Rotation.Quaternion());
+
 	GetLightCardHelper().GetSceneViewInitOptions(
 		ViewInitOptions,
 		PreviewSettings.FOV,
 		PreviewSettings.Resolution,
-		ViewOriginComponent->GetComponentLocation(),
-		FRotator(RootActor->GetActorRotation().Quaternion() * PreviewSettings.Rotation.Quaternion())
+		Location,
+		Rotation
 	);
 
 	return true;
@@ -417,13 +427,18 @@ void FStageAppRouteHandler::HandleWebSocketNDisplayPreviewRender(const FRemoteCo
 
 	RenderSettings.RenderType = StageAppRouteHandlerUtils::GetInternalPreviewRenderType(PreviewSettings.RenderType);
 	RenderSettings.EngineShowFlags.SetSelectionOutline(false);
-
+	
 	EDisplayClusterMeshProjectionType ProjectionType;
 	bool bIsOrthographic;
 
 	StageAppRouteHandlerUtils::GetProjectionSettingsFromWebSocketType(PreviewSettings.ProjectionType, ProjectionType, bIsOrthographic);
 	EngineShowFlagOrthographicOverride(!bIsOrthographic, RenderSettings.EngineShowFlags);
 
+	const FDisplayClusterLightCardEditorHelper& Helper = PerRendererData->GetLightCardHelper();
+	Helper.ConfigureRenderProjectionSettings(RenderSettings);
+	RenderSettings.PrimitiveFilter.ShouldRenderPrimitiveDelegate = Helper.CreateDefaultShouldRenderPrimitiveFilter();
+	RenderSettings.PrimitiveFilter.ShouldApplyProjectionDelegate = Helper.CreateDefaultShouldApplyProjectionToPrimitiveFilter();
+
 	FSceneViewInitOptions ViewInitOptions = RenderSettings.ViewInitOptions;
 	const FGuid ClientId = WebSocketMessage.ClientId;
 	const int32 RendererId = Body.RendererId;
@@ -434,7 +449,7 @@ void FStageAppRouteHandler::HandleWebSocketNDisplayPreviewRender(const FRemoteCo
 
 	// Render the image
 	IDisplayClusterScenePreview::Get().RenderQueued(Body.RendererId, RenderSettings, Resolution, FRenderResultDelegate::CreateLambda(
-		[this, ViewInitOptions, ClientId, RendererId, Resolution, JpegQuality, bIncludeActorPositions](FRenderTarget* RenderTarget)
+		[this, ViewInitOptions, ClientId, RendererId, ProjectionType, Resolution, JpegQuality, bIncludeActorPositions](FRenderTarget* RenderTarget)
 		{
 			FPerRendererData* PerRendererData = GetClientPerRendererData(ClientId, RendererId);
 			if (PerRendererData)
@@ -499,6 +514,7 @@ void FStageAppRouteHandler::HandleWebSocketNDisplayPreviewRender(const FRemoteCo
 					if (IDisplayClusterScenePreview::Get().GetActorsInRendererScene(RendererId, false, Actors))
 					{
 						const FSceneView View(ViewInitOptions);
+						const bool bIsInUVMode = ProjectionType == EDisplayClusterMeshProjectionType::UV;
 
 						for (const AActor* Actor : Actors)
 						{
@@ -510,9 +526,9 @@ void FStageAppRouteHandler::HandleWebSocketNDisplayPreviewRender(const FRemoteCo
 							FVector WorldPosition;
 							if (const ADisplayClusterLightCardActor* LightCard = Cast<ADisplayClusterLightCardActor>(Actor))
 							{
-								if (LightCard->bIsUVLightCard)
+								// Only send UV lightcards in UV mode, and only send non-UV lightcards in non-UV modes
+								if (LightCard->bIsUVLightCard != bIsInUVMode)
 								{
-									// TODO: add support for UV light cards and projection mode
 									continue;
 								}
 
@@ -529,7 +545,8 @@ void FStageAppRouteHandler::HandleWebSocketNDisplayPreviewRender(const FRemoteCo
 							PerRendererData->GetLightCardHelper().WorldToPixel(
 								View,
 								WorldPosition,
-								ActorInfo.Position
+								ActorInfo.Position,
+								ProjectionType // Pass this since the helper's projection mode made have changed since we kicked off the render
 							);
 
 							// Normalize to the size of the texture
@@ -745,17 +762,32 @@ void FStageAppRouteHandler::DragLightCards(FPerRendererData& PerRendererData, FV
 		{
 			return FDisplayClusterWeakStageActorPtr(WeakPtr.Get());
 		});
-	
+
 	FDisplayClusterLightCardEditorHelper& LightCardHelper = PerRendererData.GetLightCardHelper();
-	LightCardHelper.DragActors(
-		MoveTemp(StageActors),
-		PixelPos,
-		View,
-		FDisplayClusterLightCardEditorHelper::ECoordinateSystem::Spherical,
-		DragWidgetOffset,
-		EAxisList::Type::XYZ,
-		PerRendererData.PrimaryLightCard.Get()
-	);
+
+	if (LightCardHelper.GetProjectionMode() == EDisplayClusterMeshProjectionType::UV)
+	{
+		LightCardHelper.DragUVActors(
+			MoveTemp(StageActors),
+			PixelPos,
+			View,
+			DragWidgetOffset,
+			EAxisList::Type::XYZ,
+			PerRendererData.PrimaryLightCard.Get()
+		);
+	}
+	else
+	{
+		LightCardHelper.DragActors(
+			MoveTemp(StageActors),
+			PixelPos,
+			View,
+			FDisplayClusterLightCardEditorHelper::ECoordinateSystem::Spherical,
+			DragWidgetOffset,
+			EAxisList::Type::XYZ,
+			PerRendererData.PrimaryLightCard.Get()
+		);
+	}
 }
 
 bool FStageAppRouteHandler::TimeOutDrags(float DeltaTime)
