nDisplay: Fixed issue where color correction windows were showing up in the ICVFX color grading drawer in the same area as color correction regions were; now, a CCW will only show up in the drawer when it is selected in the light card editor outliner. Also, added scrollbox to drawer list, as well as set the list items to be ordered alphabetically by display name.

diff --git a/Engine/Plugins/Runtime/nDisplay/Source/DisplayClusterColorGrading/Private/DataModelGenerators/DisplayClusterColorGradingGenerator_ColorCorrectRegion.cpp b/Engine/Plugins/Runtime/nDisplay/Source/DisplayClusterColorGrading/Private/DataModelGenerators/DisplayClusterColorGradingGenerator_ColorCorrectRegion.cpp
--- a/Engine/Plugins/Runtime/nDisplay/Source/DisplayClusterColorGrading/Private/DataModelGenerators/DisplayClusterColorGradingGenerator_ColorCorrectRegion.cpp
+++ b/Engine/Plugins/Runtime/nDisplay/Source/DisplayClusterColorGrading/Private/DataModelGenerators/DisplayClusterColorGradingGenerator_ColorCorrectRegion.cpp
@@ -2,12 +2,16 @@
 
 #include "DisplayClusterColorGradingGenerator_ColorCorrectRegion.h"
 
+#include "ClassIconFinder.h"
 #include "ColorCorrectRegion.h"
 #include "IDetailTreeNode.h"
 #include "IPropertyRowGenerator.h"
 #include "PropertyHandle.h"
+#include "Widgets/Images/SImage.h"
 #include "Widgets/Input/SButton.h"
+#include "Widgets/Layout/SBox.h"
 #include "Widgets/Text/SInlineEditableTextBlock.h"
+#include "Widgets/Text/STextBlock.h"
 
 #define LOCTEXT_NAMESPACE "DisplayClusterColorGrading"
 
@@ -18,6 +22,22 @@ TSharedRef<IDisplayClusterColorGradingDataModelGenerator> FDisplayClusterColorGr
 
 void FDisplayClusterColorGradingGenerator_ColorCorrectRegion::GenerateDataModel(IPropertyRowGenerator& PropertyRowGenerator, FDisplayClusterColorGradingDataModel& OutColorGradingDataModel)
 {
+	TArray<TWeakObjectPtr<AColorCorrectRegion>> SelectedCCRs;
+	const TArray<TWeakObjectPtr<UObject>>& SelectedObjects = PropertyRowGenerator.GetSelectedObjects();
+	for (const TWeakObjectPtr<UObject>& SelectedObject : SelectedObjects)
+	{
+		if (SelectedObject.IsValid() && SelectedObject->IsA<AColorCorrectRegion>())
+		{
+			TWeakObjectPtr<AColorCorrectRegion> SelectedRootActor = CastChecked<AColorCorrectRegion>(SelectedObject.Get());
+			SelectedCCRs.Add(SelectedRootActor);
+		}
+	}
+
+	if (!SelectedCCRs.Num())
+	{
+		return;
+	}
+
 	const TArray<TSharedRef<IDetailTreeNode>>& RootNodes = PropertyRowGenerator.GetRootTreeNodes();
 
 	const TSharedRef<IDetailTreeNode>* ColorCorrectionNodePtr = RootNodes.FindByPredicate([](const TSharedRef<IDetailTreeNode>& Node)
@@ -64,6 +84,32 @@ void FDisplayClusterColorGradingGenerator_ColorCorrectRegion::GenerateDataModel(
 			}
 		}
 
+		ColorGradingGroup.GroupHeaderWidget = SNew(SHorizontalBox)
+			+ SHorizontalBox::Slot()
+			.AutoWidth()
+			.Padding(FMargin(0, 1, 6, 1))
+			.VAlign(VAlign_Center)
+			[
+				SNew(SBox)
+				.WidthOverride(16)
+				.HeightOverride(16)
+				[
+					SNew(SImage)
+					.ColorAndOpacity(FSlateColor::UseForeground())
+					.Image(FClassIconFinder::FindIconForActor(SelectedCCRs[0]))
+				]
+			]
+
+			+ SHorizontalBox::Slot()
+			.AutoWidth()
+			.VAlign(VAlign_Center)
+			.HAlign(HAlign_Left)
+			[
+				SNew(STextBlock)
+				.Text(FText::FromString(SelectedCCRs[0]->GetActorLabel()))
+				.Font(FAppStyle::Get().GetFontStyle("NormalFontBold"))
+			];
+
 		OutColorGradingDataModel.ColorGradingGroups.Add(ColorGradingGroup);
 	}
 }
diff --git a/Engine/Plugins/Runtime/nDisplay/Source/DisplayClusterColorGrading/Private/DataModelGenerators/DisplayClusterColorGradingGenerator_PostProcessVolume.cpp b/Engine/Plugins/Runtime/nDisplay/Source/DisplayClusterColorGrading/Private/DataModelGenerators/DisplayClusterColorGradingGenerator_PostProcessVolume.cpp
--- a/Engine/Plugins/Runtime/nDisplay/Source/DisplayClusterColorGrading/Private/DataModelGenerators/DisplayClusterColorGradingGenerator_PostProcessVolume.cpp
+++ b/Engine/Plugins/Runtime/nDisplay/Source/DisplayClusterColorGrading/Private/DataModelGenerators/DisplayClusterColorGradingGenerator_PostProcessVolume.cpp
@@ -2,12 +2,16 @@
 
 #include "DisplayClusterColorGradingGenerator_PostProcessVolume.h"
 
+#include "ClassIconFinder.h"
 #include "Engine/PostProcessVolume.h"
 #include "IDetailTreeNode.h"
 #include "IPropertyRowGenerator.h"
 #include "PropertyHandle.h"
+#include "Widgets/Images/SImage.h"
 #include "Widgets/Input/SButton.h"
+#include "Widgets/Layout/SBox.h"
 #include "Widgets/Text/SInlineEditableTextBlock.h"
+#include "Widgets/Text/STextBlock.h"
 
 #define LOCTEXT_NAMESPACE "DisplayClusterColorGrading"
 
@@ -18,6 +22,22 @@ TSharedRef<IDisplayClusterColorGradingDataModelGenerator> FDisplayClusterColorGr
 
 void FDisplayClusterColorGradingGenerator_PostProcessVolume::GenerateDataModel(IPropertyRowGenerator& PropertyRowGenerator, FDisplayClusterColorGradingDataModel& OutColorGradingDataModel)
 {
+	TArray<TWeakObjectPtr<APostProcessVolume>> SelectedPPVs;
+	const TArray<TWeakObjectPtr<UObject>>& SelectedObjects = PropertyRowGenerator.GetSelectedObjects();
+	for (const TWeakObjectPtr<UObject>& SelectedObject : SelectedObjects)
+	{
+		if (SelectedObject.IsValid() && SelectedObject->IsA<APostProcessVolume>())
+		{
+			TWeakObjectPtr<APostProcessVolume> SelectedRootActor = CastChecked<APostProcessVolume>(SelectedObject.Get());
+			SelectedPPVs.Add(SelectedRootActor);
+		}
+	}
+
+	if (!SelectedPPVs.Num())
+	{
+		return;
+	}
+
 	const TArray<TSharedRef<IDetailTreeNode>>& RootNodes = PropertyRowGenerator.GetRootTreeNodes();
 
 	const TSharedRef<IDetailTreeNode>* ColorGradingNodePtr = RootNodes.FindByPredicate([](const TSharedRef<IDetailTreeNode>& Node)
@@ -48,6 +68,32 @@ void FDisplayClusterColorGradingGenerator_PostProcessVolume::GenerateDataModel(I
 			AddPropertiesToDetailsView(PropertyGroupNode, ColorGradingGroup);
 		}
 
+		ColorGradingGroup.GroupHeaderWidget = SNew(SHorizontalBox)
+			+ SHorizontalBox::Slot()
+			.AutoWidth()
+			.Padding(FMargin(0, 1, 6, 1))
+			.VAlign(VAlign_Center)
+			[
+				SNew(SBox)
+				.WidthOverride(16)
+				.HeightOverride(16)
+				[
+					SNew(SImage)
+					.ColorAndOpacity(FSlateColor::UseForeground())
+					.Image(FClassIconFinder::FindIconForActor(SelectedPPVs[0]))
+				]
+			]
+
+			+ SHorizontalBox::Slot()
+			.AutoWidth()
+			.VAlign(VAlign_Center)
+			.HAlign(HAlign_Left)
+			[
+				SNew(STextBlock)
+				.Text(FText::FromString(SelectedPPVs[0]->GetActorLabel()))
+				.Font(FAppStyle::Get().GetFontStyle("NormalFontBold"))
+			];
+
 		OutColorGradingDataModel.ColorGradingGroups.Add(ColorGradingGroup);
 	}
 }
diff --git a/Engine/Plugins/Runtime/nDisplay/Source/DisplayClusterColorGrading/Private/DisplayClusterColorGradingDataModel.cpp b/Engine/Plugins/Runtime/nDisplay/Source/DisplayClusterColorGrading/Private/DisplayClusterColorGradingDataModel.cpp
--- a/Engine/Plugins/Runtime/nDisplay/Source/DisplayClusterColorGrading/Private/DisplayClusterColorGradingDataModel.cpp
+++ b/Engine/Plugins/Runtime/nDisplay/Source/DisplayClusterColorGrading/Private/DisplayClusterColorGradingDataModel.cpp
@@ -36,15 +36,41 @@ TArray<TWeakObjectPtr<UObject>> FDisplayClusterColorGradingDataModel::GetObjects
 
 void FDisplayClusterColorGradingDataModel::SetObjects(const TArray<UObject*>& InObjects)
 {
-	Reset();
-
+	// Only update the data model if the objects being set are new
+	bool bUpdateDataModel = false;
 	if (PropertyRowGenerator.IsValid())
 	{
-		PropertyRowGenerator->SetObjects(InObjects);
+		const TArray<TWeakObjectPtr<UObject>>& CurrentObjects = PropertyRowGenerator->GetSelectedObjects();
+
+		if (CurrentObjects.Num() != InObjects.Num())
+		{
+			bUpdateDataModel = true;
+		}
+		else
+		{
+			for (UObject* NewObject : InObjects)
+			{
+				if (!CurrentObjects.Contains(NewObject))
+				{
+					bUpdateDataModel = true;
+					break;
+				}
+			}
+		}
 	}
 
-	SelectedColorGradingGroupIndex = ColorGradingGroups.Num() ? 0 : INDEX_NONE;
-	SelectedColorGradingElementIndex = 0;
+	if (bUpdateDataModel)
+	{
+		Reset();
+
+		if (PropertyRowGenerator.IsValid())
+		{
+			PropertyRowGenerator->SetObjects(InObjects);
+		}
+
+		SelectedColorGradingGroupIndex = ColorGradingGroups.Num() ? 0 : INDEX_NONE;
+		SelectedColorGradingElementIndex = 0;
+	}
 }
 
 bool FDisplayClusterColorGradingDataModel::HasObjectOfType(const UClass* InClass) const
diff --git a/Engine/Plugins/Runtime/nDisplay/Source/DisplayClusterColorGrading/Private/Drawer/DisplayClusterColorGradingDrawerSingleton.cpp b/Engine/Plugins/Runtime/nDisplay/Source/DisplayClusterColorGrading/Private/Drawer/DisplayClusterColorGradingDrawerSingleton.cpp
--- a/Engine/Plugins/Runtime/nDisplay/Source/DisplayClusterColorGrading/Private/Drawer/DisplayClusterColorGradingDrawerSingleton.cpp
+++ b/Engine/Plugins/Runtime/nDisplay/Source/DisplayClusterColorGrading/Private/Drawer/DisplayClusterColorGradingDrawerSingleton.cpp
@@ -40,6 +40,7 @@ FDisplayClusterColorGradingDrawerSingleton::FDisplayClusterColorGradingDrawerSin
 	IDisplayClusterOperator::Get().OnAppendOperatorPanelCommands().AddRaw(this, &FDisplayClusterColorGradingDrawerSingleton::AppendOperatorPanelCommands);
 
 	IDisplayClusterOperator::Get().GetOperatorViewModel()->OnActiveRootActorChanged().AddRaw(this, &FDisplayClusterColorGradingDrawerSingleton::OnActiveRootActorChanged);
+	IDisplayClusterOperator::Get().GetOperatorViewModel()->OnDetailObjectsChanged().AddRaw(this, &FDisplayClusterColorGradingDrawerSingleton::OnDetailObjectsChanged);
 }
 
 FDisplayClusterColorGradingDrawerSingleton::~FDisplayClusterColorGradingDrawerSingleton()
@@ -48,6 +49,7 @@ FDisplayClusterColorGradingDrawerSingleton::~FDisplayClusterColorGradingDrawerSi
 	IDisplayClusterOperator::Get().OnRegisterStatusBarExtensions().RemoveAll(this);
 	IDisplayClusterOperator::Get().OnAppendOperatorPanelCommands().RemoveAll(this);
 	IDisplayClusterOperator::Get().GetOperatorViewModel()->OnActiveRootActorChanged().RemoveAll(this);
+	IDisplayClusterOperator::Get().GetOperatorViewModel()->OnDetailObjectsChanged().RemoveAll(this);
 
 	if (FSlateApplication::IsInitialized())
 	{
@@ -184,4 +186,10 @@ void FDisplayClusterColorGradingDrawerSingleton::OnActiveRootActorChanged(ADispl
 	PreviousDrawerState.Reset();
 }
 
+void FDisplayClusterColorGradingDrawerSingleton::OnDetailObjectsChanged(const TArray<UObject*>& NewObjects)
+{
+	// Clear the previous drawer state when the selected detail objects have changed
+	PreviousDrawerState.Reset();
+}
+
 #undef LOCTEXT_NAMESPACE
\ No newline at end of file
diff --git a/Engine/Plugins/Runtime/nDisplay/Source/DisplayClusterColorGrading/Private/Drawer/DisplayClusterColorGradingDrawerSingleton.h b/Engine/Plugins/Runtime/nDisplay/Source/DisplayClusterColorGrading/Private/Drawer/DisplayClusterColorGradingDrawerSingleton.h
--- a/Engine/Plugins/Runtime/nDisplay/Source/DisplayClusterColorGrading/Private/Drawer/DisplayClusterColorGradingDrawerSingleton.h
+++ b/Engine/Plugins/Runtime/nDisplay/Source/DisplayClusterColorGrading/Private/Drawer/DisplayClusterColorGradingDrawerSingleton.h
@@ -58,6 +58,9 @@ private:
 	/** Delegate callback that is raised when the active root actor of the operator panel has changed */
 	void OnActiveRootActorChanged(ADisplayClusterRootActor* NewRootActor);
 
+	/** Delegate callback that is raised when the list of objects displayed in the operator panel's details panel has changed */
+	void OnDetailObjectsChanged(const TArray<UObject*>& NewObjects);
+
 private:
 	/** A weak pointer to the active color grading drawer that is open */
 	TWeakPtr<SDisplayClusterColorGradingDrawer> ColorGradingDrawer;
diff --git a/Engine/Plugins/Runtime/nDisplay/Source/DisplayClusterColorGrading/Private/Drawer/SDisplayClusterColorGradingDrawer.cpp b/Engine/Plugins/Runtime/nDisplay/Source/DisplayClusterColorGrading/Private/Drawer/SDisplayClusterColorGradingDrawer.cpp
--- a/Engine/Plugins/Runtime/nDisplay/Source/DisplayClusterColorGrading/Private/Drawer/SDisplayClusterColorGradingDrawer.cpp
+++ b/Engine/Plugins/Runtime/nDisplay/Source/DisplayClusterColorGrading/Private/Drawer/SDisplayClusterColorGradingDrawer.cpp
@@ -18,6 +18,7 @@
 #include "DisplayClusterConfigurationTypes.h"
 
 #include "ColorCorrectRegion.h"
+#include "ColorCorrectWindow.h"
 #include "Engine/PostProcessVolume.h"
 #include "Modules/ModuleManager.h"
 #include "ScopedTransaction.h"
@@ -27,6 +28,7 @@
 #include "Widgets/Layout/SBorder.h"
 #include "Widgets/Layout/SBox.h"
 #include "Widgets/Layout/SExpandableArea.h"
+#include "Widgets/Layout/SScrollBox.h"
 #include "Widgets/Layout/SSeparator.h"
 #include "Widgets/Layout/SSpacer.h"
 #include "Widgets/Layout/SSplitter.h"
@@ -38,6 +40,7 @@ SDisplayClusterColorGradingDrawer::~SDisplayClusterColorGradingDrawer()
 {
 	OperatorViewModel = IDisplayClusterOperator::Get().GetOperatorViewModel();
 	OperatorViewModel->OnActiveRootActorChanged().RemoveAll(this);
+	OperatorViewModel->OnDetailObjectsChanged().RemoveAll(this);
 
 	FCoreUObjectDelegates::OnObjectsReplaced.RemoveAll(this);
 	GEngine->OnLevelActorAdded().RemoveAll(this);
@@ -80,6 +83,7 @@ void SDisplayClusterColorGradingDrawer::Construct(const FArguments& InArgs, bool
 	bIsInDrawer = bInIsInDrawer;
 	OperatorViewModel = IDisplayClusterOperator::Get().GetOperatorViewModel();
 	OperatorViewModel->OnActiveRootActorChanged().AddSP(this, &SDisplayClusterColorGradingDrawer::OnActiveRootActorChanged);
+	OperatorViewModel->OnDetailObjectsChanged().AddSP(this, &SDisplayClusterColorGradingDrawer::OnDetailObjectsChanged);
 
 	FCoreUObjectDelegates::OnObjectsReplaced.AddSP(this, &SDisplayClusterColorGradingDrawer::OnObjectsReplaced);
 	GEngine->OnLevelActorAdded().AddSP(this, &SDisplayClusterColorGradingDrawer::OnLevelActorAdded);
@@ -177,55 +181,59 @@ void SDisplayClusterColorGradingDrawer::Construct(const FArguments& InArgs, bool
 						.Padding(FMargin(0.0f))
 						.BorderImage(FAppStyle::GetBrush("Brushes.Recessed"))
 						[
-							SNew(SVerticalBox)
-
-							+SVerticalBox::Slot()
-							.AutoHeight()
+							SNew(SScrollBox)
+							+ SScrollBox::Slot()
 							[
-								SNew(SExpandableArea)
-								.BorderImage(FAppStyle::Get().GetBrush("Brushes.Header"))
-								.BodyBorderImage(FAppStyle::Get().GetBrush("Brushes.Recessed"))
-								.HeaderPadding(FMargin(4.0f, 2.0f))
-								.InitiallyCollapsed(false)
-								.AllowAnimatedTransition(false)
-								.Visibility_Lambda([this]() { return LevelColorGradingItems.Num() ? EVisibility::Visible : EVisibility::Collapsed; })
-								.HeaderContent()
-								[
-									SNew(STextBlock)
-									.Text(this, &SDisplayClusterColorGradingDrawer::GetCurrentLevelName)
-									.TextStyle(FAppStyle::Get(), "ButtonText")
-									.Font(FAppStyle::Get().GetFontStyle("NormalFontBold"))
-								]
-								.BodyContent()
-								[
-									SAssignNew(LevelActorsList, SDisplayClusterColorGradingObjectList)
-									.ColorGradingItemsSource(&LevelColorGradingItems)
-									.OnSelectionChanged(this, &SDisplayClusterColorGradingDrawer::OnListSelectionChanged)
-								]
-							]
+								SNew(SVerticalBox)
 
-							+SVerticalBox::Slot()
-							.AutoHeight()
-							[
-								SNew(SExpandableArea)
-								.BorderImage(FAppStyle::Get().GetBrush("Brushes.Header"))
-								.BodyBorderImage(FAppStyle::Get().GetBrush("Brushes.Recessed"))
-								.HeaderPadding(FMargin(4.0f, 2.0f))
-								.InitiallyCollapsed(false)
-								.AllowAnimatedTransition(false)
-								.Visibility_Lambda([this]() { return RootActorColorGradingItems.Num() ? EVisibility::Visible : EVisibility::Collapsed; })
-								.HeaderContent()
+								+SVerticalBox::Slot()
+								.AutoHeight()
 								[
-									SNew(STextBlock)
-									.Text(this, &SDisplayClusterColorGradingDrawer::GetCurrentRootActorName)
-									.TextStyle(FAppStyle::Get(), "ButtonText")
-									.Font(FAppStyle::Get().GetFontStyle("NormalFontBold"))
+									SNew(SExpandableArea)
+									.BorderImage(FAppStyle::Get().GetBrush("Brushes.Header"))
+									.BodyBorderImage(FAppStyle::Get().GetBrush("Brushes.Recessed"))
+									.HeaderPadding(FMargin(4.0f, 2.0f))
+									.InitiallyCollapsed(false)
+									.AllowAnimatedTransition(false)
+									.Visibility_Lambda([this]() { return LevelColorGradingItems.Num() ? EVisibility::Visible : EVisibility::Collapsed; })
+									.HeaderContent()
+									[
+										SNew(STextBlock)
+										.Text(this, &SDisplayClusterColorGradingDrawer::GetCurrentLevelName)
+										.TextStyle(FAppStyle::Get(), "ButtonText")
+										.Font(FAppStyle::Get().GetFontStyle("NormalFontBold"))
+									]
+									.BodyContent()
+									[
+										SAssignNew(LevelActorsList, SDisplayClusterColorGradingObjectList)
+										.ColorGradingItemsSource(&LevelColorGradingItems)
+										.OnSelectionChanged(this, &SDisplayClusterColorGradingDrawer::OnListSelectionChanged)
+									]
 								]
-								.BodyContent()
+
+								+SVerticalBox::Slot()
+								.AutoHeight()
 								[
-									SAssignNew(RootActorList, SDisplayClusterColorGradingObjectList)
-									.ColorGradingItemsSource(&RootActorColorGradingItems)
-									.OnSelectionChanged(this, &SDisplayClusterColorGradingDrawer::OnListSelectionChanged)
+									SNew(SExpandableArea)
+									.BorderImage(FAppStyle::Get().GetBrush("Brushes.Header"))
+									.BodyBorderImage(FAppStyle::Get().GetBrush("Brushes.Recessed"))
+									.HeaderPadding(FMargin(4.0f, 2.0f))
+									.InitiallyCollapsed(false)
+									.AllowAnimatedTransition(false)
+									.Visibility_Lambda([this]() { return RootActorColorGradingItems.Num() ? EVisibility::Visible : EVisibility::Collapsed; })
+									.HeaderContent()
+									[
+										SNew(STextBlock)
+										.Text(this, &SDisplayClusterColorGradingDrawer::GetCurrentRootActorName)
+										.TextStyle(FAppStyle::Get(), "ButtonText")
+										.Font(FAppStyle::Get().GetFontStyle("NormalFontBold"))
+									]
+									.BodyContent()
+									[
+										SAssignNew(RootActorList, SDisplayClusterColorGradingObjectList)
+										.ColorGradingItemsSource(&RootActorColorGradingItems)
+										.OnSelectionChanged(this, &SDisplayClusterColorGradingDrawer::OnListSelectionChanged)
+									]
 								]
 							]
 						]
@@ -538,10 +546,38 @@ void SDisplayClusterColorGradingDrawer::SetDrawerState(const FDisplayClusterColo
 
 void SDisplayClusterColorGradingDrawer::SetDrawerStateToDefault()
 {
-	// The nDisplay stage actor is always the first item in the root actor color grading items list, so set that as the currently selected item
 	if (RootActorList.IsValid() && RootActorColorGradingItems.Num() > 0)
 	{
-		RootActorList->SetSelectedItems({ RootActorColorGradingItems[0] });
+		auto FindCCW = [](const TWeakObjectPtr<UObject>& Object)
+		{
+			return Object.IsValid() && Object->IsA<AColorCorrectionWindow>();
+		};
+
+		const bool bIsCCWSelected = OperatorViewModel->GetDetailObjects().ContainsByPredicate(FindCCW);
+
+		// If there is a CCW in the currently selected detail objects, automatically select the last (most recently selected) one
+		if (bIsCCWSelected)
+		{
+			int32 CCWIndex = OperatorViewModel->GetDetailObjects().FindLastByPredicate(FindCCW);
+			if (CCWIndex > INDEX_NONE)
+			{
+				AColorCorrectionWindow* CCW = Cast<AColorCorrectionWindow>(OperatorViewModel->GetDetailObjects()[CCWIndex]);
+				FDisplayClusterColorGradingListItemRef* ListItemPtr = RootActorColorGradingItems.FindByPredicate([CCW](const FDisplayClusterColorGradingListItemRef& ListItem)
+				{
+					return ListItem.IsValid() && ListItem->Actor.Get() == CCW;
+				});
+
+				if (ListItemPtr)
+				{
+					RootActorList->SetSelectedItems({ *ListItemPtr });
+				}
+			}
+		}
+		else
+		{
+			// The nDisplay stage actor is always the first item in the root actor color grading items list, so set that as the currently selected item
+			RootActorList->SetSelectedItems({ RootActorColorGradingItems[0] });
+		}
 	}
 }
 
@@ -652,28 +688,62 @@ void SDisplayClusterColorGradingDrawer::FillLevelColorGradingList()
 	{
 		if (UWorld* World = RootActor->GetWorld())
 		{
-			for (TActorIterator<APostProcessVolume> PPVIter(World); PPVIter; ++PPVIter)
+			// Sorter that sorts the list items alphabetically by their display name
+			auto AlphabeticalSort = [](const FDisplayClusterColorGradingListItemRef& A, const FDisplayClusterColorGradingListItemRef& B)
 			{
-				APostProcessVolume* PostProcessVolume = *PPVIter;
-				BindBlueprintCompiledDelegate(PostProcessVolume->GetClass());
+				if (A.IsValid() && B.IsValid())
+				{
+					return *A < *B;
+				}
+				else
+				{
+					return false;
+				}
+			};
+
+			// Add all post process volumes that are in the same world as the stage actor
+			{
+				TArray<FDisplayClusterColorGradingListItemRef> SortedPPVs;
+				for (TActorIterator<APostProcessVolume> PPVIter(World); PPVIter; ++PPVIter)
+				{
+					APostProcessVolume* PostProcessVolume = *PPVIter;
+					BindBlueprintCompiledDelegate(PostProcessVolume->GetClass());
+
+					FDisplayClusterColorGradingListItemRef PPVListItemRef = MakeShared<FDisplayClusterColorGradingListItem>(PostProcessVolume);
+					PPVListItemRef->IsItemEnabled = CREATE_IS_ENABLED_LAMBDA(PostProcessVolume, PostProcessVolume->bEnabled);
+					PPVListItemRef->OnItemEnabledChanged = CREATE_ON_ENABLED_CHANGED_LAMBDA(PostProcessVolume, PostProcessVolume->bEnabled);
 
-				FDisplayClusterColorGradingListItemRef PPVListItemRef = MakeShared<FDisplayClusterColorGradingListItem>(PostProcessVolume);
-				PPVListItemRef->IsItemEnabled = CREATE_IS_ENABLED_LAMBDA(PostProcessVolume, PostProcessVolume->bEnabled);
-				PPVListItemRef->OnItemEnabledChanged = CREATE_ON_ENABLED_CHANGED_LAMBDA(PostProcessVolume, PostProcessVolume->bEnabled);
+					SortedPPVs.Add(PPVListItemRef);
+				}
 
-				LevelColorGradingItems.Add(PPVListItemRef);
+				SortedPPVs.Sort(AlphabeticalSort);
+				LevelColorGradingItems.Append(SortedPPVs);
 			}
 
-			for (TActorIterator<AColorCorrectRegion> CCRIter(World); CCRIter; ++CCRIter)
+			// Add all color correction regions that are in the same world as the stage actor
 			{
-				AColorCorrectRegion* ColorCorrectRegion = *CCRIter;
-				BindBlueprintCompiledDelegate(ColorCorrectRegion->GetClass());
+				TArray<FDisplayClusterColorGradingListItemRef> SortedCCRs;
+				for (TActorIterator<AColorCorrectRegion> CCRIter(World); CCRIter; ++CCRIter)
+				{
+					AColorCorrectRegion* ColorCorrectRegion = *CCRIter;
+
+					// Skip color correct windows, since those are stage managed actors and will be handled in a separate way
+					if (ColorCorrectRegion->IsA<AColorCorrectionWindow>())
+					{
+						continue;
+					}
 
-				FDisplayClusterColorGradingListItemRef CCRListItemRef = MakeShared<FDisplayClusterColorGradingListItem>(ColorCorrectRegion);
-				CCRListItemRef->IsItemEnabled = CREATE_IS_ENABLED_LAMBDA(ColorCorrectRegion, ColorCorrectRegion->Enabled);
-				CCRListItemRef->OnItemEnabledChanged = CREATE_ON_ENABLED_CHANGED_LAMBDA(ColorCorrectRegion, ColorCorrectRegion->Enabled);
+					BindBlueprintCompiledDelegate(ColorCorrectRegion->GetClass());
 
-				LevelColorGradingItems.Add(CCRListItemRef);
+					FDisplayClusterColorGradingListItemRef CCRListItemRef = MakeShared<FDisplayClusterColorGradingListItem>(ColorCorrectRegion);
+					CCRListItemRef->IsItemEnabled = CREATE_IS_ENABLED_LAMBDA(ColorCorrectRegion, ColorCorrectRegion->Enabled);
+					CCRListItemRef->OnItemEnabledChanged = CREATE_ON_ENABLED_CHANGED_LAMBDA(ColorCorrectRegion, ColorCorrectRegion->Enabled);
+
+					SortedCCRs.Add(CCRListItemRef);
+				}
+
+				SortedCCRs.Sort(AlphabeticalSort);
+				LevelColorGradingItems.Append(SortedCCRs);
 			}
 		}
 	}
@@ -711,16 +781,61 @@ void SDisplayClusterColorGradingDrawer::FillRootActorColorGradingList()
 
 		RootActorColorGradingItems.Add(RootActorListItemRef);
 
-		RootActor->ForEachComponent<UDisplayClusterICVFXCameraComponent>(false, [this, RootActor](UDisplayClusterICVFXCameraComponent* ICVFXCameraComponent)
+		auto AlphabeticalSort = [](const FDisplayClusterColorGradingListItemRef& A, const FDisplayClusterColorGradingListItemRef& B)
+		{
+			if (A.IsValid() && B.IsValid())
+			{
+				return *A < *B;
+			}
+			else
+			{
+				return false;
+			}
+		};
+
+		// Add any ICVFX camera component the root actor has to the color grading list
+		{
+			TArray<FDisplayClusterColorGradingListItemRef> SortedICVFXCameras;
+			RootActor->ForEachComponent<UDisplayClusterICVFXCameraComponent>(false, [this, RootActor, &SortedICVFXCameras](UDisplayClusterICVFXCameraComponent* ICVFXCameraComponent)
+			{
+				BindBlueprintCompiledDelegate(ICVFXCameraComponent->GetClass());
+
+				FDisplayClusterColorGradingListItemRef ICVFXCameraListItemRef = MakeShared<FDisplayClusterColorGradingListItem>(RootActor, ICVFXCameraComponent);
+				ICVFXCameraListItemRef->IsItemEnabled = CREATE_IS_ENABLED_LAMBDA(ICVFXCameraComponent, ICVFXCameraComponent->CameraSettings.EnableInnerFrustumColorGrading);
+				ICVFXCameraListItemRef->OnItemEnabledChanged = CREATE_ON_ENABLED_CHANGED_LAMBDA(ICVFXCameraComponent, ICVFXCameraComponent->CameraSettings.EnableInnerFrustumColorGrading);
+
+				SortedICVFXCameras.Add(ICVFXCameraListItemRef);
+			});
+
+			SortedICVFXCameras.Sort(AlphabeticalSort);
+			RootActorColorGradingItems.Append(SortedICVFXCameras);
+		}
+
+		// Add any color correction window that is currently selected in the operator's details panel to the list
 		{
-			BindBlueprintCompiledDelegate(ICVFXCameraComponent->GetClass());
+			const TArray<TWeakObjectPtr<UObject>>& DetailObjects = OperatorViewModel->GetDetailObjects();
 
-			FDisplayClusterColorGradingListItemRef ICVFXCameraListItemRef = MakeShared<FDisplayClusterColorGradingListItem>(RootActor, ICVFXCameraComponent);
-			ICVFXCameraListItemRef->IsItemEnabled = CREATE_IS_ENABLED_LAMBDA(ICVFXCameraComponent, ICVFXCameraComponent->CameraSettings.EnableInnerFrustumColorGrading);
-			ICVFXCameraListItemRef->OnItemEnabledChanged = CREATE_ON_ENABLED_CHANGED_LAMBDA(ICVFXCameraComponent, ICVFXCameraComponent->CameraSettings.EnableInnerFrustumColorGrading);
+			TArray<FDisplayClusterColorGradingListItemRef> SortedCCWs;
+			for (const TWeakObjectPtr<UObject>& Object : DetailObjects)
+			{
+				if (Object.IsValid())
+				{
+					if (AColorCorrectionWindow* CCW = Cast<AColorCorrectionWindow>(Object.Get()))
+					{
+						BindBlueprintCompiledDelegate(CCW->GetClass());
 
-			RootActorColorGradingItems.Add(ICVFXCameraListItemRef);
-		});
+						FDisplayClusterColorGradingListItemRef ICVFXCameraListItemRef = MakeShared<FDisplayClusterColorGradingListItem>(CCW);
+						ICVFXCameraListItemRef->IsItemEnabled = CREATE_IS_ENABLED_LAMBDA(CCW, CCW->Enabled);
+						ICVFXCameraListItemRef->OnItemEnabledChanged = CREATE_ON_ENABLED_CHANGED_LAMBDA(CCW, CCW->Enabled);
+
+						SortedCCWs.Add(ICVFXCameraListItemRef);
+					}
+				}
+			}
+
+			SortedCCWs.Sort(AlphabeticalSort);
+			RootActorColorGradingItems.Append(SortedCCWs);
+		}
 	}
 
 	if (RootActorList.IsValid())
@@ -913,6 +1028,16 @@ void SDisplayClusterColorGradingDrawer::OnActiveRootActorChanged(ADisplayCluster
 	Refresh(bPreserveDrawerState);
 }
 
+void SDisplayClusterColorGradingDrawer::OnDetailObjectsChanged(const TArray<UObject*>& Objects)
+{
+	if (RootActorList.IsValid())
+	{
+		FillRootActorColorGradingList();
+
+		SetDrawerStateToDefault();
+	}
+}
+
 void SDisplayClusterColorGradingDrawer::OnColorGradingDataModelGenerated()
 {
 	FillColorGradingGroupToolBar();
@@ -937,37 +1062,44 @@ void SDisplayClusterColorGradingDrawer::OnListSelectionChanged(TSharedRef<SDispl
 	{
 		if (RootActorList->GetSelectedItems().Num())
 		{
+			bUpdateDataModelOnSelectionChanged = false;
 			RootActorList->SetSelectedItems(TArray<FDisplayClusterColorGradingListItemRef>());
+			bUpdateDataModelOnSelectionChanged = true;
 		}
 	}
 	else if (RootActorList == SourceList && SelectedObjects.Num())
 	{
 		if (LevelActorsList->GetSelectedItems().Num())
 		{
+			bUpdateDataModelOnSelectionChanged = false;
 			LevelActorsList->SetSelectedItems(TArray<FDisplayClusterColorGradingListItemRef>());
+			bUpdateDataModelOnSelectionChanged = true;
 		}
 	}
 
-	TArray<UObject*> ObjectsToColorGrade;
-	for (const FDisplayClusterColorGradingListItemRef& SelectedObject : SelectedObjects)
+	if (bUpdateDataModelOnSelectionChanged)
 	{
-		if (SelectedObject->Component.IsValid())
+		TArray<UObject*> ObjectsToColorGrade;
+		for (const FDisplayClusterColorGradingListItemRef& SelectedObject : SelectedObjects)
 		{
-			ObjectsToColorGrade.Add(SelectedObject->Component.Get());
-		}
-		else if (SelectedObject->Actor.IsValid())
-		{
-			ObjectsToColorGrade.Add(SelectedObject->Actor.Get());
+			if (SelectedObject->Component.IsValid())
+			{
+				ObjectsToColorGrade.Add(SelectedObject->Component.Get());
+			}
+			else if (SelectedObject->Actor.IsValid())
+			{
+				ObjectsToColorGrade.Add(SelectedObject->Actor.Get());
+			}
 		}
-	}
 
-	ColorGradingDataModel->SetObjects(ObjectsToColorGrade);
+		ColorGradingDataModel->SetObjects(ObjectsToColorGrade);
 
-	// If the current drawer mode is the details view but the newly selected objects don't have any details sections to display,
-	// switch to the color grading drawer mode
-	if (CurrentDrawerMode == EDisplayClusterColorGradingDrawerMode::DetailsView && ColorGradingDataModel->DetailsSections.Num() == 0)
-	{
-		CurrentDrawerMode = EDisplayClusterColorGradingDrawerMode::ColorGrading;
+		// If the current drawer mode is the details view but the newly selected objects don't have any details sections to display,
+		// switch to the color grading drawer mode
+		if (CurrentDrawerMode == EDisplayClusterColorGradingDrawerMode::DetailsView && ColorGradingDataModel->DetailsSections.Num() == 0)
+		{
+			CurrentDrawerMode = EDisplayClusterColorGradingDrawerMode::ColorGrading;
+		}
 	}
 }
 
diff --git a/Engine/Plugins/Runtime/nDisplay/Source/DisplayClusterColorGrading/Private/Drawer/SDisplayClusterColorGradingDrawer.h b/Engine/Plugins/Runtime/nDisplay/Source/DisplayClusterColorGrading/Private/Drawer/SDisplayClusterColorGradingDrawer.h
--- a/Engine/Plugins/Runtime/nDisplay/Source/DisplayClusterColorGrading/Private/Drawer/SDisplayClusterColorGradingDrawer.h
+++ b/Engine/Plugins/Runtime/nDisplay/Source/DisplayClusterColorGrading/Private/Drawer/SDisplayClusterColorGradingDrawer.h
@@ -146,6 +146,9 @@ private:
 	/** Raised when the user has changed the active root actor selected in the nDisplay operator panel */
 	void OnActiveRootActorChanged(ADisplayClusterRootActor* NewRootActor);
 	
+	/** Raised when the obejcts being displayed in the operator's details panel have changed */
+	void OnDetailObjectsChanged(const TArray<UObject*>& Objects);
+
 	/** Raised when the color grading data model has been generated */
 	void OnColorGradingDataModelGenerated();
 
@@ -183,6 +186,9 @@ private:
 	/** Indicates that the drawer should refresh itself on the next tick */
 	bool bRefreshOnNextTick = false;
 
+	/** Indicates if the color grading data model should update when a list item selection has changed */
+	bool bUpdateDataModelOnSelectionChanged = true;
+
 	/** Delegate handle for the OnActiveRootActorChanged delegate */
 	FDelegateHandle ActiveRootActorChangedHandle;
 };
\ No newline at end of file
diff --git a/Engine/Plugins/Runtime/nDisplay/Source/DisplayClusterColorGrading/Private/SDisplayClusterColorGradingObjectList.cpp b/Engine/Plugins/Runtime/nDisplay/Source/DisplayClusterColorGrading/Private/SDisplayClusterColorGradingObjectList.cpp
--- a/Engine/Plugins/Runtime/nDisplay/Source/DisplayClusterColorGrading/Private/SDisplayClusterColorGradingObjectList.cpp
+++ b/Engine/Plugins/Runtime/nDisplay/Source/DisplayClusterColorGrading/Private/SDisplayClusterColorGradingObjectList.cpp
@@ -13,6 +13,32 @@
 
 #define LOCTEXT_NAMESPACE "DisplayClusterColorGrading"
 
+bool FDisplayClusterColorGradingListItem::operator<(const FDisplayClusterColorGradingListItem& Other) const
+{
+	FString ThisName = TEXT("");
+	FString OtherName = TEXT("");
+
+	if (Component.IsValid())
+	{
+		ThisName = Component->GetName();
+	}
+	else if (Actor.IsValid())
+	{
+		ThisName = Actor->GetActorLabel();
+	}
+
+	if (Other.Component.IsValid())
+	{
+		OtherName = Other.Component->GetName();
+	}
+	else if (Other.Actor.IsValid())
+	{
+		OtherName = Other.Actor->GetActorLabel();
+	}
+
+	return ThisName < OtherName;
+}
+
 namespace DisplayClusterColorGradingObjectListColumnNames
 {
 	const static FName ItemEnabled(TEXT("ItemEnabled"));
diff --git a/Engine/Plugins/Runtime/nDisplay/Source/DisplayClusterColorGrading/Private/SDisplayClusterColorGradingObjectList.h b/Engine/Plugins/Runtime/nDisplay/Source/DisplayClusterColorGrading/Private/SDisplayClusterColorGradingObjectList.h
--- a/Engine/Plugins/Runtime/nDisplay/Source/DisplayClusterColorGrading/Private/SDisplayClusterColorGradingObjectList.h
+++ b/Engine/Plugins/Runtime/nDisplay/Source/DisplayClusterColorGrading/Private/SDisplayClusterColorGradingObjectList.h
@@ -37,6 +37,9 @@ struct FDisplayClusterColorGradingListItem
 		, Component(InComponent)
 		, IsItemEnabled(false)
 	{ }
+
+	/** Less than operator overload that compares list items alphabetically by their display names */
+	bool operator<(const FDisplayClusterColorGradingListItem& Other) const;
 };
 
 typedef TSharedPtr<FDisplayClusterColorGradingListItem> FDisplayClusterColorGradingListItemRef;
diff --git a/Engine/Plugins/Runtime/nDisplay/Source/DisplayClusterLightCardEditor/Private/DisplayClusterLightCardEditor.cpp b/Engine/Plugins/Runtime/nDisplay/Source/DisplayClusterLightCardEditor/Private/DisplayClusterLightCardEditor.cpp
--- a/Engine/Plugins/Runtime/nDisplay/Source/DisplayClusterLightCardEditor/Private/DisplayClusterLightCardEditor.cpp
+++ b/Engine/Plugins/Runtime/nDisplay/Source/DisplayClusterLightCardEditor/Private/DisplayClusterLightCardEditor.cpp
@@ -164,7 +164,7 @@ void FDisplayClusterLightCardEditor::SelectActors(const TArray<AActor*>& ActorsT
 		SelectedActors.Empty();
 	}
 	
-	IDisplayClusterOperator::Get().ShowDetailsForObjects(*reinterpret_cast<const TArray<UObject*>*>(&ActorsToSelect));
+	IDisplayClusterOperator::Get().GetOperatorViewModel()->ShowDetailsForObjects(*reinterpret_cast<const TArray<UObject*>*>(&ActorsToSelect));
 }
 
 void FDisplayClusterLightCardEditor::SelectActorProxies(const TArray<AActor*>& ActorsToSelect)
diff --git a/Engine/Plugins/Runtime/nDisplay/Source/DisplayClusterOperator/Private/DisplayClusterOperatorModule.cpp b/Engine/Plugins/Runtime/nDisplay/Source/DisplayClusterOperator/Private/DisplayClusterOperatorModule.cpp
--- a/Engine/Plugins/Runtime/nDisplay/Source/DisplayClusterOperator/Private/DisplayClusterOperatorModule.cpp
+++ b/Engine/Plugins/Runtime/nDisplay/Source/DisplayClusterOperator/Private/DisplayClusterOperatorModule.cpp
@@ -85,18 +85,6 @@ void FDisplayClusterOperatorModule::GetRootActorLevelInstances(TArray<ADisplayCl
 	}
 }
 
-void FDisplayClusterOperatorModule::ShowDetailsForObject(UObject* Object)
-{
-	TArray<UObject*> Objects;
-	Objects.Add(Object);
-	ShowDetailsForObjects(Objects);
-}
-
-void FDisplayClusterOperatorModule::ShowDetailsForObjects(const TArray<UObject*>& Objects)
-{
-	DetailObjectsChanged.Broadcast(Objects);
-}
-
 void FDisplayClusterOperatorModule::ToggleDrawer(const FName DrawerId)
 {
 	if (ActiveOperatorPanel.IsValid())
diff --git a/Engine/Plugins/Runtime/nDisplay/Source/DisplayClusterOperator/Private/DisplayClusterOperatorModule.h b/Engine/Plugins/Runtime/nDisplay/Source/DisplayClusterOperator/Private/DisplayClusterOperatorModule.h
--- a/Engine/Plugins/Runtime/nDisplay/Source/DisplayClusterOperator/Private/DisplayClusterOperatorModule.h
+++ b/Engine/Plugins/Runtime/nDisplay/Source/DisplayClusterOperator/Private/DisplayClusterOperatorModule.h
@@ -34,7 +34,6 @@ public:
 	virtual FOnRegisterLayoutExtensions& OnRegisterLayoutExtensions() override { return RegisterLayoutExtensions; }
 	virtual FOnRegisterStatusBarExtensions& OnRegisterStatusBarExtensions() override { return RegisterStatusBarExtensions; }
 	virtual FOnAppendOperatorPanelCommands& OnAppendOperatorPanelCommands() override { return AppendOperatorPanelCommands; }
-	virtual FOnDetailObjectsChanged& OnDetailObjectsChanged() override { return DetailObjectsChanged; }
 
 	virtual FName GetPrimaryOperatorExtensionId() override;
 	virtual FName GetAuxilliaryOperatorExtensionId() override;
@@ -42,8 +41,6 @@ public:
 	virtual TSharedPtr<FExtensibilityManager> GetOperatorToolBarExtensibilityManager() override { return OperatorToolBarExtensibilityManager; }
 	virtual TSharedPtr<FExtensibilityManager> GetOperatorMenuExtensibilityManager() override { return OperatorMenuExtensibilityManager; }
 	virtual void GetRootActorLevelInstances(TArray<ADisplayClusterRootActor*>& OutRootActorInstances) override;
-	virtual void ShowDetailsForObject(UObject* Object) override;
-	virtual void ShowDetailsForObjects(const TArray<UObject*>& Objects) override;
 	virtual void ToggleDrawer(const FName DrawerId) override;
 	virtual void ForceDismissDrawers() override;
 	//~ End IDisplayClusterOperator interface
@@ -65,7 +62,6 @@ private:
 	FOnRegisterLayoutExtensions RegisterLayoutExtensions;
 	FOnRegisterStatusBarExtensions RegisterStatusBarExtensions;
 	FOnAppendOperatorPanelCommands AppendOperatorPanelCommands;
-	FOnDetailObjectsChanged DetailObjectsChanged;
 
 	TWeakPtr<SDisplayClusterOperatorPanel> ActiveOperatorPanel;
 	TSharedPtr<FDisplayClusterOperatorViewModel> OperatorViewModel;
diff --git a/Engine/Plugins/Runtime/nDisplay/Source/DisplayClusterOperator/Private/DisplayClusterOperatorViewModel.cpp b/Engine/Plugins/Runtime/nDisplay/Source/DisplayClusterOperator/Private/DisplayClusterOperatorViewModel.cpp
--- a/Engine/Plugins/Runtime/nDisplay/Source/DisplayClusterOperator/Private/DisplayClusterOperatorViewModel.cpp
+++ b/Engine/Plugins/Runtime/nDisplay/Source/DisplayClusterOperator/Private/DisplayClusterOperatorViewModel.cpp
@@ -24,6 +24,18 @@ void FDisplayClusterOperatorViewModel::SetRootActor(ADisplayClusterRootActor* In
 	RootActorChanged.Broadcast(RootActor.Get());
 }
 
+void FDisplayClusterOperatorViewModel::ShowDetailsForObject(UObject* Object)
+{
+	ShowDetailsForObjects({ Object });
+}
+
+void FDisplayClusterOperatorViewModel::ShowDetailsForObjects(const TArray<UObject*>& Objects)
+{
+	DetailObjects.Empty();
+	DetailObjects.Append(Objects);
+	DetailObjectsChanged.Broadcast(Objects);
+}
+
 TSharedRef<FTabManager> FDisplayClusterOperatorViewModel::CreateTabManager(const TSharedRef<SDockTab>& MajorTabOwner)
 {
 	TabManager = FGlobalTabmanager::Get()->NewTabManager(MajorTabOwner);
diff --git a/Engine/Plugins/Runtime/nDisplay/Source/DisplayClusterOperator/Private/DisplayClusterOperatorViewModel.h b/Engine/Plugins/Runtime/nDisplay/Source/DisplayClusterOperator/Private/DisplayClusterOperatorViewModel.h
--- a/Engine/Plugins/Runtime/nDisplay/Source/DisplayClusterOperator/Private/DisplayClusterOperatorViewModel.h
+++ b/Engine/Plugins/Runtime/nDisplay/Source/DisplayClusterOperator/Private/DisplayClusterOperatorViewModel.h
@@ -20,7 +20,10 @@ public:
 	virtual ADisplayClusterRootActor* GetRootActor(bool bEvenIfPendingKill = false) const override;
 	virtual void SetRootActor(ADisplayClusterRootActor* InRootActor) override;
 	virtual FOnActiveRootActorChanged& OnActiveRootActorChanged() override { return RootActorChanged; }
-
+	virtual const TArray<TWeakObjectPtr<UObject>> GetDetailObjects() override { return DetailObjects; }
+	virtual void ShowDetailsForObject(UObject* Object) override;
+	virtual void ShowDetailsForObjects(const TArray<UObject*>& Objects) override;
+	virtual FOnDetailObjectsChanged& OnDetailObjectsChanged() override { return DetailObjectsChanged; }
 	virtual TSharedPtr<FTabManager> GetTabManager() const override { return TabManager; }
 
 	virtual TSharedPtr<FWorkspaceItem> GetWorkspaceMenuGroup() const override { return WorkspaceItem; }
@@ -31,8 +34,10 @@ public:
 
 private:
 	TWeakObjectPtr<ADisplayClusterRootActor> RootActor;
+	TArray<TWeakObjectPtr<UObject>> DetailObjects;
 	TSharedPtr<FTabManager> TabManager;
 	TSharedPtr<FWorkspaceItem> WorkspaceItem;
 
 	FOnActiveRootActorChanged RootActorChanged;
+	FOnDetailObjectsChanged DetailObjectsChanged;
 };
\ No newline at end of file
diff --git a/Engine/Plugins/Runtime/nDisplay/Source/DisplayClusterOperator/Private/SDisplayClusterOperatorPanel.cpp b/Engine/Plugins/Runtime/nDisplay/Source/DisplayClusterOperator/Private/SDisplayClusterOperatorPanel.cpp
--- a/Engine/Plugins/Runtime/nDisplay/Source/DisplayClusterOperator/Private/SDisplayClusterOperatorPanel.cpp
+++ b/Engine/Plugins/Runtime/nDisplay/Source/DisplayClusterOperator/Private/SDisplayClusterOperatorPanel.cpp
@@ -32,7 +32,7 @@ SDisplayClusterOperatorPanel::~SDisplayClusterOperatorPanel()
 {
 	FDisplayClusterOperatorModule& OperatorModule = FModuleManager::GetModuleChecked<FDisplayClusterOperatorModule>(IDisplayClusterOperator::ModuleName);
 	OperatorModule.OnAppUnregistered().RemoveAll(this);
-	OperatorModule.OnDetailObjectsChanged().Remove(DetailObjectsChangedHandle);
+	OperatorModule.GetOperatorViewModel()->OnDetailObjectsChanged().Remove(DetailObjectsChangedHandle);
 }
 
 void SDisplayClusterOperatorPanel::Construct(const FArguments& InArgs, const TSharedRef<FTabManager>& InTabManager, const TSharedPtr<SWindow>& WindowOwner)
@@ -146,7 +146,7 @@ void SDisplayClusterOperatorPanel::Construct(const FArguments& InArgs, const TSh
 	IDisplayClusterOperator::Get().OnRegisterLayoutExtensions().Broadcast(*LayoutExtender);
 	Layout->ProcessExtensions(*LayoutExtender);
 
-	DetailObjectsChangedHandle = IDisplayClusterOperator::Get().OnDetailObjectsChanged().AddSP(this, &SDisplayClusterOperatorPanel::DisplayObjectsInDetailsPanel);
+	DetailObjectsChangedHandle = IDisplayClusterOperator::Get().GetOperatorViewModel()->OnDetailObjectsChanged().AddSP(this, &SDisplayClusterOperatorPanel::DisplayObjectsInDetailsPanel);
 	
 	ChildSlot
 	[
diff --git a/Engine/Plugins/Runtime/nDisplay/Source/DisplayClusterOperator/Public/IDisplayClusterOperator.h b/Engine/Plugins/Runtime/nDisplay/Source/DisplayClusterOperator/Public/IDisplayClusterOperator.h
--- a/Engine/Plugins/Runtime/nDisplay/Source/DisplayClusterOperator/Public/IDisplayClusterOperator.h
+++ b/Engine/Plugins/Runtime/nDisplay/Source/DisplayClusterOperator/Public/IDisplayClusterOperator.h
@@ -26,7 +26,6 @@ public:
 	DECLARE_MULTICAST_DELEGATE_OneParam(FOnRegisterLayoutExtensions, FLayoutExtender&);
 	DECLARE_MULTICAST_DELEGATE_OneParam(FOnRegisterStatusBarExtensions, FDisplayClusterOperatorStatusBarExtender&);
 	DECLARE_MULTICAST_DELEGATE_OneParam(FOnAppendOperatorPanelCommands, TSharedRef<FUICommandList>)
-	DECLARE_MULTICAST_DELEGATE_OneParam(FOnDetailObjectsChanged, const TArray<UObject*>&);
 	DECLARE_DELEGATE_RetVal_OneParam(TSharedRef<IDisplayClusterOperatorApp>, FOnGetAppInstance, TSharedRef<IDisplayClusterOperatorViewModel>);
 	
 public:
@@ -71,9 +70,6 @@ public:
 	/** Gets the event handler that is raised when the operator panel binds commands to its command list */
 	virtual FOnAppendOperatorPanelCommands& OnAppendOperatorPanelCommands() = 0;
 
-	/** Gets the event handler that is raised when the objects being displayed in the operator's details panel are changed */
-	virtual FOnDetailObjectsChanged& OnDetailObjectsChanged() = 0;
-
 	/** Gets the extension ID for the main window region that can be used to add tabs to the operator panel */
 	virtual FName GetPrimaryOperatorExtensionId() = 0;
 
@@ -92,12 +88,6 @@ public:
 	/** Gets a list of all nDisplay root actor instances that are on the currently loaded level */
 	virtual void GetRootActorLevelInstances(TArray<ADisplayClusterRootActor*>& OutRootActorInstances) = 0;
 
-	/** Displays the properties of the specified object in the operator's details panel */
-	virtual void ShowDetailsForObject(UObject* Object) = 0;
-
-	/** Displays the properties of the specified object in the operator's details panel */
-	virtual void ShowDetailsForObjects(const TArray<UObject*>& Objects) = 0;
-
 	/** Toggles the state of a drawer with the specified ID, closing the drawer if it is open, and opening the drawer if it is closed */
 	virtual void ToggleDrawer(const FName DrawerId) = 0;
 
diff --git a/Engine/Plugins/Runtime/nDisplay/Source/DisplayClusterOperator/Public/IDisplayClusterOperatorViewModel.h b/Engine/Plugins/Runtime/nDisplay/Source/DisplayClusterOperator/Public/IDisplayClusterOperatorViewModel.h
--- a/Engine/Plugins/Runtime/nDisplay/Source/DisplayClusterOperator/Public/IDisplayClusterOperatorViewModel.h
+++ b/Engine/Plugins/Runtime/nDisplay/Source/DisplayClusterOperator/Public/IDisplayClusterOperatorViewModel.h
@@ -8,11 +8,14 @@ class ADisplayClusterRootActor;
 class FTabManager;
 class FWorkspaceItem;
 
-DECLARE_EVENT_OneParam(IDisplayClusterOperatorViewModel, FOnActiveRootActorChanged, ADisplayClusterRootActor*);
 
 /** Interface for a view model object that stores any state from the operator panel that should be exposed externally */
 class DISPLAYCLUSTEROPERATOR_API IDisplayClusterOperatorViewModel
 {
+public:
+	DECLARE_MULTICAST_DELEGATE_OneParam(FOnActiveRootActorChanged, ADisplayClusterRootActor*);
+	DECLARE_MULTICAST_DELEGATE_OneParam(FOnDetailObjectsChanged, const TArray<UObject*>&);
+
 public:
 	/** Gets whether the view model has been populated with a valid root actor */
 	virtual bool HasRootActor(bool bEvenIfPendingKill = false) const = 0;
@@ -26,6 +29,18 @@ public:
 	/** Gets the event handler that is raised when the operator panel changes the root actor being operated on */
 	virtual FOnActiveRootActorChanged& OnActiveRootActorChanged() = 0;
 
+	/** Gets the list of objects being displayed in the operator's details panel */
+	virtual const TArray<TWeakObjectPtr<UObject>> GetDetailObjects() = 0;
+
+	/** Displays the properties of the specified object in the operator's details panel */
+	virtual void ShowDetailsForObject(UObject* Object) = 0;
+
+	/** Displays the properties of the specified object in the operator's details panel */
+	virtual void ShowDetailsForObjects(const TArray<UObject*>& Objects) = 0;
+
+	/** Gets the event handler that is raised when the objects being displayed in the operator's details panel are changed */
+	virtual FOnDetailObjectsChanged& OnDetailObjectsChanged() = 0;
+
 	/** Gets the tab manager of the active operator panel, if there is an open operator panel */
 	virtual TSharedPtr<FTabManager> GetTabManager() const = 0;
 
