nDisplay: Fixed issue where per-viewport detail properties were not showing up in color grading drawer, and added Include Entire Cluster/Include All Nodes properties to color grading drawer. Also fixed issue where color grading details property categories could show up in different order than desired.

diff --git a/Engine/Plugins/Runtime/nDisplay/Source/DisplayClusterColorGrading/Private/DataModelGenerators/DisplayClusterColorGradingGenerator_RootActor.cpp b/Engine/Plugins/Runtime/nDisplay/Source/DisplayClusterColorGrading/Private/DataModelGenerators/DisplayClusterColorGradingGenerator_RootActor.cpp
--- a/Engine/Plugins/Runtime/nDisplay/Source/DisplayClusterColorGrading/Private/DataModelGenerators/DisplayClusterColorGradingGenerator_RootActor.cpp
+++ b/Engine/Plugins/Runtime/nDisplay/Source/DisplayClusterColorGrading/Private/DataModelGenerators/DisplayClusterColorGradingGenerator_RootActor.cpp
@@ -28,6 +28,8 @@
 
 #define LOCTEXT_NAMESPACE "DisplayClusterColorGrading"
 
+#define CATEGORY_OVERRIDE_NAME(Category) TEXT(Category "_Override")
+
 FDisplayClusterColorGradingDataModel::FColorGradingGroup FDisplayClusterColorGradingGenerator_ColorGradingRenderingSettings::CreateColorGradingGroup(const TSharedPtr<IPropertyHandle>& GroupPropertyHandle)
 {
 	FDisplayClusterColorGradingDataModel::FColorGradingGroup ColorGradingGroup;
@@ -50,6 +52,15 @@ FDisplayClusterColorGradingDataModel::FColorGradingGroup FDisplayClusterColorGra
 	AddDetailsViewPropertyToGroup(GroupPropertyHandle, ColorGradingGroup, GET_MEMBER_NAME_CHECKED(FDisplayClusterConfigurationViewport_ColorGradingMiscSettings, ExpandGamut), TEXT("Misc"));
 	AddDetailsViewPropertyToGroup(GroupPropertyHandle, ColorGradingGroup, GET_MEMBER_NAME_CHECKED(FDisplayClusterConfigurationViewport_ColorGradingMiscSettings, SceneColorTint), TEXT("Misc"));
 
+	// Ensure the details view categories are always displayed in a specific order
+	ColorGradingGroup.DetailsViewCategorySortOrder =
+	{
+		{ CATEGORY_OVERRIDE_NAME("Exposure"), 0},
+		{ CATEGORY_OVERRIDE_NAME("Color Grading"), 1},
+		{ CATEGORY_OVERRIDE_NAME("White Balance"), 2},
+		{ CATEGORY_OVERRIDE_NAME("Misc"), 3}
+	};
+
 	return ColorGradingGroup;
 }
 
@@ -271,6 +282,14 @@ void FDisplayClusterColorGradingGenerator_RootActor::GenerateDataModel(IProperty
 				FDisplayClusterColorGradingDataModel::FColorGradingGroup PerViewportGroup = CreateColorGradingGroup(PerViewportElementHandle);
 				PerViewportGroup.EditConditionPropertyHandle = PerViewportElementHandle->GetChildHandle(GET_MEMBER_NAME_CHECKED(FDisplayClusterConfigurationViewport_PerViewportColorGrading, bIsEnabled));
 
+				// Add per-viewport group specific properties and force the category of these properties to be at the top
+				AddDetailsViewPropertyToGroup(PerViewportElementHandle, 
+					PerViewportGroup, 
+					GET_MEMBER_NAME_CHECKED(FDisplayClusterConfigurationViewport_PerViewportColorGrading, bIsEntireClusterEnabled),
+					TEXT("Per-Viewport Settings"));
+
+				PerViewportGroup.DetailsViewCategorySortOrder.Add(CATEGORY_OVERRIDE_NAME("Per-Viewport Settings"), -1);
+
 				TSharedPtr<IPropertyHandle> NamePropertyHandle = PerViewportElementHandle->GetChildHandle(GET_MEMBER_NAME_CHECKED(FDisplayClusterConfigurationViewport_PerViewportColorGrading, Name));
 				if (NamePropertyHandle.IsValid() && NamePropertyHandle->IsValidHandle())
 				{
@@ -728,6 +747,19 @@ void FDisplayClusterColorGradingGenerator_ICVFXCamera::GenerateDataModel(IProper
 				FDisplayClusterColorGradingDataModel::FColorGradingGroup PerNodeGroup = CreateColorGradingGroup(PerNodeElementHandle);
 				PerNodeGroup.EditConditionPropertyHandle = PerNodeElementHandle->GetChildHandle(GET_MEMBER_NAME_CHECKED(FDisplayClusterConfigurationViewport_PerNodeColorGrading, bIsEnabled));
 
+				// Add per-node group specific properties and force the category of these properties to be at the top
+				AddDetailsViewPropertyToGroup(PerNodeElementHandle,
+					PerNodeGroup,
+					GET_MEMBER_NAME_CHECKED(FDisplayClusterConfigurationViewport_PerNodeColorGrading, bEntireClusterColorGrading),
+					TEXT("Per-Node Settings"));
+
+				AddDetailsViewPropertyToGroup(PerNodeElementHandle,
+					PerNodeGroup,
+					GET_MEMBER_NAME_CHECKED(FDisplayClusterConfigurationViewport_PerNodeColorGrading, bAllNodesColorGrading),
+					TEXT("Per-Node Settings"));
+
+				PerNodeGroup.DetailsViewCategorySortOrder.Add(CATEGORY_OVERRIDE_NAME("Per-Node Settings"), -1);
+
 				TSharedPtr<IPropertyHandle> NamePropertyHandle = PerNodeElementHandle->GetChildHandle(GET_MEMBER_NAME_CHECKED(FDisplayClusterConfigurationViewport_PerNodeColorGrading, Name));
 				if (NamePropertyHandle.IsValid() && NamePropertyHandle->IsValidHandle())
 				{
diff --git a/Engine/Plugins/Runtime/nDisplay/Source/DisplayClusterColorGrading/Private/DisplayClusterColorGradingDataModel.h b/Engine/Plugins/Runtime/nDisplay/Source/DisplayClusterColorGrading/Private/DisplayClusterColorGradingDataModel.h
--- a/Engine/Plugins/Runtime/nDisplay/Source/DisplayClusterColorGrading/Private/DisplayClusterColorGradingDataModel.h
+++ b/Engine/Plugins/Runtime/nDisplay/Source/DisplayClusterColorGrading/Private/DisplayClusterColorGradingDataModel.h
@@ -84,6 +84,9 @@ public:
 
 		/** A list of properties to display in the details view of the color grading panel */
 		TArray<TSharedPtr<IPropertyHandle>> DetailsViewPropertyHandles;
+
+		/** Optional sort order for the categories displayed in the details view */
+		TMap<FName, int32> DetailsViewCategorySortOrder;
 	};
 
 	/** Stores data used to generate a details subsection, which is a selectable element in a details section that modifies which properties are displayed in the section */
diff --git a/Engine/Plugins/Runtime/nDisplay/Source/DisplayClusterColorGrading/Private/SDisplayClusterColorGradingColorWheelPanel.cpp b/Engine/Plugins/Runtime/nDisplay/Source/DisplayClusterColorGrading/Private/SDisplayClusterColorGradingColorWheelPanel.cpp
--- a/Engine/Plugins/Runtime/nDisplay/Source/DisplayClusterColorGrading/Private/SDisplayClusterColorGradingColorWheelPanel.cpp
+++ b/Engine/Plugins/Runtime/nDisplay/Source/DisplayClusterColorGrading/Private/SDisplayClusterColorGradingColorWheelPanel.cpp
@@ -418,8 +418,9 @@ void SDisplayClusterColorGradingColorWheelPanel::ClearColorGradingElementsToolBa
 class FColorGradingDetailsCustomization : public IDetailCustomization
 {
 public:
-	FColorGradingDetailsCustomization(TArray<TSharedPtr<IPropertyHandle>> InPropertyHandles)
-		: PropertyHandles(InPropertyHandles)
+	FColorGradingDetailsCustomization(const FDisplayClusterColorGradingDataModel::FColorGradingGroup& ColorGradingGroup)
+		: PropertyHandles(ColorGradingGroup.DetailsViewPropertyHandles)
+		, CategorySortOrder(ColorGradingGroup.DetailsViewCategorySortOrder)
 	{ }
 
 	virtual void CustomizeDetails(IDetailLayoutBuilder& DetailBuilder) override
@@ -435,6 +436,7 @@ public:
 		// TransformCommon is a custom category that doesn't get returned by GetCategoryNames that also needs to be hidden
 		DetailBuilder.HideCategory(TEXT("TransformCommon"));
 
+		TSet<FName> FinalCategories;
 		for (const TSharedPtr<IPropertyHandle>& PropertyHandle : PropertyHandles)
 		{
 			FName CategoryOverride = PropertyHandle->GetDefaultCategoryName();
@@ -447,11 +449,25 @@ public:
 
 			IDetailCategoryBuilder& CategoryBuilder = DetailBuilder.EditCategory(CategoryOverride, CategoryOverrideText);
 			CategoryBuilder.AddProperty(PropertyHandle);
+
+			FinalCategories.Add(CategoryOverride);
 		}
+
+		DetailBuilder.SortCategories([this, FinalCategories = MoveTemp(FinalCategories)](const TMap<FName, IDetailCategoryBuilder*>& CategoryMap)
+		{
+			for (const FName& Category : FinalCategories)
+			{
+				if (CategorySortOrder.Contains(Category))
+				{
+					CategoryMap[Category]->SetSortOrder(CategorySortOrder[Category]);
+				}
+			}
+		});
 	}
 
 private:
 	TArray<TSharedPtr<IPropertyHandle>> PropertyHandles;
+	TMap<FName, int32> CategorySortOrder;
 };
 
 void SDisplayClusterColorGradingColorWheelPanel::FillDetailsView(const FDisplayClusterColorGradingDataModel::FColorGradingGroup& ColorGradingGroup)
@@ -473,10 +489,11 @@ void SDisplayClusterColorGradingColorWheelPanel::FillDetailsView(const FDisplayC
 		{
 			DetailsView->RegisterInstancedCustomPropertyLayout(OuterObjects[0]->GetClass(), FOnGetDetailCustomizationInstance::CreateLambda([&ColorGradingGroup]()
 			{
-				return MakeShareable(new FColorGradingDetailsCustomization(ColorGradingGroup.DetailsViewPropertyHandles));
+				return MakeShareable(new FColorGradingDetailsCustomization(ColorGradingGroup));
 			}));
 
-			DetailsView->SetObjects(OuterObjects);
+			const bool bForceRefresh = true;
+			DetailsView->SetObjects(OuterObjects, bForceRefresh);
 		}
 	}
 }
