nDisplay: Updated the ICVFX window's toolbar style to match the style of the rest of the editor, and modified the labels on the Place Actors menu to match stylistically as well. Added the Place Actors submenu to the 3D viewport's context menu, allowing users to add new actors at their cursor's location. Additionally, fixed a bug where the light card labels were hit traceable, which was causing the hit bounds of light cards to be far bigger than they should be.

diff --git a/Engine/Plugins/Runtime/nDisplay/Source/DisplayCluster/Private/Game/EngineClasses/Scene/DisplayClusterLabelComponent.cpp b/Engine/Plugins/Runtime/nDisplay/Source/DisplayCluster/Private/Game/EngineClasses/Scene/DisplayClusterLabelComponent.cpp
--- a/Engine/Plugins/Runtime/nDisplay/Source/DisplayCluster/Private/Game/EngineClasses/Scene/DisplayClusterLabelComponent.cpp
+++ b/Engine/Plugins/Runtime/nDisplay/Source/DisplayCluster/Private/Game/EngineClasses/Scene/DisplayClusterLabelComponent.cpp
@@ -23,6 +23,9 @@ UDisplayClusterLabelComponent::UDisplayClusterLabelComponent()
 	WidgetComponent->SetBlendMode(EWidgetBlendMode::Transparent);
 	WidgetComponent->SetTranslucentSortPriority(1.f);
 	WidgetComponent->SetUsingAbsoluteScale(true);
+	WidgetComponent->SetCollisionEnabled(ECollisionEnabled::NoCollision);
+	WidgetComponent->SetCollisionResponseToAllChannels(ECollisionResponse::ECR_Ignore);
+	WidgetComponent->SetCollisionObjectType(ECollisionChannel::ECC_Visibility);
 
 	WidgetClass = TSoftClassPtr<UDisplayClusterLabelWidget>(FSoftClassPath(TEXT("/nDisplay/Widgets/BP_DisplayClusterLabelWidget.BP_DisplayClusterLabelWidget_C")));
 
diff --git a/Engine/Plugins/Runtime/nDisplay/Source/DisplayClusterLightCardEditor/Private/DisplayClusterLightCardEditor.cpp b/Engine/Plugins/Runtime/nDisplay/Source/DisplayClusterLightCardEditor/Private/DisplayClusterLightCardEditor.cpp
--- a/Engine/Plugins/Runtime/nDisplay/Source/DisplayClusterLightCardEditor/Private/DisplayClusterLightCardEditor.cpp
+++ b/Engine/Plugins/Runtime/nDisplay/Source/DisplayClusterLightCardEditor/Private/DisplayClusterLightCardEditor.cpp
@@ -1088,9 +1088,9 @@ TSharedRef<SWidget> FDisplayClusterLightCardEditor::GeneratePlaceActorsMenu()
 				continue;
 			}
 			
-			FString Label = FString::Printf(TEXT("Add %s"), *Class->GetDisplayNameText().ToString());
+			FText Label = Class->GetDisplayNameText();
 			FSlateIcon StageActorIcon = FSlateIconFinder::FindIconForClass(Class);
-			MenuBuilder.AddMenuEntry(FText::FromString(Label), LOCTEXT("AddStageActorHeader", "Add a stage actor to the scene"), StageActorIcon,
+			MenuBuilder.AddMenuEntry(Label, LOCTEXT("AddStageActorHeader", "Add a stage actor to the scene"), StageActorIcon,
 				FUIAction(FExecuteAction::CreateRaw(this, &FDisplayClusterLightCardEditor::AddNewDynamic, Class)));
 		}
 		
@@ -1442,7 +1442,7 @@ void FDisplayClusterLightCardEditor::ExtendToolbar(FToolBarBuilder& ToolbarBuild
 			FOnGetContent::CreateSP(this, &FDisplayClusterLightCardEditor::GeneratePlaceActorsMenu),
 			TAttribute<FText>(),
 			LOCTEXT("PlaceActors_ToolTip", "Place actors in the scene"),
-			FSlateIcon(FAppStyle::GetAppStyleSetName(), "Icons.PlaceActors"));
+			FSlateIcon(FAppStyle::Get().GetStyleSetName(), "LevelEditor.OpenAddContent.Background", NAME_None, "LevelEditor.OpenAddContent.Overlay"));
 	}
 	
 	ToolbarBuilder.AddSeparator();
diff --git a/Engine/Plugins/Runtime/nDisplay/Source/DisplayClusterLightCardEditor/Private/DisplayClusterLightCardEditorCommands.cpp b/Engine/Plugins/Runtime/nDisplay/Source/DisplayClusterLightCardEditor/Private/DisplayClusterLightCardEditorCommands.cpp
--- a/Engine/Plugins/Runtime/nDisplay/Source/DisplayClusterLightCardEditor/Private/DisplayClusterLightCardEditorCommands.cpp
+++ b/Engine/Plugins/Runtime/nDisplay/Source/DisplayClusterLightCardEditor/Private/DisplayClusterLightCardEditorCommands.cpp
@@ -31,8 +31,8 @@ void FDisplayClusterLightCardEditorCommands::RegisterCommands()
 	UI_COMMAND(SphericalCoordinateSystem, "Spherical Coordinate System", "Move objects using latitude and longitude", EUserInterfaceActionType::RadioButton, FInputChord());
 	UI_COMMAND(CartesianCoordinateSystem, "Cartesian Coordinate System", "Move objects using X, Y, and Z axes", EUserInterfaceActionType::RadioButton, FInputChord());
 
-	UI_COMMAND(AddNewLightCard, "Add New Light Card", "Add and assign a new Light Card to the actor", EUserInterfaceActionType::Button, FInputChord());
-	UI_COMMAND(AddNewFlag, "Add Flag", "Add and assign a new light control flag", EUserInterfaceActionType::Button, FInputChord());
+	UI_COMMAND(AddNewLightCard, "Light Card", "Add and assign a new Light Card to the actor", EUserInterfaceActionType::Button, FInputChord());
+	UI_COMMAND(AddNewFlag, "Flag", "Add and assign a new light control flag", EUserInterfaceActionType::Button, FInputChord());
 	UI_COMMAND(AddExistingLightCard, "Import Existing Content", "Add an existing Light Card to the actor", EUserInterfaceActionType::Button, FInputChord());
 	UI_COMMAND(RemoveLightCard, "Remove from Actor", "Remove the Light Card from the actor but do not delete it", EUserInterfaceActionType::Button, FInputChord());
 	UI_COMMAND(PasteHere, "Paste Here", "Paste clipboard contents at the click location", EUserInterfaceActionType::Button, FInputChord());
diff --git a/Engine/Plugins/Runtime/nDisplay/Source/DisplayClusterLightCardEditor/Private/Viewport/DisplayClusterLightcardEditorViewport.cpp b/Engine/Plugins/Runtime/nDisplay/Source/DisplayClusterLightCardEditor/Private/Viewport/DisplayClusterLightcardEditorViewport.cpp
--- a/Engine/Plugins/Runtime/nDisplay/Source/DisplayClusterLightCardEditor/Private/Viewport/DisplayClusterLightcardEditorViewport.cpp
+++ b/Engine/Plugins/Runtime/nDisplay/Source/DisplayClusterLightCardEditor/Private/Viewport/DisplayClusterLightcardEditorViewport.cpp
@@ -4,6 +4,7 @@
 
 #include "DisplayClusterLightCardEditorCommands.h"
 #include "DisplayClusterLightCardEditorStyle.h"
+#include "DisplayClusterLightCardEditorUtils.h"
 #include "DisplayClusterLightCardEditorViewportClient.h"
 #include "DisplayClusterLightCardEditor.h"
 #include "LightCardTemplates/DisplayClusterLightCardTemplate.h"
@@ -18,6 +19,7 @@
 #include "Kismet2/DebuggerCommands.h"
 #include "Slate/SceneViewport.h"
 #include "Styling/AppStyle.h"
+#include "Styling/SlateIconFinder.h"
 
 #define LOCTEXT_NAMESPACE "DisplayClusterLightcardEditorViewport"
 
@@ -675,8 +677,17 @@ TSharedRef<SWidget> SDisplayClusterLightCardEditorViewport::MakeContextMenu()
 	const bool bInShouldCloseWindowAfterMenuSelection = true;
 	FMenuBuilder MenuBuilder(bInShouldCloseWindowAfterMenuSelection, CommandList);
 
-	MenuBuilder.AddMenuEntry(FDisplayClusterLightCardEditorCommands::Get().RemoveLightCard);
-	MenuBuilder.AddMenuEntry(FDisplayClusterLightCardEditorCommands::Get().SaveLightCardTemplate);
+	MenuBuilder.BeginSection("Actors", LOCTEXT("ActorsSection", "Actors"));
+	{
+		MenuBuilder.AddSubMenu(
+			LOCTEXT("PlaceActorsSubMenuLabel", "Place Actor"),
+			LOCTEXT("PlaceActorsSubMenuToolTip", "Add new actors to the stage"),
+			FNewMenuDelegate::CreateSP(this, &SDisplayClusterLightCardEditorViewport::MakePlaceActorsSubMenu));
+
+		MenuBuilder.AddMenuEntry(FDisplayClusterLightCardEditorCommands::Get().RemoveLightCard);
+		MenuBuilder.AddMenuEntry(FDisplayClusterLightCardEditorCommands::Get().SaveLightCardTemplate);
+	}
+	MenuBuilder.EndSection();
 	
 	MenuBuilder.BeginSection("View", LOCTEXT("ViewSection", "View"));
 	{
@@ -698,6 +709,43 @@ TSharedRef<SWidget> SDisplayClusterLightCardEditorViewport::MakeContextMenu()
 	return MenuBuilder.MakeWidget();
 }
 
+void SDisplayClusterLightCardEditorViewport::MakePlaceActorsSubMenu(FMenuBuilder& MenuBuilder)
+{
+	FSlateIcon LightCardIcon = FSlateIconFinder::FindIconForClass(ADisplayClusterLightCardActor::StaticClass());
+
+	// Add custom UI actions here for the AddNew commands so that the newly added actors can be moved to the user's mouse position after the
+	// add operation is performed
+	MenuBuilder.AddMenuEntry(
+		FDisplayClusterLightCardEditorCommands::Get().AddNewFlag->GetLabel(),
+		FDisplayClusterLightCardEditorCommands::Get().AddNewFlag->GetDescription(),
+		LightCardIcon,
+		FUIAction(FExecuteAction::CreateSP(this, &SDisplayClusterLightCardEditorViewport::AddFlagHere)));
+
+	MenuBuilder.AddMenuEntry(
+		FDisplayClusterLightCardEditorCommands::Get().AddNewLightCard->GetLabel(),
+		FDisplayClusterLightCardEditorCommands::Get().AddNewLightCard->GetDescription(),
+		LightCardIcon,
+		FUIAction(FExecuteAction::CreateSP(this, &SDisplayClusterLightCardEditorViewport::AddLightCardHere)));
+
+	TSet<UClass*> StageActorClasses = UE::DisplayClusterLightCardEditorUtils::GetAllStageActorClasses();
+	for (UClass* Class : StageActorClasses)
+	{
+		if (Class == ADisplayClusterLightCardActor::StaticClass())
+		{
+			// Added manually already
+			continue;
+		}
+
+		FText Label = Class->GetDisplayNameText();
+		FSlateIcon StageActorIcon = FSlateIconFinder::FindIconForClass(Class);
+		MenuBuilder.AddMenuEntry(
+			Label,
+			LOCTEXT("AddStageActorHeader", "Add a stage actor to the scene"), 
+			StageActorIcon,
+			FUIAction(FExecuteAction::CreateRaw(this, &SDisplayClusterLightCardEditorViewport::AddStageActorHere, Class)));
+	}
+}
+
 void SDisplayClusterLightCardEditorViewport::SetEditorWidgetMode(FDisplayClusterLightCardEditorWidget::EWidgetMode InWidgetMode)
 {
 	if (ViewportClient.IsValid())
@@ -807,6 +855,60 @@ bool SDisplayClusterLightCardEditorViewport::CanPasteLightCardsHere() const
 	return false;
 }
 
+void SDisplayClusterLightCardEditorViewport::AddLightCardHere()
+{
+	if (LightCardEditorPtr.IsValid())
+	{
+		LightCardEditorPtr.Pin()->AddNewLightCard();
+
+		if (ViewportClient.IsValid())
+		{
+			// Perform the positioning of the pasted light cards on the next scene refresh, since the light card proxies will have not been 
+			// regenerated until then
+			ViewportClient->GetOnNextSceneRefresh().AddLambda([this]()
+			{
+				ViewportClient->MoveSelectedActorsToPixel(FIntPoint(PasteHerePos.X, PasteHerePos.Y));
+			});
+		}
+	}
+}
+
+void SDisplayClusterLightCardEditorViewport::AddFlagHere()
+{
+	if (LightCardEditorPtr.IsValid())
+	{
+		LightCardEditorPtr.Pin()->AddNewFlag();
+
+		if (ViewportClient.IsValid())
+		{
+			// Perform the positioning of the pasted light cards on the next scene refresh, since the light card proxies will have not been 
+			// regenerated until then
+			ViewportClient->GetOnNextSceneRefresh().AddLambda([this]()
+			{
+				ViewportClient->MoveSelectedActorsToPixel(FIntPoint(PasteHerePos.X, PasteHerePos.Y));
+			});
+		}
+	}
+}
+
+void SDisplayClusterLightCardEditorViewport::AddStageActorHere(UClass* InClass)
+{
+	if (LightCardEditorPtr.IsValid())
+	{
+		LightCardEditorPtr.Pin()->AddNewDynamic(InClass);
+
+		if (ViewportClient.IsValid())
+		{
+			// Perform the positioning of the pasted light cards on the next scene refresh, since the light card proxies will have not been 
+			// regenerated until then
+			ViewportClient->GetOnNextSceneRefresh().AddLambda([this]()
+			{
+				ViewportClient->MoveSelectedActorsToPixel(FIntPoint(PasteHerePos.X, PasteHerePos.Y));
+			});
+		}
+	}
+}
+
 void SDisplayClusterLightCardEditorViewport::ToggleLabels()
 {
 	// TODO: Handle CCR
diff --git a/Engine/Plugins/Runtime/nDisplay/Source/DisplayClusterLightCardEditor/Private/Viewport/DisplayClusterLightcardEditorViewport.h b/Engine/Plugins/Runtime/nDisplay/Source/DisplayClusterLightCardEditor/Private/Viewport/DisplayClusterLightcardEditorViewport.h
--- a/Engine/Plugins/Runtime/nDisplay/Source/DisplayClusterLightCardEditor/Private/Viewport/DisplayClusterLightcardEditorViewport.h
+++ b/Engine/Plugins/Runtime/nDisplay/Source/DisplayClusterLightCardEditor/Private/Viewport/DisplayClusterLightcardEditorViewport.h
@@ -68,6 +68,9 @@ private:
 	/** Create the popup context menu */
 	TSharedRef<SWidget> MakeContextMenu();
 
+	/** Create the submenu for placing actors */
+	void MakePlaceActorsSubMenu(FMenuBuilder& MenuBuilder);
+
 	/** Sets the widget mode for the editor viewport to the specified mode */
 	void SetEditorWidgetMode(FDisplayClusterLightCardEditorWidget::EWidgetMode InWidgetMode);
 
@@ -98,6 +101,15 @@ private:
 	/** Determines if there are any light cards in the clipboard that can be pasted */
 	bool CanPasteLightCardsHere() const;
 
+	/** Adds a new light card and moves it to the cached mouse position */
+	void AddLightCardHere();
+
+	/** Adds a new flag and moves it to the cached mouse position */
+	void AddFlagHere();
+
+	/** Adds a new stage actor of the specified class and moves it to the cached mouse position */
+	void AddStageActorHere(UClass* InClass);
+
 	/** Toggles light card labels on or off */
 	void ToggleLabels();
 
diff --git a/Engine/Plugins/Runtime/nDisplay/Source/DisplayClusterOperator/Private/SDisplayClusterOperatorToolbar.cpp b/Engine/Plugins/Runtime/nDisplay/Source/DisplayClusterOperator/Private/SDisplayClusterOperatorToolbar.cpp
--- a/Engine/Plugins/Runtime/nDisplay/Source/DisplayClusterOperator/Private/SDisplayClusterOperatorToolbar.cpp
+++ b/Engine/Plugins/Runtime/nDisplay/Source/DisplayClusterOperator/Private/SDisplayClusterOperatorToolbar.cpp
@@ -60,8 +60,9 @@ void SDisplayClusterOperatorToolbar::Construct(const FArguments& InArgs)
 		];
 
 	const TSharedPtr<FExtender> ToolBarExtender = IDisplayClusterOperator::Get().GetOperatorToolBarExtensibilityManager()->GetAllExtenders();
-	FToolBarBuilder ToolBarBuilder(CommandList, FMultiBoxCustomization::None, ToolBarExtender, true);
-	
+	FSlimHorizontalToolBarBuilder ToolBarBuilder(CommandList, FMultiBoxCustomization::None, ToolBarExtender);
+	ToolBarBuilder.SetStyle(&FAppStyle::Get(), "AssetEditorToolbar");
+
 	const TSharedPtr<SHorizontalBox> RootActorSelectionBox = SNew(SHorizontalBox)
 	+SHorizontalBox::Slot()
 	.AutoWidth()
