nDisplay: Add Hidden in Game column to light card outliner. The type column will now distinguish between a 3d light card and a uv light card.

diff --git a/Engine/Plugins/Runtime/nDisplay/Source/DisplayClusterLightCardEditor/Private/DisplayClusterLightCardEditorStyle.cpp b/Engine/Plugins/Runtime/nDisplay/Source/DisplayClusterLightCardEditor/Private/DisplayClusterLightCardEditorStyle.cpp
new file mode 100644
--- /dev/null
+++ b/Engine/Plugins/Runtime/nDisplay/Source/DisplayClusterLightCardEditor/Private/DisplayClusterLightCardEditorStyle.cpp
@@ -0,0 +1,35 @@
+// Copyright Epic Games, Inc. All Rights Reserved.
+
+#include "DisplayClusterLightCardEditorStyle.h"
+
+#include "Styling/SlateStyleMacros.h"
+
+FDisplayClusterLightCardEditorStyle::FDisplayClusterLightCardEditorStyle(): FSlateStyleSet("DisplayClusterLightCardEditorStyle")
+{
+	const FVector2D Icon40x40(40.0f, 40.0f);
+	const FVector2D Icon16x16(16.0f, 16.0f);
+
+	// Set miscellaneous icons
+	SetContentRoot(FPaths::EnginePluginsDir() / TEXT("Runtime/nDisplay/Content/Icons/OperatorPanel/"));
+	SetCoreContentRoot(FPaths::EngineContentDir() / TEXT("Editor/Slate"));
+
+	Set("DrawPoly", new IMAGE_BRUSH("PolyPath_40x", Icon40x40));
+
+	Set("DisplayClusterLightCardEditor.PasteHere", new CORE_IMAGE_BRUSH_SVG("Starship/Actors/paste-here", Icon16x16));
+
+	Set("DisplayClusterLightCardEditor.ViewOrientationTop", new CORE_IMAGE_BRUSH_SVG("Starship/Common/ViewTop", Icon16x16));
+	Set("DisplayClusterLightCardEditor.ViewOrientationBottom", new CORE_IMAGE_BRUSH_SVG("Starship/Common/ViewBottom", Icon16x16));
+	Set("DisplayClusterLightCardEditor.ViewOrientationLeft", new CORE_IMAGE_BRUSH_SVG("Starship/Common/ViewLeft", Icon16x16));
+	Set("DisplayClusterLightCardEditor.ViewOrientationRight", new CORE_IMAGE_BRUSH_SVG("Starship/Common/ViewRight", Icon16x16));
+	Set("DisplayClusterLightCardEditor.ViewOrientationFront", new CORE_IMAGE_BRUSH_SVG("Starship/Common/ViewFront", Icon16x16));
+	Set("DisplayClusterLightCardEditor.ViewOrientationBack", new CORE_IMAGE_BRUSH_SVG("Starship/Common/ViewBack", Icon16x16));
+
+	Set("UVLightCard", new CORE_IMAGE_BRUSH_SVG("Starship/Common/SetDrawUVs", Icon16x16));
+
+	Set("DisplayClusterLightCardEditor.Labels", new CORE_IMAGE_BRUSH_SVG("Starship/AssetIcons/Font_16", Icon16x16));
+
+	Set("DisplayClusterLightCardEditor.ActorHidden", new IMAGE_BRUSH_SVG("ActorHidden", Icon16x16));
+	Set("DisplayClusterLightCardEditor.ActorNotHidden", new IMAGE_BRUSH_SVG("ActorNotHidden", Icon16x16));
+		
+	FSlateStyleRegistry::RegisterSlateStyle(*this);
+}
diff --git a/Engine/Plugins/Runtime/nDisplay/Source/DisplayClusterLightCardEditor/Private/DisplayClusterLightCardEditorStyle.h b/Engine/Plugins/Runtime/nDisplay/Source/DisplayClusterLightCardEditor/Private/DisplayClusterLightCardEditorStyle.h
--- a/Engine/Plugins/Runtime/nDisplay/Source/DisplayClusterLightCardEditor/Private/DisplayClusterLightCardEditorStyle.h
+++ b/Engine/Plugins/Runtime/nDisplay/Source/DisplayClusterLightCardEditor/Private/DisplayClusterLightCardEditorStyle.h
@@ -4,7 +4,7 @@
 
 #include "CoreMinimal.h"
 #include "Styling/SlateStyleRegistry.h"
-#include "Styling/SlateStyleMacros.h"
+
 #include "Styling/SlateStyle.h"
 
 
@@ -15,33 +15,7 @@ class FDisplayClusterLightCardEditorStyle final : public FSlateStyleSet
 {
 public:
 
-	FDisplayClusterLightCardEditorStyle()
-		: FSlateStyleSet("DisplayClusterLightCardEditorStyle")
-	{
-		const FVector2D Icon40x40(40.0f, 40.0f);
-		const FVector2D Icon16x16(16.0f, 16.0f);
-
-		// Set miscellaneous icons
-		SetContentRoot(FPaths::EnginePluginsDir() / TEXT("Runtime/nDisplay/Content/Icons/OperatorPanel/"));
-		SetCoreContentRoot(FPaths::EngineContentDir() / TEXT("Editor/Slate"));
-
-		Set("DrawPoly", new IMAGE_BRUSH("PolyPath_40x", Icon40x40));
-
-		Set("DisplayClusterLightCardEditor.PasteHere", new CORE_IMAGE_BRUSH_SVG("Starship/Actors/paste-here", Icon16x16));
-
-		Set("DisplayClusterLightCardEditor.ViewOrientationTop", new CORE_IMAGE_BRUSH_SVG("Starship/Common/ViewTop", Icon16x16));
-		Set("DisplayClusterLightCardEditor.ViewOrientationBottom", new CORE_IMAGE_BRUSH_SVG("Starship/Common/ViewBottom", Icon16x16));
-		Set("DisplayClusterLightCardEditor.ViewOrientationLeft", new CORE_IMAGE_BRUSH_SVG("Starship/Common/ViewLeft", Icon16x16));
-		Set("DisplayClusterLightCardEditor.ViewOrientationRight", new CORE_IMAGE_BRUSH_SVG("Starship/Common/ViewRight", Icon16x16));
-		Set("DisplayClusterLightCardEditor.ViewOrientationFront", new CORE_IMAGE_BRUSH_SVG("Starship/Common/ViewFront", Icon16x16));
-		Set("DisplayClusterLightCardEditor.ViewOrientationBack", new CORE_IMAGE_BRUSH_SVG("Starship/Common/ViewBack", Icon16x16));
-
-		Set("UVLightCard", new CORE_IMAGE_BRUSH_SVG("Starship/Common/SetDrawUVs", Icon16x16));
-
-		Set("DisplayClusterLightCardEditor.Labels", new CORE_IMAGE_BRUSH_SVG("Starship/AssetIcons/Font_16", Icon16x16));
-		
-		FSlateStyleRegistry::RegisterSlateStyle(*this);
-	}
+	FDisplayClusterLightCardEditorStyle();
 
 	virtual ~FDisplayClusterLightCardEditorStyle()
 	{
diff --git a/Engine/Plugins/Runtime/nDisplay/Source/DisplayClusterLightCardEditor/Private/Outliner/DisplayClusterLightCardOutlinerColumns.cpp b/Engine/Plugins/Runtime/nDisplay/Source/DisplayClusterLightCardEditor/Private/Outliner/DisplayClusterLightCardOutlinerColumns.cpp
new file mode 100644
--- /dev/null
+++ b/Engine/Plugins/Runtime/nDisplay/Source/DisplayClusterLightCardEditor/Private/Outliner/DisplayClusterLightCardOutlinerColumns.cpp
@@ -0,0 +1,239 @@
+﻿// Copyright Epic Games, Inc. All Rights Reserved.
+
+#include "DisplayClusterLightCardOutlinerColumns.h"
+
+#include "DisplayClusterLightCardEditorStyle.h"
+
+#include "ActorTreeItem.h"
+#include "ISceneOutliner.h"
+#include "ISceneOutlinerMode.h"
+#include "SortHelper.h"
+#include "Widgets/Input/SCheckBox.h"
+
+/**
+ * Sets the 'Hidden in Game' field. Similar to SVisibilityWidget but is implemented as a checkbox
+ * and can produce an undetermined state.
+ */
+class SHiddenInGameVisibilityWidget : public SCompoundWidget
+{
+public:
+	SLATE_BEGIN_ARGS(SHiddenInGameVisibilityWidget) {}
+	SLATE_END_ARGS()
+
+	void Construct(const FArguments& InArgs, TWeakPtr<FSceneOutlinerGutter> InWeakColumn, TWeakPtr<ISceneOutliner> InWeakOutliner,
+		TWeakPtr<ISceneOutlinerTreeItem> InWeakTreeItem, const STableRow<FSceneOutlinerTreeItemPtr>* InRow)
+	{
+		WeakColumn = InWeakColumn;
+		WeakOutliner = InWeakOutliner;
+		WeakTreeItem = InWeakTreeItem;
+		Row = InRow;
+		
+		ActorHiddenBrush = FDisplayClusterLightCardEditorStyle::Get().GetBrush(TEXT("DisplayClusterLightCardEditor.ActorHidden"));
+		ActorNotHiddenBrush = FDisplayClusterLightCardEditorStyle::Get().GetBrush(TEXT("DisplayClusterLightCardEditor.ActorNotHidden"));
+
+		ChildSlot
+		[
+			SNew(SCheckBox)
+			.Padding(FMargin(0.f))
+			.Visibility(this, &SHiddenInGameVisibilityWidget::GetCheckBoxVisibility)
+			.ForegroundColor(this, &SHiddenInGameVisibilityWidget::GetForegroundColor)
+			.IsChecked(this, &SHiddenInGameVisibilityWidget::IsChecked)
+			.OnCheckStateChanged(this, &SHiddenInGameVisibilityWidget::OnCheckBoxStateChanged)
+			.BackgroundImage(FAppStyle::GetNoBrush())
+			.BackgroundHoveredImage(FAppStyle::GetNoBrush())
+			.BackgroundPressedImage(FAppStyle::GetNoBrush())
+			.CheckedImage(ActorHiddenBrush)
+			.CheckedHoveredImage(ActorHiddenBrush)
+			.CheckedPressedImage(ActorHiddenBrush)
+			.UncheckedImage(ActorNotHiddenBrush)
+			.UncheckedHoveredImage(ActorNotHiddenBrush)
+			.UncheckedPressedImage(ActorNotHiddenBrush)
+			.UndeterminedImage(ActorNotHiddenBrush)
+			.UndeterminedHoveredImage(ActorNotHiddenBrush)
+			.UndeterminedPressedImage(ActorNotHiddenBrush)
+		];
+	}
+
+	TWeakPtr<FSceneOutlinerGutter> WeakColumn;
+	TWeakPtr<ISceneOutliner> WeakOutliner;
+	TWeakPtr<ISceneOutlinerTreeItem> WeakTreeItem;
+	const STableRow<FSceneOutlinerTreeItemPtr>* Row;
+
+	const FSlateBrush* ActorNotHiddenBrush;
+	const FSlateBrush* ActorHiddenBrush;
+	
+protected:
+
+	virtual FSlateColor GetForegroundColor() const override
+	{
+		if (WeakOutliner.IsValid() && WeakTreeItem.IsValid())
+		{
+			const TSharedPtr<ISceneOutliner> Outliner = WeakOutliner.Pin();
+			const TSharedPtr<ISceneOutlinerTreeItem> TreeItem = WeakTreeItem.Pin();
+
+			const bool bIsSelected = Outliner->GetTree().IsItemSelected(TreeItem.ToSharedRef());
+
+			// Highlights when hovered
+			if (IsHovered() && !bIsSelected)
+			{
+				return FAppStyle::Get().GetSlateColor("Colors.ForegroundHover");
+			}
+		}
+
+		return FSlateColor::UseForeground();
+	}
+	
+	EVisibility GetCheckBoxVisibility() const
+	{
+		// Only display if hovered, selected, or the checkbox is explicitly set to hidden in game
+		const bool bVisible = Row->IsHovered() || Row->IsSelected() || IsChecked() == ECheckBoxState::Checked;
+		return bVisible ? EVisibility::Visible : EVisibility::Collapsed;
+	}
+
+	ECheckBoxState IsChecked() const
+	{
+		if (WeakTreeItem.IsValid())
+		{
+			return GetCheckStateRecursive(*WeakTreeItem.Pin());
+		}
+
+		return ECheckBoxState::Unchecked;
+	}
+	
+	void OnCheckBoxStateChanged(ECheckBoxState NewState)
+	{
+		if (WeakTreeItem.IsValid())
+		{
+			const bool bHidden = NewState == ECheckBoxState::Checked;
+
+			FScopedTransaction ScopedTransaction(NSLOCTEXT("DisplayClusterLightCardOutlinerHiddenInGameColumn","SetHiddenInGame", "Set Hidden In Game"));
+
+			const STreeView<FSceneOutlinerTreeItemPtr>& Tree = WeakOutliner.Pin()->GetTree();
+
+			// Update all selected items if this item is selected
+			if (Tree.IsItemSelected(WeakTreeItem.Pin()))
+			{
+				for (TSharedPtr<ISceneOutlinerTreeItem>& SelectedItem : Tree.GetSelectedItems())
+				{
+					SetHiddenInGame(*SelectedItem, bHidden);	
+				}
+			}
+			else
+			{
+				SetHiddenInGame(*WeakTreeItem.Pin(), bHidden);	
+			}
+		}
+	}
+	
+	/** Determine the check state based on the given item and all of its children states */
+	static ECheckBoxState GetCheckStateRecursive(const ISceneOutlinerTreeItem& Item)
+	{
+		TOptional<ECheckBoxState> CheckBoxState;
+		
+		if (Item.IsA<FActorTreeItem>())
+		{
+			const FActorTreeItem* ActorTreeItem = Item.CastTo<FActorTreeItem>();
+			if (ActorTreeItem->Actor.IsValid())
+			{
+				CheckBoxState = ActorTreeItem->Actor->IsHidden() ? ECheckBoxState::Checked : ECheckBoxState::Unchecked;
+			}
+		}
+		
+		for (const TWeakPtr<ISceneOutlinerTreeItem>& ChildPtr : Item.GetChildren())
+		{
+			TSharedPtr<ISceneOutlinerTreeItem> Child = ChildPtr.Pin();
+			if (Child.IsValid())
+			{
+				const ECheckBoxState ChildCheckBoxState = GetCheckStateRecursive(*Child);
+				if (CheckBoxState.IsSet() && *CheckBoxState != ChildCheckBoxState)
+				{
+					CheckBoxState = ECheckBoxState::Undetermined;
+					break;
+				}
+				
+				CheckBoxState = ChildCheckBoxState;
+			}
+		}
+
+		return CheckBoxState.Get(ECheckBoxState::Unchecked);
+	}
+
+	static void SetHiddenInGame(ISceneOutlinerTreeItem& Item, const bool bNewValue)
+	{
+		if (Item.IsValid() && Item.IsA<FActorTreeItem>())
+		{
+			const FActorTreeItem* ActorTreeItem = Item.CastTo<FActorTreeItem>();
+			if (ActorTreeItem->Actor.IsValid())
+			{
+				ActorTreeItem->Actor->Modify();
+				ActorTreeItem->Actor->SetActorHiddenInGame(bNewValue);
+			}
+		}
+
+		for (const TWeakPtr<ISceneOutlinerTreeItem>& ChildPtr : Item.GetChildren())
+		{
+			TSharedPtr<ISceneOutlinerTreeItem> Child = ChildPtr.Pin();
+			if (Child.IsValid())
+			{
+				SetHiddenInGame(*Child, bNewValue);
+			}
+		}
+	}
+};
+
+FName FDisplayClusterLightCardOutlinerHiddenInGameColumn::GetID()
+{
+	static FName HiddenInGame("Hidden in Game");
+	return HiddenInGame;
+}
+
+SHeaderRow::FColumn::FArguments FDisplayClusterLightCardOutlinerHiddenInGameColumn::ConstructHeaderRowColumn()
+{
+	return SHeaderRow::Column(GetColumnID())
+		.FixedWidth(24.f)
+		.HAlignHeader(HAlign_Left)
+		.VAlignHeader(VAlign_Center)
+		.HAlignCell(HAlign_Center)
+		.VAlignCell(VAlign_Center)
+		.DefaultTooltip(FText::FromName(GetColumnID()))
+		.HeaderContentPadding(FMargin(4.0f, 0.0f, 0.0f, 0.0f))
+		[
+			SNew(SImage)
+			.ColorAndOpacity(FSlateColor::UseForeground())
+			.Image(FDisplayClusterLightCardEditorStyle::Get().GetBrush(TEXT("DisplayClusterLightCardEditor.ActorNotHidden")))
+		];
+}
+
+const TSharedRef<SWidget> FDisplayClusterLightCardOutlinerHiddenInGameColumn::ConstructRowWidget(
+	FSceneOutlinerTreeItemRef TreeItem, const STableRow<FSceneOutlinerTreeItemPtr>& Row)
+{
+	return SNew(SHorizontalBox)
+		+SHorizontalBox::Slot()
+		.AutoWidth()
+		.VAlign(VAlign_Center)
+		.HAlign(HAlign_Center)
+		[
+			SNew(SHiddenInGameVisibilityWidget, SharedThis(this), WeakOutliner, TreeItem, &Row)
+		];
+}
+
+void FDisplayClusterLightCardOutlinerHiddenInGameColumn::SortItems(TArray<FSceneOutlinerTreeItemPtr>& RootItems, const EColumnSortMode::Type SortMode) const
+{
+	FSceneOutlinerSortHelper<int32, bool>()
+		/** Sort by type first */
+		.Primary([this](const ISceneOutlinerTreeItem& Item){ return WeakOutliner.Pin()->GetMode()->GetTypeSortPriority(Item); }, SortMode)
+		/** Then by hidden in game */
+		.Secondary([](const ISceneOutlinerTreeItem& Item)
+		{
+			if (const FActorTreeItem* ActorTreeItem = Item.CastTo<FActorTreeItem>())
+			{
+				if (ActorTreeItem->Actor.IsValid())
+				{
+					return ActorTreeItem->Actor->IsHidden();
+				}
+			}
+
+			return false;
+		}, SortMode)
+		.Sort(RootItems);
+}
\ No newline at end of file
diff --git a/Engine/Plugins/Runtime/nDisplay/Source/DisplayClusterLightCardEditor/Private/Outliner/DisplayClusterLightCardOutlinerColumns.h b/Engine/Plugins/Runtime/nDisplay/Source/DisplayClusterLightCardEditor/Private/Outliner/DisplayClusterLightCardOutlinerColumns.h
new file mode 100644
--- /dev/null
+++ b/Engine/Plugins/Runtime/nDisplay/Source/DisplayClusterLightCardEditor/Private/Outliner/DisplayClusterLightCardOutlinerColumns.h
@@ -0,0 +1,28 @@
+﻿// Copyright Epic Games, Inc. All Rights Reserved.
+
+#pragma once
+
+#include "SceneOutlinerFwd.h"
+#include "SceneOutlinerGutter.h"
+#include "Templates/SharedPointer.h"
+#include "UObject/NameTypes.h"
+
+class ISceneOutliner;
+class SWidget;
+template<typename ItemType> class STableRow;
+
+class FDisplayClusterLightCardOutlinerHiddenInGameColumn : public FSceneOutlinerGutter
+{
+public:
+	FDisplayClusterLightCardOutlinerHiddenInGameColumn(ISceneOutliner& SceneOutliner) : FSceneOutlinerGutter(SceneOutliner) {}
+	virtual ~FDisplayClusterLightCardOutlinerHiddenInGameColumn() {}
+	static FName GetID();
+	
+	// Begin ISceneOutlinerColumn Implementation
+	virtual FName GetColumnID() override { return GetID(); }
+	virtual SHeaderRow::FColumn::FArguments ConstructHeaderRowColumn() override;
+	virtual const TSharedRef<SWidget> ConstructRowWidget(FSceneOutlinerTreeItemRef TreeItem, const STableRow<FSceneOutlinerTreeItemPtr>& Row) override;
+	virtual bool SupportsSorting() const override { return true; }
+	virtual void SortItems(TArray<FSceneOutlinerTreeItemPtr>& OutItems, const EColumnSortMode::Type SortMode) const override;
+	// End ISceneOutlinerColumn Implementation
+};
\ No newline at end of file
diff --git a/Engine/Plugins/Runtime/nDisplay/Source/DisplayClusterLightCardEditor/Private/Outliner/SDisplayClusterLightCardOutliner.cpp b/Engine/Plugins/Runtime/nDisplay/Source/DisplayClusterLightCardEditor/Private/Outliner/SDisplayClusterLightCardOutliner.cpp
--- a/Engine/Plugins/Runtime/nDisplay/Source/DisplayClusterLightCardEditor/Private/Outliner/SDisplayClusterLightCardOutliner.cpp
+++ b/Engine/Plugins/Runtime/nDisplay/Source/DisplayClusterLightCardEditor/Private/Outliner/SDisplayClusterLightCardOutliner.cpp
@@ -5,6 +5,7 @@
 #include "DisplayClusterLightCardEditorCommands.h"
 #include "DisplayClusterLightCardOutlinerMode.h"
 #include "DisplayClusterLightCardEditor.h"
+#include "DisplayClusterLightCardOutlinerColumns.h"
 
 #include "IDisplayClusterOperator.h"
 
@@ -237,6 +238,11 @@ void SDisplayClusterLightCardOutliner::CreateWorldOutliner()
 	SceneOutlinerOptions.ColumnMap.Add(FSceneOutlinerBuiltInColumnTypes::Gutter(),
 		FSceneOutlinerColumnInfo(ESceneOutlinerColumnVisibility::Visible, 0, FCreateSceneOutlinerColumn(), false, TOptional<float>(),
 				FSceneOutlinerBuiltInColumnTypes::Gutter_Localized()));
+
+	SceneOutlinerOptions.ColumnMap.Add(FDisplayClusterLightCardOutlinerHiddenInGameColumn::GetID(),
+		FSceneOutlinerColumnInfo(FSceneOutlinerColumnInfo(ESceneOutlinerColumnVisibility::Visible,
+			1, FCreateSceneOutlinerColumn::CreateLambda([](ISceneOutliner& InSceneOutliner)
+			{ return MakeShareable(new FDisplayClusterLightCardOutlinerHiddenInGameColumn(InSceneOutliner)); }))));
 	
 	SceneOutlinerOptions.ColumnMap.Add(FSceneOutlinerBuiltInColumnTypes::Label(),
 		FSceneOutlinerColumnInfo(ESceneOutlinerColumnVisibility::Visible, 10, FCreateSceneOutlinerColumn(), false, TOptional<float>(),
@@ -252,6 +258,14 @@ void SDisplayClusterLightCardOutliner::CreateWorldOutliner()
 			{
 				if (const AActor* Actor = ActorItem->Actor.Get())
 				{
+					if (const ADisplayClusterLightCardActor* LightCardActor = Cast<ADisplayClusterLightCardActor>(Actor))
+					{
+						if (LightCardActor->bIsUVLightCard)
+						{
+							return TEXT("UV Light Card");
+						}
+					}
+					
 					return Actor->GetClass()->GetDisplayNameText().ToString();
 				}
 			
