nDisplay: Fixed sluggish selection of light cards in light card editor when color grading drawer is open. Also fixed issue where changing a property in the drawer would cause it to switch selection to the selected CCW or root actor.

diff --git a/Engine/Plugins/Runtime/nDisplay/Source/DisplayClusterColorGrading/Private/Drawer/SDisplayClusterColorGradingDrawer.cpp b/Engine/Plugins/Runtime/nDisplay/Source/DisplayClusterColorGrading/Private/Drawer/SDisplayClusterColorGradingDrawer.cpp
--- a/Engine/Plugins/Runtime/nDisplay/Source/DisplayClusterColorGrading/Private/Drawer/SDisplayClusterColorGradingDrawer.cpp
+++ b/Engine/Plugins/Runtime/nDisplay/Source/DisplayClusterColorGrading/Private/Drawer/SDisplayClusterColorGradingDrawer.cpp
@@ -1032,9 +1032,19 @@ void SDisplayClusterColorGradingDrawer::OnDetailObjectsChanged(const TArray<UObj
 {
 	if (RootActorList.IsValid())
 	{
-		FillRootActorColorGradingList();
+		const bool bIsCCWBeingSelected = Objects.ContainsByPredicate([](UObject* Object) { return Object && Object->IsA<AColorCorrectionWindow>(); });
+		const bool bAreCCWsInList = RootActorColorGradingItems.ContainsByPredicate([](const FDisplayClusterColorGradingListItemRef& ListItem)
+		{
+			return ListItem.IsValid() && ListItem->Actor.IsValid() && ListItem->Actor->IsA<AColorCorrectionWindow>();
+		});
 
-		SetDrawerStateToDefault();
+		// Only need to refresh when CCWs are being selected or unselected
+		if (bIsCCWBeingSelected || bAreCCWsInList)
+		{
+			FillRootActorColorGradingList();
+
+			SetDrawerStateToDefault();
+		}
 	}
 }
 
diff --git a/Engine/Plugins/Runtime/nDisplay/Source/DisplayClusterOperator/Private/DisplayClusterOperatorViewModel.cpp b/Engine/Plugins/Runtime/nDisplay/Source/DisplayClusterOperator/Private/DisplayClusterOperatorViewModel.cpp
--- a/Engine/Plugins/Runtime/nDisplay/Source/DisplayClusterOperator/Private/DisplayClusterOperatorViewModel.cpp
+++ b/Engine/Plugins/Runtime/nDisplay/Source/DisplayClusterOperator/Private/DisplayClusterOperatorViewModel.cpp
@@ -31,9 +31,30 @@ void FDisplayClusterOperatorViewModel::ShowDetailsForObject(UObject* Object)
 
 void FDisplayClusterOperatorViewModel::ShowDetailsForObjects(const TArray<UObject*>& Objects)
 {
-	DetailObjects.Empty();
-	DetailObjects.Append(Objects);
-	DetailObjectsChanged.Broadcast(Objects);
+	bool bDetailObjectsChanged = false;
+
+	if (DetailObjects.Num() == Objects.Num())
+	{
+		for (UObject* NewObject : Objects)
+		{
+			if (NewObject && !DetailObjects.Contains(NewObject))
+			{
+				bDetailObjectsChanged = true;
+				break;
+			}
+		}
+	}
+	else
+	{
+		bDetailObjectsChanged = true;
+	}
+
+	if (bDetailObjectsChanged)
+	{
+		DetailObjects.Empty();
+		DetailObjects.Append(Objects);
+		DetailObjectsChanged.Broadcast(Objects);
+	}
 }
 
 TSharedRef<FTabManager> FDisplayClusterOperatorViewModel::CreateTabManager(const TSharedRef<SDockTab>& MajorTabOwner)
