nDisplay: Added separate view configurations for orthographic projection in light card editor viewport, and made its default orientation to be looking at the top of the stage

diff --git a/Engine/Plugins/Runtime/nDisplay/Source/DisplayClusterLightCardEditor/Private/Viewport/DisplayClusterLightCardEditorViewportClient.cpp b/Engine/Plugins/Runtime/nDisplay/Source/DisplayClusterLightCardEditor/Private/Viewport/DisplayClusterLightCardEditorViewportClient.cpp
--- a/Engine/Plugins/Runtime/nDisplay/Source/DisplayClusterLightCardEditor/Private/Viewport/DisplayClusterLightCardEditorViewportClient.cpp
+++ b/Engine/Plugins/Runtime/nDisplay/Source/DisplayClusterLightCardEditor/Private/Viewport/DisplayClusterLightCardEditorViewportClient.cpp
@@ -112,8 +112,7 @@ FDisplayClusterLightCardEditorViewportClient::FDisplayClusterLightCardEditorView
 	SetRealtime(true);
 	SetShowStats(true);
 
-	ResetFOVs();
-	ResetCameraProjectionTransforms();
+	ResetProjectionViewConfigurations();
 
 	//This seems to be needed to get the correct world time in the preview.
 	SetIsSimulateInEditorViewport(true);
@@ -1681,9 +1680,12 @@ void FDisplayClusterLightCardEditorViewportClient::SetProjectionMode(EDisplayClu
 float FDisplayClusterLightCardEditorViewportClient::GetProjectionModeFOV(EDisplayClusterMeshProjectionType InProjectionMode) const
 {
 	int32 ProjectionModeIndex = (int32)InProjectionMode;
-	if (ProjectionFOVs.Num() > ProjectionModeIndex)
+	if (ProjectionViewConfigurations.Num() > ProjectionModeIndex)
 	{
-		return ProjectionFOVs[ProjectionModeIndex];
+		const bool bIsOrthographic = RenderViewportType != ELevelViewportType::LVT_Perspective;
+		return bIsOrthographic
+			? ProjectionViewConfigurations[ProjectionModeIndex].OrthographicFOV
+			: ProjectionViewConfigurations[ProjectionModeIndex].PerspectiveFOV;
 	}
 	else
 	{
@@ -1694,9 +1696,17 @@ float FDisplayClusterLightCardEditorViewportClient::GetProjectionModeFOV(EDispla
 void FDisplayClusterLightCardEditorViewportClient::SetProjectionModeFOV(EDisplayClusterMeshProjectionType InProjectionMode, float NewFOV)
 {
 	int32 ProjectionModeIndex = (int32)InProjectionMode;
-	if (ProjectionFOVs.Num() > ProjectionModeIndex)
+	if (ProjectionViewConfigurations.Num() > ProjectionModeIndex)
 	{
-		ProjectionFOVs[ProjectionModeIndex] = NewFOV;
+		const bool bIsOrthographic = RenderViewportType != ELevelViewportType::LVT_Perspective;
+		if (bIsOrthographic)
+		{
+			ProjectionViewConfigurations[ProjectionModeIndex].OrthographicFOV = NewFOV;
+		}
+		else
+		{
+			ProjectionViewConfigurations[ProjectionModeIndex].PerspectiveFOV = NewFOV;
+		}
 	}
 	else
 	{
@@ -1724,8 +1734,7 @@ void FDisplayClusterLightCardEditorViewportClient::ResetCamera(bool bLocationOnl
 	
 	SetProjectionMode(GetProjectionMode(), GetRenderViewportType());
 
-	ResetFOVs();
-	ResetCameraProjectionTransforms();
+	ResetProjectionViewConfigurations();
 }
 
 void FDisplayClusterLightCardEditorViewportClient::FrameSelection()
@@ -2642,7 +2651,7 @@ void FDisplayClusterLightCardEditorViewportClient::FixCameraTransform()
 void FDisplayClusterLightCardEditorViewportClient::SaveProjectionCameraTransform()
 {
 	int32 ProjectionModeIndex = (int32)ProjectionMode;
-	if (ProjectionCameraTransforms.Num() > ProjectionModeIndex)
+	if (ProjectionViewConfigurations.Num() > ProjectionModeIndex)
 	{
 		FViewportCameraTransform CameraTransform;
 		CameraTransform.SetRotation(GetViewRotation());
@@ -2652,16 +2661,28 @@ void FDisplayClusterLightCardEditorViewportClient::SaveProjectionCameraTransform
 			CameraTransform.SetLocation(GetViewLocation());
 		}
 
-		ProjectionCameraTransforms[ProjectionModeIndex] = CameraTransform;
+		const bool bIsOrthographic = RenderViewportType != ELevelViewportType::LVT_Perspective;
+		if (bIsOrthographic)
+		{
+			ProjectionViewConfigurations[ProjectionModeIndex].OrthographicTransform = CameraTransform;
+		}
+		else
+		{
+			ProjectionViewConfigurations[ProjectionModeIndex].PerspectiveTransform = CameraTransform;
+		}
 	}
 }
 
 void FDisplayClusterLightCardEditorViewportClient::RestoreProjectionCameraTransform()
 {
 	int32 ProjectionModeIndex = (int32)ProjectionMode;
-	if (ProjectionCameraTransforms.Num() > ProjectionModeIndex)
+	if (ProjectionViewConfigurations.Num() > ProjectionModeIndex)
 	{
-		const FViewportCameraTransform& CameraTransform = ProjectionCameraTransforms[ProjectionModeIndex];
+		const bool bIsOrthographic = RenderViewportType != ELevelViewportType::LVT_Perspective;
+		const FViewportCameraTransform& CameraTransform = bIsOrthographic 
+			? ProjectionViewConfigurations[ProjectionModeIndex].OrthographicTransform
+			: ProjectionViewConfigurations[ProjectionModeIndex].PerspectiveTransform;
+
 		SetViewRotation(CameraTransform.GetRotation());
 
 		if (ProjectionMode == EDisplayClusterMeshProjectionType::UV)
@@ -2671,33 +2692,27 @@ void FDisplayClusterLightCardEditorViewportClient::RestoreProjectionCameraTransf
 	}
 }
 
-void FDisplayClusterLightCardEditorViewportClient::ResetCameraProjectionTransforms()
+void FDisplayClusterLightCardEditorViewportClient::ResetProjectionViewConfigurations()
 {
 	constexpr int32 MaxProjections = 3;
-	if (ProjectionCameraTransforms.Num() < MaxProjections)
+	if (ProjectionViewConfigurations.Num() < MaxProjections)
 	{
-		ProjectionCameraTransforms.AddDefaulted(MaxProjections - ProjectionCameraTransforms.Num());
+		ProjectionViewConfigurations.AddDefaulted(MaxProjections - ProjectionViewConfigurations.Num());
 	}
 
-	ProjectionCameraTransforms[static_cast<int32>(EDisplayClusterMeshProjectionType::Linear)].SetRotation(FVector::ForwardVector.Rotation());
-	ProjectionCameraTransforms[static_cast<int32>(EDisplayClusterMeshProjectionType::Azimuthal)].SetRotation(FVector::UpVector.Rotation());
-	ProjectionCameraTransforms[static_cast<int32>(EDisplayClusterMeshProjectionType::UV)].SetRotation(FVector::ForwardVector.Rotation());
+	ProjectionViewConfigurations[static_cast<int32>(EDisplayClusterMeshProjectionType::Linear)].PerspectiveTransform.SetRotation(FVector::ForwardVector.Rotation());
+	ProjectionViewConfigurations[static_cast<int32>(EDisplayClusterMeshProjectionType::Linear)].OrthographicTransform.SetRotation(FVector::UpVector.Rotation());
+	ProjectionViewConfigurations[static_cast<int32>(EDisplayClusterMeshProjectionType::Linear)].PerspectiveFOV = 90.0f;
+	ProjectionViewConfigurations[static_cast<int32>(EDisplayClusterMeshProjectionType::Linear)].OrthographicFOV = 90.0f;
 
-	ProjectionCameraTransforms[static_cast<int32>(EDisplayClusterMeshProjectionType::UV)].SetLocation(FVector::ZeroVector);
+	ProjectionViewConfigurations[static_cast<int32>(EDisplayClusterMeshProjectionType::Azimuthal)].PerspectiveTransform.SetRotation(FVector::UpVector.Rotation());
+	ProjectionViewConfigurations[static_cast<int32>(EDisplayClusterMeshProjectionType::Azimuthal)].PerspectiveFOV = 130.0f;
 
-	RestoreProjectionCameraTransform();
-}
+	ProjectionViewConfigurations[static_cast<int32>(EDisplayClusterMeshProjectionType::UV)].OrthographicTransform.SetRotation(FVector::ForwardVector.Rotation());
+	ProjectionViewConfigurations[static_cast<int32>(EDisplayClusterMeshProjectionType::UV)].OrthographicTransform.SetLocation(FVector::ZeroVector);
+	ProjectionViewConfigurations[static_cast<int32>(EDisplayClusterMeshProjectionType::UV)].OrthographicFOV = 45.0f;
 
-void FDisplayClusterLightCardEditorViewportClient::ResetFOVs()
-{
-	constexpr int32 MaxFOVs = 3;
-	if (ProjectionFOVs.Num() < MaxFOVs)
-	{
-		ProjectionFOVs.AddDefaulted(MaxFOVs - ProjectionFOVs.Num());
-	}
-	ProjectionFOVs[static_cast<int32>(EDisplayClusterMeshProjectionType::Linear)] = 90.0f;
-	ProjectionFOVs[static_cast<int32>(EDisplayClusterMeshProjectionType::Azimuthal)] = 130.0f;
-	ProjectionFOVs[static_cast<int32>(EDisplayClusterMeshProjectionType::UV)] = 45.0f;
+	RestoreProjectionCameraTransform();
 }
 
 void FDisplayClusterLightCardEditorViewportClient::EnterDrawingLightCardMode()
diff --git a/Engine/Plugins/Runtime/nDisplay/Source/DisplayClusterLightCardEditor/Private/Viewport/DisplayClusterLightCardEditorViewportClient.h b/Engine/Plugins/Runtime/nDisplay/Source/DisplayClusterLightCardEditor/Private/Viewport/DisplayClusterLightCardEditorViewportClient.h
--- a/Engine/Plugins/Runtime/nDisplay/Source/DisplayClusterLightCardEditor/Private/Viewport/DisplayClusterLightCardEditorViewportClient.h
+++ b/Engine/Plugins/Runtime/nDisplay/Source/DisplayClusterLightCardEditor/Private/Viewport/DisplayClusterLightCardEditorViewportClient.h
@@ -232,11 +232,8 @@ private:
 	/** Restores the saved camera transform for the current projection mode */
 	void RestoreProjectionCameraTransform();
 
-	/** Resets the stored camera transforms to default values for each projection mode */
-	void ResetCameraProjectionTransforms();
-
-	/** Resets the camera FOVs */
-	void ResetFOVs();
+	/** Resets the stored view configurations to default values for each projection mode */
+	void ResetProjectionViewConfigurations();
 
 	/** Creates a new light card using a polygon alpha mask as defined by the given mouse positions on the viewport */
 	void CreateDrawnLightCard(const TArray<FIntPoint>& MousePositions);
@@ -292,6 +289,16 @@ private:
 		static FSpriteProxy FromBillboard(const UBillboardComponent* InBillboardComponent);
 	};
 
+	/** Stores view configuration values (view transform and field of view) for the supported projection modes */
+	struct FProjectionViewConfiguration
+	{
+	public:
+		FViewportCameraTransform PerspectiveTransform = FViewportCameraTransform();
+		FViewportCameraTransform OrthographicTransform = FViewportCameraTransform();
+		float PerspectiveFOV = 0.0f;
+		float OrthographicFOV = 0.0f;
+	};
+
 	TArray<FActorProxy> ActorProxies;
 	TArray<TWeakObjectPtr<UBillboardComponent>> BillboardComponentProxies;
 	TArray<FDisplayClusterWeakStageActorPtr> SelectedActors;
@@ -344,11 +351,8 @@ private:
 	/** The component of the root actor that is acting as the projection origin. Can be either the root component (stage origin) or a view origin component */
 	TWeakObjectPtr<USceneComponent> ProjectionOriginComponent;
 
-	/** Stores each projection mode's field of view separately */
-	TArray<float> ProjectionFOVs;
-
-	/** Stores each projecion mode's camera transforms separately so they can be restored when the projection mode is selected  */
-	TArray<FViewportCameraTransform> ProjectionCameraTransforms;
+	/** Stores each projecion mode's view configuration separately so they can be restored when the projection mode is selected  */
+	TArray<FProjectionViewConfiguration> ProjectionViewConfigurations;
 
 	/** The increment to change the FOV by when using the scroll wheel */
 	float FOVScrollIncrement = 5.0f;
