nDisplay: Added an additional list in the color grading drawer dedicated to color correction regions. Also fixed up the appearance of the color gradable object lists to match conventions elsewhere in the editor.

diff --git a/Engine/Plugins/Runtime/nDisplay/Source/DisplayClusterColorGrading/Private/Drawer/SDisplayClusterColorGradingDrawer.cpp b/Engine/Plugins/Runtime/nDisplay/Source/DisplayClusterColorGrading/Private/Drawer/SDisplayClusterColorGradingDrawer.cpp
--- a/Engine/Plugins/Runtime/nDisplay/Source/DisplayClusterColorGrading/Private/Drawer/SDisplayClusterColorGradingDrawer.cpp
+++ b/Engine/Plugins/Runtime/nDisplay/Source/DisplayClusterColorGrading/Private/Drawer/SDisplayClusterColorGradingDrawer.cpp
@@ -48,29 +48,19 @@ SDisplayClusterColorGradingDrawer::~SDisplayClusterColorGradingDrawer()
 
 	GEditor->UnregisterForUndo(this);
 
-	for (const FDisplayClusterColorGradingListItemRef& LevelColorGradingItem : LevelColorGradingItems)
+	for (const TArray<FDisplayClusterColorGradingListItemRef>& ColorGradingItemList : ColorGradingItemLists)
 	{
-		if (LevelColorGradingItem->Component.IsValid())
+		for (const FDisplayClusterColorGradingListItemRef& ColorGradingItem : ColorGradingItemList)
 		{
-			UnbindBlueprintCompiledDelegate(LevelColorGradingItem->Component->GetClass());
-		}
-
-		if (LevelColorGradingItem->Actor.IsValid())
-		{
-			UnbindBlueprintCompiledDelegate(LevelColorGradingItem->Actor->GetClass());
-		}
-	}
-
-	for (const FDisplayClusterColorGradingListItemRef& RootActorColorGradingItem : RootActorColorGradingItems)
-	{
-		if (RootActorColorGradingItem->Component.IsValid())
-		{
-			UnbindBlueprintCompiledDelegate(RootActorColorGradingItem->Component->GetClass());
-		}
+			if (ColorGradingItem->Component.IsValid())
+			{
+				UnbindBlueprintCompiledDelegate(ColorGradingItem->Component->GetClass());
+			}
 
-		if (RootActorColorGradingItem->Actor.IsValid())
-		{
-			UnbindBlueprintCompiledDelegate(RootActorColorGradingItem->Actor->GetClass());
+			if (ColorGradingItem->Actor.IsValid())
+			{
+				UnbindBlueprintCompiledDelegate(ColorGradingItem->Actor->GetClass());
+			}
 		}
 	}
 }
@@ -91,8 +81,46 @@ void SDisplayClusterColorGradingDrawer::Construct(const FArguments& InArgs, bool
 
 	GEditor->RegisterForUndo(this);
 
-	FillLevelColorGradingList();
-	FillRootActorColorGradingList();
+	RefreshColorGradingLists();
+	
+	TSharedRef<SVerticalBox> ColorGradingObjectListBox = SNew(SVerticalBox);
+
+	for (int32 Index = 0; Index < ColorGradingItemLists.Num(); ++Index)
+	{
+		TSharedPtr<SDisplayClusterColorGradingObjectList> NewListView;
+		ColorGradingObjectListBox->AddSlot()
+			.AutoHeight()
+			.Padding(0.0f, 2.0f, 0.0f, 0.0f)
+			[
+				SNew(SExpandableArea)
+				.BorderImage(FAppStyle::Get().GetBrush("Brushes.Header"))
+				.BodyBorderImage(FAppStyle::Get().GetBrush("Brushes.Recessed"))
+				.HeaderPadding(FMargin(4.0f, 2.0f))
+				.InitiallyCollapsed(false)
+				.AllowAnimatedTransition(false)
+				.Visibility_Lambda([this, Index]() { return ColorGradingItemLists[Index].Num() ? EVisibility::Visible : EVisibility::Collapsed; })
+				.HeaderContent()
+				[
+					SNew(SBox)
+					.HeightOverride(24.0f)
+					.VAlign(VAlign_Center)
+					[
+						SNew(STextBlock)
+						.Text(this, &SDisplayClusterColorGradingDrawer::GetColorGradingListName, Index)
+						.TextStyle(FAppStyle::Get(), "ButtonText")
+						.Font(FAppStyle::Get().GetFontStyle("NormalFontBold"))
+					]
+				]
+				.BodyContent()
+				[
+					SAssignNew(NewListView, SDisplayClusterColorGradingObjectList)
+					.ColorGradingItemsSource(&ColorGradingItemLists[Index])
+					.OnSelectionChanged(this, &SDisplayClusterColorGradingDrawer::OnListSelectionChanged)
+				]
+			];
+
+		ColorGradingObjectListViews.Add(NewListView);
+	}
 
 	ChildSlot
 	[
@@ -184,57 +212,7 @@ void SDisplayClusterColorGradingDrawer::Construct(const FArguments& InArgs, bool
 							SNew(SScrollBox)
 							+ SScrollBox::Slot()
 							[
-								SNew(SVerticalBox)
-
-								+SVerticalBox::Slot()
-								.AutoHeight()
-								[
-									SNew(SExpandableArea)
-									.BorderImage(FAppStyle::Get().GetBrush("Brushes.Header"))
-									.BodyBorderImage(FAppStyle::Get().GetBrush("Brushes.Recessed"))
-									.HeaderPadding(FMargin(4.0f, 2.0f))
-									.InitiallyCollapsed(false)
-									.AllowAnimatedTransition(false)
-									.Visibility_Lambda([this]() { return LevelColorGradingItems.Num() ? EVisibility::Visible : EVisibility::Collapsed; })
-									.HeaderContent()
-									[
-										SNew(STextBlock)
-										.Text(this, &SDisplayClusterColorGradingDrawer::GetCurrentLevelName)
-										.TextStyle(FAppStyle::Get(), "ButtonText")
-										.Font(FAppStyle::Get().GetFontStyle("NormalFontBold"))
-									]
-									.BodyContent()
-									[
-										SAssignNew(LevelActorsList, SDisplayClusterColorGradingObjectList)
-										.ColorGradingItemsSource(&LevelColorGradingItems)
-										.OnSelectionChanged(this, &SDisplayClusterColorGradingDrawer::OnListSelectionChanged)
-									]
-								]
-
-								+SVerticalBox::Slot()
-								.AutoHeight()
-								[
-									SNew(SExpandableArea)
-									.BorderImage(FAppStyle::Get().GetBrush("Brushes.Header"))
-									.BodyBorderImage(FAppStyle::Get().GetBrush("Brushes.Recessed"))
-									.HeaderPadding(FMargin(4.0f, 2.0f))
-									.InitiallyCollapsed(false)
-									.AllowAnimatedTransition(false)
-									.Visibility_Lambda([this]() { return RootActorColorGradingItems.Num() ? EVisibility::Visible : EVisibility::Collapsed; })
-									.HeaderContent()
-									[
-										SNew(STextBlock)
-										.Text(this, &SDisplayClusterColorGradingDrawer::GetCurrentRootActorName)
-										.TextStyle(FAppStyle::Get(), "ButtonText")
-										.Font(FAppStyle::Get().GetFontStyle("NormalFontBold"))
-									]
-									.BodyContent()
-									[
-										SAssignNew(RootActorList, SDisplayClusterColorGradingObjectList)
-										.ColorGradingItemsSource(&RootActorColorGradingItems)
-										.OnSelectionChanged(this, &SDisplayClusterColorGradingDrawer::OnListSelectionChanged)
-									]
-								]
+								ColorGradingObjectListBox
 							]
 						]
 					]
@@ -365,8 +343,7 @@ void SDisplayClusterColorGradingDrawer::Refresh(bool bPreserveDrawerState)
 
 	ColorGradingDataModel->Reset();
 
-	FillLevelColorGradingList();
-	FillRootActorColorGradingList();
+	RefreshColorGradingLists();
 
 	if (ColorWheelPanel.IsValid())
 	{
@@ -437,41 +414,24 @@ FDisplayClusterColorGradingDrawerState SDisplayClusterColorGradingDrawer::GetDra
 		DetailsPanel->GetDrawerState(DrawerState);
 	}
 
-	if (LevelActorsList.IsValid())
+	for (const TSharedPtr<SDisplayClusterColorGradingObjectList>& ListView : ColorGradingObjectListViews)
 	{
-		TArray<FDisplayClusterColorGradingListItemRef> SelectedLevelItems = LevelActorsList->GetSelectedItems();
-
-		for (const FDisplayClusterColorGradingListItemRef& SelectedLevelItem : SelectedLevelItems)
+		if (ListView.IsValid())
 		{
-			if (SelectedLevelItem.IsValid())
-			{
-				if (SelectedLevelItem->Component.IsValid())
-				{
-					DrawerState.SelectedObjects.Add(SelectedLevelItem->Component);
-				}
-				else if (SelectedLevelItem->Actor.IsValid())
-				{
-					DrawerState.SelectedObjects.Add(SelectedLevelItem->Actor);
-				}
-			}
-		}
-	}
+			TArray<FDisplayClusterColorGradingListItemRef> SelectedItems = ListView->GetSelectedItems();
 
-	if (RootActorList.IsValid())
-	{
-		TArray<FDisplayClusterColorGradingListItemRef> SelectedRootActorItems = RootActorList->GetSelectedItems();
-
-		for (const FDisplayClusterColorGradingListItemRef& SelectedRootActorItem : SelectedRootActorItems)
-		{
-			if (SelectedRootActorItem.IsValid())
+			for (const FDisplayClusterColorGradingListItemRef& SelectedItem : SelectedItems)
 			{
-				if (SelectedRootActorItem->Component.IsValid())
+				if (SelectedItem.IsValid())
 				{
-					DrawerState.SelectedObjects.Add(SelectedRootActorItem->Component);
-				}
-				else if (SelectedRootActorItem->Actor.IsValid())
-				{
-					DrawerState.SelectedObjects.Add(SelectedRootActorItem->Actor);
+					if (SelectedItem->Component.IsValid())
+					{
+						DrawerState.SelectedObjects.Add(SelectedItem->Component);
+					}
+					else if (SelectedItem->Actor.IsValid())
+					{
+						DrawerState.SelectedObjects.Add(SelectedItem->Actor);
+					}
 				}
 			}
 		}
@@ -482,49 +442,38 @@ FDisplayClusterColorGradingDrawerState SDisplayClusterColorGradingDrawer::GetDra
 
 void SDisplayClusterColorGradingDrawer::SetDrawerState(const FDisplayClusterColorGradingDrawerState& InDrawerState)
 {
-	TArray<FDisplayClusterColorGradingListItemRef> LevelItemsToSelect;
-	TArray<FDisplayClusterColorGradingListItemRef> RootActorItemsToSelect;
+	TArray<TArray<FDisplayClusterColorGradingListItemRef>> ItemsToSelect;
+	ItemsToSelect.AddDefaulted(ColorGradingItemLists.Num());
+
 	for (const TWeakObjectPtr<UObject>& SelectedObject : InDrawerState.SelectedObjects)
 	{
 		if (SelectedObject.IsValid())
 		{
-			FDisplayClusterColorGradingListItemRef* FoundColorGradingItemPtr = LevelColorGradingItems.FindByPredicate([&SelectedObject](const FDisplayClusterColorGradingListItemRef& ColorGradingItem)
+			auto FindColorGraidngItem = [&SelectedObject](const FDisplayClusterColorGradingListItemRef& ColorGradingItem)
 			{
 				return ColorGradingItem->Actor == SelectedObject || ColorGradingItem->Component == SelectedObject;
-			});
+			};
 
-			if (FoundColorGradingItemPtr)
+			for (int32 Index = 0; Index < ColorGradingItemLists.Num(); ++Index)
 			{
-				LevelItemsToSelect.Add(*FoundColorGradingItemPtr);
-			}
-			else
-			{
-				FoundColorGradingItemPtr = RootActorColorGradingItems.FindByPredicate([&SelectedObject](const FDisplayClusterColorGradingListItemRef& ColorGradingItem)
-				{
-					return ColorGradingItem->Actor == SelectedObject || ColorGradingItem->Component == SelectedObject;
-				});
-
-				if (FoundColorGradingItemPtr)
+				if (FDisplayClusterColorGradingListItemRef* FoundItem = ColorGradingItemLists[Index].FindByPredicate(FindColorGraidngItem))
 				{
-					RootActorItemsToSelect.Add(*FoundColorGradingItemPtr);
+					ItemsToSelect[Index].Add(*FoundItem);
+					break;
 				}
 			}
 		}
 	}
 
 	// TODO: For now, since we don't support multiple color grading items selected at once, ensure either a level item or a root actor item is seleccted, not both
-	if (LevelItemsToSelect.Num())
+	for (int32 Index = 0; Index < ColorGradingItemLists.Num(); ++Index)
 	{
-		if (LevelActorsList.IsValid())
+		if (ItemsToSelect[Index].Num())
 		{
-			LevelActorsList->SetSelectedItems(LevelItemsToSelect);
-		}
-	}
-	else if (RootActorItemsToSelect.Num())
-	{
-		if (RootActorList.IsValid())
-		{
-			RootActorList->SetSelectedItems(RootActorItemsToSelect);
+			if (ColorGradingObjectListViews.Num() > Index && ColorGradingObjectListViews[Index].IsValid())
+			{
+				ColorGradingObjectListViews[Index]->SetSelectedItems(ItemsToSelect[Index]);
+			}
 		}
 	}
 
@@ -546,37 +495,44 @@ void SDisplayClusterColorGradingDrawer::SetDrawerState(const FDisplayClusterColo
 
 void SDisplayClusterColorGradingDrawer::SetDrawerStateToDefault()
 {
-	if (RootActorList.IsValid() && RootActorColorGradingItems.Num() > 0)
+	constexpr int32 RootActorListIndex = 1;
+	if (ColorGradingItemLists.Num() > RootActorListIndex && ColorGradingObjectListViews.Num() > RootActorListIndex)
 	{
-		auto FindCCW = [](const TWeakObjectPtr<UObject>& Object)
+		TSharedPtr<SDisplayClusterColorGradingObjectList>& RootActorList = ColorGradingObjectListViews[RootActorListIndex];
+		TArray<FDisplayClusterColorGradingListItemRef>& RootActorColorGradingItems = ColorGradingItemLists[RootActorListIndex];
+
+		if (RootActorList.IsValid() && RootActorColorGradingItems.Num() > 0)
 		{
-			return Object.IsValid() && Object->IsA<AColorCorrectionWindow>();
-		};
+			auto FindCCW = [](const TWeakObjectPtr<UObject>& Object)
+			{
+				return Object.IsValid() && Object->IsA<AColorCorrectionWindow>();
+			};
 
-		const bool bIsCCWSelected = OperatorViewModel->GetDetailObjects().ContainsByPredicate(FindCCW);
+			const bool bIsCCWSelected = OperatorViewModel->GetDetailObjects().ContainsByPredicate(FindCCW);
 
-		// If there is a CCW in the currently selected detail objects, automatically select the last (most recently selected) one
-		if (bIsCCWSelected)
-		{
-			int32 CCWIndex = OperatorViewModel->GetDetailObjects().FindLastByPredicate(FindCCW);
-			if (CCWIndex > INDEX_NONE)
+			// If there is a CCW in the currently selected detail objects, automatically select the last (most recently selected) one
+			if (bIsCCWSelected)
 			{
-				AColorCorrectionWindow* CCW = Cast<AColorCorrectionWindow>(OperatorViewModel->GetDetailObjects()[CCWIndex]);
-				FDisplayClusterColorGradingListItemRef* ListItemPtr = RootActorColorGradingItems.FindByPredicate([CCW](const FDisplayClusterColorGradingListItemRef& ListItem)
+				int32 CCWIndex = OperatorViewModel->GetDetailObjects().FindLastByPredicate(FindCCW);
+				if (CCWIndex > INDEX_NONE)
 				{
-					return ListItem.IsValid() && ListItem->Actor.Get() == CCW;
-				});
+					AColorCorrectionWindow* CCW = Cast<AColorCorrectionWindow>(OperatorViewModel->GetDetailObjects()[CCWIndex]);
+					FDisplayClusterColorGradingListItemRef* ListItemPtr = RootActorColorGradingItems.FindByPredicate([CCW](const FDisplayClusterColorGradingListItemRef& ListItem)
+					{
+						return ListItem.IsValid() && ListItem->Actor.Get() == CCW;
+					});
 
-				if (ListItemPtr)
-				{
-					RootActorList->SetSelectedItems({ *ListItemPtr });
+					if (ListItemPtr)
+					{
+						RootActorList->SetSelectedItems({ *ListItemPtr });
+					}
 				}
 			}
-		}
-		else
-		{
-			// The nDisplay stage actor is always the first item in the root actor color grading items list, so set that as the currently selected item
-			RootActorList->SetSelectedItems({ RootActorColorGradingItems[0] });
+			else
+			{
+				// The nDisplay stage actor is always the first item in the root actor color grading items list, so set that as the currently selected item
+				RootActorList->SetSelectedItems({ RootActorColorGradingItems[0] });
+			}
 		}
 	}
 }
@@ -616,6 +572,24 @@ TSharedRef<SWidget> SDisplayClusterColorGradingDrawer::CreateDockInLayoutButton(
 	return SNullWidget::NullWidget;
 }
 
+FText SDisplayClusterColorGradingDrawer::GetColorGradingListName(int32 ListIndex) const
+{
+	switch (ListIndex)
+	{
+	case 0:
+		return GetCurrentLevelName();
+		
+	case 1:
+		return GetCurrentRootActorName();
+
+	case 2:
+		return LOCTEXT("ColorCorrectionRegionsListLabel", "Color Correction Regions");
+
+	default:
+		return FText::GetEmpty();
+	}
+}
+
 FText SDisplayClusterColorGradingDrawer::GetCurrentLevelName() const
 {
 	if (OperatorViewModel->HasRootActor())
@@ -667,23 +641,61 @@ void SDisplayClusterColorGradingDrawer::UnbindBlueprintCompiledDelegate(const UC
 		IsEnabledProperty = bIsEnabled; \
 	})
 
-void SDisplayClusterColorGradingDrawer::FillLevelColorGradingList()
+void SDisplayClusterColorGradingDrawer::RefreshColorGradingLists()
+{
+	constexpr int32 NumLists = 3;
+	if (ColorGradingItemLists.Num() < NumLists)
+	{
+		ColorGradingItemLists.AddDefaulted(NumLists - ColorGradingItemLists.Num());
+	}
+
+	for (int32 Index = 0; Index < NumLists; ++Index)
+	{
+		RefreshColorGradingList(Index);
+	}
+}
+
+void SDisplayClusterColorGradingDrawer::RefreshColorGradingList(int32 Index)
 {
-	for (const FDisplayClusterColorGradingListItemRef& LevelColorGradingItem : LevelColorGradingItems)
+	TArray<FDisplayClusterColorGradingListItemRef>& ItemList = ColorGradingItemLists[Index];
+	for (const FDisplayClusterColorGradingListItemRef& Item : ItemList)
 	{
-		if (LevelColorGradingItem->Component.IsValid())
+		if (Item->Component.IsValid())
 		{
-			UnbindBlueprintCompiledDelegate(LevelColorGradingItem->Component->GetClass());
+			UnbindBlueprintCompiledDelegate(Item->Component->GetClass());
 		}
 
-		if (LevelColorGradingItem->Actor.IsValid())
+		if (Item->Actor.IsValid())
 		{
-			UnbindBlueprintCompiledDelegate(LevelColorGradingItem->Actor->GetClass());
+			UnbindBlueprintCompiledDelegate(Item->Actor->GetClass());
 		}
 	}
 
-	LevelColorGradingItems.Empty();
+	ItemList.Empty();
+
+	switch (Index)
+	{
+	case 0:
+		FillLevelColorGradingList(ItemList);
+		break;
+
+	case 1:
+		FillRootActorColorGradingList(ItemList);
+		break;
+
+	case 2:
+		FillColorCorrectionRegionColorGradingList(ItemList);
+		break;
+	}
+
+	if (ColorGradingObjectListViews.Num() > Index && ColorGradingObjectListViews[Index].IsValid())
+	{
+		ColorGradingObjectListViews[Index]->RefreshList();
+	}
+}
 
+void SDisplayClusterColorGradingDrawer::FillLevelColorGradingList(TArray<FDisplayClusterColorGradingListItemRef>& List)
+{
 	if (ADisplayClusterRootActor* RootActor = OperatorViewModel->GetRootActor())
 	{
 		if (UWorld* World = RootActor->GetWorld())
@@ -717,60 +729,14 @@ void SDisplayClusterColorGradingDrawer::FillLevelColorGradingList()
 				}
 
 				SortedPPVs.Sort(AlphabeticalSort);
-				LevelColorGradingItems.Append(SortedPPVs);
-			}
-
-			// Add all color correction regions that are in the same world as the stage actor
-			{
-				TArray<FDisplayClusterColorGradingListItemRef> SortedCCRs;
-				for (TActorIterator<AColorCorrectRegion> CCRIter(World); CCRIter; ++CCRIter)
-				{
-					AColorCorrectRegion* ColorCorrectRegion = *CCRIter;
-
-					// Skip color correct windows, since those are stage managed actors and will be handled in a separate way
-					if (ColorCorrectRegion->IsA<AColorCorrectionWindow>())
-					{
-						continue;
-					}
-
-					BindBlueprintCompiledDelegate(ColorCorrectRegion->GetClass());
-
-					FDisplayClusterColorGradingListItemRef CCRListItemRef = MakeShared<FDisplayClusterColorGradingListItem>(ColorCorrectRegion);
-					CCRListItemRef->IsItemEnabled = CREATE_IS_ENABLED_LAMBDA(ColorCorrectRegion, ColorCorrectRegion->Enabled);
-					CCRListItemRef->OnItemEnabledChanged = CREATE_ON_ENABLED_CHANGED_LAMBDA(ColorCorrectRegion, ColorCorrectRegion->Enabled);
-
-					SortedCCRs.Add(CCRListItemRef);
-				}
-
-				SortedCCRs.Sort(AlphabeticalSort);
-				LevelColorGradingItems.Append(SortedCCRs);
+				List.Append(SortedPPVs);
 			}
 		}
 	}
-
-	if (LevelActorsList.IsValid())
-	{
-		LevelActorsList->RefreshList();
-	}
 }
 
-void SDisplayClusterColorGradingDrawer::FillRootActorColorGradingList()
+void SDisplayClusterColorGradingDrawer::FillRootActorColorGradingList(TArray<FDisplayClusterColorGradingListItemRef>& List)
 {
-	for (const FDisplayClusterColorGradingListItemRef& RootActorColorGradingItem : RootActorColorGradingItems)
-	{
-		if (RootActorColorGradingItem->Component.IsValid())
-		{
-			UnbindBlueprintCompiledDelegate(RootActorColorGradingItem->Component->GetClass());
-		}
-
-		if (RootActorColorGradingItem->Actor.IsValid())
-		{
-			UnbindBlueprintCompiledDelegate(RootActorColorGradingItem->Actor->GetClass());
-		}
-	}
-
-	RootActorColorGradingItems.Empty();
-
 	if (ADisplayClusterRootActor* RootActor = OperatorViewModel->GetRootActor())
 	{
 		BindBlueprintCompiledDelegate(RootActor->GetClass());
@@ -779,7 +745,7 @@ void SDisplayClusterColorGradingDrawer::FillRootActorColorGradingList()
 		RootActorListItemRef->IsItemEnabled = CREATE_IS_ENABLED_LAMBDA(RootActor, RootActor->GetConfigData()->StageSettings.EnableColorGrading);
 		RootActorListItemRef->OnItemEnabledChanged = CREATE_ON_ENABLED_CHANGED_LAMBDA(RootActor, RootActor->GetConfigData()->StageSettings.EnableColorGrading);
 
-		RootActorColorGradingItems.Add(RootActorListItemRef);
+		List.Add(RootActorListItemRef);
 
 		auto AlphabeticalSort = [](const FDisplayClusterColorGradingListItemRef& A, const FDisplayClusterColorGradingListItemRef& B)
 		{
@@ -808,7 +774,7 @@ void SDisplayClusterColorGradingDrawer::FillRootActorColorGradingList()
 			});
 
 			SortedICVFXCameras.Sort(AlphabeticalSort);
-			RootActorColorGradingItems.Append(SortedICVFXCameras);
+			List.Append(SortedICVFXCameras);
 		}
 
 		// Add any color correction window that is currently selected in the operator's details panel to the list
@@ -834,13 +800,56 @@ void SDisplayClusterColorGradingDrawer::FillRootActorColorGradingList()
 			}
 
 			SortedCCWs.Sort(AlphabeticalSort);
-			RootActorColorGradingItems.Append(SortedCCWs);
+			List.Append(SortedCCWs);
 		}
 	}
+}
 
-	if (RootActorList.IsValid())
+void SDisplayClusterColorGradingDrawer::FillColorCorrectionRegionColorGradingList(TArray<FDisplayClusterColorGradingListItemRef>& List)
+{
+	if (ADisplayClusterRootActor* RootActor = OperatorViewModel->GetRootActor())
 	{
-		RootActorList->RefreshList();
+		if (UWorld* World = RootActor->GetWorld())
+		{
+			// Sorter that sorts the list items alphabetically by their display name
+			auto AlphabeticalSort = [](const FDisplayClusterColorGradingListItemRef& A, const FDisplayClusterColorGradingListItemRef& B)
+			{
+				if (A.IsValid() && B.IsValid())
+				{
+					return *A < *B;
+				}
+				else
+				{
+					return false;
+				}
+			};
+
+			// Add all color correction regions that are in the same world as the stage actor
+			{
+				TArray<FDisplayClusterColorGradingListItemRef> SortedCCRs;
+				for (TActorIterator<AColorCorrectRegion> CCRIter(World); CCRIter; ++CCRIter)
+				{
+					AColorCorrectRegion* ColorCorrectRegion = *CCRIter;
+
+					// Skip color correct windows, since those are stage managed actors and will be handled in a separate way
+					if (ColorCorrectRegion->IsA<AColorCorrectionWindow>())
+					{
+						continue;
+					}
+
+					BindBlueprintCompiledDelegate(ColorCorrectRegion->GetClass());
+
+					FDisplayClusterColorGradingListItemRef CCRListItemRef = MakeShared<FDisplayClusterColorGradingListItem>(ColorCorrectRegion);
+					CCRListItemRef->IsItemEnabled = CREATE_IS_ENABLED_LAMBDA(ColorCorrectRegion, ColorCorrectRegion->Enabled);
+					CCRListItemRef->OnItemEnabledChanged = CREATE_ON_ENABLED_CHANGED_LAMBDA(ColorCorrectRegion, ColorCorrectRegion->Enabled);
+
+					SortedCCRs.Add(CCRListItemRef);
+				}
+
+				SortedCCRs.Sort(AlphabeticalSort);
+				List.Append(SortedCCRs);
+			}
+		}
 	}
 }
 
@@ -933,13 +942,18 @@ void SDisplayClusterColorGradingDrawer::OnObjectsReplaced(const TMap<UObject*, U
 	{
 		if (Pair.Key && Pair.Value)
 		{
-			// Must use GetEvenIfUnreachable on the weak pointers here because most of the time, the objects being replaced have already been marked for GC, and TWeakObjectPtr
-			// will return nullptr from Get on GC-marked objects
-			FDisplayClusterColorGradingListItemRef* FoundColorGradingItemPtr =
-				RootActorColorGradingItems.FindByPredicate([&Pair](const FDisplayClusterColorGradingListItemRef& ColorGradingItem)
+			FDisplayClusterColorGradingListItemRef* FoundColorGradingItemPtr = nullptr;
+
+			constexpr int32 RootActorListIndex = 1;
+			if (ColorGradingItemLists.Num() > RootActorListIndex)
+			{
+				// Must use GetEvenIfUnreachable on the weak pointers here because most of the time, the objects being replaced have already been marked for GC, and TWeakObjectPtr
+				// will return nullptr from Get on GC-marked objects
+				FoundColorGradingItemPtr = ColorGradingItemLists[RootActorListIndex].FindByPredicate([&Pair](const FDisplayClusterColorGradingListItemRef& ColorGradingItem)
 				{
 					return ColorGradingItem->Actor.GetEvenIfUnreachable() == Pair.Key || ColorGradingItem->Component.GetEvenIfUnreachable() == Pair.Key;
 				});
+			}
 
 			if (FoundColorGradingItemPtr)
 			{
@@ -970,14 +984,12 @@ void SDisplayClusterColorGradingDrawer::OnObjectsReplaced(const TMap<UObject*, U
 	}
 	else if (bNeedsListRefresh)
 	{
-		if (LevelActorsList.IsValid())
-		{
-			LevelActorsList->RefreshList();
-		}
-
-		if (RootActorList.IsValid())
+		for (const TSharedPtr<SDisplayClusterColorGradingObjectList>& ListView : ColorGradingObjectListViews)
 		{
-			RootActorList->RefreshList();
+			if (ListView.IsValid())
+			{
+				ListView->RefreshList();
+			}
 		}
 	}
 }
@@ -1008,11 +1020,14 @@ void SDisplayClusterColorGradingDrawer::OnLevelActorDeleted(AActor* Actor)
 		return ColorGradingItem->Actor.GetEvenIfUnreachable() == Actor;
 	};
 
-	// Only refresh if the actor being deleted is being referenced by the drawer
-	if (RootActorColorGradingItems.ContainsByPredicate(ContainsActorRef) || LevelColorGradingItems.ContainsByPredicate(ContainsActorRef))
+	for (const TArray<FDisplayClusterColorGradingListItemRef>& ItemList : ColorGradingItemLists)
 	{
-		// Must wait for next tick to refresh because the actor has not actually been removed from the level at this point
-		bRefreshOnNextTick = true;
+		if (ItemList.ContainsByPredicate(ContainsActorRef))
+		{
+			// Must wait for next tick to refresh because the actor has not actually been removed from the level at this point
+			bRefreshOnNextTick = true;
+			break;
+		}
 	}
 }
 
@@ -1030,20 +1045,27 @@ void SDisplayClusterColorGradingDrawer::OnActiveRootActorChanged(ADisplayCluster
 
 void SDisplayClusterColorGradingDrawer::OnDetailObjectsChanged(const TArray<UObject*>& Objects)
 {
-	if (RootActorList.IsValid())
+	constexpr int32 RootActorListIndex = 1;
+	if (ColorGradingItemLists.Num() > RootActorListIndex && ColorGradingObjectListViews.Num() > RootActorListIndex)
 	{
-		const bool bIsCCWBeingSelected = Objects.ContainsByPredicate([](UObject* Object) { return Object && Object->IsA<AColorCorrectionWindow>(); });
-		const bool bAreCCWsInList = RootActorColorGradingItems.ContainsByPredicate([](const FDisplayClusterColorGradingListItemRef& ListItem)
+		TSharedPtr<SDisplayClusterColorGradingObjectList>& RootActorList = ColorGradingObjectListViews[RootActorListIndex];
+		TArray<FDisplayClusterColorGradingListItemRef>& RootActorColorGradingItems = ColorGradingItemLists[RootActorListIndex];
+		if (RootActorList.IsValid())
 		{
-			return ListItem.IsValid() && ListItem->Actor.IsValid() && ListItem->Actor->IsA<AColorCorrectionWindow>();
-		});
+			const bool bIsCCWBeingSelected = Objects.ContainsByPredicate([](UObject* Object) { return Object && Object->IsA<AColorCorrectionWindow>(); });
+			const bool bAreCCWsInList = RootActorColorGradingItems.ContainsByPredicate([](const FDisplayClusterColorGradingListItemRef& ListItem)
+			{
+				return ListItem.IsValid() && ListItem->Actor.IsValid() && ListItem->Actor->IsA<AColorCorrectionWindow>();
+			});
 
-		// Only need to refresh when CCWs are being selected or unselected
-		if (bIsCCWBeingSelected || bAreCCWsInList)
-		{
-			FillRootActorColorGradingList();
+			// Only need to refresh when CCWs are being selected or unselected
+			if (bIsCCWBeingSelected || bAreCCWsInList)
+			{
+				
+				RefreshColorGradingList(RootActorListIndex);
 
-			SetDrawerStateToDefault();
+				SetDrawerStateToDefault();
+			}
 		}
 	}
 }
@@ -1066,24 +1088,22 @@ void SDisplayClusterColorGradingDrawer::OnColorGradingDataModelGenerated()
 void SDisplayClusterColorGradingDrawer::OnListSelectionChanged(TSharedRef<SDisplayClusterColorGradingObjectList> SourceList, FDisplayClusterColorGradingListItemRef SelectedItem, ESelectInfo::Type SelectInfo)
 {
 	TArray<FDisplayClusterColorGradingListItemRef> SelectedObjects = SourceList->GetSelectedItems();
-
-	// When an item on one list is selected, clear the selection of the other list
-	if (LevelActorsList == SourceList && SelectedObjects.Num())
-	{
-		if (RootActorList->GetSelectedItems().Num())
-		{
-			bUpdateDataModelOnSelectionChanged = false;
-			RootActorList->SetSelectedItems(TArray<FDisplayClusterColorGradingListItemRef>());
-			bUpdateDataModelOnSelectionChanged = true;
-		}
-	}
-	else if (RootActorList == SourceList && SelectedObjects.Num())
+	if (SelectedObjects.Num())
 	{
-		if (LevelActorsList->GetSelectedItems().Num())
+		// When an item on one list is selected, clear the selection of the other list
+		for (const TSharedPtr<SDisplayClusterColorGradingObjectList>& List : ColorGradingObjectListViews)
 		{
-			bUpdateDataModelOnSelectionChanged = false;
-			LevelActorsList->SetSelectedItems(TArray<FDisplayClusterColorGradingListItemRef>());
-			bUpdateDataModelOnSelectionChanged = true;
+			if (List == SourceList)
+			{
+				continue;
+			}
+
+			if (List->GetSelectedItems().Num())
+			{
+				bUpdateDataModelOnSelectionChanged = false;
+				List->SetSelectedItems(TArray<FDisplayClusterColorGradingListItemRef>());
+				bUpdateDataModelOnSelectionChanged = true;
+			}
 		}
 	}
 
diff --git a/Engine/Plugins/Runtime/nDisplay/Source/DisplayClusterColorGrading/Private/Drawer/SDisplayClusterColorGradingDrawer.h b/Engine/Plugins/Runtime/nDisplay/Source/DisplayClusterColorGrading/Private/Drawer/SDisplayClusterColorGradingDrawer.h
--- a/Engine/Plugins/Runtime/nDisplay/Source/DisplayClusterColorGrading/Private/Drawer/SDisplayClusterColorGradingDrawer.h
+++ b/Engine/Plugins/Runtime/nDisplay/Source/DisplayClusterColorGrading/Private/Drawer/SDisplayClusterColorGradingDrawer.h
@@ -92,6 +92,9 @@ private:
 	/** Creates the button used to dock the drawer in the operator panel */
 	TSharedRef<SWidget> CreateDockInLayoutButton();
 
+	/** Gets the display name of the color grading list at the specified index */
+	FText GetColorGradingListName(int32 ListIndex) const;
+
 	/** Gets the name of the current level the active root actor is in */
 	FText GetCurrentLevelName() const;
 
@@ -104,11 +107,20 @@ private:
 	/** Unbinds a callback to the BlueprintCompiled delegate of the specified class */
 	void UnbindBlueprintCompiledDelegate(const UClass* Class);
 
-	/** Fills the level color grading list with any color gradable items found in the currently loaded level */
-	void FillLevelColorGradingList();
+	/** Refreshes all of the lists in the drawer, filling them with the current color gradable objects from the root actor and world  */
+	void RefreshColorGradingLists();
+
+	/** Refreshes the specified list, filling it with the current color gradable objects from the relevant source */
+	void RefreshColorGradingList(int32 Index);
+
+	/** Fills the specified list with any color gradable items found in the currently loaded level */
+	void FillLevelColorGradingList(TArray<FDisplayClusterColorGradingListItemRef>& InOutList);
+
+	/** Fills the specified list with any color gradable components of the active root actor */
+	void FillRootActorColorGradingList(TArray<FDisplayClusterColorGradingListItemRef>& InOutList);
 
-	/** Fills the root actor color grading list with any color gradable components of the active root actor */
-	void FillRootActorColorGradingList();
+	/** Fills the specified list with any color gradable color correction regions in the currently loaded level */
+	void FillColorCorrectionRegionColorGradingList(TArray<FDisplayClusterColorGradingListItemRef>& InOutList);
 
 	/** Fills the color grading group toolbar using the color grading data model */
 	void FillColorGradingGroupToolBar();
@@ -162,18 +174,16 @@ private:
 	/** The operator panel's view model */
 	TSharedPtr<IDisplayClusterOperatorViewModel> OperatorViewModel;
 
-	TSharedPtr<SDisplayClusterColorGradingObjectList> LevelActorsList;
-	TSharedPtr<SDisplayClusterColorGradingObjectList> RootActorList;
+	/** A list of all color grading object list widgets being displayed in the drawer's list panel */
+	TArray<TSharedPtr<SDisplayClusterColorGradingObjectList>> ColorGradingObjectListViews;
+
+	/** A list of source lists for the color grading object list widgets */
+	TArray<TArray<FDisplayClusterColorGradingListItemRef>> ColorGradingItemLists;
+
 	TSharedPtr<SHorizontalBox> ColorGradingGroupToolBarBox;
 	TSharedPtr<SDisplayClusterColorGradingColorWheelPanel> ColorWheelPanel;
 	TSharedPtr<SDisplayClusterColorGradingDetailsPanel> DetailsPanel;
 
-	/** List of color gradable items found in the current level */
-	TArray<FDisplayClusterColorGradingListItemRef> LevelColorGradingItems;
-
-	/** List of color gradable actors and components attached to the root actor */
-	TArray<FDisplayClusterColorGradingListItemRef> RootActorColorGradingItems;
-
 	/** The color grading data model for the currently selected objects */
 	TSharedPtr<FDisplayClusterColorGradingDataModel> ColorGradingDataModel;
 
diff --git a/Engine/Plugins/Runtime/nDisplay/Source/DisplayClusterColorGrading/Private/SDisplayClusterColorGradingObjectList.cpp b/Engine/Plugins/Runtime/nDisplay/Source/DisplayClusterColorGrading/Private/SDisplayClusterColorGradingObjectList.cpp
--- a/Engine/Plugins/Runtime/nDisplay/Source/DisplayClusterColorGrading/Private/SDisplayClusterColorGradingObjectList.cpp
+++ b/Engine/Plugins/Runtime/nDisplay/Source/DisplayClusterColorGrading/Private/SDisplayClusterColorGradingObjectList.cpp
@@ -75,7 +75,7 @@ public:
 
 				+ SHorizontalBox::Slot()
 				.AutoWidth()
-				.Padding(FMargin(0, 1, 6, 1))
+				.Padding(FMargin(6, 1))
 				.VAlign(VAlign_Center)
 				[
 					SNew(SBox)
@@ -199,6 +199,9 @@ void SDisplayClusterColorGradingObjectList::Construct(const FArguments& InArgs)
 				+ SHeaderRow::Column(DisplayClusterColorGradingObjectListColumnNames::ItemEnabled)
 				.DefaultLabel(FText::GetEmpty())
 				.FixedWidth(24.0f)
+				.HAlignCell(HAlign_Center)
+				.VAlignCell(VAlign_Center)
+				.DefaultTooltip(LOCTEXT("ItemEnabledColumnTooltip", "Color Grading Enabled"))
 
 				+ SHeaderRow::Column(DisplayClusterColorGradingObjectListColumnNames::ItemLabel)
 				.DefaultLabel(LOCTEXT("ColorGradingDrawerList_ItemLabelHeader", "Item Label"))
