nDisplay: Support setting the default template to use when spawning a new light card or flag.
Rename project editor settings from Light Card Editor to ICVFX Editor.
Fix templates not supporting UV lightcards.

diff --git a/Engine/Plugins/Runtime/nDisplay/Source/DisplayClusterLightCardEditor/Private/DisplayClusterLightCardEditor.cpp b/Engine/Plugins/Runtime/nDisplay/Source/DisplayClusterLightCardEditor/Private/DisplayClusterLightCardEditor.cpp
--- a/Engine/Plugins/Runtime/nDisplay/Source/DisplayClusterLightCardEditor/Private/DisplayClusterLightCardEditor.cpp
+++ b/Engine/Plugins/Runtime/nDisplay/Source/DisplayClusterLightCardEditor/Private/DisplayClusterLightCardEditor.cpp
@@ -189,8 +189,15 @@ void FDisplayClusterLightCardEditor::CenterActorInView(AActor* Actor)
 	}
 }
 
-AActor* FDisplayClusterLightCardEditor::SpawnActor(TSubclassOf<AActor> InActorClass, const FName& InActorName)
+AActor* FDisplayClusterLightCardEditor::SpawnActor(TSubclassOf<AActor> InActorClass, const FName& InActorName, const UDisplayClusterLightCardTemplate* InTemplate)
 {
+	if (InTemplate)
+	{
+		// For now only light card templates are supported
+		check(InActorClass && InActorClass->IsChildOf(ADisplayClusterLightCardActor::StaticClass()));
+		return SpawnLightCardFromTemplate(InTemplate);
+	}
+	
 	if (!ActiveRootActor.IsValid())
 	{
 		return nullptr;
@@ -219,16 +226,7 @@ AActor* FDisplayClusterLightCardEditor::SpawnActor(TSubclassOf<AActor> InActorCl
 
 	if (ADisplayClusterLightCardActor* NewLightCard = Cast<ADisplayClusterLightCardActor>(NewActor))
 	{
-		if (ViewportView.IsValid())
-		{
-			const EDisplayClusterMeshProjectionType ProjectionMode = ViewportView->GetLightCardEditorViewportClient()->GetProjectionMode();
-			if (ProjectionMode == EDisplayClusterMeshProjectionType::UV)
-			{
-				NewLightCard->bIsUVLightCard = true;
-				NewLightCard->Scale /= 4;
-				NewLightCard->Feathering = 0.05; // Just enough to avoid jagged look on UV lightcards.
-			}
-		}
+		TryConfigureLightCardForUV(NewLightCard);
 		const TArray<ADisplayClusterLightCardActor*> LightCards { NewLightCard } ;
 		AddLightCardsToActor(LightCards);
 	}
@@ -267,6 +265,8 @@ ADisplayClusterLightCardActor* FDisplayClusterLightCardEditor::SpawnLightCardFro
 	// and prevent the map from saving.
 	ADisplayClusterLightCardActor* NewLightCard = CastChecked<ADisplayClusterLightCardActor>(StaticDuplicateObject(InTemplate->LightCardActor.Get(), Level, UniqueName));
 	Level->AddLoadedActor(NewLightCard);
+
+	TryConfigureLightCardForUV(NewLightCard);
 	
 	if (!bIsPreview)
 	{
@@ -292,11 +292,17 @@ void FDisplayClusterLightCardEditor::AddNewLightCard()
 
 	FScopedTransaction Transaction(LOCTEXT("AddNewLightCardTransactionMessage", "Add New Light Card"));
 
-	// When adding a new lightcard, usually the desired location is in the middle of the viewport
-	if (ADisplayClusterLightCardActor* NewLightCard = SpawnActorAs<ADisplayClusterLightCardActor>(TEXT("LightCard")))
+	const UDisplayClusterLightCardEditorProjectSettings* Settings = GetDefault<UDisplayClusterLightCardEditorProjectSettings>();
+	const UDisplayClusterLightCardTemplate* Template = Settings->DefaultLightCardTemplate.LoadSynchronous();
+	
+	if (ADisplayClusterLightCardActor* NewLightCard = SpawnActorAs<ADisplayClusterLightCardActor>(TEXT("LightCard"), Template))
 	{
-		CenterActorInView(NewLightCard);
-
+		if (!Template)
+		{
+			// When adding a new lightcard, usually the desired location is in the middle of the viewport
+			CenterActorInView(NewLightCard);
+		}
+		
 		FDisplayClusterLightCardEditorRecentItem RecentlyPlacedItem;
 		RecentlyPlacedItem.ObjectPath = NewLightCard->GetClass();
 		RecentlyPlacedItem.ItemType = FDisplayClusterLightCardEditorRecentItem::Type_LightCard;
@@ -420,13 +426,19 @@ void FDisplayClusterLightCardEditor::AddNewFlag()
 
 	FScopedTransaction Transaction(LOCTEXT("AddNewFlagTransactionMessage", "Add New Flag"));
 
-	// When adding a new lightcard, usually the desired location is in the middle of the viewport
-	if (ADisplayClusterLightCardActor* NewLightCard = SpawnActorAs<ADisplayClusterLightCardActor>(TEXT("Flag")))
+	const UDisplayClusterLightCardEditorProjectSettings* Settings = GetDefault<UDisplayClusterLightCardEditorProjectSettings>();
+	const UDisplayClusterLightCardTemplate* Template = Settings->DefaultFlagTemplate.LoadSynchronous();
+	
+	if (ADisplayClusterLightCardActor* NewLightCard = SpawnActorAs<ADisplayClusterLightCardActor>(TEXT("Flag"), Template))
 	{
-		NewLightCard->Color = FLinearColor(0.f, 0.f, 0.f);
-		
-		CenterActorInView(NewLightCard);
+		if (!Template)
+		{
+			NewLightCard->Color = FLinearColor(0.f, 0.f, 0.f);
 
+			// When adding a new lightcard, usually the desired location is in the middle of the viewport
+			CenterActorInView(NewLightCard);
+		}
+		
 		FDisplayClusterLightCardEditorRecentItem RecentlyPlacedItem;
 		RecentlyPlacedItem.ObjectPath = NewLightCard->GetClass();
 		RecentlyPlacedItem.ItemType = FDisplayClusterLightCardEditorRecentItem::Type_Flag;
@@ -1544,6 +1556,21 @@ bool FDisplayClusterLightCardEditor::IsOurObject(UObject* InObject,
 	return bIsOurActor;
 }
 
+void FDisplayClusterLightCardEditor::TryConfigureLightCardForUV(ADisplayClusterLightCardActor* InLightCard)
+{
+	check(InLightCard);
+	if (ViewportView.IsValid())
+	{
+		const EDisplayClusterMeshProjectionType ProjectionMode = ViewportView->GetLightCardEditorViewportClient()->GetProjectionMode();
+		if (ProjectionMode == EDisplayClusterMeshProjectionType::UV)
+		{
+			InLightCard->bIsUVLightCard = true;
+			InLightCard->Scale /= 4;
+			InLightCard->Feathering = 0.05; // Just enough to avoid jagged look on UV lightcards.
+		}
+	}
+}
+
 void FDisplayClusterLightCardEditor::BindCompileDelegates()
 {
 	if (LightCardOutliner.IsValid())
diff --git a/Engine/Plugins/Runtime/nDisplay/Source/DisplayClusterLightCardEditor/Private/DisplayClusterLightCardEditor.h b/Engine/Plugins/Runtime/nDisplay/Source/DisplayClusterLightCardEditor/Private/DisplayClusterLightCardEditor.h
--- a/Engine/Plugins/Runtime/nDisplay/Source/DisplayClusterLightCardEditor/Private/DisplayClusterLightCardEditor.h
+++ b/Engine/Plugins/Runtime/nDisplay/Source/DisplayClusterLightCardEditor/Private/DisplayClusterLightCardEditor.h
@@ -88,12 +88,12 @@ public:
 	void CenterActorInView(AActor* Actor);
 
 	/** Spawns a new actor and adds it to the root actor if it is a light card */
-	AActor* SpawnActor(TSubclassOf<AActor> InActorClass, const FName& InActorName = NAME_None);
+	AActor* SpawnActor(TSubclassOf<AActor> InActorClass, const FName& InActorName = NAME_None, const UDisplayClusterLightCardTemplate* InTemplate = nullptr);
 
 	template<typename T>
-	T* SpawnActorAs(const FName& InActorName = NAME_None)
+	T* SpawnActorAs(const FName& InActorName = NAME_None, const UDisplayClusterLightCardTemplate* InTemplate = nullptr)
 	{
-		return Cast<T>(SpawnActor(T::StaticClass(), InActorName));
+		return Cast<T>(SpawnActor(T::StaticClass(), InActorName, InTemplate));
 	}
 	
 	/** Spawns a new light card from a light card template */
@@ -270,6 +270,9 @@ private:
 	 * @return True if our object, false if not
 	 */
 	bool IsOurObject(UObject* InObject, EDisplayClusterLightCardEditorProxyType& OutProxyType) const;
+
+	/** Configure a light card actor for UV if possible */
+	void TryConfigureLightCardForUV(ADisplayClusterLightCardActor* InLightCard);
 	
 	/** Bind delegates to when a BP compiles */
 	void BindCompileDelegates();
diff --git a/Engine/Plugins/Runtime/nDisplay/Source/DisplayClusterLightCardEditor/Private/DisplayClusterLightCardEditorModule.cpp b/Engine/Plugins/Runtime/nDisplay/Source/DisplayClusterLightCardEditor/Private/DisplayClusterLightCardEditorModule.cpp
--- a/Engine/Plugins/Runtime/nDisplay/Source/DisplayClusterLightCardEditor/Private/DisplayClusterLightCardEditorModule.cpp
+++ b/Engine/Plugins/Runtime/nDisplay/Source/DisplayClusterLightCardEditor/Private/DisplayClusterLightCardEditorModule.cpp
@@ -77,9 +77,9 @@ void FDisplayClusterLightCardEditorModule::RegisterSettings()
 		UDisplayClusterLightCardEditorProjectSettings* Settings = GetMutableDefault<UDisplayClusterLightCardEditorProjectSettings>();
 		// Needs transactional for undo/redo support in the light card editor
 		Settings->SetFlags(RF_Transactional);
-		SettingsModule->RegisterSettings("Project", "Plugins", "nDisplayLightCardEditor",
-			LOCTEXT("nDisplayLightCardEditorName", "nDisplay Light Card Editor"),
-			LOCTEXT("nDisplayLightCardEditorDescription", "Configure settings for the nDisplay Light Card Editor."),
+		SettingsModule->RegisterSettings("Project", "Plugins", "nDisplayICVFXEditor",
+			LOCTEXT("nDisplayICVFXEditorName", "nDisplay - ICVFX Editor"),
+			LOCTEXT("nDisplayICVFXEditorDescription", "Configure settings for the nDisplay ICVFX Editor."),
 			Settings);
 	}
 }
diff --git a/Engine/Plugins/Runtime/nDisplay/Source/DisplayClusterLightCardEditor/Private/Settings/DisplayClusterLightCardEditorSettings.h b/Engine/Plugins/Runtime/nDisplay/Source/DisplayClusterLightCardEditor/Private/Settings/DisplayClusterLightCardEditorSettings.h
--- a/Engine/Plugins/Runtime/nDisplay/Source/DisplayClusterLightCardEditor/Private/Settings/DisplayClusterLightCardEditorSettings.h
+++ b/Engine/Plugins/Runtime/nDisplay/Source/DisplayClusterLightCardEditor/Private/Settings/DisplayClusterLightCardEditorSettings.h
@@ -6,6 +6,7 @@
 
 #include "DisplayClusterLightCardEditorSettings.generated.h"
 
+class UDisplayClusterLightCardTemplate;
 struct FSlateBrush;
 
 /**
@@ -20,15 +21,23 @@ public:
 	UDisplayClusterLightCardEditorProjectSettings();
 
 	/** The default path to save new light card templates */
-	UPROPERTY(config, EditAnywhere, Category = LightCardTemplates)
+	UPROPERTY(config, EditAnywhere, Category = Templates)
 	FDirectoryPath LightCardTemplateDefaultPath;
 
+	/** The default template to use when creating a new light card */
+	UPROPERTY(config, EditAnywhere, Category = Templates)
+	TSoftObjectPtr<UDisplayClusterLightCardTemplate> DefaultLightCardTemplate;
+
+	/** The default template to use when creating a new flag */
+	UPROPERTY(config, EditAnywhere, Category = Templates)
+	TSoftObjectPtr<UDisplayClusterLightCardTemplate> DefaultFlagTemplate;
+	
 	/** Whether light card labels should be displayed. Handled through the light card editor */
 	UPROPERTY()
 	bool bDisplayLightCardLabels;
 	
 	/** The scale to use for light card labels */
-	UPROPERTY(config, EditAnywhere, Category = LightCardLabels)
+	UPROPERTY(config, EditAnywhere, Category = Labels)
 	float LightCardLabelScale;
 };
 
