nDisplay: Replaced placeholder icons with final icons for a variety of UI elements in the light card editor

diff --git a/Engine/Plugins/Runtime/nDisplay/Source/DisplayCluster/Private/Game/EngineClasses/Scene/DisplayClusterLightCardActor.cpp b/Engine/Plugins/Runtime/nDisplay/Source/DisplayCluster/Private/Game/EngineClasses/Scene/DisplayClusterLightCardActor.cpp
--- a/Engine/Plugins/Runtime/nDisplay/Source/DisplayCluster/Private/Game/EngineClasses/Scene/DisplayClusterLightCardActor.cpp
+++ b/Engine/Plugins/Runtime/nDisplay/Source/DisplayCluster/Private/Game/EngineClasses/Scene/DisplayClusterLightCardActor.cpp
@@ -186,7 +186,7 @@ FName ADisplayClusterLightCardActor::GetCustomIconName() const
 {
 	if (bIsUVLightCard)
 	{
-		return TEXT("UVLightCard");
+		return TEXT("ClassIcon.DisplayClusterLightCardActor.UVLightCard");
 	}
 
 	return Super::GetCustomIconName();
diff --git a/Engine/Plugins/Runtime/nDisplay/Source/DisplayClusterColorGrading/Private/DisplayClusterColorGradingStyle.h b/Engine/Plugins/Runtime/nDisplay/Source/DisplayClusterColorGrading/Private/DisplayClusterColorGradingStyle.h
--- a/Engine/Plugins/Runtime/nDisplay/Source/DisplayClusterColorGrading/Private/DisplayClusterColorGradingStyle.h
+++ b/Engine/Plugins/Runtime/nDisplay/Source/DisplayClusterColorGrading/Private/DisplayClusterColorGradingStyle.h
@@ -23,7 +23,8 @@ public:
 		SetCoreContentRoot(FPaths::EngineContentDir() / TEXT("Editor/Slate"));
 
 		Set("ColorGradingDrawer.Icon", new IMAGE_BRUSH_SVG("Components/nDisplayCamera", Icon16x16));
-		Set("ColorGradingDrawer.ColorGradingMode", new CORE_IMAGE_BRUSH_SVG("Starship/Common/PostProcessing", Icon16x16));
+		Set("ColorGradingDrawer.ColorGradingMode", new IMAGE_BRUSH_SVG("OperatorPanel/Colors", Icon16x16));
+		Set("ColorGradingDrawer.DetailViewMode", new IMAGE_BRUSH_SVG("OperatorPanel/Sliders", Icon16x16));
 		Set("ColorGradingDrawer.Viewports", new CORE_IMAGE_BRUSH_SVG("Starship/Common/Viewports", Icon16x16));
 		Set("ColorGradingDrawer.Nodes", new IMAGE_BRUSH_SVG("Cluster/ClusterNode", Icon16x16));
 
diff --git a/Engine/Plugins/Runtime/nDisplay/Source/DisplayClusterColorGrading/Private/Drawer/SDisplayClusterColorGradingDrawer.cpp b/Engine/Plugins/Runtime/nDisplay/Source/DisplayClusterColorGrading/Private/Drawer/SDisplayClusterColorGradingDrawer.cpp
--- a/Engine/Plugins/Runtime/nDisplay/Source/DisplayClusterColorGrading/Private/Drawer/SDisplayClusterColorGradingDrawer.cpp
+++ b/Engine/Plugins/Runtime/nDisplay/Source/DisplayClusterColorGrading/Private/Drawer/SDisplayClusterColorGradingDrawer.cpp
@@ -184,7 +184,7 @@ void SDisplayClusterColorGradingDrawer::Construct(const FArguments& InArgs, bool
 								[
 									SNew(SImage)
 									.ColorAndOpacity(FSlateColor::UseForeground())
-									.Image(FAppStyle::Get().GetBrush("EditorPreferences.TabIcon"))
+									.Image(FDisplayClusterColorGradingStyle::Get().GetBrush("ColorGradingDrawer.DetailViewMode"))
 								]
 							]
 						]
diff --git a/Engine/Plugins/Runtime/nDisplay/Source/DisplayClusterConfigurator/Private/DisplayClusterConfiguratorStyle.cpp b/Engine/Plugins/Runtime/nDisplay/Source/DisplayClusterConfigurator/Private/DisplayClusterConfiguratorStyle.cpp
--- a/Engine/Plugins/Runtime/nDisplay/Source/DisplayClusterConfigurator/Private/DisplayClusterConfiguratorStyle.cpp
+++ b/Engine/Plugins/Runtime/nDisplay/Source/DisplayClusterConfigurator/Private/DisplayClusterConfiguratorStyle.cpp
@@ -70,6 +70,15 @@ void FDisplayClusterConfiguratorStyle::Initialize()
 
 		Set("ClassIcon.DisplayClusterOriginComponent", new IMAGE_BRUSH_SVG("Icons/Components/nDisplayRootComponent", Icon16x16));
 		Set("ClassThumbnail.DisplayClusterOriginComponent", new IMAGE_BRUSH_SVG("Icons/Components/nDisplayRootComponent", Icon64x64));
+
+		Set("ClassIcon.DisplayClusterLightCardActor", new IMAGE_BRUSH_SVG("Icons/LightCard/LightCard", Icon16x16));
+		Set("ClassThumbnail.DisplayClusterLightCardActor", new IMAGE_BRUSH_SVG("Icons/LightCard/LightCard", Icon64x64));
+
+		Set("ClassIcon.DisplayClusterLightCardActor.UVLightCard", new IMAGE_BRUSH_SVG("Icons/LightCard/LightCardUV", Icon16x16));
+		Set("ClassThumbnail.DisplayClusterLightCardActor.UVLightCard", new IMAGE_BRUSH_SVG("Icons/UVLightCard/LightCardUV", Icon64x64));
+
+		Set("ClassIcon.DisplayClusterLightCardActor.Flag", new IMAGE_BRUSH_SVG("Icons/LightCard/LightCardFlag", Icon16x16));
+		Set("ClassThumbnail.DisplayClusterLightCardActor.Flag", new IMAGE_BRUSH_SVG("Icons/UVLightCard/LightCardFlag", Icon64x64));
 	}
 
 	// Config Editor Tabs
diff --git a/Engine/Plugins/Runtime/nDisplay/Source/DisplayClusterLightCardEditor/Private/DisplayClusterLightCardEditor.cpp b/Engine/Plugins/Runtime/nDisplay/Source/DisplayClusterLightCardEditor/Private/DisplayClusterLightCardEditor.cpp
--- a/Engine/Plugins/Runtime/nDisplay/Source/DisplayClusterLightCardEditor/Private/DisplayClusterLightCardEditor.cpp
+++ b/Engine/Plugins/Runtime/nDisplay/Source/DisplayClusterLightCardEditor/Private/DisplayClusterLightCardEditor.cpp
@@ -1085,11 +1085,19 @@ TSharedRef<SWidget> FDisplayClusterLightCardEditor::GeneratePlaceActorsMenu()
 	MenuBuilder.BeginSection("PlaceActors", LOCTEXT("PlaceActorsMenuHeader", "Place Actors"));
 	{
 		FSlateIcon LightCardIcon = FSlateIconFinder::FindIconForClass(ADisplayClusterLightCardActor::StaticClass());
+		FSlateIcon FlagIcon = FSlateIconFinder::FindIcon(TEXT("ClassIcon.DisplayClusterLightCardActor.Flag"));
+		FSlateIcon UVLightCardIcon = FSlateIconFinder::FindIcon(TEXT("ClassIcon.DisplayClusterLightCardActor.UVLightCard"));
+
+		bool bIsUVMode = false;
+		if (ViewportView.IsValid())
+		{
+			bIsUVMode = ViewportView->GetLightCardEditorViewportClient()->GetProjectionMode() == EDisplayClusterMeshProjectionType::UV;
+		}
 
 		MenuBuilder.AddMenuEntry(FDisplayClusterLightCardEditorCommands::Get().AddNewFlag,
-			NAME_None, TAttribute<FText>(), TAttribute<FText>(), LightCardIcon);
+			NAME_None, TAttribute<FText>(), TAttribute<FText>(), FlagIcon);
 		MenuBuilder.AddMenuEntry(FDisplayClusterLightCardEditorCommands::Get().AddNewLightCard,
-			NAME_None, TAttribute<FText>(), TAttribute<FText>(), LightCardIcon);
+			NAME_None, TAttribute<FText>(), TAttribute<FText>(), bIsUVMode ? UVLightCardIcon : LightCardIcon);
 		
 		TSet<UClass*> StageActorClasses = UE::DisplayClusterLightCardEditorUtils::GetAllStageActorClasses();
 		for (UClass* Class : StageActorClasses)
@@ -1110,7 +1118,7 @@ TSharedRef<SWidget> FDisplayClusterLightCardEditor::GeneratePlaceActorsMenu()
 		LOCTEXT("AllTemplatesTooltip", "Select a template"),
 		FNewMenuDelegate::CreateSP(this, &FDisplayClusterLightCardEditor::GenerateTemplateSubMenu),
 		false,
-		FSlateIcon(),
+		FSlateIcon(FDisplayClusterLightCardEditorStyle::Get().GetStyleSetName(), TEXT("DisplayClusterLightCardEditor.Template")),
 		true);
 	}
 	MenuBuilder.EndSection();
diff --git a/Engine/Plugins/Runtime/nDisplay/Source/DisplayClusterLightCardEditor/Private/DisplayClusterLightCardEditorStyle.cpp b/Engine/Plugins/Runtime/nDisplay/Source/DisplayClusterLightCardEditor/Private/DisplayClusterLightCardEditorStyle.cpp
--- a/Engine/Plugins/Runtime/nDisplay/Source/DisplayClusterLightCardEditor/Private/DisplayClusterLightCardEditorStyle.cpp
+++ b/Engine/Plugins/Runtime/nDisplay/Source/DisplayClusterLightCardEditor/Private/DisplayClusterLightCardEditorStyle.cpp
@@ -24,12 +24,12 @@ FDisplayClusterLightCardEditorStyle::FDisplayClusterLightCardEditorStyle(): FSla
 	Set("DisplayClusterLightCardEditor.ViewOrientationFront", new CORE_IMAGE_BRUSH_SVG("Starship/Common/ViewFront", Icon16x16));
 	Set("DisplayClusterLightCardEditor.ViewOrientationBack", new CORE_IMAGE_BRUSH_SVG("Starship/Common/ViewBack", Icon16x16));
 
-	Set("UVLightCard", new CORE_IMAGE_BRUSH_SVG("Starship/Common/SetDrawUVs", Icon16x16));
-
-	Set("DisplayClusterLightCardEditor.Labels", new CORE_IMAGE_BRUSH_SVG("Starship/AssetIcons/Font_16", Icon16x16));
+	Set("DisplayClusterLightCardEditor.Labels", new CORE_IMAGE_BRUSH_SVG("Starship/Common/IssueTracker", Icon16x16));
 
 	Set("DisplayClusterLightCardEditor.ActorHidden", new IMAGE_BRUSH_SVG("ActorHidden", Icon16x16));
 	Set("DisplayClusterLightCardEditor.ActorNotHidden", new IMAGE_BRUSH_SVG("ActorNotHidden", Icon16x16));
+
+	Set("DisplayClusterLightCardEditor.Template", new IMAGE_BRUSH_SVG("Template", Icon16x16));
 		
 	FSlateStyleRegistry::RegisterSlateStyle(*this);
 }
diff --git a/Engine/Plugins/Runtime/nDisplay/Source/DisplayClusterLightCardEditor/Private/Viewport/DisplayClusterLightcardEditorViewport.cpp b/Engine/Plugins/Runtime/nDisplay/Source/DisplayClusterLightCardEditor/Private/Viewport/DisplayClusterLightcardEditorViewport.cpp
--- a/Engine/Plugins/Runtime/nDisplay/Source/DisplayClusterLightCardEditor/Private/Viewport/DisplayClusterLightcardEditorViewport.cpp
+++ b/Engine/Plugins/Runtime/nDisplay/Source/DisplayClusterLightCardEditor/Private/Viewport/DisplayClusterLightcardEditorViewport.cpp
@@ -712,19 +712,27 @@ TSharedRef<SWidget> SDisplayClusterLightCardEditorViewport::MakeContextMenu()
 void SDisplayClusterLightCardEditorViewport::MakePlaceActorsSubMenu(FMenuBuilder& MenuBuilder)
 {
 	FSlateIcon LightCardIcon = FSlateIconFinder::FindIconForClass(ADisplayClusterLightCardActor::StaticClass());
+	FSlateIcon FlagIcon = FSlateIconFinder::FindIcon(TEXT("ClassIcon.DisplayClusterLightCardActor.Flag"));
+	FSlateIcon UVLightCardIcon = FSlateIconFinder::FindIcon(TEXT("ClassIcon.DisplayClusterLightCardActor.UVLightCard"));
+
+	bool bIsUVMode = false;
+	if (ViewportClient.IsValid())
+	{
+		bIsUVMode = ViewportClient->GetProjectionMode() == EDisplayClusterMeshProjectionType::UV;
+	}
 
 	// Add custom UI actions here for the AddNew commands so that the newly added actors can be moved to the user's mouse position after the
 	// add operation is performed
 	MenuBuilder.AddMenuEntry(
 		FDisplayClusterLightCardEditorCommands::Get().AddNewFlag->GetLabel(),
 		FDisplayClusterLightCardEditorCommands::Get().AddNewFlag->GetDescription(),
-		LightCardIcon,
+		FlagIcon,
 		FUIAction(FExecuteAction::CreateSP(this, &SDisplayClusterLightCardEditorViewport::AddFlagHere)));
 
 	MenuBuilder.AddMenuEntry(
 		FDisplayClusterLightCardEditorCommands::Get().AddNewLightCard->GetLabel(),
 		FDisplayClusterLightCardEditorCommands::Get().AddNewLightCard->GetDescription(),
-		LightCardIcon,
+		bIsUVMode ? UVLightCardIcon : LightCardIcon,
 		FUIAction(FExecuteAction::CreateSP(this, &SDisplayClusterLightCardEditorViewport::AddLightCardHere)));
 
 	TSet<UClass*> StageActorClasses = UE::DisplayClusterLightCardEditorUtils::GetAllStageActorClasses();
