nDisplay: Add over/under frustum selection to main toolbar.

diff --git a/Engine/Plugins/Runtime/nDisplay/Source/DisplayClusterLightCardEditor/Private/DisplayClusterLightCardEditor.cpp b/Engine/Plugins/Runtime/nDisplay/Source/DisplayClusterLightCardEditor/Private/DisplayClusterLightCardEditor.cpp
--- a/Engine/Plugins/Runtime/nDisplay/Source/DisplayClusterLightCardEditor/Private/DisplayClusterLightCardEditor.cpp
+++ b/Engine/Plugins/Runtime/nDisplay/Source/DisplayClusterLightCardEditor/Private/DisplayClusterLightCardEditor.cpp
@@ -1358,6 +1358,132 @@ TSharedRef<SWidget> FDisplayClusterLightCardEditor::GenerateLabelsMenu()
 	return MenuBuilder.MakeWidget();
 }
 
+namespace UE::LightCardEditorFrustumTypes
+{
+	static FString UnderFrustumString = TEXT("Under Frustum");
+	static FString OverFrustumString = TEXT("Over Frustum");
+}
+
+/**
+ * A single widget representing a frustum selection
+ */
+class SFrustumItemWidget : public SCompoundWidget
+{
+public:
+	SLATE_BEGIN_ARGS(SFrustumItemWidget) {}
+	SLATE_ATTRIBUTE(FString, Selection)
+	SLATE_END_ARGS()
+
+	void Construct(const FArguments& InArgs)
+	{
+		static FText UnderFrustumText = LOCTEXT("UnderFrustumText", "Under Frustum");
+		static FText OverFrustumText = LOCTEXT("OverFrustumText", "Over Frustum");
+		static FText NoFrustumText = LOCTEXT("NoBlendModeSelectedText", "Blend Mode");
+		
+		Selection = InArgs._Selection;
+		ChildSlot
+		[
+			SNew(SHorizontalBox)
+			+SHorizontalBox::Slot()
+			.AutoWidth()
+			[
+				SNew(SImage)
+				.ColorAndOpacity(FSlateColor::UseForeground())
+				.Image_Lambda([this]()
+				{
+					const FString CurrentSelection = Selection.Get();
+					const FString BrushName = CurrentSelection == UE::LightCardEditorFrustumTypes::OverFrustumString ?
+						"DisplayClusterLightCardEditor.FrustumOnTop" : "DisplayClusterLightCardEditor.FrustumUnderneath";
+					return FDisplayClusterLightCardEditorStyle::Get().GetBrush(*BrushName);
+				})
+			]
+			+SHorizontalBox::Slot()
+			.Padding(4.f, 0.f, 0.f, 0.f)
+			[
+				SNew(SBox)
+				.MinDesiredWidth(90.f) // Keeps the widget from resizing on different selections
+				[
+					SNew(STextBlock).Text_Lambda([this]()
+					{
+						const FString CurrentSelection = Selection.Get();
+						const FText Text = CurrentSelection == UE::LightCardEditorFrustumTypes::OverFrustumString ? OverFrustumText :
+						CurrentSelection == UE::LightCardEditorFrustumTypes::UnderFrustumString ? UnderFrustumText : NoFrustumText;
+						return Text;
+					})
+				]
+			]
+		];
+	}
+
+private:
+	TAttribute<FString> Selection;
+};
+
+TSharedRef<SWidget> FDisplayClusterLightCardEditor::CreateFrustumWidget()
+{
+	FrustumSelections.Add(MakeShared<FString>(UE::LightCardEditorFrustumTypes::OverFrustumString));
+	FrustumSelections.Add(MakeShared<FString>(UE::LightCardEditorFrustumTypes::UnderFrustumString));
+
+	auto GetFrustumStringFromRootActor = [this]()
+	{
+		if (ActiveRootActor.IsValid())
+		{
+			switch (ActiveRootActor->GetStageSettings().Lightcard.Blendingmode)
+			{
+			case EDisplayClusterConfigurationICVFX_LightcardRenderMode::Over:
+				{
+					return UE::LightCardEditorFrustumTypes::OverFrustumString;
+				}
+			case EDisplayClusterConfigurationICVFX_LightcardRenderMode::Under:
+				{
+					return UE::LightCardEditorFrustumTypes::UnderFrustumString;
+				}
+			}
+		}
+
+		return FString();
+	};
+
+	return SNew(SComboBox<TSharedPtr<FString>>)
+		.OptionsSource(&FrustumSelections)
+		.ToolTipText(LOCTEXT("BlendModeTooltip", "Blending Mode: Specify how to render Light Cards in relation to the inner frustum."))
+		.IsEnabled_Lambda([this]()
+		{
+			return ActiveRootActor.IsValid();
+		})
+		.OnSelectionChanged_Lambda([this](TSharedPtr<FString> InItem, ESelectInfo::Type)
+		{
+			if (ActiveRootActor.IsValid())
+			{
+				checkSlow(InItem.IsValid());
+
+				UDisplayClusterConfigurationData* ConfigData = ActiveRootActor->GetConfigData();
+				check(ConfigData);
+
+				FScopedTransaction Transaction(LOCTEXT("SetBlendModeTransactionMessage", "Set Light Card Blend Mode"), !GIsTransacting);
+				ActiveRootActor->Modify();
+				ConfigData->Modify();
+				if (*InItem == UE::LightCardEditorFrustumTypes::UnderFrustumString)
+				{
+					ConfigData->StageSettings.Lightcard.Blendingmode = EDisplayClusterConfigurationICVFX_LightcardRenderMode::Under;
+				}
+				else if (*InItem == UE::LightCardEditorFrustumTypes::OverFrustumString)
+				{
+					ConfigData->StageSettings.Lightcard.Blendingmode = EDisplayClusterConfigurationICVFX_LightcardRenderMode::Over;
+				}
+			}
+		})
+		.OnGenerateWidget_Lambda([](TSharedPtr<FString> InItem)
+		{
+			return SNew(SFrustumItemWidget).Selection(*InItem.Get());
+		})
+		.Content()
+		[
+			SNew(SFrustumItemWidget)
+			.Selection_Lambda(GetFrustumStringFromRootActor)
+		];
+}
+
 TSharedRef<SWidget> FDisplayClusterLightCardEditor::CreateLightCardTemplateWidget()
 {
 	return SNew(SDisplayClusterLightCardTemplateList, SharedThis(this))
@@ -1474,7 +1600,10 @@ void FDisplayClusterLightCardEditor::ExtendToolbar(FToolBarBuilder& ToolbarBuild
 			TAttribute<FText>(),
 			LOCTEXT("Labels_ToolTip", "Configure options for labels"),
 			FSlateIcon(FDisplayClusterLightCardEditorStyle::Get().GetStyleSetName(), "DisplayClusterLightCardEditor.Labels"));
+		
+		ToolbarBuilder.AddWidget(CreateFrustumWidget());
 	}
+	
 }
 
 void FDisplayClusterLightCardEditor::ExtendFileMenu(FMenuBuilder& MenuBuilder)
diff --git a/Engine/Plugins/Runtime/nDisplay/Source/DisplayClusterLightCardEditor/Private/DisplayClusterLightCardEditor.h b/Engine/Plugins/Runtime/nDisplay/Source/DisplayClusterLightCardEditor/Private/DisplayClusterLightCardEditor.h
--- a/Engine/Plugins/Runtime/nDisplay/Source/DisplayClusterLightCardEditor/Private/DisplayClusterLightCardEditor.h
+++ b/Engine/Plugins/Runtime/nDisplay/Source/DisplayClusterLightCardEditor/Private/DisplayClusterLightCardEditor.h
@@ -239,6 +239,9 @@ private:
 	/** Generate the labels dropdown menu */
 	TSharedRef<SWidget> GenerateLabelsMenu();
 
+	/** Creates the frustum combo box */
+	TSharedRef<SWidget> CreateFrustumWidget();
+
 	/** Creates the editor's command list and binds commands to it */
 	void BindCommands();
 
@@ -331,6 +334,9 @@ private:
 	/** Selected actors */
 	TArray<TWeakObjectPtr<AActor>> SelectedActors;
 
+	/** Available frustums to choose */
+	TArray<TSharedPtr<FString>> FrustumSelections;
+
 	/** Delegate handle for the OnActiveRootActorChanged delegate */
 	FDelegateHandle ActiveRootActorChangedHandle;
 
diff --git a/Engine/Plugins/Runtime/nDisplay/Source/DisplayClusterLightCardEditor/Private/DisplayClusterLightCardEditorStyle.cpp b/Engine/Plugins/Runtime/nDisplay/Source/DisplayClusterLightCardEditor/Private/DisplayClusterLightCardEditorStyle.cpp
--- a/Engine/Plugins/Runtime/nDisplay/Source/DisplayClusterLightCardEditor/Private/DisplayClusterLightCardEditorStyle.cpp
+++ b/Engine/Plugins/Runtime/nDisplay/Source/DisplayClusterLightCardEditor/Private/DisplayClusterLightCardEditorStyle.cpp
@@ -30,6 +30,9 @@ FDisplayClusterLightCardEditorStyle::FDisplayClusterLightCardEditorStyle(): FSla
 	Set("DisplayClusterLightCardEditor.ActorNotHidden", new IMAGE_BRUSH_SVG("ActorNotHidden", Icon16x16));
 
 	Set("DisplayClusterLightCardEditor.Template", new IMAGE_BRUSH_SVG("Template", Icon16x16));
+
+	Set("DisplayClusterLightCardEditor.FrustumOnTop", new IMAGE_BRUSH_SVG("FrustumOnTop", Icon16x16));
+	Set("DisplayClusterLightCardEditor.FrustumUnderneath", new IMAGE_BRUSH_SVG("FrustumUnderneath", Icon16x16));
 		
 	FSlateStyleRegistry::RegisterSlateStyle(*this);
 }
