nDisplay: Fix crash when shutting down the ICVFX editor panel.
This was caused by the viewport client executing its destructor after the viewport owner causing the PreviewScene raw pointer to be invalid.

diff --git a/Engine/Plugins/Runtime/nDisplay/Source/DisplayClusterLightCardEditor/Private/Viewport/DisplayClusterLightCardEditorViewportClient.cpp b/Engine/Plugins/Runtime/nDisplay/Source/DisplayClusterLightCardEditor/Private/Viewport/DisplayClusterLightCardEditorViewportClient.cpp
--- a/Engine/Plugins/Runtime/nDisplay/Source/DisplayClusterLightCardEditor/Private/Viewport/DisplayClusterLightCardEditorViewportClient.cpp
+++ b/Engine/Plugins/Runtime/nDisplay/Source/DisplayClusterLightCardEditor/Private/Viewport/DisplayClusterLightCardEditorViewportClient.cpp
@@ -137,14 +137,6 @@ FDisplayClusterLightCardEditorViewportClient::~FDisplayClusterLightCardEditorVie
 	{
 		RootActorLevelInstance->UnsubscribeFromPostProcessRenderTarget(reinterpret_cast<uint8*>(this));
 	}
-
-	if (RefreshTimerHandle.IsValid())
-	{
-		if (const UWorld* PreviewWorld = PreviewScene->GetWorld())
-		{
-			PreviewWorld->GetTimerManager().ClearTimer(RefreshTimerHandle);
-		}
-	}
 }
 
 FLinearColor FDisplayClusterLightCardEditorViewportClient::GetBackgroundColor() const
@@ -1370,11 +1362,17 @@ void FDisplayClusterLightCardEditorViewportClient::UpdatePreviewActor(ADisplayCl
 		check(PreviewWorld);
 
 		const TWeakObjectPtr<ADisplayClusterRootActor> RootActorPtr (RootActor);
+		const TWeakPtr<FDisplayClusterLightCardEditorViewportClient> WeakPtrThis = SharedThis(this);
 		
 		// Schedule for the next tick so CDO changes get propagated first in the event of config editor skeleton
 		// regeneration & compiles. nDisplay's custom propagation may have issues if the archetype isn't correct.
-		RefreshTimerHandle = PreviewWorld->GetTimerManager().SetTimerForNextTick([=]()
-		{			
+		PreviewWorld->GetTimerManager().SetTimerForNextTick([=]()
+		{
+			if (!WeakPtrThis.IsValid())
+			{
+				return;
+			}
+			
 			DestroyProxies(ProxyType);
 
 			if (!RootActorPtr.IsValid())
diff --git a/Engine/Plugins/Runtime/nDisplay/Source/DisplayClusterLightCardEditor/Private/Viewport/DisplayClusterLightCardEditorViewportClient.h b/Engine/Plugins/Runtime/nDisplay/Source/DisplayClusterLightCardEditor/Private/Viewport/DisplayClusterLightCardEditorViewportClient.h
--- a/Engine/Plugins/Runtime/nDisplay/Source/DisplayClusterLightCardEditor/Private/Viewport/DisplayClusterLightCardEditorViewportClient.h
+++ b/Engine/Plugins/Runtime/nDisplay/Source/DisplayClusterLightCardEditor/Private/Viewport/DisplayClusterLightCardEditorViewportClient.h
@@ -250,9 +250,6 @@ private:
 
 	/** Proxy types that are currently refreshing over 1-2 frames */
 	TSet<EDisplayClusterLightCardEditorProxyType> ProxyTypesRefreshing;
-
-	/** Handle to world timer when refreshing */
-	FTimerHandle RefreshTimerHandle;
 	
 	/** The radius of the bounding sphere that entirely encapsulates the root actor */
 	float RootActorBoundingRadius = 0.0f;
