nDisplay: Fix ICVFX billboard render triggering a checkSlow.
Move the GetSceneViewInitOptions to Tick to avoid its internal methods calling GetWorldSettings off of the game thread triggering the checkSlow.
This also is a slight optimization as we are now only calling this once instead of for each billboard component.

diff --git a/Engine/Plugins/Runtime/nDisplay/Source/DisplayClusterLightCardEditor/Private/Viewport/DisplayClusterLightCardEditorViewportClient.cpp b/Engine/Plugins/Runtime/nDisplay/Source/DisplayClusterLightCardEditor/Private/Viewport/DisplayClusterLightCardEditorViewportClient.cpp
--- a/Engine/Plugins/Runtime/nDisplay/Source/DisplayClusterLightCardEditor/Private/Viewport/DisplayClusterLightCardEditorViewportClient.cpp
+++ b/Engine/Plugins/Runtime/nDisplay/Source/DisplayClusterLightCardEditor/Private/Viewport/DisplayClusterLightCardEditorViewportClient.cpp
@@ -221,6 +221,13 @@ void FDisplayClusterLightCardEditorViewportClient::Tick(float DeltaSeconds)
 			DesiredLookAtLocation.Reset();
 		}
 	}
+
+	if (BillboardComponentProxies.Num() > 0)
+	{
+		FSceneViewInitOptions SceneViewInitOptions;
+		GetSceneViewInitOptions(SceneViewInitOptions);
+		BillboardViewMatrices = FViewMatrices(SceneViewInitOptions);
+	}
 }
 
 void FDisplayClusterLightCardEditorViewportClient::Draw(FViewport* InViewport, FCanvas* Canvas)
@@ -457,11 +464,7 @@ void FDisplayClusterLightCardEditorViewportClient::Draw(const FSceneView* View,
 		}
 		FSpriteProxy SpriteProxy = FSpriteProxy::FromBillboard(BillboardComponent.Get());
 		
-		FSceneViewInitOptions SceneViewInitOptions;
-		GetSceneViewInitOptions(SceneViewInitOptions);
-		FViewMatrices ViewMatrices(SceneViewInitOptions);
-		
-		FVector ProjectedLocation = ProjectWorldPosition(SpriteProxy.WorldPosition, ViewMatrices);
+		FVector ProjectedLocation = ProjectWorldPosition(SpriteProxy.WorldPosition, BillboardViewMatrices);
 		
 		if (const FTexture* TextureResource = SpriteProxy.Sprite ? SpriteProxy.Sprite->GetResource() : nullptr)
 		{
diff --git a/Engine/Plugins/Runtime/nDisplay/Source/DisplayClusterLightCardEditor/Private/Viewport/DisplayClusterLightCardEditorViewportClient.h b/Engine/Plugins/Runtime/nDisplay/Source/DisplayClusterLightCardEditor/Private/Viewport/DisplayClusterLightCardEditorViewportClient.h
--- a/Engine/Plugins/Runtime/nDisplay/Source/DisplayClusterLightCardEditor/Private/Viewport/DisplayClusterLightCardEditorViewportClient.h
+++ b/Engine/Plugins/Runtime/nDisplay/Source/DisplayClusterLightCardEditor/Private/Viewport/DisplayClusterLightCardEditorViewportClient.h
@@ -300,6 +300,9 @@ private:
 	TArray<TWeakObjectPtr<UBillboardComponent>> BillboardComponentProxies;
 	TArray<FDisplayClusterWeakStageActorPtr> SelectedActors;
 
+	/** View matrices set per tick if there are any billboard components */
+	FViewMatrices BillboardViewMatrices;
+
 	/** The actor in the selected actor list that was selected last */
 	FDisplayClusterWeakStageActorPtr LastSelectedActor;
 
