nDisplay:
- added copying alpha channel from scene color to the RenderTarget (for lightcard and chromakey viewports with OCIO, PP etc.)
- removed unused old code about OCIO lightcard viewports (LightCardColor and LightCardAlpha viewports)
- Added  parent TSharedFromThis for FDisplayClusterViewportManagerProxy,FDisplayClusterViewportProxy and FDisplayClusterViewportPostProcessManager. Added changes to support it.
- improved FDisplayClusterViewportManagerViewExtension. Added scene color capture feature. . Added changes to support it.
- added CVar 'nDisplay.render.icvfx.lightcard'

diff --git a/Engine/Plugins/Runtime/nDisplay/Source/DisplayCluster/Private/Render/Viewport/Configuration/DisplayClusterViewportConfigurationHelpers_ICVFX.cpp b/Engine/Plugins/Runtime/nDisplay/Source/DisplayCluster/Private/Render/Viewport/Configuration/DisplayClusterViewportConfigurationHelpers_ICVFX.cpp
--- a/Engine/Plugins/Runtime/nDisplay/Source/DisplayCluster/Private/Render/Viewport/Configuration/DisplayClusterViewportConfigurationHelpers_ICVFX.cpp
+++ b/Engine/Plugins/Runtime/nDisplay/Source/DisplayCluster/Private/Render/Viewport/Configuration/DisplayClusterViewportConfigurationHelpers_ICVFX.cpp
@@ -348,10 +348,10 @@ FDisplayClusterViewport* FDisplayClusterViewportConfigurationHelpers_ICVFX::GetO
 	return ChromakeyViewport;
 }
 
-FDisplayClusterViewport* FDisplayClusterViewportConfigurationHelpers_ICVFX::GetOrCreateLightcardViewport(FDisplayClusterViewport& BaseViewport, ADisplayClusterRootActor& RootActor, bool bIsOpenColorIOViewportExist)
+FDisplayClusterViewport* FDisplayClusterViewportConfigurationHelpers_ICVFX::GetOrCreateLightcardViewport(FDisplayClusterViewport& BaseViewport, ADisplayClusterRootActor& RootActor)
 {
 	// Create new lightcard viewport
-	const FString ResourceId = bIsOpenColorIOViewportExist ? DisplayClusterViewportStrings::icvfx::lightcard_OCIO : DisplayClusterViewportStrings::icvfx::lightcard;
+	const FString ResourceId = DisplayClusterViewportStrings::icvfx::lightcard;
 
 	FDisplayClusterViewport* LightcardViewport = ImplFindViewport(RootActor, BaseViewport.GetId(), ResourceId);
 	if (LightcardViewport == nullptr)
@@ -374,10 +374,6 @@ FDisplayClusterViewport* FDisplayClusterViewportConfigurationHelpers_ICVFX::GetO
 
 	// Add viewport ICVFX usage as Lightcard
 	LightcardViewport->RenderSettingsICVFX.RuntimeFlags |= ViewportRuntime_ICVFXLightcard;
-	if (bIsOpenColorIOViewportExist)
-	{
-		LightcardViewport->RenderSettingsICVFX.RuntimeFlags |= ViewportRuntime_ICVFXLightcardAlpha;
-	}
 
 	return LightcardViewport;
 }
@@ -718,7 +714,7 @@ bool FDisplayClusterViewportConfigurationHelpers_ICVFX::IsShouldUseLightcard(con
 	return FDisplayClusterViewportConfigurationHelpers_Visibility::IsValid(InLightcardSettings.ShowOnlyList);
 }
 
-void FDisplayClusterViewportConfigurationHelpers_ICVFX::UpdateLightcardViewportSetting(FDisplayClusterViewport& DstViewport, FDisplayClusterViewport& BaseViewport, ADisplayClusterRootActor& RootActor, bool bIsOpenColorIOViewportExist)
+void FDisplayClusterViewportConfigurationHelpers_ICVFX::UpdateLightcardViewportSetting(FDisplayClusterViewport& DstViewport, FDisplayClusterViewport& BaseViewport, ADisplayClusterRootActor& RootActor)
 {
 	const FDisplayClusterConfigurationICVFX_StageSettings& StageSettings = RootActor.GetStageSettings();
 	const FDisplayClusterConfigurationICVFX_LightcardSettings& LightcardSettings = StageSettings.Lightcard;
@@ -731,32 +727,10 @@ void FDisplayClusterViewportConfigurationHelpers_ICVFX::UpdateLightcardViewportS
 	// LIghtcard texture used as overlay
 	DstViewport.RenderSettings.bVisible = false;
 
-	bool bIsLightcardUseResolvedScene = false;
-	if (!bIsOpenColorIOViewportExist)
-	{
-		if (FDisplayClusterViewportConfigurationHelpers_Postprocess::UpdateLightcardPostProcessSettings(DstViewport, BaseViewport, RootActor))
-		{
-			bIsLightcardUseResolvedScene = true;
-		}
+	FDisplayClusterViewportConfigurationHelpers_Postprocess::UpdateLightcardPostProcessSettings(DstViewport, BaseViewport, RootActor);
+	FDisplayClusterViewportConfigurationHelpers_OpenColorIO::UpdateLightcardViewport(DstViewport, BaseViewport, RootActor);
 
-		if (FDisplayClusterViewportConfigurationHelpers_OpenColorIO::UpdateLightcardViewport(DstViewport, BaseViewport, RootActor))
-		{
-			bIsLightcardUseResolvedScene = true;
-		}
-	}
-
-	if (bIsLightcardUseResolvedScene)
-	{
-		// ligthcard used OCIO - capture 2 vp, [1] color with OCIO and [2] alpha
-		// to support OCIO capture mode must be without alpha, so we need to render second vp with alpha
-		DstViewport.RenderSettings.CaptureMode = EDisplayClusterViewportCaptureMode::Lightcard_OCIO;
-		DstViewport.RenderSettingsICVFX.RuntimeFlags |= ViewportRuntime_ICVFXLightcardColor;
-	}
-	else
-	{
-		// no OCIO, capture 1 vp - color+alpha
-		DstViewport.RenderSettings.CaptureMode = EDisplayClusterViewportCaptureMode::Lightcard;
-	}
+	DstViewport.RenderSettings.CaptureMode = EDisplayClusterViewportCaptureMode::Lightcard;
 
 	const FDisplayClusterConfigurationICVFX_LightcardRenderSettings& InRenderSettings = LightcardSettings.RenderSettings;
 	{
@@ -800,14 +774,7 @@ void FDisplayClusterViewportConfigurationHelpers_ICVFX::UpdateLightcardViewportS
 	}
 	else
 	{
-		if (bIsLightcardUseResolvedScene)
-		{
-			BaseViewport.RenderSettingsICVFX.ICVFX.Lightcard_OCIO.ViewportId = DstViewport.GetId();
-		}
-		else
-		{
-			BaseViewport.RenderSettingsICVFX.ICVFX.Lightcard.ViewportId = DstViewport.GetId();
-		}
+		BaseViewport.RenderSettingsICVFX.ICVFX.Lightcard.ViewportId = DstViewport.GetId();
 	}
 }
 
diff --git a/Engine/Plugins/Runtime/nDisplay/Source/DisplayCluster/Private/Render/Viewport/Configuration/DisplayClusterViewportConfigurationHelpers_ICVFX.h b/Engine/Plugins/Runtime/nDisplay/Source/DisplayCluster/Private/Render/Viewport/Configuration/DisplayClusterViewportConfigurationHelpers_ICVFX.h
--- a/Engine/Plugins/Runtime/nDisplay/Source/DisplayCluster/Private/Render/Viewport/Configuration/DisplayClusterViewportConfigurationHelpers_ICVFX.h
+++ b/Engine/Plugins/Runtime/nDisplay/Source/DisplayCluster/Private/Render/Viewport/Configuration/DisplayClusterViewportConfigurationHelpers_ICVFX.h
@@ -37,7 +37,7 @@ public:
 
 	static FDisplayClusterViewport* GetOrCreateCameraViewport(ADisplayClusterRootActor& RootActor, UDisplayClusterICVFXCameraComponent& InCameraComponent);
 	static FDisplayClusterViewport* GetOrCreateChromakeyViewport(ADisplayClusterRootActor& RootActor, UDisplayClusterICVFXCameraComponent& InCameraComponent);
-	static FDisplayClusterViewport* GetOrCreateLightcardViewport(FDisplayClusterViewport& BaseViewport, ADisplayClusterRootActor& RootActor, bool bIsOpenColorIOViewportExist);
+	static FDisplayClusterViewport* GetOrCreateLightcardViewport(FDisplayClusterViewport& BaseViewport, ADisplayClusterRootActor& RootActor);
 
 #if WITH_EDITOR
 	static TArray<FDisplayClusterViewport*> PreviewGetRenderedInCameraViewports(ADisplayClusterRootActor& RootActor, UDisplayClusterICVFXCameraComponent& InCameraComponent, bool bGetChromakey = false);
@@ -47,7 +47,7 @@ public:
 
 	static void UpdateCameraViewportSettings(FDisplayClusterViewport& DstViewport, ADisplayClusterRootActor& RootActor, UDisplayClusterICVFXCameraComponent& InCameraComponent);
 	static void UpdateChromakeyViewportSettings(FDisplayClusterViewport& DstViewport, FDisplayClusterViewport& InCameraViewport, ADisplayClusterRootActor& RootActor, UDisplayClusterICVFXCameraComponent& InCameraComponent);
-	static void UpdateLightcardViewportSetting(FDisplayClusterViewport& DstViewport, FDisplayClusterViewport& BaseViewport, ADisplayClusterRootActor& RootActor, bool bIsOpenColorIOViewportExist);
+	static void UpdateLightcardViewportSetting(FDisplayClusterViewport& DstViewport, FDisplayClusterViewport& BaseViewport, ADisplayClusterRootActor& RootActor);
 
 	static bool IsCameraUsed(UDisplayClusterICVFXCameraComponent& InCameraComponent);
 	static bool GetCameraContext(ADisplayClusterRootActor& RootActor, UDisplayClusterICVFXCameraComponent& InCameraComponent, FCameraContext_ICVFX& OutCameraContext);
diff --git a/Engine/Plugins/Runtime/nDisplay/Source/DisplayCluster/Private/Render/Viewport/Configuration/DisplayClusterViewportConfigurationICVFX.cpp b/Engine/Plugins/Runtime/nDisplay/Source/DisplayCluster/Private/Render/Viewport/Configuration/DisplayClusterViewportConfigurationICVFX.cpp
--- a/Engine/Plugins/Runtime/nDisplay/Source/DisplayCluster/Private/Render/Viewport/Configuration/DisplayClusterViewportConfigurationICVFX.cpp
+++ b/Engine/Plugins/Runtime/nDisplay/Source/DisplayCluster/Private/Render/Viewport/Configuration/DisplayClusterViewportConfigurationICVFX.cpp
@@ -279,30 +279,15 @@ FDisplayClusterViewportConfigurationICVFX::~FDisplayClusterViewportConfiguration
 
 bool FDisplayClusterViewportConfigurationICVFX::CreateLightcardViewport(FDisplayClusterViewport& BaseViewport)
 {
-	FDisplayClusterViewport* LightcardViewport = FDisplayClusterViewportConfigurationHelpers_ICVFX::GetOrCreateLightcardViewport(BaseViewport, RootActor, false);
+	FDisplayClusterViewport* LightcardViewport = FDisplayClusterViewportConfigurationHelpers_ICVFX::GetOrCreateLightcardViewport(BaseViewport, RootActor);
 	if (LightcardViewport)
 	{
 		// Update lightcard viewport settings
-		FDisplayClusterViewportConfigurationHelpers_ICVFX::UpdateLightcardViewportSetting(*LightcardViewport, BaseViewport, RootActor, false);
+		FDisplayClusterViewportConfigurationHelpers_ICVFX::UpdateLightcardViewportSetting(*LightcardViewport, BaseViewport, RootActor);
 
 		// Support projection policy update
 		FDisplayClusterViewportConfigurationHelpers::UpdateProjectionPolicy(*LightcardViewport);
 
-		// When lightcard use OCIO or PP for render we need resolve scene. Now get alpha channel by second render pass
-		if ((LightcardViewport->RenderSettingsICVFX.RuntimeFlags & ViewportRuntime_ICVFXLightcardColor) != 0)
-		{
-			// Now OCIO require second vp for alpha
-			FDisplayClusterViewport* LightcardViewportAlpha = FDisplayClusterViewportConfigurationHelpers_ICVFX::GetOrCreateLightcardViewport(BaseViewport, RootActor, true);
-			if (LightcardViewportAlpha)
-			{
-				// Update lightcard viewport settings
-				FDisplayClusterViewportConfigurationHelpers_ICVFX::UpdateLightcardViewportSetting(*LightcardViewportAlpha, BaseViewport, RootActor, true);
-
-				// Support projection policy update
-				FDisplayClusterViewportConfigurationHelpers::UpdateProjectionPolicy(*LightcardViewportAlpha);
-			}
-		}
-
 		return true;
 	}
 
@@ -411,7 +396,7 @@ void FDisplayClusterViewportConfigurationICVFX::Update()
 				// Freeze render for lightcards when outer viewports freezed
 				if (StageSettings.Lightcard.bIgnoreOuterViewportsFreezingForLightcards == false)
 				{
-					const EDisplayClusterViewportRuntimeICVFXFlags LightcardViewportMask = ViewportRuntime_ICVFXLightcard | ViewportRuntime_ICVFXLightcardColor | ViewportRuntime_ICVFXLightcardAlpha;
+					const EDisplayClusterViewportRuntimeICVFXFlags LightcardViewportMask = ViewportRuntime_ICVFXLightcard;
 
 					for (FDisplayClusterViewport* ViewportIt : ViewportManager->ImplGetViewports())
 					{
diff --git a/Engine/Plugins/Runtime/nDisplay/Source/DisplayCluster/Private/Render/Viewport/Containers/DisplayClusterViewportProxyData.h b/Engine/Plugins/Runtime/nDisplay/Source/DisplayCluster/Private/Render/Viewport/Containers/DisplayClusterViewportProxyData.h
--- a/Engine/Plugins/Runtime/nDisplay/Source/DisplayCluster/Private/Render/Viewport/Containers/DisplayClusterViewportProxyData.h
+++ b/Engine/Plugins/Runtime/nDisplay/Source/DisplayCluster/Private/Render/Viewport/Containers/DisplayClusterViewportProxyData.h
@@ -27,7 +27,7 @@ public:
 	void UpdateProxy_RenderThread() const;
 
 private:
-	FDisplayClusterViewportProxy* DstViewportProxy;
+	TSharedPtr<FDisplayClusterViewportProxy, ESPMode::ThreadSafe> DstViewportProxy;
 
 	TSharedPtr<IDisplayClusterRender_MeshComponent, ESPMode::ThreadSafe> RemapMesh;
 
diff --git a/Engine/Plugins/Runtime/nDisplay/Source/DisplayCluster/Private/Render/Viewport/DisplayClusterViewport.cpp b/Engine/Plugins/Runtime/nDisplay/Source/DisplayCluster/Private/Render/Viewport/DisplayClusterViewport.cpp
--- a/Engine/Plugins/Runtime/nDisplay/Source/DisplayCluster/Private/Render/Viewport/DisplayClusterViewport.cpp
+++ b/Engine/Plugins/Runtime/nDisplay/Source/DisplayCluster/Private/Render/Viewport/DisplayClusterViewport.cpp
@@ -4,6 +4,7 @@
 #include "Render/Viewport/DisplayClusterViewportHelpers.h"
 #include "Render/Viewport/DisplayClusterViewportManager.h"
 #include "Render/Viewport/DisplayClusterViewportProxy.h"
+#include "Render/Viewport/DisplayClusterViewportManagerProxy.h"
 
 #include "Render/Viewport/Configuration/DisplayClusterViewportConfiguration.h"
 
@@ -37,14 +38,37 @@ FDisplayClusterViewport::FDisplayClusterViewport(FDisplayClusterViewportManager&
 	check(UninitializedProjectionPolicy.IsValid());
 
 	// Create scene proxy pair with on game thread. Outside, in ViewportManager added to proxy array on render thread
-	ViewportProxy = new FDisplayClusterViewportProxy(Owner.ImplGetProxy(), *this);
+	ViewportProxy = MakeShared<FDisplayClusterViewportProxy, ESPMode::ThreadSafe>(*this);
+
+	// Add viewport proxy on renderthread
+	ENQUEUE_RENDER_COMMAND(CreateDisplayClusterViewportProxy)(
+		[ViewportManagerProxy = Owner.GetViewportManagerProxy(), ViewportProxy = ViewportProxy](FRHICommandListImmediate& RHICmdList)
+		{
+			ViewportManagerProxy->CreateViewport_RenderThread(ViewportProxy);
+		}
+	);
 }
 
 FDisplayClusterViewport::~FDisplayClusterViewport()
 {
+	// Handle projection policy event
+	ProjectionPolicy.Reset();
+	UninitializedProjectionPolicy.Reset();
+
+	// Remove viewport proxy on render_thread
+	ENQUEUE_RENDER_COMMAND(DeleteDisplayClusterViewportProxy)(
+		[ViewportManagerProxy = Owner.GetViewportManagerProxy(), ViewportProxy = ViewportProxy](FRHICommandListImmediate& RHICmdList)
+		{
+			ViewportManagerProxy->DeleteViewport_RenderThread(ViewportProxy);
+		}
+	);
+
+	ViewportProxy.Reset();
+
 	HandleEndScene();
 
-	// ViewportProxy deleted on render thread from FDisplayClusterViewportManagerProxy::ImplDeleteViewport()
+	// Reset RTT size after viewport delete
+	Owner.ResetSceneRenderTargetSize();
 }
 
 IDisplayClusterViewportManager& FDisplayClusterViewport::GetOwner() const
@@ -220,7 +244,6 @@ void FDisplayClusterViewport::SetupSceneView(uint32 ContextNum, class UWorld* Wo
 	{
 	case EDisplayClusterViewportCaptureMode::Chromakey:
 	case EDisplayClusterViewportCaptureMode::Lightcard:
-	case EDisplayClusterViewportCaptureMode::Lightcard_OCIO:
 
 		InOutView.bAllowRayTracing = false;
 		break;
diff --git a/Engine/Plugins/Runtime/nDisplay/Source/DisplayCluster/Private/Render/Viewport/DisplayClusterViewport.h b/Engine/Plugins/Runtime/nDisplay/Source/DisplayCluster/Private/Render/Viewport/DisplayClusterViewport.h
--- a/Engine/Plugins/Runtime/nDisplay/Source/DisplayCluster/Private/Render/Viewport/DisplayClusterViewport.h
+++ b/Engine/Plugins/Runtime/nDisplay/Source/DisplayCluster/Private/Render/Viewport/DisplayClusterViewport.h
@@ -122,6 +122,11 @@ public:
 
 	virtual IDisplayClusterViewportManager& GetOwner() const override;
 
+	FDisplayClusterViewportManager& ImplGetOwner() const
+	{
+		return Owner;
+	}
+
 	const FDisplayClusterRenderFrameSettings& GetRenderFrameSettings() const;
 
 	//////////////////////////////////////////////////////
@@ -246,7 +251,7 @@ protected:
 	friend FDisplayClusterViewportRemap;
 
 	// viewport render thread data
-	FDisplayClusterViewportProxy* ViewportProxy = nullptr;
+	TSharedPtr<FDisplayClusterViewportProxy, ESPMode::ThreadSafe> ViewportProxy;
 
 	// Unique viewport name
 	const FString ViewportId;
diff --git a/Engine/Plugins/Runtime/nDisplay/Source/DisplayCluster/Private/Render/Viewport/DisplayClusterViewportManager.cpp b/Engine/Plugins/Runtime/nDisplay/Source/DisplayCluster/Private/Render/Viewport/DisplayClusterViewportManager.cpp
--- a/Engine/Plugins/Runtime/nDisplay/Source/DisplayCluster/Private/Render/Viewport/DisplayClusterViewportManager.cpp
+++ b/Engine/Plugins/Runtime/nDisplay/Source/DisplayCluster/Private/Render/Viewport/DisplayClusterViewportManager.cpp
@@ -60,9 +60,9 @@ FDisplayClusterViewportManager::FDisplayClusterViewportManager()
 	Configuration      = MakeUnique<FDisplayClusterViewportConfiguration>(*this);
 	RenderFrameManager = MakeUnique<FDisplayClusterRenderFrameManager>();
 
-	ViewportManagerProxy = new FDisplayClusterViewportManagerProxy();
+	ViewportManagerProxy = MakeShared<FDisplayClusterViewportManagerProxy, ESPMode::ThreadSafe>();
 
-	RenderTargetManager = MakeShared<FDisplayClusterRenderTargetManager, ESPMode::ThreadSafe>(ViewportManagerProxy);
+	RenderTargetManager = MakeShared<FDisplayClusterRenderTargetManager, ESPMode::ThreadSafe>(ViewportManagerProxy.Get());
 	PostProcessManager  = MakeShared<FDisplayClusterViewportPostProcessManager, ESPMode::ThreadSafe>(*this);
 	LightCardManager = MakeShared<FDisplayClusterViewportLightCardManager, ESPMode::ThreadSafe>(*this);
 
@@ -76,30 +76,41 @@ FDisplayClusterViewportManager::FDisplayClusterViewportManager()
 FDisplayClusterViewportManager::~FDisplayClusterViewportManager()
 {
 	// Remove viewports
+	Viewports.Reset();
+	ClusterNodeViewports.Reset();
+
+	RenderTargetManager.Reset();
+	PostProcessManager.Reset();
+
+	if (LightCardManager.IsValid())
 	{
-		TArray<FDisplayClusterViewport*> ExistViewports = Viewports;
-		for (FDisplayClusterViewport* Viewport : ExistViewports)
-		{
-			ImplDeleteViewport(Viewport);
-		}
-		ExistViewports.Empty();
+		LightCardManager->Release();
+		LightCardManager.Reset();
 	}
 
-	if (ViewportManagerProxy)
+	Configuration.Reset();
+
+	if (ViewportManagerProxy.IsValid())
 	{
-		ViewportManagerProxy->ImplSafeRelease();
-		ViewportManagerProxy = nullptr;
+		// Remove viewport manager proxy on render_thread
+		ENQUEUE_RENDER_COMMAND(DeleteDisplayClusterViewportManagerProxy)(
+			[ViewportManagerProxy = ViewportManagerProxy](FRHICommandListImmediate& RHICmdList)
+			{
+				ViewportManagerProxy->Release_RenderThread();
+			});
+
+		ViewportManagerProxy.Reset();
 	}
 }
 
 const IDisplayClusterViewportManagerProxy* FDisplayClusterViewportManager::GetProxy() const
 {
-	return ViewportManagerProxy;
+	return ViewportManagerProxy.Get();
 }
 
 IDisplayClusterViewportManagerProxy* FDisplayClusterViewportManager::GetProxy()
 {
-	return ViewportManagerProxy;
+	return ViewportManagerProxy.Get();
 }
 
 UWorld* FDisplayClusterViewportManager::GetCurrentWorld() const
@@ -612,13 +623,6 @@ FSceneViewFamily::ConstructionValues FDisplayClusterViewportManager::CreateViewF
 	{
 		case EDisplayClusterViewportCaptureMode::Chromakey:
 		case EDisplayClusterViewportCaptureMode::Lightcard:
-		case EDisplayClusterViewportCaptureMode::Lightcard_OCIO:
-
-			if (InFrameTarget.CaptureMode != EDisplayClusterViewportCaptureMode::Lightcard_OCIO)
-			{
-				bResolveScene = false;
-				InEngineShowFlags.PostProcessing = 0;
-			}
 
 			InEngineShowFlags.SetAtmosphere(0);
 			InEngineShowFlags.SetFog(0);
@@ -649,7 +653,6 @@ FSceneViewFamily::ConstructionValues FDisplayClusterViewportManager::CreateViewF
 		break;
 
 	case EDisplayClusterViewportCaptureMode::Lightcard:
-	case EDisplayClusterViewportCaptureMode::Lightcard_OCIO:
 
 		if (!GDisplayClusterLightcardsAllowNanite)
 		{
@@ -817,6 +820,7 @@ bool FDisplayClusterViewportManager::DeleteViewport(const FString& ViewportId)
 	if (ExistViewport != nullptr)
 	{
 		ImplDeleteViewport(ExistViewport);
+
 		return true;
 	}
 
@@ -840,9 +844,6 @@ FDisplayClusterViewport* FDisplayClusterViewportManager::ImplCreateViewport(cons
 	Viewports.Add(NewViewport);
 	ClusterNodeViewports.Add(NewViewport);
 
-	// Add viewport proxy on renderthread
-	ViewportManagerProxy->ImplCreateViewport(NewViewport->ViewportProxy);
-
 	// Handle start scene for viewport
 	NewViewport->HandleStartScene();
 
@@ -855,11 +856,8 @@ void FDisplayClusterViewportManager::ImplDeleteViewport(FDisplayClusterViewport*
 	ExistViewport->ProjectionPolicy.Reset();
 	ExistViewport->UninitializedProjectionPolicy.Reset();
 
-	// Delete viewport proxy on render thread
-	ViewportManagerProxy->ImplDeleteViewport(ExistViewport->ViewportProxy);
-
 	{
-		// Remove viewport obj from manager
+		// Remove viewport from the whole viewports list
 		int32 ViewportIndex = Viewports.Find(ExistViewport);
 		if (ViewportIndex != INDEX_NONE)
 		{
@@ -869,7 +867,7 @@ void FDisplayClusterViewportManager::ImplDeleteViewport(FDisplayClusterViewport*
 	}
 
 	{
-		// Remove cluster node viewport
+		// Remove viewport from the cluster viewports list
 		int32 ViewportIndex = ClusterNodeViewports.Find(ExistViewport);
 		if (ViewportIndex != INDEX_NONE)
 		{
diff --git a/Engine/Plugins/Runtime/nDisplay/Source/DisplayCluster/Private/Render/Viewport/DisplayClusterViewportManager.h b/Engine/Plugins/Runtime/nDisplay/Source/DisplayCluster/Private/Render/Viewport/DisplayClusterViewportManager.h
--- a/Engine/Plugins/Runtime/nDisplay/Source/DisplayCluster/Private/Render/Viewport/DisplayClusterViewportManager.h
+++ b/Engine/Plugins/Runtime/nDisplay/Source/DisplayCluster/Private/Render/Viewport/DisplayClusterViewportManager.h
@@ -116,12 +116,11 @@ public:
 
 	static TSharedPtr<IDisplayClusterProjectionPolicy, ESPMode::ThreadSafe> CreateProjectionPolicy(const FString& InViewportId, const FDisplayClusterConfigurationProjection* InConfigurationProjectionPolicy);
 
-	const FDisplayClusterViewportManagerProxy& ImplGetProxy() const
+	TSharedPtr<FDisplayClusterViewportManagerProxy, ESPMode::ThreadSafe> GetViewportManagerProxy() const
 	{
-		check(ViewportManagerProxy);
 		check(IsInGameThread());
 
-		return *ViewportManagerProxy;
+		return ViewportManagerProxy;
 	}
 
 	const FDisplayClusterRenderFrameSettings& GetRenderFrameSettings() const;
@@ -133,10 +132,10 @@ public:
 	bool ShouldUseFullSizeFrameTargetableResource() const;
 
 	void SetViewportBufferRatio(FDisplayClusterViewport& DstViewport, float InBufferRatio);
+	void ResetSceneRenderTargetSize();
 
 private:
 	void UpdateDesiredNumberOfViews(FDisplayClusterRenderFrame& InOutRenderFrame);
-	void ResetSceneRenderTargetSize();
 	void UpdateSceneRenderTargetSize();
 	void HandleViewportRTTChanges(const TArray<FDisplayClusterViewport_Context>& PrevContexts, const TArray<FDisplayClusterViewport_Context>& Contexts);
 	void ImplUpdateClusterNodeViewports(const EDisplayClusterRenderFrameMode InRenderMode, const FString& InClusterNodeId);
@@ -159,7 +158,7 @@ private:
 	TArray<FDisplayClusterViewport*> ClusterNodeViewports;
 
 	/** Render thread proxy manager. Deleted on render thread */
-	FDisplayClusterViewportManagerProxy* ViewportManagerProxy = nullptr;
+	TSharedPtr<FDisplayClusterViewportManagerProxy, ESPMode::ThreadSafe> ViewportManagerProxy;
 
 	// Pointer to the current scene
 	TWeakObjectPtr<UWorld> CurrentWorldRef;
diff --git a/Engine/Plugins/Runtime/nDisplay/Source/DisplayCluster/Private/Render/Viewport/DisplayClusterViewportManagerProxy.cpp b/Engine/Plugins/Runtime/nDisplay/Source/DisplayCluster/Private/Render/Viewport/DisplayClusterViewportManagerProxy.cpp
--- a/Engine/Plugins/Runtime/nDisplay/Source/DisplayCluster/Private/Render/Viewport/DisplayClusterViewportManagerProxy.cpp
+++ b/Engine/Plugins/Runtime/nDisplay/Source/DisplayCluster/Private/Render/Viewport/DisplayClusterViewportManagerProxy.cpp
@@ -81,24 +81,35 @@ FDisplayClusterViewportManagerProxy::FDisplayClusterViewportManagerProxy()
 { }
 
 FDisplayClusterViewportManagerProxy::~FDisplayClusterViewportManagerProxy()
-{
-	Release();
-}
+{ }
 
-void FDisplayClusterViewportManagerProxy::Release()
+void FDisplayClusterViewportManagerProxy::Release_RenderThread()
 {
+	check(IsInRenderingThread());
+
 	// Delete viewport proxy objects
-	for (FDisplayClusterViewportProxy* ViewportProxy : ViewportProxies)
-	{
-		if (ViewportProxy)
-		{
-			delete ViewportProxy;
-		}
-	}
 	ViewportProxies.Empty();
 	ClusterNodeViewportProxies.Empty();
 
-	RenderTargetManager.Reset();
+	if (RenderTargetManager.IsValid())
+	{
+		RenderTargetManager->Release();
+		RenderTargetManager.Reset();
+	}
+
+	if (PostProcessManager.IsValid())
+	{
+		PostProcessManager->Release();
+		PostProcessManager.Reset();
+	}
+
+	if (LightCardManager.IsValid())
+	{
+		LightCardManager->Release();
+		LightCardManager.Reset();
+	}
+
+	ViewportManagerViewExtension.Reset();
 }
 
 void FDisplayClusterViewportManagerProxy::Initialize(FDisplayClusterViewportManager& InViewportManager)
@@ -106,6 +117,7 @@ void FDisplayClusterViewportManagerProxy::Initialize(FDisplayClusterViewportMana
 	RenderTargetManager = InViewportManager.RenderTargetManager;
 	PostProcessManager = InViewportManager.PostProcessManager;
 	LightCardManager = InViewportManager.LightCardManager;
+	ViewportManagerViewExtension = InViewportManager.ViewportManagerViewExtension;
 }
 
 void FDisplayClusterViewportManagerProxy::DeleteResource_RenderThread(FDisplayClusterViewportResource* InDeletedResourcePtr)
@@ -113,9 +125,9 @@ void FDisplayClusterViewportManagerProxy::DeleteResource_RenderThread(FDisplayCl
 	if (InDeletedResourcePtr)
 	{
 		// Handle resource refs must be removed from all viewports
-		for (FDisplayClusterViewportProxy* ViewportProxyIt : ViewportProxies)
+		for (TSharedPtr<FDisplayClusterViewportProxy, ESPMode::ThreadSafe>& ViewportProxyIt : ViewportProxies)
 		{
-			if (ViewportProxyIt)
+			if (ViewportProxyIt.IsValid())
 			{
 				ViewportProxyIt->HandleResourceDelete_RenderThread(InDeletedResourcePtr);
 			}
@@ -130,77 +142,43 @@ void FDisplayClusterViewportManagerProxy::ImplUpdateClusterNodeViewportProxies()
 {
 	ClusterNodeViewportProxies.Empty();
 	// Collect viewport proxies for rendered cluster node
-	for (FDisplayClusterViewportProxy* ViewportProxyIt : ViewportProxies)
+	for (TSharedPtr<FDisplayClusterViewportProxy, ESPMode::ThreadSafe>& ViewportProxyIt : ViewportProxies)
 	{
-		if (ViewportProxyIt && (RenderFrameSettings.ClusterNodeId.IsEmpty() || ViewportProxyIt->GetClusterNodeId() == RenderFrameSettings.ClusterNodeId))
+		if (ViewportProxyIt.IsValid() && (RenderFrameSettings.ClusterNodeId.IsEmpty() || ViewportProxyIt->GetClusterNodeId() == RenderFrameSettings.ClusterNodeId))
 		{
 			ClusterNodeViewportProxies.Add(ViewportProxyIt);
 		}
 	}
 }
 
-void FDisplayClusterViewportManagerProxy::ImplSafeRelease()
+void FDisplayClusterViewportManagerProxy::CreateViewport_RenderThread(const TSharedPtr<FDisplayClusterViewportProxy, ESPMode::ThreadSafe>& InViewportProxy)
 {
-	check(IsInGameThread());
-
-	if (RenderTargetManager.IsValid())
-	{
-		RenderTargetManager->Release();
-	}
-	
-	if (PostProcessManager.IsValid())
-	{
-		PostProcessManager->Release();
-	}
-
-	if (LightCardManager.IsValid())
-	{
-		LightCardManager->Release();
-	}
+	check(IsInRenderingThread());
 
-	// Remove viewport manager proxy on render_thread
-	ENQUEUE_RENDER_COMMAND(DeleteDisplayClusterViewportManagerProxy)(
-		[ViewportManagerProxy = this](FRHICommandListImmediate& RHICmdList)
-	{
-		delete ViewportManagerProxy;
-	});
+	ViewportProxies.Add(InViewportProxy);
+	ImplUpdateClusterNodeViewportProxies();
 }
 
-void FDisplayClusterViewportManagerProxy::ImplCreateViewport(FDisplayClusterViewportProxy* InViewportProxy)
+void FDisplayClusterViewportManagerProxy::DeleteViewport_RenderThread(const TSharedPtr<FDisplayClusterViewportProxy, ESPMode::ThreadSafe>& InViewportProxy)
 {
-	check(IsInGameThread());
+	check(IsInRenderingThread());
 
-	if (InViewportProxy)
+	// Remove viewport obj from manager
+	int32 ViewportProxyIndex = ViewportProxies.Find(InViewportProxy);
+	if (ViewportProxyIndex != INDEX_NONE)
 	{
-		ENQUEUE_RENDER_COMMAND(CreateDisplayClusterViewportProxy)(
-			[ViewportManagerProxy = this, ViewportProxy = InViewportProxy](FRHICommandListImmediate& RHICmdList)
-		{
-			ViewportManagerProxy->ViewportProxies.Add(ViewportProxy);
-			ViewportManagerProxy->ImplUpdateClusterNodeViewportProxies();
-		});
+		ViewportProxies[ViewportProxyIndex].Reset();
+		ViewportProxies.RemoveAt(ViewportProxyIndex);
 	}
-}
 
-void FDisplayClusterViewportManagerProxy::ImplDeleteViewport(FDisplayClusterViewportProxy* InViewportProxy)
-{
-	check(IsInGameThread());
-
-	// Remove viewport sceneproxy on renderthread
-	ENQUEUE_RENDER_COMMAND(DeleteDisplayClusterViewportProxy)(
-		[ViewportManagerProxy = this, ViewportProxy = InViewportProxy](FRHICommandListImmediate& RHICmdList)
+	int32 ClusterViewportProxyIndex = ClusterNodeViewportProxies.Find(InViewportProxy);
+	if (ClusterViewportProxyIndex != INDEX_NONE)
 	{
-		// Remove viewport obj from manager
-		int32 ViewportProxyIndex = ViewportManagerProxy->ViewportProxies.Find(ViewportProxy);
-		if (ViewportProxyIndex != INDEX_NONE)
-		{
-			ViewportManagerProxy->ViewportProxies[ViewportProxyIndex] = nullptr;
-			ViewportManagerProxy->ViewportProxies.RemoveAt(ViewportProxyIndex);
-			ViewportManagerProxy->ImplUpdateClusterNodeViewportProxies();
-		}
-
-		delete ViewportProxy;
-	});
+		ClusterNodeViewportProxies[ClusterViewportProxyIndex].Reset();
+		ClusterNodeViewportProxies.RemoveAt(ClusterViewportProxyIndex);
+	}
 }
+
 void FDisplayClusterViewportManagerProxy::ImplUpdateRenderFrameSettings(const FDisplayClusterRenderFrameSettings& InRenderFrameSettings)
 {
 	check(IsInGameThread());
@@ -209,7 +187,7 @@ void FDisplayClusterViewportManagerProxy::ImplUpdateRenderFrameSettings(const FD
 
 	// Send frame settings to renderthread
 	ENQUEUE_RENDER_COMMAND(DeleteDisplayClusterViewportProxy)(
-		[ViewportManagerProxy = this, Settings](FRHICommandListImmediate& RHICmdList)
+		[ViewportManagerProxy = SharedThis(this), Settings](FRHICommandListImmediate& RHICmdList)
 	{
 		ViewportManagerProxy->RenderFrameSettings = *Settings;
 		delete Settings;
@@ -247,11 +225,13 @@ DECLARE_GPU_STAT_NAMED(nDisplay_ViewportManager_RenderFrame, TEXT("nDisplay View
 void FDisplayClusterViewportManagerProxy::ImplRenderFrame(FViewport* InViewport)
 {
 	ENQUEUE_RENDER_COMMAND(DeleteDisplayClusterViewportProxy)(
-		[ViewportManagerProxy = this, InViewport](FRHICommandListImmediate& RHICmdList)
+		[InViewportManagerProxy = SharedThis(this), InViewport](FRHICommandListImmediate& RHICmdList)
 	{
 		SCOPED_GPU_STAT(RHICmdList, nDisplay_ViewportManager_RenderFrame);
 		SCOPED_DRAW_EVENT(RHICmdList, nDisplay_ViewportManager_RenderFrame);
 
+		const FDisplayClusterViewportManagerProxy* ViewportManagerProxy = &InViewportManagerProxy.Get();
+
 		// Handle render setup
 		if (ViewportManagerProxy->PostProcessManager.IsValid())
 		{
@@ -322,10 +302,10 @@ void FDisplayClusterViewportManagerProxy::UpdateDeferredResources_RenderThread(F
 {
 	check(IsInRenderingThread());
 
-	TArray<FDisplayClusterViewportProxy*> OverriddenViewports;
+	TArray<TSharedPtr<FDisplayClusterViewportProxy, ESPMode::ThreadSafe>> OverriddenViewports;
 	OverriddenViewports.Reserve(ClusterNodeViewportProxies.Num());
 
-	for (FDisplayClusterViewportProxy* ViewportProxy : ClusterNodeViewportProxies)
+	for (const TSharedPtr<FDisplayClusterViewportProxy, ESPMode::ThreadSafe>& ViewportProxy : ClusterNodeViewportProxies)
 	{
 		if (ViewportProxy->RenderSettings.OverrideViewportId.IsEmpty())
 		{
@@ -339,7 +319,7 @@ void FDisplayClusterViewportManagerProxy::UpdateDeferredResources_RenderThread(F
 	}
 
 	// Update deferred viewports after all
-	for (FDisplayClusterViewportProxy* ViewportProxy : OverriddenViewports)
+	for (TSharedPtr<FDisplayClusterViewportProxy, ESPMode::ThreadSafe>& ViewportProxy : OverriddenViewports)
 	{
 		ViewportProxy->UpdateDeferredResources(RHICmdList);
 	}
@@ -392,11 +372,11 @@ void FDisplayClusterViewportManagerProxy::UpdateFrameResources_RenderThread(FRHI
 	}
 
 	// Support viewport overlap order sorting:
-	TArray<FDisplayClusterViewportProxy*> SortedViewportProxy = ClusterNodeViewportProxies;
+	TArray<TSharedPtr<FDisplayClusterViewportProxy, ESPMode::ThreadSafe>> SortedViewportProxy = ClusterNodeViewportProxies;
 	SortedViewportProxy.Sort(
-		[](const FDisplayClusterViewportProxy& VP1, const FDisplayClusterViewportProxy& VP2)
+		[](const TSharedPtr<FDisplayClusterViewportProxy, ESPMode::ThreadSafe>& VP1, const TSharedPtr<FDisplayClusterViewportProxy, ESPMode::ThreadSafe>& VP2)
 		{
-			return  VP1.GetRenderSettings_RenderThread().OverlapOrder < VP2.GetRenderSettings_RenderThread().OverlapOrder;
+			return  VP1->GetRenderSettings_RenderThread().OverlapOrder < VP2->GetRenderSettings_RenderThread().OverlapOrder;
 		}
 	);
 
@@ -413,28 +393,28 @@ void FDisplayClusterViewportManagerProxy::UpdateFrameResources_RenderThread(FRHI
 	for (uint8 WarpPass = 0; WarpPass < (uint8)EWarpPass::COUNT; WarpPass++)
 	{
 		// Update deferred resources for viewports
-		for (FDisplayClusterViewportProxy* ViewportProxy : SortedViewportProxy)
+		for (TSharedPtr<FDisplayClusterViewportProxy, ESPMode::ThreadSafe>& ViewportProxyIt : SortedViewportProxy)
 		{
 			// Iterate over visible viewports:
-			if (ViewportProxy && ViewportProxy->GetRenderSettings_RenderThread().bVisible)
+			if (ViewportProxyIt.IsValid() && ViewportProxyIt->GetRenderSettings_RenderThread().bVisible)
 			{
-				if (bWarpBlendEnabled && ShouldApplyWarpBlend(ViewportProxy))
+				if (bWarpBlendEnabled && ShouldApplyWarpBlend(ViewportProxyIt.Get()))
 				{
-					const TSharedPtr<IDisplayClusterProjectionPolicy, ESPMode::ThreadSafe>& PrjPolicy = ViewportProxy->GetProjectionPolicy_RenderThread();
+					const TSharedPtr<IDisplayClusterProjectionPolicy, ESPMode::ThreadSafe>& PrjPolicy = ViewportProxyIt->GetProjectionPolicy_RenderThread();
 					switch ((EWarpPass)WarpPass)
 					{
 					case EWarpPass::Begin:
-						IDisplayCluster::Get().GetCallbacks().OnDisplayClusterPreWarpViewport_RenderThread().Broadcast(RHICmdList, ViewportProxy);
-						PrjPolicy->BeginWarpBlend_RenderThread(RHICmdList, ViewportProxy);
+						IDisplayCluster::Get().GetCallbacks().OnDisplayClusterPreWarpViewport_RenderThread().Broadcast(RHICmdList, ViewportProxyIt.Get());
+						PrjPolicy->BeginWarpBlend_RenderThread(RHICmdList, ViewportProxyIt.Get());
 						break;
 
 					case EWarpPass::Render:
-						PrjPolicy->ApplyWarpBlend_RenderThread(RHICmdList, ViewportProxy);
+						PrjPolicy->ApplyWarpBlend_RenderThread(RHICmdList, ViewportProxyIt.Get());
 						break;
 
 					case EWarpPass::End:
-						PrjPolicy->EndWarpBlend_RenderThread(RHICmdList, ViewportProxy);
-						IDisplayCluster::Get().GetCallbacks().OnDisplayClusterPostWarpViewport_RenderThread().Broadcast(RHICmdList, ViewportProxy);
+						PrjPolicy->EndWarpBlend_RenderThread(RHICmdList, ViewportProxyIt.Get());
+						IDisplayCluster::Get().GetCallbacks().OnDisplayClusterPostWarpViewport_RenderThread().Broadcast(RHICmdList, ViewportProxyIt.Get());
 						break;
 
 					default:
@@ -456,15 +436,15 @@ void FDisplayClusterViewportManagerProxy::UpdateFrameResources_RenderThread(FRHI
 
 	// Post resolve to Frame RTT
 	// All warp&blend results are now inside AdditionalTargetableResource. Viewport images of other projection policies are still stored in the InputShaderResource.
-	for (FDisplayClusterViewportProxy* ViewportProxy : SortedViewportProxy)
+	for (TSharedPtr<FDisplayClusterViewportProxy, ESPMode::ThreadSafe>& ViewportProxyIt : SortedViewportProxy)
 	{
 		// Iterate over visible viewports:
-		if (ViewportProxy && ViewportProxy->GetRenderSettings_RenderThread().bVisible)
+		if (ViewportProxyIt.IsValid() && ViewportProxyIt->GetRenderSettings_RenderThread().bVisible)
 		{
 			EDisplayClusterViewportResourceType ViewportSource = EDisplayClusterViewportResourceType::InputShaderResource;
-			if (bWarpBlendEnabled && ShouldApplyWarpBlend(ViewportProxy))
+			if (bWarpBlendEnabled && ShouldApplyWarpBlend(ViewportProxyIt.Get()))
 			{
-				const TSharedPtr<IDisplayClusterProjectionPolicy, ESPMode::ThreadSafe>& PrjPolicy = ViewportProxy->GetProjectionPolicy_RenderThread();
+				const TSharedPtr<IDisplayClusterProjectionPolicy, ESPMode::ThreadSafe>& PrjPolicy = ViewportProxyIt->GetProjectionPolicy_RenderThread();
 				if (PrjPolicy->ShouldUseAdditionalTargetableResource())
 				{
 					ViewportSource = EDisplayClusterViewportResourceType::AdditionalTargetableResource;
@@ -472,10 +452,10 @@ void FDisplayClusterViewportManagerProxy::UpdateFrameResources_RenderThread(FRHI
 			}
 
 			// resolve viewports to the frame target texture
-			ViewportProxy->ResolveResources_RenderThread(RHICmdList, ViewportSource, ViewportProxy->GetOutputResourceType_RenderThread());
+			ViewportProxyIt->ResolveResources_RenderThread(RHICmdList, ViewportSource, ViewportProxyIt->GetOutputResourceType_RenderThread());
 
 			// Apply post-warp (viewport remap, etc)
-			ViewportProxy->PostResolveViewport_RenderThread(RHICmdList);
+			ViewportProxyIt->PostResolveViewport_RenderThread(RHICmdList);
 		}
 	}
 
@@ -501,27 +481,27 @@ void FDisplayClusterViewportManagerProxy::DoCrossGPUTransfers_RenderThread(FRHIC
 	// Copy the view render results to all GPUs that are native to the viewport.
 	TArray<FTransferResourceParams> TransferResources;
 
-	for (FDisplayClusterViewportProxy* ViewportProxy : ClusterNodeViewportProxies)
+	for (const TSharedPtr<FDisplayClusterViewportProxy, ESPMode::ThreadSafe>& ViewportProxyIt : ClusterNodeViewportProxies)
 	{
 		bool bShouldCrossGPUTransfersViewportRenderResources = true;
 
 		// Skip a frozen viewport that has already been transferred between GPUs
 		// The first time freezing should do the transfer (RenderTargets must be assigned on the first pass)
-		if (ViewportProxy->RenderSettings.bFreezeRendering && ViewportProxy->RenderTargets.Num() == 0)
+		if (ViewportProxyIt->RenderSettings.bFreezeRendering && ViewportProxyIt->RenderTargets.Num() == 0)
 		{
 			bShouldCrossGPUTransfersViewportRenderResources = false;
 		}
 
 		if (bShouldCrossGPUTransfersViewportRenderResources)
 		{
-			for (FDisplayClusterViewport_Context& ViewportContext : ViewportProxy->Contexts)
+			for (FDisplayClusterViewport_Context& ViewportContext : ViewportProxyIt->Contexts)
 			{
 				if (ViewportContext.bAllowGPUTransferOptimization && ViewportContext.GPUIndex >= 0)
 				{
 					// Use optimized cross GPU transfer for this context
 
-					FRenderTarget* RenderTarget = ViewportProxy->RenderTargets[ViewportContext.ContextNum];
-					FRHITexture2D* TextureRHI = ViewportProxy->RenderTargets[ViewportContext.ContextNum]->GetViewportRenderTargetResourceRHI();
+					FRenderTarget* RenderTarget = ViewportProxyIt->RenderTargets[ViewportContext.ContextNum];
+					FRHITexture2D* TextureRHI = ViewportProxyIt->RenderTargets[ViewportContext.ContextNum]->GetViewportRenderTargetResourceRHI();
 
 					FRHIGPUMask RenderTargetGPUMask = (GNumExplicitGPUsForRendering > 1 && RenderTarget) ? RenderTarget->GetGPUMask(RHICmdList) : FRHIGPUMask::GPU0();
 					FRHIGPUMask ContextGPUMask = FRHIGPUMask::FromIndex(ViewportContext.GPUIndex);
@@ -561,12 +541,12 @@ bool FDisplayClusterViewportManagerProxy::GetFrameTargets_RenderThread(TArray<FR
 	check(IsInRenderingThread());
 
 	// Get any defined frame targets from first visible viewport
-	for (FDisplayClusterViewportProxy* ViewportProxy : ClusterNodeViewportProxies)
+	for (const TSharedPtr<FDisplayClusterViewportProxy, ESPMode::ThreadSafe>& ViewportProxyIt : ClusterNodeViewportProxies)
 	{
-		if (ViewportProxy)
+		if (ViewportProxyIt.IsValid())
 		{
-			const TArray<FDisplayClusterViewportTextureResource*>& Frames = ViewportProxy->OutputFrameTargetableResources;
-			const TArray<FDisplayClusterViewportTextureResource*>& AdditionalFrames = ViewportProxy->AdditionalFrameTargetableResources;
+			const TArray<FDisplayClusterViewportTextureResource*>& Frames = ViewportProxyIt->OutputFrameTargetableResources;
+			const TArray<FDisplayClusterViewportTextureResource*>& AdditionalFrames = ViewportProxyIt->AdditionalFrameTargetableResources;
 
 			if (Frames.Num() > 0)
 			{
@@ -656,23 +636,23 @@ FDisplayClusterViewportProxy* FDisplayClusterViewportManagerProxy::ImplFindViewp
 	check(IsInRenderingThread());
 
 	// Ok, we have a request for a particular viewport. Let's find it.
-	FDisplayClusterViewportProxy* const* DesiredViewport = ViewportProxies.FindByPredicate([ViewportId](const FDisplayClusterViewportProxy* ItemViewport)
+	TSharedPtr<FDisplayClusterViewportProxy, ESPMode::ThreadSafe> const* DesiredViewport = ViewportProxies.FindByPredicate([ViewportId](const TSharedPtr<FDisplayClusterViewportProxy, ESPMode::ThreadSafe>& ItemViewport)
 	{
 		return ViewportId.Equals(ItemViewport->GetId(), ESearchCase::IgnoreCase);
 	});
 
-	return (DesiredViewport != nullptr) ? *DesiredViewport : nullptr;
+	return (DesiredViewport != nullptr) ? DesiredViewport->Get() : nullptr;
 }
 
 IDisplayClusterViewportProxy* FDisplayClusterViewportManagerProxy::FindViewport_RenderThread(const int32 StereoViewIndex, uint32* OutContextNum) const
 {
 	check(IsInRenderingThread());
 
-	for (FDisplayClusterViewportProxy* ViewportProxy : ViewportProxies)
+	for (const TSharedPtr<FDisplayClusterViewportProxy, ESPMode::ThreadSafe>&  ViewportProxyIt : ViewportProxies)
 	{
-		if (ViewportProxy && ViewportProxy->FindContext_RenderThread(StereoViewIndex, OutContextNum))
+		if (ViewportProxyIt.IsValid() && ViewportProxyIt->FindContext_RenderThread(StereoViewIndex, OutContextNum))
 		{
-			return ViewportProxy;
+			return ViewportProxyIt.Get();
 		}
 	}
 
diff --git a/Engine/Plugins/Runtime/nDisplay/Source/DisplayCluster/Private/Render/Viewport/DisplayClusterViewportManagerProxy.h b/Engine/Plugins/Runtime/nDisplay/Source/DisplayCluster/Private/Render/Viewport/DisplayClusterViewportManagerProxy.h
--- a/Engine/Plugins/Runtime/nDisplay/Source/DisplayCluster/Private/Render/Viewport/DisplayClusterViewportManagerProxy.h
+++ b/Engine/Plugins/Runtime/nDisplay/Source/DisplayCluster/Private/Render/Viewport/DisplayClusterViewportManagerProxy.h
@@ -8,6 +8,7 @@
 #include "Render/Viewport/DisplayClusterViewportProxy.h"
 #include "Render/Viewport/Containers/DisplayClusterViewport_Enums.h"
 #include "Render/Viewport/RenderFrame/DisplayClusterRenderFrameSettings.h"
+#include "Templates/SharedPointer.h"
 
 class FDisplayClusterRenderTargetManager;
 class FDisplayClusterViewportPostProcessManager;
@@ -22,12 +23,13 @@ class FViewport;
 
 class FDisplayClusterViewportManagerProxy
 	: public IDisplayClusterViewportManagerProxy
+	, public TSharedFromThis<FDisplayClusterViewportManagerProxy, ESPMode::ThreadSafe>
 {
 public:
 	FDisplayClusterViewportManagerProxy();
 	virtual ~FDisplayClusterViewportManagerProxy();
 
-	void Release();
+	void Release_RenderThread();
 
 public:
 	void DoCrossGPUTransfers_RenderThread(FRHICommandListImmediate& RHICmdList) const;
@@ -42,9 +44,9 @@ public:
 
 	virtual IDisplayClusterViewportProxy* FindViewport_RenderThread(const int32 StereoViewIndex, uint32* OutContextNum = nullptr) const override;
 
-	virtual const TArrayView<IDisplayClusterViewportProxy*> GetViewports_RenderThread() const override
+	virtual const TArrayView<TSharedPtr<IDisplayClusterViewportProxy, ESPMode::ThreadSafe>> GetViewports_RenderThread() const override
 	{
-		return TArrayView<IDisplayClusterViewportProxy*>((IDisplayClusterViewportProxy**)(ClusterNodeViewportProxies.GetData()), ClusterNodeViewportProxies.Num());
+		return TArrayView<TSharedPtr<IDisplayClusterViewportProxy, ESPMode::ThreadSafe>>((TSharedPtr<IDisplayClusterViewportProxy, ESPMode::ThreadSafe>*)(ClusterNodeViewportProxies.GetData()), ClusterNodeViewportProxies.Num());
 	}
 
 	virtual bool GetFrameTargets_RenderThread(TArray<FRHITexture2D*>& OutFrameResources, TArray<FIntPoint>& OutTargetOffsets, TArray<FRHITexture2D*>* OutAdditionalFrameResources = nullptr) const override;
@@ -53,6 +55,9 @@ public:
 	virtual TSharedPtr<IDisplayClusterViewportLightCardManager, ESPMode::ThreadSafe> GetLightCardManager_RenderThread() const override;
 
 	// internal use only
+	void CreateViewport_RenderThread(const TSharedPtr<FDisplayClusterViewportProxy, ESPMode::ThreadSafe>& InViewportProxy);
+	void DeleteViewport_RenderThread(const TSharedPtr<FDisplayClusterViewportProxy, ESPMode::ThreadSafe>& InViewportProxy);
+
 	void DeleteResource_RenderThread(FDisplayClusterViewportResource* InDeletedResourcePtr);
 	void Initialize(FDisplayClusterViewportManager& InViewportManager);
 
@@ -63,7 +68,7 @@ public:
 		return RenderFrameSettings;
 	}
 
-	const TArray<FDisplayClusterViewportProxy*>& ImplGetViewportProxies_RenderThread() const
+	const TArray<TSharedPtr<FDisplayClusterViewportProxy, ESPMode::ThreadSafe>>& ImplGetViewportProxies_RenderThread() const
 	{
 		check(IsInRenderingThread());
 
@@ -72,11 +77,8 @@ public:
 
 	FDisplayClusterViewportProxy* ImplFindViewport_RenderThread(const FString& InViewportId) const;
 
-	void ImplCreateViewport(FDisplayClusterViewportProxy* InViewportProxy);
-	void ImplDeleteViewport(FDisplayClusterViewportProxy* InViewportProxy);
 	void ImplUpdateRenderFrameSettings(const FDisplayClusterRenderFrameSettings& InRenderFrameSettings);
 	void ImplUpdateViewports(const TArray<FDisplayClusterViewport*>& InViewports);
-	void ImplSafeRelease();
 
 	void ImplRenderFrame(FViewport* InViewport);
 	void ImplClearFrameTargets_RenderThread(FRHICommandListImmediate& RHICmdList) const;
@@ -91,6 +93,9 @@ private:
 
 	FDisplayClusterRenderFrameSettings RenderFrameSettings;
 
-	TArray<FDisplayClusterViewportProxy*> ViewportProxies;
-	TArray<FDisplayClusterViewportProxy*> ClusterNodeViewportProxies;
+	TArray<TSharedPtr<FDisplayClusterViewportProxy, ESPMode::ThreadSafe>> ViewportProxies;
+	TArray<TSharedPtr<FDisplayClusterViewportProxy, ESPMode::ThreadSafe>> ClusterNodeViewportProxies;
+
+	// Handle special features
+	TSharedPtr<class FDisplayClusterViewportManagerViewExtension, ESPMode::ThreadSafe> ViewportManagerViewExtension;
 };
diff --git a/Engine/Plugins/Runtime/nDisplay/Source/DisplayCluster/Private/Render/Viewport/DisplayClusterViewportManagerViewExtension.cpp b/Engine/Plugins/Runtime/nDisplay/Source/DisplayCluster/Private/Render/Viewport/DisplayClusterViewportManagerViewExtension.cpp
--- a/Engine/Plugins/Runtime/nDisplay/Source/DisplayCluster/Private/Render/Viewport/DisplayClusterViewportManagerViewExtension.cpp
+++ b/Engine/Plugins/Runtime/nDisplay/Source/DisplayCluster/Private/Render/Viewport/DisplayClusterViewportManagerViewExtension.cpp
@@ -10,10 +10,27 @@
 #include "IDisplayClusterCallbacks.h"
 
 
+namespace DisplayClusterViewportManagerViewExtensionHelpers
+{
+	static const FName RendererModuleName(TEXT("Renderer"));
+};
+using namespace DisplayClusterViewportManagerViewExtensionHelpers;
+
+///////////////////////////////////////////////////////////////////////////////////////
+// FDisplayClusterViewportManagerViewExtension
+///////////////////////////////////////////////////////////////////////////////////////
 FDisplayClusterViewportManagerViewExtension::FDisplayClusterViewportManagerViewExtension(const FAutoRegister& AutoRegister, const IDisplayClusterViewportManager* InViewportManager)
 	: FSceneViewExtensionBase(AutoRegister)
 	, ViewportManager(InViewportManager)
-{ }
+	, ViewportManagerProxy(InViewportManager->GetProxy())
+{
+	RegisterCallbacks();
+}
+
+FDisplayClusterViewportManagerViewExtension::~FDisplayClusterViewportManagerViewExtension()
+{
+	UnregisterCallbacks();
+}
 
 void FDisplayClusterViewportManagerViewExtension::PreRenderView_RenderThread(FRDGBuilder& GraphBuilder, FSceneView& InView)
 {
@@ -42,21 +59,71 @@ bool FDisplayClusterViewportManagerViewExtension::IsActiveThisFrame_Internal(con
 	return false;
 }
 
-void FDisplayClusterViewportManagerViewExtension::PostRenderViewFamily_RenderThread(FRDGBuilder& GraphBuilder, FSceneViewFamily& InViewFamily)
+void FDisplayClusterViewportManagerViewExtension::PreRenderViewFamily_RenderThread(FRDGBuilder& GraphBuilder, FSceneViewFamily& InViewFamily)
 {
-	// FSceneViewFamily::ProfileDescription is used currently as a viewport ID holder, so we use it to find a proper viewport proxy.
-	if (ViewportManager)
+	Viewports.Empty();
+
+	// Get all viewport proxies that are uses in this view family
+	for (int32 ViewIndex = 0; ViewIndex < InViewFamily.Views.Num(); ViewIndex++)
 	{
-		if (const IDisplayClusterViewportManagerProxy* const ViewportMgrPoxy = ViewportManager->GetProxy())
+		if (const FSceneView* SceneView = InViewFamily.Views[ViewIndex])
 		{
-			if (!InViewFamily.ProfileDescription.IsEmpty())
+			FViewportProxy NewViewportProxy;
+			NewViewportProxy.ViewportProxy = ViewportManagerProxy->FindViewport_RenderThread(SceneView->StereoViewIndex, &NewViewportProxy.ViewportProxyContext.ContextNum);
+			NewViewportProxy.ViewportProxyContext.ViewFamilyProfileDescription = InViewFamily.ProfileDescription;
+
+			NewViewportProxy.ViewIndex = ViewIndex;
+
+			if (NewViewportProxy.IsEnabled())
 			{
-				if (IDisplayClusterViewportProxy* ViewportProxy = ViewportMgrPoxy->FindViewport_RenderThread(InViewFamily.ProfileDescription))
-				{
-					// Notify about viewport rendering finish
-					IDisplayCluster::Get().GetCallbacks().OnDisplayClusterPostRenderViewFamily_RenderThread().Broadcast(GraphBuilder, InViewFamily, ViewportProxy);
-				}
+				Viewports.Add(NewViewportProxy);
 			}
 		}
 	}
 }
+
+void FDisplayClusterViewportManagerViewExtension::PostRenderViewFamily_RenderThread(FRDGBuilder& GraphBuilder, FSceneViewFamily& InViewFamily)
+{
+	for (FViewportProxy& ViewportIt : Viewports)
+	{
+		if (const FSceneView* SceneView = InViewFamily.Views[ViewportIt.ViewIndex])
+		{
+			ViewportIt.ViewportProxy->PostRenderViewFamily_RenderThread(GraphBuilder, InViewFamily, *SceneView, ViewportIt.ViewportProxyContext);
+		}
+	}
+}
+
+//////////////////////////////////////////////////////////////////////////////////////////////
+void FDisplayClusterViewportManagerViewExtension::OnResolvedSceneColor_RenderThread(FRDGBuilder& GraphBuilder, const FSceneTextures& SceneTextures)
+{
+	check(IsInRenderingThread());
+
+	for (FViewportProxy& ViewportIt : Viewports)
+	{
+		ViewportIt.ViewportProxy->OnResolvedSceneColor_RenderThread(GraphBuilder, SceneTextures, ViewportIt.ViewportProxyContext);
+	}
+}
+
+void FDisplayClusterViewportManagerViewExtension::RegisterCallbacks()
+{
+	if (!ResolvedSceneColorCallbackHandle.IsValid())
+	{
+		if (IRendererModule* RendererModule = FModuleManager::GetModulePtr<IRendererModule>(RendererModuleName))
+		{
+			ResolvedSceneColorCallbackHandle = RendererModule->GetResolvedSceneColorCallbacks().AddRaw(this, &FDisplayClusterViewportManagerViewExtension::OnResolvedSceneColor_RenderThread);
+		}
+	}
+}
+
+void FDisplayClusterViewportManagerViewExtension::UnregisterCallbacks()
+{
+	if (ResolvedSceneColorCallbackHandle.IsValid())
+	{
+		if (IRendererModule* RendererModule = FModuleManager::GetModulePtr<IRendererModule>(RendererModuleName))
+		{
+			RendererModule->GetResolvedSceneColorCallbacks().Remove(ResolvedSceneColorCallbackHandle);
+		}
+
+		ResolvedSceneColorCallbackHandle.Reset();
+	}
+}
diff --git a/Engine/Plugins/Runtime/nDisplay/Source/DisplayCluster/Private/Render/Viewport/DisplayClusterViewportManagerViewExtension.h b/Engine/Plugins/Runtime/nDisplay/Source/DisplayCluster/Private/Render/Viewport/DisplayClusterViewportManagerViewExtension.h
--- a/Engine/Plugins/Runtime/nDisplay/Source/DisplayCluster/Private/Render/Viewport/DisplayClusterViewportManagerViewExtension.h
+++ b/Engine/Plugins/Runtime/nDisplay/Source/DisplayCluster/Private/Render/Viewport/DisplayClusterViewportManagerViewExtension.h
@@ -2,9 +2,10 @@
 
 #pragma once
 #include "SceneViewExtension.h"
+#include "Render/Viewport/Containers/DisplayClusterViewportProxy_Context.h"
 
 class IDisplayClusterViewportManager;
-
+class IDisplayClusterViewportManagerProxy;
 
 /**
  * View extension applying an DC Viewport features
@@ -13,6 +14,7 @@ class FDisplayClusterViewportManagerViewExtension : public FSceneViewExtensionBa
 {
 public:
 	FDisplayClusterViewportManagerViewExtension(const FAutoRegister& AutoRegister, const IDisplayClusterViewportManager* InViewportManager);
+	virtual ~FDisplayClusterViewportManagerViewExtension();
 
 public:
 	//~ Begin ISceneViewExtension interface
@@ -34,9 +36,37 @@ public:
 
 	virtual bool IsActiveThisFrame_Internal(const FSceneViewExtensionContext& Context) const override;
 
+	virtual void PreRenderViewFamily_RenderThread(FRDGBuilder& GraphBuilder, FSceneViewFamily& InViewFamily) override;
 	virtual void PostRenderViewFamily_RenderThread(FRDGBuilder& GraphBuilder, FSceneViewFamily& InViewFamily) override;
 	//~End ISceneVIewExtension interface
 
+private:
+	/** Rendered callback (get scene textures to share) */
+	void OnResolvedSceneColor_RenderThread(FRDGBuilder& GraphBuilder, const FSceneTextures& SceneTextures);
+
+	void RegisterCallbacks();
+	void UnregisterCallbacks();
+
+private:
+	FDelegateHandle ResolvedSceneColorCallbackHandle;
+
 private:
 	const IDisplayClusterViewportManager* ViewportManager;
+	const IDisplayClusterViewportManagerProxy* ViewportManagerProxy;
+
+	struct FViewportProxy
+	{
+		inline bool IsEnabled() const
+		{
+			return ViewportProxy != nullptr;
+		}
+
+		class IDisplayClusterViewportProxy* ViewportProxy = nullptr;
+		FDisplayClusterViewportProxy_Context ViewportProxyContext;
+
+		int32 ViewIndex = 0;
+	};
+
+	TArray<FViewportProxy> Viewports;
+	
 };
diff --git a/Engine/Plugins/Runtime/nDisplay/Source/DisplayCluster/Private/Render/Viewport/DisplayClusterViewportProxy.cpp b/Engine/Plugins/Runtime/nDisplay/Source/DisplayCluster/Private/Render/Viewport/DisplayClusterViewportProxy.cpp
--- a/Engine/Plugins/Runtime/nDisplay/Source/DisplayCluster/Private/Render/Viewport/DisplayClusterViewportProxy.cpp
+++ b/Engine/Plugins/Runtime/nDisplay/Source/DisplayCluster/Private/Render/Viewport/DisplayClusterViewportProxy.cpp
@@ -4,6 +4,7 @@
 
 #include "Render/Viewport/DisplayClusterViewportManagerProxy.h"
 #include "Render/Viewport/DisplayClusterViewport.h"
+#include "Render/Viewport/DisplayClusterViewportManager.h"
 
 #include "Render/Viewport/Containers/DisplayClusterViewport_PostRenderSettings.h"
 #include "Render/Viewport/RenderTarget/DisplayClusterRenderTargetResource.h"
@@ -15,6 +16,10 @@
 #include "Render/Containers/IDisplayClusterRender_MeshComponent.h"
 #include "Render/Containers/IDisplayClusterRender_MeshComponentProxy.h"
 
+#include "IDisplayCluster.h"
+#include "IDisplayClusterCallbacks.h"
+#include "IDisplayClusterShaders.h"
+
 #if WITH_EDITOR
 #include "DisplayClusterRootActor.h"
 #include "Render/Viewport/Containers/DisplayClusterViewportReadPixels.h"
@@ -33,8 +38,9 @@
 
 #include "ScreenRendering.h"
 #include "PostProcess/SceneFilterRendering.h"
-
-#include "IDisplayClusterShaders.h"
+#include "SceneRendering.h"
+#include "RenderGraphUtils.h"
+#include "RenderGraphResources.h"
 
 static TAutoConsoleVariable<int32> CVarDisplayClusterRenderOverscanResolve(
 	TEXT("nDisplay.render.overscan.resolve"),
@@ -148,12 +154,12 @@ namespace DisplayClusterViewportProxyHelpers
 using namespace DisplayClusterViewportProxyHelpers;
 
 ///////////////////////////////////////////////////////////////////////////////////////
-FDisplayClusterViewportProxy::FDisplayClusterViewportProxy(const FDisplayClusterViewportManagerProxy& InOwner, const FDisplayClusterViewport& RenderViewport)
+FDisplayClusterViewportProxy::FDisplayClusterViewportProxy(const FDisplayClusterViewport& RenderViewport)
 	: ViewportId(RenderViewport.ViewportId)
 	, ClusterNodeId(RenderViewport.ClusterNodeId)
 	, RenderSettings(RenderViewport.RenderSettings)
 	, ProjectionPolicy(RenderViewport.UninitializedProjectionPolicy)
-	, Owner(InOwner)
+	, Owner(RenderViewport.ImplGetOwner().GetViewportManagerProxy().ToSharedRef())
 	, ShadersAPI(IDisplayClusterShaders::Get())
 {
 	check(ProjectionPolicy.IsValid());
@@ -167,7 +173,7 @@ const IDisplayClusterViewportManagerProxy& FDisplayClusterViewportProxy::GetOwne
 {
 	check(IsInRenderingThread());
 
-	return Owner;
+	return Owner.Get();
 }
 
 void FDisplayClusterViewportProxy::HandleResourceDelete_RenderThread(FDisplayClusterViewportResource* InDeletedResourcePtr)
@@ -220,7 +226,7 @@ bool FDisplayClusterViewportProxy::ImplGetResources_RenderThread(const EDisplayC
 	{
 		if (InRecursionDepth < DisplayClusterViewportProxyResourcesOverrideRecursionDepthMax)
 		{
-			if (FDisplayClusterViewportProxy const* OverrideViewportProxy = Owner.ImplFindViewport_RenderThread(RenderSettings.OverrideViewportId))
+			if (FDisplayClusterViewportProxy const* OverrideViewportProxy = Owner->ImplFindViewport_RenderThread(RenderSettings.OverrideViewportId))
 			{
 				return OverrideViewportProxy->ImplGetResources_RenderThread(InResourceType, OutResources, InRecursionDepth + 1);
 			}
@@ -320,7 +326,7 @@ void FDisplayClusterViewportProxy::PostResolveViewport_RenderThread(FRHICommandL
 void FDisplayClusterViewportProxy::ImplViewportRemap_RenderThread(FRHICommandListImmediate& RHICmdList) const
 {
 #if WITH_EDITOR
-	switch(Owner.GetRenderFrameSettings_RenderThread().RenderMode)
+	switch(Owner->GetRenderFrameSettings_RenderThread().RenderMode)
 	{
 	case EDisplayClusterRenderFrameMode::PreviewInScene:
 		// Preview in editor not support this feature
@@ -369,7 +375,7 @@ void FDisplayClusterViewportProxy::ImplPreviewReadPixels_RenderThread(FRHIComman
 	// Now try to read until success:
 	if (bPreviewReadPixels && Contexts.Num() > 0 && OutputPreviewTargetableResource.IsValid())
 	{
-		check(Owner.GetRenderFrameSettings_RenderThread().RenderMode == EDisplayClusterRenderFrameMode::PreviewInScene);
+		check(Owner->GetRenderFrameSettings_RenderThread().RenderMode == EDisplayClusterRenderFrameMode::PreviewInScene);
 
 		// We should synchronize thread for preview read
 		FScopeLock Lock(&PreviewPixelsCSGuard);
@@ -427,7 +433,7 @@ EDisplayClusterViewportResourceType FDisplayClusterViewportProxy::GetOutputResou
 	check(IsInRenderingThread());
 
 #if WITH_EDITOR
-	switch(Owner.GetRenderFrameSettings_RenderThread().RenderMode)
+	switch(Owner->GetRenderFrameSettings_RenderThread().RenderMode)
 	{
 	case EDisplayClusterRenderFrameMode::PreviewInScene:
 		return EDisplayClusterViewportResourceType::OutputPreviewTargetableResource;
@@ -464,7 +470,7 @@ bool FDisplayClusterViewportProxy::ImplGetResourcesWithRects_RenderThread(const
 	{
 		if (InRecursionDepth < DisplayClusterViewportProxyResourcesOverrideRecursionDepthMax)
 		{
-			if (FDisplayClusterViewportProxy const* OverrideViewportProxy = Owner.ImplFindViewport_RenderThread(RenderSettings.OverrideViewportId))
+			if (FDisplayClusterViewportProxy const* OverrideViewportProxy = Owner->ImplFindViewport_RenderThread(RenderSettings.OverrideViewportId))
 			{
 				return OverrideViewportProxy->ImplGetResourcesWithRects_RenderThread(InResourceType, OutResources, OutResourceRects, InRecursionDepth + 1);
 			}
@@ -571,14 +577,14 @@ void FDisplayClusterViewportProxy::UpdateDeferredResources(FRHICommandListImmedi
 }
 
 template<class TScreenPixelShader>
-void ResampleCopyTextureImpl_RenderThread(FRHICommandListImmediate& RHICmdList, FRHITexture* SrcTexture, FRHITexture* DstTexture, const FIntRect& SrcRect, const FIntRect& DstRect)
+void ResampleCopyTextureImpl_RenderThread(FRHICommandListImmediate& RHICmdList, FRHITexture* SrcTexture, FRHITexture* DstTexture, const FIntRect& SrcRect, const FIntRect& DstRect, const bool bCopyOnlyAlphaChannel)
 {
 	// Texture format mismatch, use a shader to do the copy.
 	// #todo-renderpasses there's no explicit resolve here? Do we need one?
 	FRHIRenderPassInfo RPInfo(DstTexture, ERenderTargetActions::Load_Store);
 	RHICmdList.Transition(FRHITransitionInfo(DstTexture, ERHIAccess::Unknown, ERHIAccess::RTV));
 
-	RHICmdList.BeginRenderPass(RPInfo, TEXT("DisaplyClusterRender_ResampleTexture"));
+	RHICmdList.BeginRenderPass(RPInfo, TEXT("DisplayCluster_ResampleCopyTexture"));
 	{
 		FIntVector SrcSizeXYZ = SrcTexture->GetSizeXYZ();
 		FIntVector DstSizeXYZ = DstTexture->GetSizeXYZ();
@@ -591,11 +597,23 @@ void ResampleCopyTextureImpl_RenderThread(FRHICommandListImmediate& RHICmdList,
 		FGraphicsPipelineStateInitializer GraphicsPSOInit;
 		RHICmdList.ApplyCachedRenderTargets(GraphicsPSOInit);
 
-		GraphicsPSOInit.BlendState = TStaticBlendState<>::GetRHI();
+		if(bCopyOnlyAlphaChannel)
+		{
+			// Copy alpha channel from source to dest
+			// RT0ColorWriteMask,RT0ColorBlendOp,RT0ColorSrcBlend,RT0ColorDestBlend,RT0AlphaBlendOp,RT0AlphaSrcBlend,RT0AlphaDestBlend,
+			GraphicsPSOInit.BlendState = TStaticBlendState <CW_ALPHA, BO_Add, BF_One, BF_Zero, BO_Add, BF_One, BF_Zero>::GetRHI();
+		}
+		else
+		{
+			GraphicsPSOInit.BlendState = TStaticBlendState<>::GetRHI();
+		}
+
 		GraphicsPSOInit.RasterizerState = TStaticRasterizerState<>::GetRHI();
 		GraphicsPSOInit.DepthStencilState = TStaticDepthStencilState<false, CF_Always>::GetRHI();
 
 
+
+
 		FGlobalShaderMap* ShaderMap = GetGlobalShaderMap(GMaxRHIFeatureLevel);
 		TShaderMapRef<FScreenVS> VertexShader(ShaderMap);
 		TShaderMapRef<TScreenPixelShader> PixelShader(ShaderMap);
@@ -633,6 +651,11 @@ void ResampleCopyTextureImpl_RenderThread(FRHICommandListImmediate& RHICmdList,
 	RHICmdList.Transition(FRHITransitionInfo(DstTexture, ERHIAccess::Unknown, ERHIAccess::SRVMask));
 }
 
+void ImplResolveResourceAlpha(FRHICommandListImmediate& RHICmdList, FRHITexture2D* InputResource, const FIntRect& InputRect, FRHITexture2D* OutputResource, const FIntRect& OutputRect)
+{
+	ResampleCopyTextureImpl_RenderThread<FScreenPS>(RHICmdList, InputResource, OutputResource, InputRect, OutputRect, true);
+}
+
 void ImplResolveResource(FRHICommandListImmediate& RHICmdList, FRHITexture2D* InputResource, const FIntRect& InputRect, FRHITexture2D* OutputResource, const FIntRect& OutputRect, const bool bOutputIsMipsResource, const bool bOutputIsPreviewResource)
 {
 	check(InputResource);
@@ -641,7 +664,7 @@ void ImplResolveResource(FRHICommandListImmediate& RHICmdList, FRHITexture2D* In
 	if (bOutputIsPreviewResource)
 	{
 		// Preview require a normal alpha (re-invert)
-		ResampleCopyTextureImpl_RenderThread<FScreenPSInvertAlpha>(RHICmdList, InputResource, OutputResource, InputRect, OutputRect);
+		ResampleCopyTextureImpl_RenderThread<FScreenPSInvertAlpha>(RHICmdList, InputResource, OutputResource, InputRect, OutputRect, false);
 	}
 	else if (InputRect.Size() == OutputRect.Size() && InputResource->GetFormat() == OutputResource->GetFormat())
 	{
@@ -656,17 +679,36 @@ void ImplResolveResource(FRHICommandListImmediate& RHICmdList, FRHITexture2D* In
 	}
 	else
 	{
-		ResampleCopyTextureImpl_RenderThread<FScreenPS>(RHICmdList, InputResource, OutputResource, InputRect, OutputRect);
+		ResampleCopyTextureImpl_RenderThread<FScreenPS>(RHICmdList, InputResource, OutputResource, InputRect, OutputRect, false);
+	}
+}
+
+FIntRect FDisplayClusterViewportProxy::GetFinalContextRect(const EDisplayClusterViewportResourceType InputResourceType, const FIntRect& InRect) const
+{
+	// When resolving without warp, apply overscan
+	switch (InputResourceType)
+	{
+	case EDisplayClusterViewportResourceType::InternalRenderTargetResource:
+		if (OverscanSettings.bIsEnabled && CVarDisplayClusterRenderOverscanResolve.GetValueOnRenderThread() != 0)
+		{
+			// Support overscan crop
+			return OverscanSettings.OverscanPixels.GetInnerRect(InRect);
+		}
+		break;
+	default:
+		break;
 	}
+
+	return InRect;
 }
 
 // Resolve resource contexts
-bool FDisplayClusterViewportProxy::ResolveResources_RenderThread(FRHICommandListImmediate& RHICmdList, const EDisplayClusterViewportResourceType InputResourceType, const EDisplayClusterViewportResourceType OutputResourceType) const
+bool FDisplayClusterViewportProxy::ResolveResources_RenderThread(FRHICommandListImmediate& RHICmdList, const EDisplayClusterViewportResourceType InputResourceType, const EDisplayClusterViewportResourceType OutputResourceType, const int32 InContextNum) const
 {
-	return ImplResolveResources_RenderThread(RHICmdList, this, InputResourceType, OutputResourceType);
+	return ImplResolveResources_RenderThread(RHICmdList, this, InputResourceType, OutputResourceType, InContextNum);
 }
 
-bool FDisplayClusterViewportProxy::ImplResolveResources_RenderThread(FRHICommandListImmediate& RHICmdList, FDisplayClusterViewportProxy const* SourceProxy, const EDisplayClusterViewportResourceType InputResourceType, const EDisplayClusterViewportResourceType OutputResourceType) const
+bool FDisplayClusterViewportProxy::ImplResolveResources_RenderThread(FRHICommandListImmediate& RHICmdList, FDisplayClusterViewportProxy const* SourceProxy, const EDisplayClusterViewportResourceType InputResourceType, const EDisplayClusterViewportResourceType OutputResourceType, const int32 InContextNum) const
 {
 	check(IsInRenderingThread());
 
@@ -684,46 +726,163 @@ bool FDisplayClusterViewportProxy::ImplResolveResources_RenderThread(FRHICommand
 	const bool bOutputIsPreviewResource = false;
 #endif
 
-	TArray<FRHITexture2D*> InputResources, OutputResources;
-	TArray<FIntRect> InputResourcesRect, OutputResourcesRect;
-	if (SourceProxy->GetResourcesWithRects_RenderThread(InputResourceType, InputResources, InputResourcesRect) && GetResourcesWithRects_RenderThread(OutputResourceType, OutputResources, OutputResourcesRect))
+	TArray<FRHITexture2D*> SrcResources, DestResources;
+	TArray<FIntRect> SrcResourcesRect, DestResourcesRect;
+	if (SourceProxy->GetResourcesWithRects_RenderThread(InputResourceType, SrcResources, SrcResourcesRect) && GetResourcesWithRects_RenderThread(OutputResourceType, DestResources, DestResourcesRect))
 	{
-		const int32 InputAmount = FMath::Min(InputResources.Num(), OutputResources.Num());
-		for (int32 InputIt = 0; InputIt < InputAmount; InputIt++)
+		const int32 SrcResourcesAmount = FMath::Min(SrcResources.Num(), DestResources.Num());
+		for (int32 SrcResourceContextIndex = 0; SrcResourceContextIndex < SrcResourcesAmount; SrcResourceContextIndex++)
 		{
-			FIntRect SourceRect = InputResourcesRect[InputIt];
+			// Get context num to copy
+			const int32 SrcContextNum = (InContextNum == INDEX_NONE) ? SrcResourceContextIndex : InContextNum;
+			const FIntRect SrcRect = GetFinalContextRect(InputResourceType, SrcResourcesRect[SrcContextNum]);
 
-			// When resolving without warp, apply overscan
-			switch (InputResourceType)
+			if ((SrcContextNum + 1) == SrcResourcesAmount)
 			{
-			case EDisplayClusterViewportResourceType::InternalRenderTargetResource:
-				if (OverscanSettings.bIsEnabled && CVarDisplayClusterRenderOverscanResolve.GetValueOnRenderThread() != 0)
+				// last input mono -> stereo outputs
+				for (int32 DestResourceContextIndex = SrcContextNum; DestResourceContextIndex < DestResources.Num(); DestResourceContextIndex++)
 				{
-					// Support overscan crop
-					SourceRect = OverscanSettings.OverscanPixels.GetInnerRect(SourceRect);
+					const FIntRect DestRect = GetFinalContextRect(OutputResourceType, DestResourcesRect[DestResourceContextIndex]);
+					ImplResolveResource(RHICmdList, SrcResources[SrcContextNum], SrcRect, DestResources[DestResourceContextIndex], DestRect, bOutputIsMipsResource, bOutputIsPreviewResource);
 				}
 				break;
-			default:
-				break;
+			}
+			else
+			{
+				const FIntRect DestRect = GetFinalContextRect(OutputResourceType, DestResourcesRect[SrcContextNum]);
+				ImplResolveResource(RHICmdList, SrcResources[SrcContextNum], SrcRect, DestResources[SrcContextNum], DestRect, bOutputIsMipsResource, bOutputIsPreviewResource);
 			}
 
-			if ((InputIt + 1) == InputAmount)
+			if (InContextNum != INDEX_NONE)
 			{
-				// last input mono -> stereo outputs
-				for (int32 OutputIt = InputIt; OutputIt < OutputResources.Num(); OutputIt++)
-				{
-					ImplResolveResource(RHICmdList, InputResources[InputIt], SourceRect, OutputResources[OutputIt], OutputResourcesRect[OutputIt], bOutputIsMipsResource, bOutputIsPreviewResource);
-				}
+				// Copy only one texture
 				break;
 			}
+		}
+
+		return true;
+	}
+
+	return false;
+}
+
+bool FDisplayClusterViewportProxy::ShouldCopyAlphaChannel() const
+{
+	if ((RenderSettingsICVFX.RuntimeFlags & (ViewportRuntime_ICVFXLightcard | ViewportRuntime_ICVFXChromakey)) != 0)
+	{
+		return true;
+	}
+
+	return false;
+}
+
+void FDisplayClusterViewportProxy::CopyRenderTargetAlphaChannel_RehderThread(FRHICommandListImmediate& RHICmdList, const int32 InContextNum, FRHITexture* InSrcTexture, bool bInvertCopyDirection)
+{
+	if (InputShaderResources.IsValidIndex(InContextNum) && InputShaderResources[InContextNum] != nullptr && InSrcTexture)
+	{
+		if(FRHITexture* DestTexture = InputShaderResources[InContextNum]->GetTextureRHI())
+		{
+			const FIntRect SrcRect = GetFinalContextRect(EDisplayClusterViewportResourceType::InternalRenderTargetResource, Contexts[InContextNum].RenderTargetRect);
+			const FIntRect DestRect = FIntRect(FIntPoint(0, 0), DestTexture->GetSizeXY());
+
+			if (bInvertCopyDirection)
+			{
+
+				ImplResolveResourceAlpha(RHICmdList, DestTexture, DestRect, InSrcTexture, SrcRect);
+			}
 			else
 			{
-				ImplResolveResource(RHICmdList, InputResources[InputIt], SourceRect, OutputResources[InputIt], OutputResourcesRect[InputIt], bOutputIsMipsResource, bOutputIsPreviewResource);
+				ImplResolveResourceAlpha(RHICmdList, InSrcTexture, SrcRect, DestTexture, DestRect);
+				SceneColorResourcesState |= (1 << InContextNum);
 			}
 		}
+	}
+}
 
-		return true;
+BEGIN_SHADER_PARAMETER_STRUCT(FDisplayClusterCopyTextureParameters, )
+RDG_TEXTURE_ACCESS(SrcTexture, ERHIAccess::CopySrc)
+END_SHADER_PARAMETER_STRUCT()
+
+void FDisplayClusterViewportProxy::OnResolvedSceneColor_RenderThread(FRDGBuilder& GraphBuilder, const FSceneTextures& SceneTextures, const FDisplayClusterViewportProxy_Context& InProxyContext)
+{
+	const uint32 InContextNum = InProxyContext.ContextNum;
+
+	check(InContextNum < sizeof(SceneColorResourcesState) * 8);
+
+	SceneColorResourcesState &= ~(1 << InContextNum);
+
+	if (ShouldCopyAlphaChannel())
+	{
+		// Copy scene color with alpha channel to InputShaderResource texture
+		FRDGTextureRef InTextureRef = SceneTextures.Color.Resolve;
+		if (HasBeenProduced(InTextureRef))
+		{
+			check(InputShaderResources.IsValidIndex(InContextNum));
+
+			FDisplayClusterCopyTextureParameters* PassParameters = GraphBuilder.AllocParameters<FDisplayClusterCopyTextureParameters>();
+			PassParameters->SrcTexture = InTextureRef;
+
+			GraphBuilder.AddPass(
+				RDG_EVENT_NAME("DisplayCluster_SaveViewportAlphaChannel_%s", *GetId()),
+				PassParameters,
+				ERDGPassFlags::Copy | ERDGPassFlags::NeverCull,
+				[ViewportProxy = SharedThis(this), InTextureRef, InContextNum](FRHICommandListImmediate& RHICmdList)
+				{
+					ViewportProxy->CopyRenderTargetAlphaChannel_RehderThread(RHICmdList, InContextNum, InTextureRef->GetRHI());
+				});
+		}
 	}
+}
 
-	return false;
+BEGIN_SHADER_PARAMETER_STRUCT(FDisplayClusterEmptyParameters, )
+END_SHADER_PARAMETER_STRUCT()
+
+void FDisplayClusterViewportProxy::PostRenderViewFamily_RenderThread(FRDGBuilder& GraphBuilder, FSceneViewFamily& InViewFamily, const FSceneView& InSceneView, const FDisplayClusterViewportProxy_Context& InProxyContext)
+{
+	const uint32 InContextNum = InProxyContext.ContextNum;
+
+	check(InContextNum < sizeof(SceneColorResourcesState) * 8);
+
+	// Save GPUMask for each viewport context
+	if (Contexts.IsValidIndex(InContextNum))
+	{
+		if (GPUMask.Num() < Contexts.Num())
+		{
+			GPUMask.AddDefaulted(Contexts.Num() - GPUMask.Num());
+		}
+
+		checkSlow(InSceneView.bIsViewInfo);
+		const FViewInfo& ViewInfo = static_cast<const FViewInfo&>(InSceneView);
+
+		GPUMask[InContextNum] = ViewInfo.GPUMask;
+	}
+
+	if (ShouldCopyAlphaChannel())
+	{
+		FTextureRHIRef DestRenderTarget = InViewFamily.RenderTarget->GetRenderTargetTexture();
+		
+		FDisplayClusterEmptyParameters* PassParameters = GraphBuilder.AllocParameters<FDisplayClusterEmptyParameters>();
+		GraphBuilder.AddPass(
+			RDG_EVENT_NAME("DisplayCluster_LoadViewportAlphaChannel_%s", *GetId()),
+			PassParameters,
+			ERDGPassFlags::Copy | ERDGPassFlags::NeverCull,
+			[ViewportProxy = SharedThis(this), DestRenderTarget, InContextNum](FRHICommandListImmediate& RHICmdList)
+			{
+				// Copy alpha channel back
+				if ((ViewportProxy->SceneColorResourcesState & (1 << InContextNum)) != 0)
+				{
+					ViewportProxy->CopyRenderTargetAlphaChannel_RehderThread(RHICmdList, InContextNum, DestRenderTarget.GetReference(), true);
+				}
+			});
+	}
+
+	if (!InProxyContext.ViewFamilyProfileDescription.IsEmpty())
+	{
+		static IDisplayClusterCallbacks& DCCallbacksAPI = IDisplayCluster::Get().GetCallbacks();
+		if (DCCallbacksAPI.OnDisplayClusterPostRenderViewFamily_RenderThread().IsBound())
+		{
+			// Now we can perform viewport notification
+			DCCallbacksAPI.OnDisplayClusterPostRenderViewFamily_RenderThread().Broadcast(GraphBuilder, InViewFamily, this);
+		}
+	}
 }
diff --git a/Engine/Plugins/Runtime/nDisplay/Source/DisplayCluster/Private/Render/Viewport/DisplayClusterViewportProxy.h b/Engine/Plugins/Runtime/nDisplay/Source/DisplayCluster/Private/Render/Viewport/DisplayClusterViewportProxy.h
--- a/Engine/Plugins/Runtime/nDisplay/Source/DisplayCluster/Private/Render/Viewport/DisplayClusterViewportProxy.h
+++ b/Engine/Plugins/Runtime/nDisplay/Source/DisplayCluster/Private/Render/Viewport/DisplayClusterViewportProxy.h
@@ -2,6 +2,13 @@
 
 #pragma once
 
+#include "CoreMinimal.h"
+
+#include "EngineUtils.h"
+#include "ScreenRendering.h"
+#include "SceneView.h"
+#include "Templates/SharedPointer.h"
+
 #include "Render/Viewport/IDisplayClusterViewportProxy.h"
 #include "Render/Viewport/Containers/DisplayClusterViewport_OverscanSettings.h"
 
@@ -18,9 +25,10 @@ class FDisplayClusterViewportReadPixelsData;
 
 class FDisplayClusterViewportProxy
 	: public IDisplayClusterViewportProxy
+	, public TSharedFromThis<FDisplayClusterViewportProxy, ESPMode::ThreadSafe>
 {
 public:
-	FDisplayClusterViewportProxy(const FDisplayClusterViewportManagerProxy& Owner, const FDisplayClusterViewport& RenderViewport);
+	FDisplayClusterViewportProxy(const FDisplayClusterViewport& RenderViewport);
 	virtual ~FDisplayClusterViewportProxy();
 
 public:
@@ -90,12 +98,14 @@ public:
 	virtual bool GetResourcesWithRects_RenderThread(const EDisplayClusterViewportResourceType InResourceType, TArray<FRHITexture2D*>& OutResources, TArray<FIntRect>& OutRects) const override;
 
 	// Resolve resource contexts
-	virtual bool ResolveResources_RenderThread(FRHICommandListImmediate& RHICmdList, const EDisplayClusterViewportResourceType InputResourceType, const EDisplayClusterViewportResourceType OutputResourceType) const override;
+	virtual bool ResolveResources_RenderThread(FRHICommandListImmediate& RHICmdList, const EDisplayClusterViewportResourceType InputResourceType, const EDisplayClusterViewportResourceType OutputResourceType, const int32 InContextNum = INDEX_NONE) const override;
 
 	virtual EDisplayClusterViewportResourceType GetOutputResourceType_RenderThread() const override;
 
 	virtual const IDisplayClusterViewportManagerProxy& GetOwner_RenderThread() const override;
 
+	virtual void OnResolvedSceneColor_RenderThread(FRDGBuilder& GraphBuilder, const FSceneTextures& SceneTextures, const FDisplayClusterViewportProxy_Context& InProxyContext) override;
+	virtual void PostRenderViewFamily_RenderThread(FRDGBuilder& InGraphBuilder, class FSceneViewFamily& InViewFamily, const class FSceneView& InSceneView, const FDisplayClusterViewportProxy_Context& InProxyContext) override;
 
 	///////////////////////////////
 	// ~IDisplayClusterViewportProxy
@@ -106,6 +116,8 @@ public:
 
 	void PostResolveViewport_RenderThread(FRHICommandListImmediate& RHICmdList) const;
 
+	bool ShouldCopyAlphaChannel() const;
+
 #if WITH_EDITOR
 	bool GetPreviewPixels_GameThread(TSharedPtr<FDisplayClusterViewportReadPixelsData, ESPMode::ThreadSafe>& OutPixelsData) const;
 #endif
@@ -130,6 +142,8 @@ public:
 		return false;
 	}
 
+	FIntRect GetFinalContextRect(const EDisplayClusterViewportResourceType InputResourceType, const FIntRect& InRect) const;
+
 private:
 	bool ImplGetResourcesWithRects_RenderThread(const EDisplayClusterViewportResourceType InResourceType, TArray<FRHITexture2D*>& OutResources, TArray<FIntRect>& OutResourceRects, const int32 InRecursionDepth) const;
 	bool ImplGetResources_RenderThread(const EDisplayClusterViewportResourceType InResourceType, TArray<FRHITexture2D*>& OutResources, const int32 InRecursionDepth) const;
@@ -137,10 +151,13 @@ private:
 	void ImplViewportRemap_RenderThread(FRHICommandListImmediate& RHICmdList) const;
 	void ImplPreviewReadPixels_RenderThread(FRHICommandListImmediate& RHICmdList) const;
 
-	bool ImplResolveResources_RenderThread(FRHICommandListImmediate& RHICmdList, FDisplayClusterViewportProxy const* SourceProxy, const EDisplayClusterViewportResourceType InputResourceType, const EDisplayClusterViewportResourceType OutputResourceType) const;
+	bool ImplResolveResources_RenderThread(FRHICommandListImmediate& RHICmdList, FDisplayClusterViewportProxy const* SourceProxy, const EDisplayClusterViewportResourceType InputResourceType, const EDisplayClusterViewportResourceType OutputResourceType, const int32 InContextNum) const;
 
 	bool IsShouldOverrideViewportResource(const EDisplayClusterViewportResourceType InResourceType) const;
 
+	// The InputShaderResources[] uses as temporary storage for alpha channel data
+	void CopyRenderTargetAlphaChannel_RehderThread(FRHICommandListImmediate& RHICmdList, const int32 InContextNum, FRHITexture* InSrcTexture, bool bInvertCopyDirection = false);
+
 protected:
 	friend FDisplayClusterViewportProxyData;
 	friend FDisplayClusterViewportManagerProxy;
@@ -166,6 +183,10 @@ protected:
 	// Viewport contexts (left/center/right eyes)
 	mutable TArray<FDisplayClusterViewport_Context> Contexts;
 
+	/** The GPU nodes on which to render this view. */
+	TArray<FRHIGPUMask> GPUMask;
+	uint32 SceneColorResourcesState = 0;
+
 	// View family render to this resources
 	TArray<FDisplayClusterViewportRenderTargetResource*> RenderTargets;
 
@@ -186,7 +207,7 @@ protected:
 	TArray<FDisplayClusterViewportTextureResource*> AdditionalTargetableResources;
 	TArray<FDisplayClusterViewportTextureResource*> MipsShaderResources;
 
-	const FDisplayClusterViewportManagerProxy& Owner;
+	const TSharedRef<FDisplayClusterViewportManagerProxy, ESPMode::ThreadSafe> Owner;
 	IDisplayClusterShaders& ShadersAPI;
 };
 
diff --git a/Engine/Plugins/Runtime/nDisplay/Source/DisplayCluster/Private/Render/Viewport/Postprocess/DisplayClusterViewportPostProcessManager.cpp b/Engine/Plugins/Runtime/nDisplay/Source/DisplayCluster/Private/Render/Viewport/Postprocess/DisplayClusterViewportPostProcessManager.cpp
--- a/Engine/Plugins/Runtime/nDisplay/Source/DisplayCluster/Private/Render/Viewport/Postprocess/DisplayClusterViewportPostProcessManager.cpp
+++ b/Engine/Plugins/Runtime/nDisplay/Source/DisplayCluster/Private/Render/Viewport/Postprocess/DisplayClusterViewportPostProcessManager.cpp
@@ -369,7 +369,7 @@ void FDisplayClusterViewportPostProcessManager::FinalizeNewFrame()
 	check(IsInGameThread());
 
 	ENQUEUE_RENDER_COMMAND(DisplayClusterViewportPostProcessManager_FinalizeNewFrame)(
-		[PostprocessManager = this, PostprocessData = Postprocess](FRHICommandListImmediate& RHICmdList)
+		[PostprocessManager = SharedThis(this), PostprocessData = Postprocess](FRHICommandListImmediate& RHICmdList)
 	{
 		PostprocessManager->PostprocessProxy.Empty();
 		PostprocessManager->PostprocessProxy.Append(PostprocessData);
@@ -485,12 +485,12 @@ void FDisplayClusterViewportPostProcessManager::ImplPerformPostProcessViewBefore
 				// send nullptr viewport first to handle begin
 				It->PerformPostProcessViewBeforeWarpBlend_RenderThread(RHICmdList, nullptr);
 
-				for (IDisplayClusterViewportProxy* ViewportProxyIt : InViewportManagerProxy->GetViewports_RenderThread())
+				for (TSharedPtr<IDisplayClusterViewportProxy, ESPMode::ThreadSafe>& ViewportProxyIt : InViewportManagerProxy->GetViewports_RenderThread())
 				{
-					if (ViewportProxyIt)
+					if (ViewportProxyIt.IsValid())
 					{
 						UE_LOG(LogDisplayClusterRender, VeryVerbose, TEXT("Postprocess VIEW before WarpBlend - Viewport '%s'"), *ViewportProxyIt->GetId());
-						It->PerformPostProcessViewBeforeWarpBlend_RenderThread(RHICmdList, ViewportProxyIt);
+						It->PerformPostProcessViewBeforeWarpBlend_RenderThread(RHICmdList, ViewportProxyIt.Get());
 					}
 				}
 			}
@@ -515,12 +515,12 @@ void FDisplayClusterViewportPostProcessManager::ImplPerformPostProcessViewAfterW
 				// send nullptr viewport first to handle begin
 				It->PerformPostProcessViewAfterWarpBlend_RenderThread(RHICmdList, nullptr);
 
-				for (IDisplayClusterViewportProxy* ViewportProxyIt : InViewportManagerProxy->GetViewports_RenderThread())
+				for (const TSharedPtr<IDisplayClusterViewportProxy, ESPMode::ThreadSafe>& ViewportProxyIt : InViewportManagerProxy->GetViewports_RenderThread())
 				{
-					if (ViewportProxyIt)
+					if (ViewportProxyIt.IsValid())
 					{
 						UE_LOG(LogDisplayClusterRender, VeryVerbose, TEXT("Postprocess VIEW after WarpBlend - Viewport '%s'"), *ViewportProxyIt->GetId());
-						It->PerformPostProcessViewAfterWarpBlend_RenderThread(RHICmdList, ViewportProxyIt);
+						It->PerformPostProcessViewAfterWarpBlend_RenderThread(RHICmdList, ViewportProxyIt.Get());
 					}
 				}
 			}
diff --git a/Engine/Plugins/Runtime/nDisplay/Source/DisplayCluster/Private/Render/Viewport/Postprocess/DisplayClusterViewportPostProcessManager.h b/Engine/Plugins/Runtime/nDisplay/Source/DisplayCluster/Private/Render/Viewport/Postprocess/DisplayClusterViewportPostProcessManager.h
--- a/Engine/Plugins/Runtime/nDisplay/Source/DisplayCluster/Private/Render/Viewport/Postprocess/DisplayClusterViewportPostProcessManager.h
+++ b/Engine/Plugins/Runtime/nDisplay/Source/DisplayCluster/Private/Render/Viewport/Postprocess/DisplayClusterViewportPostProcessManager.h
@@ -4,6 +4,7 @@
 
 #include "Render/PostProcess/IDisplayClusterPostProcess.h"
 #include "Render/IDisplayClusterRenderManager.h"
+#include "Templates/SharedPointer.h"
 
 class FDisplayClusterViewportManager;
 class FDisplayClusterViewportManagerProxy;
@@ -15,6 +16,7 @@ struct FDisplayClusterConfigurationPostprocess;
  */
 
 class FDisplayClusterViewportPostProcessManager
+	: public TSharedFromThis<FDisplayClusterViewportPostProcessManager, ESPMode::ThreadSafe>
 {
 public:
 	FDisplayClusterViewportPostProcessManager(FDisplayClusterViewportManager& InViewportManager);
diff --git a/Engine/Plugins/Runtime/nDisplay/Source/DisplayCluster/Private/Render/Viewport/RenderTarget/DisplayClusterRenderTargetResourcesPool.cpp b/Engine/Plugins/Runtime/nDisplay/Source/DisplayCluster/Private/Render/Viewport/RenderTarget/DisplayClusterRenderTargetResourcesPool.cpp
--- a/Engine/Plugins/Runtime/nDisplay/Source/DisplayCluster/Private/Render/Viewport/RenderTarget/DisplayClusterRenderTargetResourcesPool.cpp
+++ b/Engine/Plugins/Runtime/nDisplay/Source/DisplayCluster/Private/Render/Viewport/RenderTarget/DisplayClusterRenderTargetResourcesPool.cpp
@@ -28,19 +28,19 @@ namespace DisplayClusterRenderTargetResourcesPool
 
 	static void ImplReleaseViewportResourcesRHI(FDisplayClusterViewportManagerProxy* InViewportManagerProxy, const TArrayView<FDisplayClusterViewportResource*>& InResources)
 	{
-		TArray<FDisplayClusterViewportResource*> ResourcesForRenderThread(InResources);
-
-		ENQUEUE_RENDER_COMMAND(DisplayCluster_ReleaseViewportResourcesRHI)(
-			[InViewportManagerProxy, ReleasedResources = std::move(ResourcesForRenderThread)](FRHICommandListImmediate& RHICmdList)
+		if (InViewportManagerProxy != nullptr && !InResources.IsEmpty())
 		{
-			for (FDisplayClusterViewportResource* DeletedResourcePtrIt : ReleasedResources)
-			{
-				if (InViewportManagerProxy)
+			TArray<FDisplayClusterViewportResource*> ResourcesForRenderThread(InResources);
+
+			ENQUEUE_RENDER_COMMAND(DisplayCluster_ReleaseViewportResourcesRHI)(
+				[InViewportManagerProxy = InViewportManagerProxy->AsShared(), ReleasedResources = std::move(ResourcesForRenderThread)](FRHICommandListImmediate& RHICmdList)
 				{
-					InViewportManagerProxy->DeleteResource_RenderThread(DeletedResourcePtrIt);
-				}
-			}
-		});
+					for (FDisplayClusterViewportResource* DeletedResourcePtrIt : ReleasedResources)
+					{
+						InViewportManagerProxy->DeleteResource_RenderThread(DeletedResourcePtrIt);
+					}
+				});
+		}
 	}
 
 	static FDisplayClusterViewportResourceSettings* ImplCreateViewportResourceSettings(const FDisplayClusterRenderFrameSettings& InRenderFrameSettings, FViewport* InViewport)
diff --git a/Engine/Plugins/Runtime/nDisplay/Source/DisplayCluster/Public/Render/Viewport/Containers/DisplayClusterViewportProxy_Context.h b/Engine/Plugins/Runtime/nDisplay/Source/DisplayCluster/Public/Render/Viewport/Containers/DisplayClusterViewportProxy_Context.h
new file mode 100644
--- /dev/null
+++ b/Engine/Plugins/Runtime/nDisplay/Source/DisplayCluster/Public/Render/Viewport/Containers/DisplayClusterViewportProxy_Context.h
@@ -0,0 +1,14 @@
+// Copyright Epic Games, Inc. All Rights Reserved.
+
+#pragma once
+
+#include "CoreMinimal.h"
+
+class FDisplayClusterViewportProxy_Context
+{
+public:
+	uint32  ContextNum = 0;
+
+	// Profile description from ViewFamily
+	FString ViewFamilyProfileDescription;
+};
diff --git a/Engine/Plugins/Runtime/nDisplay/Source/DisplayCluster/Public/Render/Viewport/Containers/DisplayClusterViewport_RenderSettings.h b/Engine/Plugins/Runtime/nDisplay/Source/DisplayCluster/Public/Render/Viewport/Containers/DisplayClusterViewport_RenderSettings.h
--- a/Engine/Plugins/Runtime/nDisplay/Source/DisplayCluster/Public/Render/Viewport/Containers/DisplayClusterViewport_RenderSettings.h
+++ b/Engine/Plugins/Runtime/nDisplay/Source/DisplayCluster/Public/Render/Viewport/Containers/DisplayClusterViewport_RenderSettings.h
@@ -14,7 +14,6 @@ enum class EDisplayClusterViewportCaptureMode : uint8
 
 	// use hi-res float texture with alpha for compisiting
 	Lightcard,
-	Lightcard_OCIO,
 
 	// Special hi-res mode for movie pipeline
 	MoviePipeline,
diff --git a/Engine/Plugins/Runtime/nDisplay/Source/DisplayCluster/Public/Render/Viewport/Containers/DisplayClusterViewport_RenderSettingsICVFX.h b/Engine/Plugins/Runtime/nDisplay/Source/DisplayCluster/Public/Render/Viewport/Containers/DisplayClusterViewport_RenderSettingsICVFX.h
--- a/Engine/Plugins/Runtime/nDisplay/Source/DisplayCluster/Public/Render/Viewport/Containers/DisplayClusterViewport_RenderSettingsICVFX.h
+++ b/Engine/Plugins/Runtime/nDisplay/Source/DisplayCluster/Public/Render/Viewport/Containers/DisplayClusterViewport_RenderSettingsICVFX.h
@@ -39,8 +39,6 @@ enum EDisplayClusterViewportRuntimeICVFXFlags
 	ViewportRuntime_ICVFXIncamera       = 1 << 1,
 	ViewportRuntime_ICVFXChromakey      = 1 << 2,
 	ViewportRuntime_ICVFXLightcard      = 1 << 3,
-	ViewportRuntime_ICVFXLightcardAlpha = 1 << 4,
-	ViewportRuntime_ICVFXLightcardColor = 1 << 5,
 
 	// This viewport used as internal icvfx composing resource (created and deleted inside icvfx logic)
 	ViewportRuntime_InternalResource = 1 << 30,
diff --git a/Engine/Plugins/Runtime/nDisplay/Source/DisplayCluster/Public/Render/Viewport/IDisplayClusterViewportManagerProxy.h b/Engine/Plugins/Runtime/nDisplay/Source/DisplayCluster/Public/Render/Viewport/IDisplayClusterViewportManagerProxy.h
--- a/Engine/Plugins/Runtime/nDisplay/Source/DisplayCluster/Public/Render/Viewport/IDisplayClusterViewportManagerProxy.h
+++ b/Engine/Plugins/Runtime/nDisplay/Source/DisplayCluster/Public/Render/Viewport/IDisplayClusterViewportManagerProxy.h
@@ -41,7 +41,7 @@ public:
 	*
 	* @return - arrays with viewport render thread proxy objects refs
 	*/
-	virtual const TArrayView<IDisplayClusterViewportProxy*> GetViewports_RenderThread() const = 0;
+	virtual const TArrayView<TSharedPtr<IDisplayClusterViewportProxy, ESPMode::ThreadSafe>> GetViewports_RenderThread() const = 0;
 
 	/**
 	* Return render frame targets for current frame
diff --git a/Engine/Plugins/Runtime/nDisplay/Source/DisplayCluster/Public/Render/Viewport/IDisplayClusterViewportProxy.h b/Engine/Plugins/Runtime/nDisplay/Source/DisplayCluster/Public/Render/Viewport/IDisplayClusterViewportProxy.h
--- a/Engine/Plugins/Runtime/nDisplay/Source/DisplayCluster/Public/Render/Viewport/IDisplayClusterViewportProxy.h
+++ b/Engine/Plugins/Runtime/nDisplay/Source/DisplayCluster/Public/Render/Viewport/IDisplayClusterViewportProxy.h
@@ -13,7 +13,7 @@
 #include "Render/Viewport/Containers/DisplayClusterViewport_RenderSettings.h"
 #include "Render/Viewport/Containers/DisplayClusterViewport_RenderSettingsICVFX.h"
 #include "Render/Viewport/Containers/DisplayClusterViewport_PostRenderSettings.h"
-
+#include "Render/Viewport/Containers/DisplayClusterViewportProxy_Context.h"
 
 class DISPLAYCLUSTER_API IDisplayClusterViewportProxy
 {
@@ -37,7 +37,7 @@ public:
 	virtual bool GetResourcesWithRects_RenderThread(const EDisplayClusterViewportResourceType InResourceType, TArray<FRHITexture2D*>& OutResources, TArray<FIntRect>& OutRects) const = 0;
 
 	// Resolve resource contexts
-	virtual bool ResolveResources_RenderThread(FRHICommandListImmediate& RHICmdList, const EDisplayClusterViewportResourceType InputResourceType, const EDisplayClusterViewportResourceType OutputResourceType) const = 0;
+	virtual bool ResolveResources_RenderThread(FRHICommandListImmediate& RHICmdList, const EDisplayClusterViewportResourceType InputResourceType, const EDisplayClusterViewportResourceType OutputResourceType, const int32 InContextNum = INDEX_NONE) const = 0;
 
 	virtual EDisplayClusterViewportResourceType   GetOutputResourceType_RenderThread() const = 0;
 
@@ -45,4 +45,9 @@ public:
 
 	virtual void SetRenderSettings_RenderThread(const FDisplayClusterViewport_RenderSettings& InRenderSettings) const = 0;
 	virtual void SetContexts_RenderThread(const TArray<FDisplayClusterViewport_Context>& InContexts) const = 0;
+
+	virtual void OnResolvedSceneColor_RenderThread(class FRDGBuilder& GraphBuilder, const struct FSceneTextures& SceneTextures, const FDisplayClusterViewportProxy_Context& InProxyContext) = 0;
+
+	// Callback. After the viewfamily has finished rendering
+	virtual void PostRenderViewFamily_RenderThread(class FRDGBuilder& InGraphBuilder, class FSceneViewFamily& InViewFamily, const class FSceneView& InSceneView, const FDisplayClusterViewportProxy_Context& InProxyContext) = 0;
 };
diff --git a/Engine/Plugins/Runtime/nDisplay/Source/DisplayClusterMedia/Private/Synchronization/DisplayClusterFrameQueue.cpp b/Engine/Plugins/Runtime/nDisplay/Source/DisplayClusterMedia/Private/Synchronization/DisplayClusterFrameQueue.cpp
--- a/Engine/Plugins/Runtime/nDisplay/Source/DisplayClusterMedia/Private/Synchronization/DisplayClusterFrameQueue.cpp
+++ b/Engine/Plugins/Runtime/nDisplay/Source/DisplayClusterMedia/Private/Synchronization/DisplayClusterFrameQueue.cpp
@@ -108,18 +108,18 @@ void FDisplayClusterFrameQueue::HandleProcessLatency_RenderThread(FRHICommandLis
 			{
 				if (IDisplayClusterViewportManagerProxy* ViewportMgrProxy = ViewportMgr->GetProxy())
 				{
-					TArrayView<IDisplayClusterViewportProxy*> ViewportProxies = ViewportMgrProxy->GetViewports_RenderThread();
-					for (IDisplayClusterViewportProxy* ViewportProxy : ViewportProxies)
+					TArrayView<TSharedPtr<IDisplayClusterViewportProxy, ESPMode::ThreadSafe>> ViewportProxies = ViewportMgrProxy->GetViewports_RenderThread();
+					for (TSharedPtr<IDisplayClusterViewportProxy, ESPMode::ThreadSafe>& ViewportProxyIt : ViewportProxies)
 					{
 						// Handle textures of all viewports that don't have any media input attached. Those don't get rendered.
-						if (!ViewportProxy->GetRenderSettings_RenderThread().bSkipSceneRenderingButLeaveResourcesAvailable)
+						if (ViewportProxyIt.IsValid() && !ViewportProxyIt->GetRenderSettings_RenderThread().bSkipSceneRenderingButLeaveResourcesAvailable)
 						{
 							// Get viewport texture
 							TArray<FRHITexture*> Textures;
-							ViewportProxy->GetResources_RenderThread(EDisplayClusterViewportResourceType::InternalRenderTargetResource, Textures);
+							ViewportProxyIt->GetResources_RenderThread(EDisplayClusterViewportResourceType::InternalRenderTargetResource, Textures);
 							check(Textures.Num() > 0);
 
-							const FString ViewportId(ViewportProxy->GetId());
+							const FString ViewportId(ViewportProxyIt->GetId());
 
 							// Enqueue this view (head)
 							Frames[IdxHead].SaveView(RHICmdList, ViewportId, Textures[0]);
diff --git a/Engine/Plugins/Runtime/nDisplay/Source/DisplayClusterShaders/Private/Shaders/DisplayClusterShadersPreprocess_UVLightCards.cpp b/Engine/Plugins/Runtime/nDisplay/Source/DisplayClusterShaders/Private/Shaders/DisplayClusterShadersPreprocess_UVLightCards.cpp
--- a/Engine/Plugins/Runtime/nDisplay/Source/DisplayClusterShaders/Private/Shaders/DisplayClusterShadersPreprocess_UVLightCards.cpp
+++ b/Engine/Plugins/Runtime/nDisplay/Source/DisplayClusterShaders/Private/Shaders/DisplayClusterShadersPreprocess_UVLightCards.cpp
@@ -146,6 +146,22 @@ bool FDisplayClusterShadersPreprocess_UVLightCards::RenderPreprocess_UVLightCard
 
 	FEngineShowFlags EngineShowFlags(ESFIM_Game);
 
+	// LightCard settings from the FDisplayClusterViewportManager::ConfigureViewFamily
+	{
+		EngineShowFlags.PostProcessing = 0;
+		EngineShowFlags.SetAtmosphere(0);
+		EngineShowFlags.SetFog(0);
+		EngineShowFlags.SetVolumetricFog(0);
+		EngineShowFlags.SetMotionBlur(0); // motion blur doesn't work correctly with scene captures.
+		EngineShowFlags.SetSeparateTranslucency(0);
+		EngineShowFlags.SetHMDDistortion(0);
+		EngineShowFlags.SetOnScreenDebug(0);
+
+		EngineShowFlags.SetLumenReflections(0);
+		EngineShowFlags.SetLumenGlobalIllumination(0);
+		EngineShowFlags.SetGlobalIllumination(0);
+	}
+
 	FSceneViewFamilyContext ViewFamily(FSceneViewFamily::ConstructionValues(
 		InRenderTarget,
 		InScene,
@@ -193,7 +209,7 @@ bool FDisplayClusterShadersPreprocess_UVLightCards::RenderPreprocess_UVLightCard
 		);
 	}
 
-	ViewInitOptions.BackgroundColor = FLinearColor::Transparent;
+	ViewInitOptions.BackgroundColor = FLinearColor::Black;
 
 	GetRendererModule().CreateAndInitSingleView(RHICmdList, &ViewFamily, &ViewInitOptions);
 	const FSceneView* View = ViewFamily.Views[0];
@@ -239,4 +255,4 @@ bool FDisplayClusterShadersPreprocess_UVLightCards::RenderPreprocess_UVLightCard
 	GraphBuilder.Execute();
 
 	return true;
-}
\ No newline at end of file
+}
diff --git a/Engine/Plugins/Runtime/nDisplay/Source/DisplayClusterShaders/Private/Shaders/DisplayClusterShadersWarpblend_ICVFX.cpp b/Engine/Plugins/Runtime/nDisplay/Source/DisplayClusterShaders/Private/Shaders/DisplayClusterShadersWarpblend_ICVFX.cpp
--- a/Engine/Plugins/Runtime/nDisplay/Source/DisplayClusterShaders/Private/Shaders/DisplayClusterShadersWarpblend_ICVFX.cpp
+++ b/Engine/Plugins/Runtime/nDisplay/Source/DisplayClusterShaders/Private/Shaders/DisplayClusterShadersWarpblend_ICVFX.cpp
@@ -45,6 +45,14 @@ static TAutoConsoleVariable<int32> CVarIcvfxMPCDIShaderType(
 	ECVF_RenderThreadSafe
 );
 
+int32 GDisplayClusterShadersICVFXEnableLightCard = 1;
+static FAutoConsoleVariableRef CVarDisplayClusterShadersICVFXEnableLightCard(
+	TEXT("nDisplay.render.icvfx.lightcard"),
+	GDisplayClusterShadersICVFXEnableLightCard,
+	TEXT("Enable ICVFX lightcards render (0 - disable).\n"),
+	ECVF_RenderThreadSafe
+);
+
 enum class EIcvfxShaderType : uint8
 {
 	Passthrough,  // Viewport frustum (no warpblend, only frustum math)
@@ -321,7 +329,7 @@ private:
 			// 2 - second camera additive pass
 			// N - camera N
 			// N+1 - overlayOver (optional)
-			return (ICVFXParameters.IsLightcardOverUsed()) ? ( CameraIndex == (TotalUsedCameras + 1) ) : ( CameraIndex == TotalUsedCameras );
+			return (ICVFXParameters.IsLightcardOverUsed() && GDisplayClusterShadersICVFXEnableLightCard) ? ( CameraIndex == (TotalUsedCameras + 1) ) : ( CameraIndex == TotalUsedCameras );
 		}
 
 		if (TotalUsedCameras>1)
@@ -331,7 +339,7 @@ private:
 			// 1   - camera2
 			// N-1 - camera N
 			// N   - overlayOver (optional)
-			return (ICVFXParameters.IsLightcardOverUsed()) ? ( CameraIndex == TotalUsedCameras ) : ( CameraIndex == (TotalUsedCameras-1) );
+			return (ICVFXParameters.IsLightcardOverUsed() && GDisplayClusterShadersICVFXEnableLightCard) ? ( CameraIndex == TotalUsedCameras ) : ( CameraIndex == (TotalUsedCameras-1) );
 		}
 
 		// By default render single camera in one pass
@@ -341,7 +349,7 @@ private:
 	inline bool IsOverlayUnderUsed() const
 	{
 		// Render OverlayUnder only on first pass
-		return IsFirstPass() && ICVFXParameters.IsLightcardUnderUsed();
+		return IsFirstPass() && ICVFXParameters.IsLightcardUnderUsed() && GDisplayClusterShadersICVFXEnableLightCard;
 	}
 
 	inline bool IsOverlayOverUsed() const
@@ -349,10 +357,10 @@ private:
 		if (IsMultiPassRender())
 		{
 			// Render OverlayOver only on last additive pass
-			return IsLastPass() && ICVFXParameters.IsLightcardOverUsed();
+			return IsLastPass() && ICVFXParameters.IsLightcardOverUsed() && GDisplayClusterShadersICVFXEnableLightCard;
 		}
 		
-		return ICVFXParameters.IsLightcardOverUsed();
+		return ICVFXParameters.IsLightcardOverUsed() && GDisplayClusterShadersICVFXEnableLightCard;
 	}
 
 	inline int32 GetUsedCameraIndex() const
@@ -521,22 +529,9 @@ public:
 		if (IsOverlayUnderUsed() || IsOverlayOverUsed())
 		{
 
-			if (ICVFXParameters.Lightcard_OCIO.IsValid())
-			{
-				// use OCIO with separated alpha lightmap texture
-				RenderPassData.PSParameters.OverlaySampler = TStaticSamplerState<SF_Trilinear, AM_Clamp, AM_Clamp, AM_Clamp>::GetRHI();
-				RenderPassData.PSParameters.OverlayTexture = ICVFXParameters.Lightcard_OCIO.Texture;
-
-				RenderPassData.PSParameters.OverlayAlphaSampler = TStaticSamplerState<SF_Trilinear, AM_Clamp, AM_Clamp, AM_Clamp>::GetRHI();
-				RenderPassData.PSParameters.OverlayAlphaTexture = ICVFXParameters.Lightcard.Texture;
+			RenderPassData.PSParameters.OverlaySampler = TStaticSamplerState<SF_Trilinear, AM_Clamp, AM_Clamp, AM_Clamp>::GetRHI();
+			RenderPassData.PSParameters.OverlayTexture = ICVFXParameters.Lightcard.Texture;
 
-				RenderPassData.PSPermutationVector.Set<IcvfxShaderPermutation::FIcvfxShaderOverlayAlpha>(true);
-			}
-			else
-			{
-				RenderPassData.PSParameters.OverlaySampler = TStaticSamplerState<SF_Trilinear, AM_Clamp, AM_Clamp, AM_Clamp>::GetRHI();
-				RenderPassData.PSParameters.OverlayTexture = ICVFXParameters.Lightcard.Texture;
-			}
 			if (IsOverlayUnderUsed())
 			{
 				RenderPassData.PSPermutationVector.Set<IcvfxShaderPermutation::FIcvfxShaderOverlayUnder>(true);
@@ -646,7 +641,7 @@ public:
 
 	bool GetUVLightCardParameters(FIcvfxRenderPassData& RenderPassData)
 	{
-		if (ICVFXParameters.UVLightCardMap.IsValid() && RenderPassData.PSPermutationVector.Get<IcvfxShaderPermutation::FIcvfxShaderMeshWarp>())
+		if (ICVFXParameters.UVLightCardMap.IsValid() && GDisplayClusterShadersICVFXEnableLightCard && RenderPassData.PSPermutationVector.Get<IcvfxShaderPermutation::FIcvfxShaderMeshWarp>())
 		{
 			RenderPassData.PSParameters.UVLightCardMapSampler = TStaticSamplerState<SF_Trilinear, AM_Clamp, AM_Clamp, AM_Clamp>::GetRHI();
 			RenderPassData.PSParameters.UVLightCardMapTexture = ICVFXParameters.UVLightCardMap;
diff --git a/Engine/Plugins/Runtime/nDisplay/Source/DisplayClusterShaders/Public/ShaderParameters/DisplayClusterShaderParameters_ICVFX.h b/Engine/Plugins/Runtime/nDisplay/Source/DisplayClusterShaders/Public/ShaderParameters/DisplayClusterShaderParameters_ICVFX.h
--- a/Engine/Plugins/Runtime/nDisplay/Source/DisplayClusterShaders/Public/ShaderParameters/DisplayClusterShaderParameters_ICVFX.h
+++ b/Engine/Plugins/Runtime/nDisplay/Source/DisplayClusterShaders/Public/ShaderParameters/DisplayClusterShaderParameters_ICVFX.h
@@ -113,7 +113,6 @@ public:
 	{
 		Cameras.Empty();
 		Lightcard.Reset();
-		Lightcard_OCIO.Reset();
 
 		UVLightCardMap = nullptr;
 	}
@@ -126,7 +125,6 @@ public:
 		Cameras = InParameters.Cameras;
 
 		Lightcard      = InParameters.Lightcard;
-		Lightcard_OCIO = InParameters.Lightcard_OCIO;
 		LightcardMode = InParameters.LightcardMode;
 
 		UVLightCardMap = InParameters.UVLightCardMap;
@@ -139,11 +137,6 @@ public:
 			Dst.Add(&Lightcard);
 		}
 
-		if (Lightcard_OCIO.IsDefined())
-		{
-			Dst.Add(&Lightcard_OCIO);
-		}
-
 		for (FCameraSettings& CameraIt : Cameras)
 		{
 			CameraIt.CollectRefViewports(Dst);
@@ -249,7 +242,6 @@ public:
 
 	// Lightcard settings
 	FDisplayClusterShaderParametersICVFX_ViewportResource    Lightcard;
-	FDisplayClusterShaderParametersICVFX_ViewportResource    Lightcard_OCIO;
 	EDisplayClusterShaderParametersICVFX_LightcardRenderMode LightcardMode = EDisplayClusterShaderParametersICVFX_LightcardRenderMode::Under;
 
 	/** Texture containing a UV map of the rendered UV light cards */
diff --git a/Engine/Plugins/VirtualProduction/TextureShare/Source/TextureShareDisplayCluster/Private/PostProcess/TextureSharePostprocess_FrameProxy.cpp b/Engine/Plugins/VirtualProduction/TextureShare/Source/TextureShareDisplayCluster/Private/PostProcess/TextureSharePostprocess_FrameProxy.cpp
--- a/Engine/Plugins/VirtualProduction/TextureShare/Source/TextureShareDisplayCluster/Private/PostProcess/TextureSharePostprocess_FrameProxy.cpp
+++ b/Engine/Plugins/VirtualProduction/TextureShare/Source/TextureShareDisplayCluster/Private/PostProcess/TextureSharePostprocess_FrameProxy.cpp
@@ -42,13 +42,13 @@ void FTextureSharePostprocess::ShareViewport_RenderThread(FRHICommandListImmedia
 {
 	if (InViewportManagerProxy)
 	{
-		for (IDisplayClusterViewportProxy* ViewportProxyIt : InViewportManagerProxy->GetViewports_RenderThread())
+		for (TSharedPtr<IDisplayClusterViewportProxy, ESPMode::ThreadSafe>& ViewportProxyIt : InViewportManagerProxy->GetViewports_RenderThread())
 		{
-			if (ViewportProxyIt)
+			if (ViewportProxyIt.IsValid())
 			{
 				// Get viewport resource type
 				EDisplayClusterViewportResourceType ResourceType = InResourceType;
-				if (bAfterWarpBlend && ShouldApplyWarpBlend(ViewportProxyIt))
+				if (bAfterWarpBlend && ShouldApplyWarpBlend(ViewportProxyIt.Get()))
 				{
 					const TSharedPtr<IDisplayClusterProjectionPolicy, ESPMode::ThreadSafe>& PrjPolicy = ViewportProxyIt->GetProjectionPolicy_RenderThread();
 					if (PrjPolicy->ShouldUseAdditionalTargetableResource())
