nDisplay: fixed Incorrect work of Dome and Vioso projection policies

diff --git a/Engine/Plugins/Runtime/nDisplay/Source/DisplayClusterProjection/Private/Policy/Domeprojection/DisplayClusterProjectionDomeprojectionPolicyBase.cpp b/Engine/Plugins/Runtime/nDisplay/Source/DisplayClusterProjection/Private/Policy/Domeprojection/DisplayClusterProjectionDomeprojectionPolicyBase.cpp
--- a/Engine/Plugins/Runtime/nDisplay/Source/DisplayClusterProjection/Private/Policy/Domeprojection/DisplayClusterProjectionDomeprojectionPolicyBase.cpp
+++ b/Engine/Plugins/Runtime/nDisplay/Source/DisplayClusterProjection/Private/Policy/Domeprojection/DisplayClusterProjectionDomeprojectionPolicyBase.cpp
@@ -14,16 +14,22 @@
 #include "Render/Viewport/IDisplayClusterViewport.h"
 #include "Render/Viewport/IDisplayClusterViewportProxy.h"
 
+int32 GDisplayClusterProjectionDomeprojectionPolicyEnableNearClippingPlane = 0;
+static FAutoConsoleVariableRef CVarDisplayClusterProjectionDomeprojectionPolicyEnableNearClippingPlane(
+	TEXT("nDisplay.render.projection.EnableNearClippingPlane.Dome"),
+	GDisplayClusterProjectionDomeprojectionPolicyEnableNearClippingPlane,
+	TEXT("Enable NearClippingPlane for DomeProjection (0 - disable).\n"),
+	ECVF_RenderThreadSafe
+);
 
+//////////////////////////////////////////////////////////////////////////////////////////////
+// FDisplayClusterProjectionDomeprojectionPolicyBase
+//////////////////////////////////////////////////////////////////////////////////////////////
 FDisplayClusterProjectionDomeprojectionPolicyBase::FDisplayClusterProjectionDomeprojectionPolicyBase(const FString& ProjectionPolicyId, const FDisplayClusterConfigurationProjection* InConfigurationProjectionPolicy)
 	: FDisplayClusterProjectionPolicyBase(ProjectionPolicyId, InConfigurationProjectionPolicy)
 {
 }
 
-
-//////////////////////////////////////////////////////////////////////////////////////////////
-// IDisplayClusterProjectionPolicy
-//////////////////////////////////////////////////////////////////////////////////////////////
 const FString& FDisplayClusterProjectionDomeprojectionPolicyBase::GetType() const
 {
 	static const FString& Type(DisplayClusterProjectionStrings::projection::Domeprojection);
@@ -95,7 +101,7 @@ void FDisplayClusterProjectionDomeprojectionPolicyBase::HandleEndScene(IDisplayC
 }
 
 
-bool FDisplayClusterProjectionDomeprojectionPolicyBase::CalculateView(IDisplayClusterViewport* InViewport, const uint32 InContextNum, FVector& InOutViewLocation, FRotator& InOutViewRotation, const FVector& ViewOffset, const float WorldToMeters, const float NCP, const float FCP)
+bool FDisplayClusterProjectionDomeprojectionPolicyBase::CalculateView(IDisplayClusterViewport* InViewport, const uint32 InContextNum, FVector& InOutViewLocation, FRotator& InOutViewRotation, const FVector& ViewOffset, const float WorldToMeters, const float InNCP, const float InFCP)
 {
 	check(IsInGameThread());
 
@@ -104,6 +110,9 @@ bool FDisplayClusterProjectionDomeprojectionPolicyBase::CalculateView(IDisplayCl
 		return false;
 	}
 
+	const float NCP = GDisplayClusterProjectionDomeprojectionPolicyEnableNearClippingPlane ? InNCP : 1;
+	const float FCP = GDisplayClusterProjectionDomeprojectionPolicyEnableNearClippingPlane ? InFCP : 1;
+
 	// Get origin component
 	const USceneComponent* const OriginComp = GetOriginComp();
 	check(OriginComp);
diff --git a/Engine/Plugins/Runtime/nDisplay/Source/DisplayClusterProjection/Private/Policy/VIOSO/Windows/DisplayClusterProjectionVIOSOPolicy.cpp b/Engine/Plugins/Runtime/nDisplay/Source/DisplayClusterProjection/Private/Policy/VIOSO/Windows/DisplayClusterProjectionVIOSOPolicy.cpp
--- a/Engine/Plugins/Runtime/nDisplay/Source/DisplayClusterProjection/Private/Policy/VIOSO/Windows/DisplayClusterProjectionVIOSOPolicy.cpp
+++ b/Engine/Plugins/Runtime/nDisplay/Source/DisplayClusterProjection/Private/Policy/VIOSO/Windows/DisplayClusterProjectionVIOSOPolicy.cpp
@@ -36,6 +36,13 @@
 #include "ID3D12DynamicRHI.h"
 #endif // VIOSO_USE_GRAPHICS_API_D3D12
 
+int32 GDisplayClusterProjectionVIOSOPolicyEnableNearClippingPlane = 0;
+static FAutoConsoleVariableRef CVarDisplayClusterProjectionVIOSOPolicyEnableNearClippingPlane(
+	TEXT("nDisplay.render.projection.EnableNearClippingPlane.vioso"),
+	GDisplayClusterProjectionVIOSOPolicyEnableNearClippingPlane,
+	TEXT("Enable NearClippingPlane for VIOSO (0 - disable).\n"),
+	ECVF_RenderThreadSafe
+);
 
 #if VIOSO_USE_GRAPHICS_API_D3D11
 /**
@@ -223,10 +230,13 @@ void FDisplayClusterProjectionVIOSOPolicy::ImplRelease()
 	}
 }
 
-bool FDisplayClusterProjectionVIOSOPolicy::CalculateView(IDisplayClusterViewport* InViewport, const uint32 InContextNum, FVector& InOutViewLocation, FRotator& InOutViewRotation, const FVector& ViewOffset, const float WorldToMeters, const float NCP, const float FCP)
+bool FDisplayClusterProjectionVIOSOPolicy::CalculateView(IDisplayClusterViewport* InViewport, const uint32 InContextNum, FVector& InOutViewLocation, FRotator& InOutViewRotation, const FVector& ViewOffset, const float WorldToMeters, const float InNCP, const float InFCP)
 {
 	check(Views.Num() > (int32)InContextNum);
 
+	const float NCP = GDisplayClusterProjectionVIOSOPolicyEnableNearClippingPlane ? InNCP : 1;
+	const float FCP = GDisplayClusterProjectionVIOSOPolicyEnableNearClippingPlane ? InFCP : 1;
+
 	// Get view location in local space
 	const USceneComponent* const OriginComp = GetOriginComp();
 	const FTransform& World2LocalTransform = (OriginComp ? OriginComp->GetComponentTransform() : FTransform::Identity);
