nDisplay: fixed crash that happens if nD session start failed (Linux)

diff --git a/Engine/Plugins/Runtime/nDisplay/Source/DisplayCluster/Private/Game/DisplayClusterGameManager.cpp b/Engine/Plugins/Runtime/nDisplay/Source/DisplayCluster/Private/Game/DisplayClusterGameManager.cpp
--- a/Engine/Plugins/Runtime/nDisplay/Source/DisplayCluster/Private/Game/DisplayClusterGameManager.cpp
+++ b/Engine/Plugins/Runtime/nDisplay/Source/DisplayCluster/Private/Game/DisplayClusterGameManager.cpp
@@ -73,7 +73,10 @@ bool FDisplayClusterGameManager::StartScene(UWorld* InWorld)
 {
 	FScopeLock Lock(&InternalsSyncScope);
 
-	check(InWorld && ConfigData);
+	if(!InWorld || !ConfigData)
+	{
+		return false;
+	}
 
 	CurrentWorld = InWorld;
 
