nDisplay: Removed unnecessary Settings button from the color grading drawer, since there are currently no drawer specific settings to expose

diff --git a/Engine/Plugins/Runtime/nDisplay/Source/DisplayClusterColorGrading/Private/Drawer/SDisplayClusterColorGradingDrawer.cpp b/Engine/Plugins/Runtime/nDisplay/Source/DisplayClusterColorGrading/Private/Drawer/SDisplayClusterColorGradingDrawer.cpp
--- a/Engine/Plugins/Runtime/nDisplay/Source/DisplayClusterColorGrading/Private/Drawer/SDisplayClusterColorGradingDrawer.cpp
+++ b/Engine/Plugins/Runtime/nDisplay/Source/DisplayClusterColorGrading/Private/Drawer/SDisplayClusterColorGradingDrawer.cpp
@@ -263,39 +263,6 @@ void SDisplayClusterColorGradingDrawer::Construct(const FArguments& InArgs, bool
 							[
 								CreateDockInLayoutButton()
 							]
-
-							// Slot for the Settings button
-							+ SHorizontalBox::Slot()
-							.AutoWidth()
-							.Padding(5.0f, 0.0f, 0.0f, 0.0f)
-							.HAlign(HAlign_Right)
-							.VAlign(VAlign_Center)
-							[
-								SNew(SComboButton)
-								.ComboButtonStyle(&FAppStyle::Get().GetWidgetStyle<FComboButtonStyle>("SimpleComboButton"))
-								.HasDownArrow(false)
-								.ButtonContent()
-								[
-									SNew(SHorizontalBox)
-									+ SHorizontalBox::Slot()
-									.AutoWidth()
-									.VAlign(VAlign_Center)
-									.Padding(4.0, 0.0f)
-									[
-										SNew(SImage)
-										.ColorAndOpacity(FSlateColor::UseForeground())
-										.Image(FAppStyle::Get().GetBrush("Icons.Settings"))
-									]
-									+ SHorizontalBox::Slot()
-									.VAlign(VAlign_Center)
-									.Padding(4.0, 0.0f)
-									[
-										SNew(STextBlock)
-										.Text(LOCTEXT("Settings", "Settings"))
-										.ColorAndOpacity(FSlateColor::UseForeground())
-									]
-								]
-							]
 						]
 					]
 				]
