nDisplay: fixed the inability to establish necessary TCP connections (nDisplay in Linux)

diff --git a/Engine/Plugins/Runtime/nDisplay/Source/DisplayCluster/Private/Network/Transport/DisplayClusterSocketOperations.h b/Engine/Plugins/Runtime/nDisplay/Source/DisplayCluster/Private/Network/Transport/DisplayClusterSocketOperations.h
--- a/Engine/Plugins/Runtime/nDisplay/Source/DisplayCluster/Private/Network/Transport/DisplayClusterSocketOperations.h
+++ b/Engine/Plugins/Runtime/nDisplay/Source/DisplayCluster/Private/Network/Transport/DisplayClusterSocketOperations.h
@@ -32,7 +32,8 @@ public:
 	// Connect
 	inline bool ConnectSocket(const FInternetAddr& Addr)
 	{
-		if (!IsOpen())
+		// if(!IsOpen()) check was removed from here because of an issue in BSD sockets.
+		if (Socket)
 		{
 			return Socket->Connect(Addr);
 		}
