nDisplay: Fix light card references being lost when restoring level snapshots by adding new light cards to a light card layer.
Fix stage app not creating new proxies in the operator panel when creating a new light card.
Fix possible crash on light card editor viewport startup.
Fix light card editor not refreshing when an actor is spawned in the level instance, such as through duplication.
Fix light card editor not always refreshing if certain tabs were closed.

diff --git a/Engine/Plugins/Runtime/nDisplay/Source/DisplayCluster/Private/Blueprints/DisplayClusterBlueprintLib.cpp b/Engine/Plugins/Runtime/nDisplay/Source/DisplayCluster/Private/Blueprints/DisplayClusterBlueprintLib.cpp
--- a/Engine/Plugins/Runtime/nDisplay/Source/DisplayCluster/Private/Blueprints/DisplayClusterBlueprintLib.cpp
+++ b/Engine/Plugins/Runtime/nDisplay/Source/DisplayCluster/Private/Blueprints/DisplayClusterBlueprintLib.cpp
@@ -64,12 +64,13 @@ ADisplayClusterLightCardActor* UDisplayClusterBlueprintLib::CreateLightCard(ADis
 	NewActor->AttachToActor(RootActor, AttachmentRules);
 
 	// Add it to the root actor
-	UDisplayClusterConfigurationData* ConfigData = RootActor->GetConfigData();
-	ConfigData->Modify();
-	FDisplayClusterConfigurationICVFX_VisibilityList& RootActorLightCards = ConfigData->StageSettings.Lightcard.ShowOnlyList;
-
-	RootActorLightCards.Actors.Add(NewActor);
+	NewActor->AddToLightCardLayer(RootActor);
 
+#if WITH_EDITOR
+	// Required so operator panel updates
+	RootActor->PostEditChange();
+#endif
+	
 	return NewActor;
 }
 
@@ -120,11 +121,12 @@ void UDisplayClusterBlueprintLib::DuplicateLightCards(TArray<ADisplayClusterLigh
 		NewLightCard->AttachToActor(RootActor, AttachmentRules);
 
 		// Add it to the root actor
-		UDisplayClusterConfigurationData* ConfigData = RootActor->GetConfigData();
-		ConfigData->Modify();
-		FDisplayClusterConfigurationICVFX_VisibilityList& RootActorLightCards = ConfigData->StageSettings.Lightcard.ShowOnlyList;
+		NewLightCard->AddToLightCardLayer(RootActor);
 
-		RootActorLightCards.Actors.Add(NewLightCard);
+#if WITH_EDITOR
+		// Required so operator panel updates
+		RootActor->PostEditChange();
+#endif
 
 #if WITH_EDITOR
 		if (GIsEditor)
@@ -185,4 +187,4 @@ void UDisplayClusterBlueprintLib::FindLightCardsForRootActor(ADisplayClusterRoot
 	}
 }
 
-#undef LOCTEXT_NAMESPACE
\ No newline at end of file
+#undef LOCTEXT_NAMESPACE
diff --git a/Engine/Plugins/Runtime/nDisplay/Source/DisplayCluster/Private/Game/EngineClasses/Scene/DisplayClusterLightCardActor.cpp b/Engine/Plugins/Runtime/nDisplay/Source/DisplayCluster/Private/Game/EngineClasses/Scene/DisplayClusterLightCardActor.cpp
--- a/Engine/Plugins/Runtime/nDisplay/Source/DisplayCluster/Private/Game/EngineClasses/Scene/DisplayClusterLightCardActor.cpp
+++ b/Engine/Plugins/Runtime/nDisplay/Source/DisplayCluster/Private/Game/EngineClasses/Scene/DisplayClusterLightCardActor.cpp
@@ -1,6 +1,8 @@
 // Copyright Epic Games, Inc. All Rights Reserved.
 
 #include "DisplayClusterLightCardActor.h"
+
+#include "DisplayClusterConfigurationTypes.h"
 #include "DisplayClusterRootActor.h"
 #include "IDisplayClusterLightCardActorExtender.h"
 
@@ -11,6 +13,13 @@
 #include "Engine/StaticMesh.h"
 #include "UObject/ConstructorHelpers.h"
 
+#if WITH_EDITOR
+
+#include "Editor.h"
+#include "Layers/LayersSubsystem.h"
+
+#endif
+
 #if WITH_OPENCV
 
 #include "OpenCVHelper.h"
@@ -458,6 +467,61 @@ void ADisplayClusterLightCardActor::SetRootActorOwner(ADisplayClusterRootActor*
 	RootActorOwner = MakeWeakObjectPtr(InRootActor);
 }
 
+void ADisplayClusterLightCardActor::AddToLightCardLayer(ADisplayClusterRootActor* InRootActor)
+{
+	check(InRootActor);
+
+	UDisplayClusterConfigurationData* ConfigData = InRootActor->GetConfigData();
+	check(ConfigData);
+	
+	// Add light cards to an existing layer or a new layer. This helps with restoring snapshots of a DCRA
+	// so they can be done separately from light cards and still have access to all light cards in the layer.
+				
+	static const FString LightCardPrefix = TEXT("LightCards");
+	FName LightCardLayerName = NAME_None;
+
+	FDisplayClusterConfigurationICVFX_VisibilityList& RootActorLightCards = ConfigData->StageSettings.Lightcard.ShowOnlyList;
+	
+	if (const FActorLayer* ExistingLightCardLayer = RootActorLightCards.ActorLayers.FindByPredicate([](const FActorLayer& Layer)
+	{
+		return Layer.Name.ToString().StartsWith(LightCardPrefix);
+	}))
+	{
+		LightCardLayerName = ExistingLightCardLayer->Name;
+	}
+	else
+	{
+		// Find a unique name so it's LightCards_XXXX
+		do
+		{
+			const FGuid LayerGuid = FGuid::NewGuid();
+			const FString GuidStr = LayerGuid.ToString().Left(4);
+			LightCardLayerName = *FString::Printf(TEXT("%s_%s"), *LightCardPrefix, *GuidStr);
+		}
+		while (FindObject<UObject>(GetWorld(), *LightCardLayerName.ToString()) != nullptr);
+		
+		ConfigData->Modify();
+		RootActorLightCards.ActorLayers.AddDefaulted_GetRef().Name = LightCardLayerName;
+	}
+	
+	check(!LightCardLayerName.IsNone());
+
+	bool bLayerAdded = false;
+#if WITH_EDITOR
+	ULayersSubsystem* LayersSubsystem = GEditor ? GEditor->GetEditorSubsystem<ULayersSubsystem>() : nullptr;
+	if (LayersSubsystem)
+	{
+		LayersSubsystem->AddActorsToLayer(TArray<AActor*>{ this }, LightCardLayerName);
+		bLayerAdded = true;
+	}
+#endif
+	if (!bLayerAdded)
+	{
+		// Likely -game and we need to add the layer directly
+		Layers.AddUnique(LightCardLayerName);
+	}
+}
+
 void ADisplayClusterLightCardActor::UpdateStageActorTransform()
 {
 	// If the light card is in UV space, set the spring arm's trasform to be zero, effectively removing it from the transform hierarchy
diff --git a/Engine/Plugins/Runtime/nDisplay/Source/DisplayCluster/Public/DisplayClusterLightCardActor.h b/Engine/Plugins/Runtime/nDisplay/Source/DisplayCluster/Public/DisplayClusterLightCardActor.h
--- a/Engine/Plugins/Runtime/nDisplay/Source/DisplayCluster/Public/DisplayClusterLightCardActor.h
+++ b/Engine/Plugins/Runtime/nDisplay/Source/DisplayCluster/Public/DisplayClusterLightCardActor.h
@@ -103,7 +103,10 @@ public:
 	
 	/** Return the current owner, providing one was set */
 	TWeakObjectPtr<ADisplayClusterRootActor> GetRootActorOwner() const { return RootActorOwner; }
-
+	
+	/** Add this light card to a light card layer on the given root actor */
+	void AddToLightCardLayer(ADisplayClusterRootActor* InRootActor);
+	
 	// ~Begin IDisplayClusterStageActor interface
 	/** Updates the Light Card transform based on its positional properties (Lat, Long, etc.) */
 	virtual void UpdateStageActorTransform() override;
diff --git a/Engine/Plugins/Runtime/nDisplay/Source/DisplayClusterLightCardEditor/Private/DisplayClusterLightCardEditor.cpp b/Engine/Plugins/Runtime/nDisplay/Source/DisplayClusterLightCardEditor/Private/DisplayClusterLightCardEditor.cpp
--- a/Engine/Plugins/Runtime/nDisplay/Source/DisplayClusterLightCardEditor/Private/DisplayClusterLightCardEditor.cpp
+++ b/Engine/Plugins/Runtime/nDisplay/Source/DisplayClusterLightCardEditor/Private/DisplayClusterLightCardEditor.cpp
@@ -40,6 +40,7 @@
 #include "Framework/Docking/TabManager.h"
 #include "Framework/MultiBox/MultiBoxBuilder.h"
 #include "Framework/MultiBox/MultiBoxExtender.h"
+#include "Layers/LayersSubsystem.h"
 #include "Misc/FileHelper.h"
 #include "Misc/TransactionObjectEvent.h"
 #include "Styling/SlateIconFinder.h"
@@ -70,6 +71,7 @@ FDisplayClusterLightCardEditor::~FDisplayClusterLightCardEditor()
 
 	if (GEngine != nullptr)
 	{
+		GEngine->OnLevelActorAdded().RemoveAll(this);
 		GEngine->OnLevelActorDeleted().RemoveAll(this);
 	}
 
@@ -100,6 +102,7 @@ void FDisplayClusterLightCardEditor::Initialize(TSharedRef<IDisplayClusterOperat
 	ActiveRootActorChangedHandle = InViewModel->OnActiveRootActorChanged().AddSP(this, &FDisplayClusterLightCardEditor::OnActiveRootActorChanged);
 	if (GEngine != nullptr)
 	{
+		GEngine->OnLevelActorAdded().AddSP(this, &FDisplayClusterLightCardEditor::OnLevelActorAdded);
 		GEngine->OnLevelActorDeleted().AddSP(this, &FDisplayClusterLightCardEditor::OnLevelActorDeleted);
 	}
 	
@@ -232,7 +235,7 @@ AActor* FDisplayClusterLightCardEditor::SpawnActor(TSubclassOf<AActor> InActorCl
 	}
 	else
 	{
-		RefreshPreviewActors(EDisplayClusterLightCardEditorProxyType::LightCards);
+		RefreshPreviewActors(EDisplayClusterLightCardEditorProxyType::StageActor);
 	}
 	
 	SelectActors({ NewActor });
@@ -496,11 +499,11 @@ void FDisplayClusterLightCardEditor::AddLightCardsToActor(const TArray<ADisplayC
 						return Actor == LightCardSoftObject;
 					});
 
-				RootActorLightCards.Actors.Add(LightCard);
+				LightCard->AddToLightCardLayer(ActiveRootActor.Get());
 			}
 		}
 
-		RefreshPreviewActors(EDisplayClusterLightCardEditorProxyType::LightCards);
+		RefreshPreviewActors(EDisplayClusterLightCardEditorProxyType::StageActor);
 	}
 }
 
@@ -766,9 +769,24 @@ void FDisplayClusterLightCardEditor::RemoveActors(
 					{
 						return InActor.Get() == LightCard;
 					});
-
+				
 				if (!bDeleteActors)
 				{
+					// Remove from any layers shared with the DCRA. If we aren't deleting this has the possibility to
+					// remove from a layer used by another DCRA, but this is rare and acceptable.
+					ULayersSubsystem* LayersSubsystem = GEditor->GetEditorSubsystem<ULayersSubsystem>();
+					TArray<FName> LightCardLayerNames = LightCard->Layers;
+					for (const FName& LightCardLayerName : LightCardLayerNames)
+					{
+						if (RootActorLightCards.ActorLayers.ContainsByPredicate([&](const FActorLayer& LightCardLayer)
+						{
+							return LightCardLayer.Name == LightCardLayerName;
+						}))
+						{
+							LayersSubsystem->RemoveActorFromLayer(LightCard, LightCardLayerName);
+						}
+					}
+					
 					LightCard->ShowLightCardLabel(false, *GetLightCardLabelScale(), ActiveRootActor.Get());
 				}
 			}
@@ -793,7 +811,7 @@ void FDisplayClusterLightCardEditor::RemoveActors(
 		}
 	}
 
-	RefreshPreviewActors(EDisplayClusterLightCardEditorProxyType::LightCards);
+	RefreshPreviewActors(EDisplayClusterLightCardEditorProxyType::StageActor);
 }
 
 bool FDisplayClusterLightCardEditor::CanRemoveSelectedActors() const
@@ -919,7 +937,7 @@ void FDisplayClusterLightCardEditor::ShowLightCardLabels(bool bVisible)
 	
 	LightCardEditorModule.ShowLabels(MoveTemp(Args));
 
-	RefreshPreviewActors(EDisplayClusterLightCardEditorProxyType::LightCards);
+	RefreshPreviewActors(EDisplayClusterLightCardEditorProxyType::StageActor);
 }
 
 bool FDisplayClusterLightCardEditor::ShouldShowLightCardLabels() const
@@ -1676,17 +1694,30 @@ bool FDisplayClusterLightCardEditor::IsOurObject(UObject* InObject,
 	EDisplayClusterLightCardEditorProxyType ProxyType = EDisplayClusterLightCardEditorProxyType::All;
 	
 	bool bIsOurActor = IsOurActor(GetActiveRootActor().Get());
-	if (!bIsOurActor && LightCardOutliner.IsValid())
+	if (!bIsOurActor)
 	{
-		for (const TSharedPtr<SDisplayClusterLightCardOutliner::FStageActorTreeItem>& LightCard : LightCardOutliner->GetStageActorTreeItems())
+		if (LightCardOutliner.IsValid())
 		{
-			bIsOurActor = IsOurActor(LightCard->Actor.Get());
-			if (bIsOurActor)
+			for (const TSharedPtr<SDisplayClusterLightCardOutliner::FStageActorTreeItem>& LightCard : LightCardOutliner->GetStageActorTreeItems())
 			{
-				ProxyType = EDisplayClusterLightCardEditorProxyType::LightCards;
-				break;
+				if (LightCard.IsValid())
+				{
+					bIsOurActor = IsOurActor(LightCard->Actor.Get());
+					break;
+				}
 			}
 		}
+
+		// Fallback to generic check
+		if (!bIsOurActor && UE::DisplayClusterLightCardEditorUtils::IsManagedActor(Cast<AActor>(InObject)))
+		{
+			bIsOurActor = true;
+		}
+
+		if (bIsOurActor)
+		{
+			ProxyType = EDisplayClusterLightCardEditorProxyType::StageActor;
+		}
 	}
 
 	OutProxyType = ProxyType;
@@ -1770,13 +1801,17 @@ void FDisplayClusterLightCardEditor::OnActorPropertyChanged(UObject* ObjectBeing
 	}
 }
 
-void FDisplayClusterLightCardEditor::OnLevelActorDeleted(AActor* Actor)
+void FDisplayClusterLightCardEditor::OnLevelActorAdded(AActor* Actor)
 {
-	if (LightCardOutliner.IsValid() &&
-		LightCardOutliner->GetStageActorTreeItems().ContainsByPredicate([Actor](const TSharedPtr<SDisplayClusterLightCardOutliner::FStageActorTreeItem>& LightCardTreeItem)
+	if (UE::DisplayClusterLightCardEditorUtils::IsManagedActor(Actor))
 	{
-		return LightCardTreeItem.IsValid() && Actor == LightCardTreeItem->Actor.Get();
-	}))
+		RefreshPreviewActors(EDisplayClusterLightCardEditorProxyType::StageActor);
+	}
+}
+
+void FDisplayClusterLightCardEditor::OnLevelActorDeleted(AActor* Actor)
+{
+	if (UE::DisplayClusterLightCardEditorUtils::IsManagedActor(Actor))
 	{
 		if (Actor && Actor->GetClass()->HasAnyClassFlags(CLASS_NewerVersionExists))
 		{
@@ -1785,13 +1820,13 @@ void FDisplayClusterLightCardEditor::OnLevelActorDeleted(AActor* Actor)
 			return;
 		}
 		
-		if (ViewportView.IsValid())
+		if (GEditor && GEditor->GetWorld())
 		{
-			ViewportView->GetLightCardEditorViewportClient()->GetWorld()->GetTimerManager().SetTimerForNextTick([=]()
+			GEditor->GetWorld()->GetTimerManager().SetTimerForNextTick([=]()
 			{
 				// Schedule for next tick so available selections are properly updated once the
 				// actor is fully deleted.
-				RefreshPreviewActors(EDisplayClusterLightCardEditorProxyType::LightCards);
+				RefreshPreviewActors(EDisplayClusterLightCardEditorProxyType::StageActor);
 			});
 		}
 	}
@@ -1800,7 +1835,7 @@ void FDisplayClusterLightCardEditor::OnLevelActorDeleted(AActor* Actor)
 void FDisplayClusterLightCardEditor::OnBlueprintCompiled(UBlueprint* Blueprint)
 {
 	// Right now only LightCard blueprints are handled here.
-	RefreshPreviewActors(EDisplayClusterLightCardEditorProxyType::LightCards);
+	RefreshPreviewActors(EDisplayClusterLightCardEditorProxyType::StageActor);
 }
 
 void FDisplayClusterLightCardEditor::OnObjectTransacted(UObject* Object,
diff --git a/Engine/Plugins/Runtime/nDisplay/Source/DisplayClusterLightCardEditor/Private/DisplayClusterLightCardEditor.h b/Engine/Plugins/Runtime/nDisplay/Source/DisplayClusterLightCardEditor/Private/DisplayClusterLightCardEditor.h
--- a/Engine/Plugins/Runtime/nDisplay/Source/DisplayClusterLightCardEditor/Private/DisplayClusterLightCardEditor.h
+++ b/Engine/Plugins/Runtime/nDisplay/Source/DisplayClusterLightCardEditor/Private/DisplayClusterLightCardEditor.h
@@ -286,6 +286,9 @@ private:
 	/** When a property on the actor has changed */
 	void OnActorPropertyChanged(UObject* ObjectBeingModified, FPropertyChangedEvent& PropertyChangedEvent);
 
+	/** Raised when the user adds an actor to the level */
+	void OnLevelActorAdded(AActor* Actor);
+	
 	/** Raised when the user deletes an actor from the level */
 	void OnLevelActorDeleted(AActor* Actor);
 
diff --git a/Engine/Plugins/Runtime/nDisplay/Source/DisplayClusterLightCardEditor/Private/DisplayClusterLightCardEditorProxyType.h b/Engine/Plugins/Runtime/nDisplay/Source/DisplayClusterLightCardEditor/Private/DisplayClusterLightCardEditorProxyType.h
--- a/Engine/Plugins/Runtime/nDisplay/Source/DisplayClusterLightCardEditor/Private/DisplayClusterLightCardEditorProxyType.h
+++ b/Engine/Plugins/Runtime/nDisplay/Source/DisplayClusterLightCardEditor/Private/DisplayClusterLightCardEditorProxyType.h
@@ -8,5 +8,5 @@ enum class EDisplayClusterLightCardEditorProxyType : uint8
 {
 	All,
 	RootActor,
-	LightCards
+	StageActor
 };
\ No newline at end of file
diff --git a/Engine/Plugins/Runtime/nDisplay/Source/DisplayClusterLightCardEditor/Private/Viewport/DisplayClusterLightCardEditorViewportClient.cpp b/Engine/Plugins/Runtime/nDisplay/Source/DisplayClusterLightCardEditor/Private/Viewport/DisplayClusterLightCardEditorViewportClient.cpp
--- a/Engine/Plugins/Runtime/nDisplay/Source/DisplayClusterLightCardEditor/Private/Viewport/DisplayClusterLightCardEditorViewportClient.cpp
+++ b/Engine/Plugins/Runtime/nDisplay/Source/DisplayClusterLightCardEditor/Private/Viewport/DisplayClusterLightCardEditorViewportClient.cpp
@@ -116,8 +116,6 @@ FDisplayClusterLightCardEditorViewportClient::FDisplayClusterLightCardEditorView
 
 	//This seems to be needed to get the correct world time in the preview.
 	SetIsSimulateInEditorViewport(true);
-	
-	UpdatePreviewActor(LightCardEditorPtr.Pin()->GetActiveRootActor().Get());
 
 	const UDisplayClusterLightCardEditorSettings* Settings = GetDefault<UDisplayClusterLightCardEditorSettings>();
 	
@@ -1428,7 +1426,7 @@ void FDisplayClusterLightCardEditorViewportClient::UpdatePreviewActor(ADisplayCl
 			IDisplayClusterScenePreview::Get().SetRendererRootActor(PreviewRendererId, RootActorProxy.Get());
 
 			if (ProxyType == EDisplayClusterLightCardEditorProxyType::All ||
-				ProxyType == EDisplayClusterLightCardEditorProxyType::LightCards)
+				ProxyType == EDisplayClusterLightCardEditorProxyType::StageActor)
 			{
 				SelectActor(nullptr);
 				const FTransform RALevelTransformNoScale(RootActorLevelInstance->GetActorRotation(), RootActorLevelInstance->GetActorLocation(), FVector::OneVector);
@@ -1562,7 +1560,7 @@ void FDisplayClusterLightCardEditorViewportClient::DestroyProxies(
 		}
 		break;
 
-	case EDisplayClusterLightCardEditorProxyType::LightCards:
+	case EDisplayClusterLightCardEditorProxyType::StageActor:
 		for (const FActorProxy& ActorProxy : ActorProxies)
 		{
 			if (ActorProxy.Proxy.IsValid())
@@ -1600,7 +1598,7 @@ void FDisplayClusterLightCardEditorViewportClient::DestroyProxies(
 	}
 	
 	if (ProxyType == EDisplayClusterLightCardEditorProxyType::All ||
-		ProxyType == EDisplayClusterLightCardEditorProxyType::LightCards)
+		ProxyType == EDisplayClusterLightCardEditorProxyType::StageActor)
 	{
 		for (const FActorProxy& ActorProxy : ActorProxies)
 		{
diff --git a/Engine/Plugins/Runtime/nDisplay/Source/DisplayClusterLightCardEditor/Private/Viewport/DisplayClusterLightcardEditorViewport.cpp b/Engine/Plugins/Runtime/nDisplay/Source/DisplayClusterLightCardEditor/Private/Viewport/DisplayClusterLightcardEditorViewport.cpp
--- a/Engine/Plugins/Runtime/nDisplay/Source/DisplayClusterLightCardEditor/Private/Viewport/DisplayClusterLightcardEditorViewport.cpp
+++ b/Engine/Plugins/Runtime/nDisplay/Source/DisplayClusterLightCardEditor/Private/Viewport/DisplayClusterLightcardEditorViewport.cpp
@@ -410,6 +410,13 @@ TSharedRef<FEditorViewportClient> SDisplayClusterLightCardEditorViewport::MakeEd
 	
 	ViewportClient = MakeShareable(new FDisplayClusterLightCardEditorViewportClient(*PreviewScene.Get(),
 		SharedThis(this)));
+
+	if (LightCardEditorPtr.IsValid())
+	{
+		// Call after construction as UpdatePreviewActor will create a weak reference to itself
+		ViewportClient->UpdatePreviewActor(LightCardEditorPtr.Pin()->GetActiveRootActor().Get());
+	}
+	
 	return ViewportClient.ToSharedRef();
 }
 
