nDisplay: fixed crash with Vioso nDisplay config on Dx12

diff --git a/Engine/Plugins/Runtime/nDisplay/Source/DisplayClusterProjection/Private/Policy/VIOSO/Windows/DisplayClusterProjectionVIOSOPolicy.cpp b/Engine/Plugins/Runtime/nDisplay/Source/DisplayClusterProjection/Private/Policy/VIOSO/Windows/DisplayClusterProjectionVIOSOPolicy.cpp
--- a/Engine/Plugins/Runtime/nDisplay/Source/DisplayClusterProjection/Private/Policy/VIOSO/Windows/DisplayClusterProjectionVIOSOPolicy.cpp
+++ b/Engine/Plugins/Runtime/nDisplay/Source/DisplayClusterProjection/Private/Policy/VIOSO/Windows/DisplayClusterProjectionVIOSOPolicy.cpp
@@ -145,21 +145,6 @@ private:
 FDisplayClusterProjectionVIOSOPolicy::FDisplayClusterProjectionVIOSOPolicy(const FString& ProjectionPolicyId, const FDisplayClusterConfigurationProjection* InConfigurationProjectionPolicy)
 	: FDisplayClusterProjectionPolicyBase(ProjectionPolicyId, InConfigurationProjectionPolicy)
 {
-	const ERHIInterfaceType RHIType = RHIGetInterfaceType();
-
-	// Create warper for view
-	if (RHIType == ERHIInterfaceType::D3D11)
-	{
-		RenderDevice = ERenderDevice::D3D11;
-	}
-	else if (RHIType == ERHIInterfaceType::D3D12)
-	{
-		RenderDevice = ERenderDevice::D3D12;
-	}
-	else
-	{
-		UE_LOG(LogDisplayClusterProjectionVIOSO, Error, TEXT("VIOSO warp projection not supported by '%s' rhi"), GDynamicRHI->GetName());
-	}
 }
 
 FDisplayClusterProjectionVIOSOPolicy::~FDisplayClusterProjectionVIOSOPolicy()
@@ -193,17 +178,12 @@ bool FDisplayClusterProjectionVIOSOPolicy::HandleStartScene(IDisplayClusterViewp
 	InitializeOriginComponent(InViewport, ViosoConfigData.OriginCompId);
 
 	Views.AddDefaulted(2);
-	//ViewportSize = InViewportSize;
 
 	// Initialize data for all views
 	FScopeLock lock(&DllAccessCS);
-	for (FViewData& ViewIt : Views)
+	for (TSharedPtr<FDisplayClusterProjectionVIOSOPolicyViewData, ESPMode::ThreadSafe>& ViewIt : Views)
 	{
-		if (!ViewIt.Initialize(RenderDevice, ViosoConfigData))
-		{
-			UE_LOG(LogDisplayClusterProjectionVIOSO, Error, TEXT("VIOSO not initialized"));
-			return false;
-		}
+		ViewIt = MakeShared<FDisplayClusterProjectionVIOSOPolicyViewData>();
 	}
 
 	UE_LOG(LogDisplayClusterProjectionVIOSO, Verbose, TEXT("VIOSO policy has been initialized: %s"), *ViosoConfigData.ToString());
@@ -224,15 +204,17 @@ void FDisplayClusterProjectionVIOSOPolicy::ImplRelease()
 
 	// Destroy VIOSO for all views
 	FScopeLock lock(&DllAccessCS);
-	for (FViewData& ViewIt : Views)
-	{
-		ViewIt.DestroyVIOSO();
-	}
+	Views.Reset();
 }
 
 bool FDisplayClusterProjectionVIOSOPolicy::CalculateView(IDisplayClusterViewport* InViewport, const uint32 InContextNum, FVector& InOutViewLocation, FRotator& InOutViewRotation, const FVector& ViewOffset, const float WorldToMeters, const float InNCP, const float InFCP)
 {
-	check(Views.Num() > (int32)InContextNum);
+	check(Views.IsValidIndex(InContextNum));
+	TSharedPtr<FDisplayClusterProjectionVIOSOPolicyViewData, ESPMode::ThreadSafe>& CurrentView = Views[InContextNum];
+	if (!CurrentView.IsValid())
+	{
+		return false;
+	}
 
 	const float NCP = GDisplayClusterProjectionVIOSOPolicyEnableNearClippingPlane ? InNCP : 1;
 	const float FCP = GDisplayClusterProjectionVIOSOPolicyEnableNearClippingPlane ? InFCP : 1;
@@ -248,9 +230,9 @@ bool FDisplayClusterProjectionVIOSOPolicy::CalculateView(IDisplayClusterViewport
 
 	// Get view prj data from VIOSO
 	FScopeLock lock(&DllAccessCS);
-	if (!Views[InContextNum].UpdateVIOSO(InViewport, InContextNum, LocalEyeOrigin, LocalRotator, WorldToMeters, NCP, FCP))
+	if (!CurrentView->UpdateVIOSO(InViewport, InContextNum, LocalEyeOrigin, LocalRotator, WorldToMeters, NCP, FCP))
 	{
-		if (Views[InContextNum].IsValid() && !FDisplayClusterProjectionPolicyBase::IsEditorOperationMode(InViewport))
+		if (!FDisplayClusterProjectionPolicyBase::IsEditorOperationMode(InViewport))
 		{
 			// Vioso api used, but failed inside math. The config base matrix or vioso geometry is invalid
 			UE_LOG(LogDisplayClusterProjectionVIOSO, Error, TEXT("Couldn't Calculate View for VIOSO viewport '%s'"), *InViewport->GetId());
@@ -260,8 +242,8 @@ bool FDisplayClusterProjectionVIOSOPolicy::CalculateView(IDisplayClusterViewport
 	}
 
 	// Transform rotation to world space
-	InOutViewRotation = World2LocalTransform.TransformRotation(Views[InContextNum].ViewRotation.Quaternion()).Rotator();
-	InOutViewLocation = World2LocalTransform.TransformPosition(Views[InContextNum].ViewLocation);
+	InOutViewRotation = World2LocalTransform.TransformRotation(CurrentView->ViewRotation.Quaternion()).Rotator();
+	InOutViewLocation = World2LocalTransform.TransformPosition(CurrentView->ViewLocation);
 
 	return true;
 }
@@ -269,8 +251,14 @@ bool FDisplayClusterProjectionVIOSOPolicy::CalculateView(IDisplayClusterViewport
 bool FDisplayClusterProjectionVIOSOPolicy::GetProjectionMatrix(IDisplayClusterViewport* InViewport, const uint32 InContextNum, FMatrix& OutPrjMatrix)
 {
 	check(IsInGameThread());
+	check(Views.IsValidIndex(InContextNum));
+	TSharedPtr<FDisplayClusterProjectionVIOSOPolicyViewData, ESPMode::ThreadSafe>& CurrentView = Views[InContextNum];
+	if (!CurrentView.IsValid())
+	{
+		return false;
+	}
 
-	OutPrjMatrix = Views[InContextNum].ProjectionMatrix;
+	OutPrjMatrix = CurrentView->ProjectionMatrix;
 	
 	return true;
 }
@@ -323,16 +311,18 @@ bool FDisplayClusterProjectionVIOSOPolicy::ImplApplyWarpBlend_RenderThread(FRHIC
 	// External SDK not use our RHI flow, call flush to finish resolve context image to input resource
 	RHICmdList.ImmediateFlush(EImmediateFlushType::FlushRHIThreadFlushResources);
 
-
 	TRACE_CPUPROFILER_EVENT_SCOPE(nDisplay VIOSO::Render);
 	{
 		FScopeLock lock(&DllAccessCS);
 
 		for (int32 ContextNum = 0; ContextNum < InputTextures.Num(); ContextNum++)
 		{
-			if (!Views[ContextNum].RenderVIOSO_RenderThread(RHICmdList, InputTextures[ContextNum], OutputTextures[ContextNum], ViosoConfigData))
+			if (Views[ContextNum].IsValid())
 			{
-				return false;
+				if (!Views[ContextNum]->RenderVIOSO_RenderThread(RHICmdList, InputTextures[ContextNum], OutputTextures[ContextNum], ViosoConfigData))
+				{
+					return false;
+				}
 			}
 		}
 	}
@@ -342,28 +332,41 @@ bool FDisplayClusterProjectionVIOSOPolicy::ImplApplyWarpBlend_RenderThread(FRHIC
 }
 
 //////////////////////////////////////////////////////////////////////////////////////////////
-// FDisplayClusterProjectionVIOSOPolicy::FViewData
+// FDisplayClusterProjectionVIOSOPolicyViewData
 //////////////////////////////////////////////////////////////////////////////////////////////
-bool FDisplayClusterProjectionVIOSOPolicy::FViewData::IsValid()
+FDisplayClusterProjectionVIOSOPolicyViewData::FDisplayClusterProjectionVIOSOPolicyViewData()
 {
-	return bDataInitialized && bInitialized && Warper.IsValid();
+	const ERHIInterfaceType RHIType = RHIGetInterfaceType();
+
+	// Create warper for view
+	if (RHIType == ERHIInterfaceType::D3D11)
+	{
+		RenderDevice = ERenderDevice::D3D11;
+	}
+	else if (RHIType == ERHIInterfaceType::D3D12)
+	{
+		RenderDevice = ERenderDevice::D3D12;
+	}
+	else
+	{
+		UE_LOG(LogDisplayClusterProjectionVIOSO, Error, TEXT("VIOSO warp projection not supported by '%s' rhi"), GDynamicRHI->GetName());
+	}
 }
 
-bool FDisplayClusterProjectionVIOSOPolicy::FViewData::Initialize(ERenderDevice InRenderDevice, const FViosoPolicyConfiguration& InConfigData)
+FDisplayClusterProjectionVIOSOPolicyViewData::~FDisplayClusterProjectionVIOSOPolicyViewData()
 {
-	check(IsInGameThread());
-
-	Warper.Release();
-
-	// VIOSO requre render thread resources to initialize.
-	// Store args, and initialize latter, on render thread
-	RenderDevice = InRenderDevice;
+	if (IsValid())
+	{
+		Warper.Release();
+	}
+}
 
-	bDataInitialized = true;
-	return true;
+bool FDisplayClusterProjectionVIOSOPolicyViewData::IsValid()
+{
+	return bInitialized && Warper.IsValid();
 }
 
-bool FDisplayClusterProjectionVIOSOPolicy::FViewData::UpdateVIOSO(IDisplayClusterViewport* InViewport, const uint32 InContextNum, const FVector& LocalLocation, const FRotator& LocalRotator, const float WorldToMeters, const float NCP, const float FCP)
+bool FDisplayClusterProjectionVIOSOPolicyViewData::UpdateVIOSO(IDisplayClusterViewport* InViewport, const uint32 InContextNum, const FVector& LocalLocation, const FRotator& LocalRotator, const float WorldToMeters, const float NCP, const float FCP)
 {
 	if (IsValid())
 	{
@@ -376,7 +379,7 @@ bool FDisplayClusterProjectionVIOSOPolicy::FViewData::UpdateVIOSO(IDisplayCluste
 	return false;
 }
 
-bool FDisplayClusterProjectionVIOSOPolicy::FViewData::RenderVIOSO_RenderThread(FRHICommandListImmediate& RHICmdList, FRHITexture2D* ShaderResourceTexture, FRHITexture2D* RenderTargetTexture, const FViosoPolicyConfiguration& InConfigData)
+bool FDisplayClusterProjectionVIOSOPolicyViewData::RenderVIOSO_RenderThread(FRHICommandListImmediate& RHICmdList, FRHITexture2D* ShaderResourceTexture, FRHITexture2D* RenderTargetTexture, const FViosoPolicyConfiguration& InConfigData)
 {
 	check(IsInRenderingThread());
 
@@ -388,7 +391,7 @@ bool FDisplayClusterProjectionVIOSOPolicy::FViewData::RenderVIOSO_RenderThread(F
 		if (bRequireInitialize)
 		{
 			TRACE_CPUPROFILER_EVENT_SCOPE(nDisplay Vioso::Initialize);
-			InitializeVIOSO(RenderTargetTexture, InConfigData);
+			InitializeVIOSO_RenderThread(RHICmdList, RenderTargetTexture, InConfigData);
 		}
 
 		if (IsValid())
@@ -413,23 +416,29 @@ bool FDisplayClusterProjectionVIOSOPolicy::FViewData::RenderVIOSO_RenderThread(F
 #ifdef VIOSO_USE_GRAPHICS_API_D3D12
 			case ERenderDevice::D3D12:
 			{
-				ID3D12DynamicRHI* D3D12RHI = GetID3D12DynamicRHI();
-				const D3D12_CPU_DESCRIPTOR_HANDLE RTVHandle = D3D12RHI->RHIGetRenderTargetView(RenderTargetTexture);
+				// Insert an outer query that encloses the whole batch
+				RHICmdList.EnqueueLambda([ViosoViewData = SharedThis(this), RenderTargetTexture = RenderTargetTexture, ShaderResourceTexture= ShaderResourceTexture](FRHICommandList& ExecutingCmdList)
+					{
+						ID3D12DynamicRHI* D3D12RHI = GetID3D12DynamicRHI();
+						const D3D12_CPU_DESCRIPTOR_HANDLE RTVHandle = D3D12RHI->RHIGetRenderTargetView(RenderTargetTexture);
 
-				VWB_D3D12_RENDERINPUT RenderInput = {};
-				RenderInput.textureResource = D3D12RHI->RHIGetResource(ShaderResourceTexture);
-				RenderInput.renderTarget = D3D12RHI->RHIGetResource(RenderTargetTexture);
-				RenderInput.rtvHandlePtr = RTVHandle.ptr;
+						VWB_D3D12_RENDERINPUT RenderInput = {};
+						RenderInput.textureResource = D3D12RHI->RHIGetResource(ShaderResourceTexture);
+						RenderInput.renderTarget = D3D12RHI->RHIGetResource(RenderTargetTexture);
+						RenderInput.rtvHandlePtr = RTVHandle.ptr;
 
-				//experimental: add resource barrier
-				D3D12RHI->RHITransitionResource(RHICmdList, RenderTargetTexture, D3D12_RESOURCE_STATE_RENDER_TARGET, D3D12_RESOURCE_BARRIER_ALL_SUBRESOURCES);
+						//experimental: add resource barrier
+						D3D12RHI->RHITransitionResource(ExecutingCmdList, RenderTargetTexture, D3D12_RESOURCE_STATE_RENDER_TARGET, D3D12_RESOURCE_BARRIER_ALL_SUBRESOURCES);
+
+						if (ViosoViewData->Warper.Render(&RenderInput, VWB_STATEMASK_DEFAULT_D3D12))
+						{
+							//experimental: add resource barrier
+							D3D12RHI->RHITransitionResource(ExecutingCmdList, RenderTargetTexture, D3D12_RESOURCE_STATE_PRESENT, D3D12_RESOURCE_BARRIER_ALL_SUBRESOURCES);
+						}
+					});
+
+				return true;
 
-				if (Warper.Render(&RenderInput, VWB_STATEMASK_DEFAULT_D3D12))
-				{
-					//experimental: add resource barrier
-					D3D12RHI->RHITransitionResource(RHICmdList, RenderTargetTexture, D3D12_RESOURCE_STATE_PRESENT, D3D12_RESOURCE_BARRIER_ALL_SUBRESOURCES);
-					return true;
-				}
 				break;
 			}
 #endif // VIOSO_USE_GRAPHICS_API_D3D12
@@ -443,13 +452,8 @@ bool FDisplayClusterProjectionVIOSOPolicy::FViewData::RenderVIOSO_RenderThread(F
 	return false;
 }
 
-bool FDisplayClusterProjectionVIOSOPolicy::FViewData::InitializeVIOSO(FRHITexture2D* RenderTargetTexture, const FViosoPolicyConfiguration& InConfigData)
+bool FDisplayClusterProjectionVIOSOPolicyViewData::InitializeVIOSO_RenderThread(FRHICommandListImmediate& RHICmdList, FRHITexture2D* RenderTargetTexture, const FViosoPolicyConfiguration& InConfigData)
 {
-	if (!bDataInitialized)
-	{
-		return false;
-	}
-
 	if (bInitialized)
 	{
 		return true;
@@ -457,8 +461,6 @@ bool FDisplayClusterProjectionVIOSOPolicy::FViewData::InitializeVIOSO(FRHITextur
 
 	bInitialized = true;
 
-	DestroyVIOSO();
-
 	switch (RenderDevice)
 	{
 
@@ -480,7 +482,6 @@ bool FDisplayClusterProjectionVIOSOPolicy::FViewData::InitializeVIOSO(FRHITextur
 	case ERenderDevice::D3D12:
 	{
 		ID3D12CommandQueue* D3D12CommandQueue = GetID3D12DynamicRHI()->RHIGetCommandQueue();
-
 		if (Warper.Initialize(D3D12CommandQueue, InConfigData))
 		{
 			return true;
@@ -497,11 +498,3 @@ bool FDisplayClusterProjectionVIOSOPolicy::FViewData::InitializeVIOSO(FRHITextur
 	UE_LOG(LogDisplayClusterProjectionVIOSO, Error, TEXT("Couldn't initialize VIOSO internals"));
 	return false;
 }
-
-void FDisplayClusterProjectionVIOSOPolicy::FViewData::DestroyVIOSO()
-{
-	if (IsValid())
-	{
-		Warper.Release();
-	}
-}
diff --git a/Engine/Plugins/Runtime/nDisplay/Source/DisplayClusterProjection/Private/Policy/VIOSO/Windows/DisplayClusterProjectionVIOSOPolicy.h b/Engine/Plugins/Runtime/nDisplay/Source/DisplayClusterProjection/Private/Policy/VIOSO/Windows/DisplayClusterProjectionVIOSOPolicy.h
--- a/Engine/Plugins/Runtime/nDisplay/Source/DisplayClusterProjection/Private/Policy/VIOSO/Windows/DisplayClusterProjectionVIOSOPolicy.h
+++ b/Engine/Plugins/Runtime/nDisplay/Source/DisplayClusterProjection/Private/Policy/VIOSO/Windows/DisplayClusterProjectionVIOSOPolicy.h
@@ -8,13 +8,13 @@
 #include "RHI.h"
 #include "RHICommandList.h"
 #include "RHIResources.h"
-
+#include "Templates/SharedPointer.h"
 
 /**
- * VIOSO projection policy
+ * VIOSO projection policy data
  */
-class FDisplayClusterProjectionVIOSOPolicy
-	: public FDisplayClusterProjectionPolicyBase
+struct FDisplayClusterProjectionVIOSOPolicyViewData
+	: public TSharedFromThis<FDisplayClusterProjectionVIOSOPolicyViewData, ESPMode::ThreadSafe>
 {
 	enum class ERenderDevice : uint8
 	{
@@ -23,6 +23,36 @@ class FDisplayClusterProjectionVIOSOPolicy
 		D3D12
 	};
 
+	FDisplayClusterProjectionVIOSOPolicyViewData();
+	virtual ~FDisplayClusterProjectionVIOSOPolicyViewData();
+
+public:
+	bool IsValid();
+
+	bool UpdateVIOSO(IDisplayClusterViewport* InViewport, const uint32 InContextNum, const FVector& LocalLocation, const FRotator& LocalRotator, const float WorldToMeters, const float NCP, const float FCP);
+	bool RenderVIOSO_RenderThread(FRHICommandListImmediate& RHICmdList, FRHITexture2D* ShaderResourceTexture, FRHITexture2D* RenderTargetTexture, const FViosoPolicyConfiguration& InConfigData);
+
+protected:
+	bool InitializeVIOSO_RenderThread(FRHICommandListImmediate& RHICmdList, FRHITexture2D* RenderTargetTexture, const FViosoPolicyConfiguration& InConfigData);
+
+public:
+	FVector  ViewLocation;
+	FRotator ViewRotation;
+	FMatrix  ProjectionMatrix;
+
+private:
+	ERenderDevice RenderDevice = ERenderDevice::Unsupported;
+	FViosoWarper Warper;
+
+	bool bInitialized = false;
+};
+
+/**
+ * VIOSO projection policy
+ */
+class FDisplayClusterProjectionVIOSOPolicy
+	: public FDisplayClusterProjectionPolicyBase
+{
 public:
 	FDisplayClusterProjectionVIOSOPolicy(const FString& ProjectionPolicyId, const FDisplayClusterConfigurationProjection* InConfigurationProjectionPolicy);
 	virtual ~FDisplayClusterProjectionVIOSOPolicy();
@@ -60,35 +90,7 @@ protected:
 protected:
 	FViosoPolicyConfiguration ViosoConfigData;
 
-	//FIntPoint ViewportSize;
-	ERenderDevice RenderDevice = ERenderDevice::Unsupported;
-
-	struct FViewData
-	{
-		bool IsValid();
-
-		bool Initialize(ERenderDevice RenderDevice, const FViosoPolicyConfiguration& InConfigData);
-		void DestroyVIOSO();
-
-		bool UpdateVIOSO(IDisplayClusterViewport* InViewport, const uint32 InContextNum, const FVector& LocalLocation, const FRotator& LocalRotator, const float WorldToMeters, const float NCP, const float FCP);
-		bool RenderVIOSO_RenderThread(FRHICommandListImmediate& RHICmdList, FRHITexture2D* ShaderResourceTexture, FRHITexture2D* RenderTargetTexture, const FViosoPolicyConfiguration& InConfigData);
-
-	protected:
-		bool InitializeVIOSO(FRHITexture2D* RenderTargetTexture, const FViosoPolicyConfiguration& InConfigData);
-
-	public:
-		FVector  ViewLocation;
-		FRotator ViewRotation;
-		FMatrix  ProjectionMatrix;
-
-	private:
-		ERenderDevice RenderDevice = ERenderDevice::Unsupported;
-		FViosoWarper Warper;
-
-		bool bInitialized = false;
-		bool bDataInitialized = false;
-	};
+	TArray<TSharedPtr<FDisplayClusterProjectionVIOSOPolicyViewData, ESPMode::ThreadSafe>> Views;
 
-	TArray<FViewData> Views;
 	FCriticalSection DllAccessCS;
 };
