nDisplay: fixed 'Freeze Viewports' bug that causes outer viewports to take wrong place on the backbuffer

diff --git a/Engine/Plugins/Runtime/nDisplay/Source/DisplayCluster/Private/Render/Viewport/DisplayClusterViewport.cpp b/Engine/Plugins/Runtime/nDisplay/Source/DisplayCluster/Private/Render/Viewport/DisplayClusterViewport.cpp
--- a/Engine/Plugins/Runtime/nDisplay/Source/DisplayCluster/Private/Render/Viewport/DisplayClusterViewport.cpp
+++ b/Engine/Plugins/Runtime/nDisplay/Source/DisplayCluster/Private/Render/Viewport/DisplayClusterViewport.cpp
@@ -421,6 +421,7 @@ bool FDisplayClusterViewport::UpdateFrameContexts(const uint32 InStereoViewIndex
 				ContextIt.StereoscopicPass = FDisplayClusterViewportStereoscopicPass::EncodeStereoscopicPass(ContextIt.ContextNum, ViewportContextAmount, InFrameSettings);
 				ContextIt.StereoViewIndex = (int32)(InStereoViewIndex + ContextIt.ContextNum);
 				ContextIt.bDisableRender = true;
+				ContextIt.FrameTargetRect = GetValidRect(DesiredFrameTargetRect, TEXT("Context Frame"));
 			}
 
 			return true;
