[nDisplay] Fix crash when we try to render a UV lightcard in the middle of being destroyed.

diff --git a/Engine/Plugins/Runtime/nDisplay/Source/DisplayCluster/Private/Render/Viewport/LightCard/DisplayClusterViewportLightCardManager.cpp b/Engine/Plugins/Runtime/nDisplay/Source/DisplayCluster/Private/Render/Viewport/LightCard/DisplayClusterViewportLightCardManager.cpp
--- a/Engine/Plugins/Runtime/nDisplay/Source/DisplayCluster/Private/Render/Viewport/LightCard/DisplayClusterViewportLightCardManager.cpp
+++ b/Engine/Plugins/Runtime/nDisplay/Source/DisplayCluster/Private/Render/Viewport/LightCard/DisplayClusterViewportLightCardManager.cpp
@@ -123,7 +123,7 @@ void FDisplayClusterViewportLightCardManager::RenderFrame()
 			bool bLoadedPrimitives = false;
 			for (ADisplayClusterLightCardActor* LightCard : UVLightCards)
 			{
-				if (LightCard->IsHidden())
+				if (LightCard->IsHidden() || LightCard->IsActorBeingDestroyed() || LightCard->GetWorld() == nullptr)
 				{
 					continue;
 				}
