nDisplay: Fix restoring a level snapshot not updating proxy positions in ICVFX editor.
Fix outliner column localization warning.
Fix light card UV templates having some of their settings overwritten.
Fix light card labels not being able to be toggled. Regression from 22346211.
Refactor light card editor spawn methods to FDisplayClusterLightCardEditorHelper so the stage app can access it.
Add UDisplayClusterStageActorTemplate to light card extender and implement with UDisplayClusterLightCardTemplate to allow easier access from various modules.

diff --git a/Engine/Plugins/Runtime/nDisplay/Source/DisplayClusterLightCardEditor/Private/DisplayClusterLightCardEditor.cpp b/Engine/Plugins/Runtime/nDisplay/Source/DisplayClusterLightCardEditor/Private/DisplayClusterLightCardEditor.cpp
--- a/Engine/Plugins/Runtime/nDisplay/Source/DisplayClusterLightCardEditor/Private/DisplayClusterLightCardEditor.cpp
+++ b/Engine/Plugins/Runtime/nDisplay/Source/DisplayClusterLightCardEditor/Private/DisplayClusterLightCardEditor.cpp
@@ -34,6 +34,7 @@
 #include "Selection.h"
 #include "UnrealEdGlobals.h"
 #include "AssetRegistry/AssetRegistryModule.h"
+#include "Editor/Transactor.h"
 #include "Editor/UnrealEdEngine.h"
 #include "Framework/Commands/GenericCommands.h"
 #include "Framework/Docking/LayoutExtender.h"
@@ -42,6 +43,7 @@
 #include "Framework/MultiBox/MultiBoxExtender.h"
 #include "Layers/LayersSubsystem.h"
 #include "Misc/FileHelper.h"
+#include "Misc/ITransaction.h"
 #include "Misc/TransactionObjectEvent.h"
 #include "Styling/SlateIconFinder.h"
 #include "Toolkits/AssetEditorToolkit.h"
@@ -192,109 +194,60 @@ void FDisplayClusterLightCardEditor::CenterActorInView(AActor* Actor)
 	}
 }
 
-AActor* FDisplayClusterLightCardEditor::SpawnActor(TSubclassOf<AActor> InActorClass, const FName& InActorName, const UDisplayClusterLightCardTemplate* InTemplate)
-{
-	if (InTemplate)
-	{
-		// For now only light card templates are supported
-		check(InActorClass && InActorClass->IsChildOf(ADisplayClusterLightCardActor::StaticClass()));
-		return SpawnLightCardFromTemplate(InTemplate);
-	}
-	
-	if (!ActiveRootActor.IsValid())
-	{
-		return nullptr;
-	}
-
-	FScopedTransaction Transaction(LOCTEXT("SpawnActorTransactionMessage", "Spawn Actor"));
-	
-	check(InActorClass);
-
-	const FVector SpawnLocation = ActiveRootActor->GetDefaultCamera()->GetComponentLocation();
-	FRotator SpawnRotation = ActiveRootActor->GetDefaultCamera()->GetComponentRotation();
-	SpawnRotation.Yaw -= 180.f;
-
-	FActorSpawnParameters SpawnParameters;
-	SpawnParameters.bNoFail = true;
-	SpawnParameters.SpawnCollisionHandlingOverride = ESpawnActorCollisionHandlingMethod::AdjustIfPossibleButAlwaysSpawn;
-	SpawnParameters.Name = InActorName;
-	SpawnParameters.NameMode = FActorSpawnParameters::ESpawnActorNameMode::Requested;
-	SpawnParameters.OverrideLevel = ActiveRootActor->GetWorld()->GetCurrentLevel();
-
-	AActor* NewActor = CastChecked<AActor>(
-		ActiveRootActor->GetWorld()->SpawnActor(InActorClass,
-			&SpawnLocation, &SpawnRotation, MoveTemp(SpawnParameters)));
-
-	NewActor->SetActorLabel(NewActor->GetName());
-
-	if (ADisplayClusterLightCardActor* NewLightCard = Cast<ADisplayClusterLightCardActor>(NewActor))
-	{
-		TryConfigureLightCardForUV(NewLightCard);
-		const TArray<ADisplayClusterLightCardActor*> LightCards { NewLightCard } ;
-		AddLightCardsToActor(LightCards);
-	}
-	else
-	{
-		RefreshPreviewActors(EDisplayClusterLightCardEditorProxyType::StageActor);
-	}
-	
-	SelectActors({ NewActor });
-
-	return NewActor;
-}
-
-ADisplayClusterLightCardActor* FDisplayClusterLightCardEditor::SpawnLightCardFromTemplate(
+AActor* FDisplayClusterLightCardEditor::SpawnActor(TSubclassOf<AActor> InActorClass, const FName& InActorName,
 	const UDisplayClusterLightCardTemplate* InTemplate, ULevel* InLevel, bool bIsPreview)
 {
 	if (!ActiveRootActor.IsValid())
 	{
 		return nullptr;
 	}
-
-	FScopedTransaction Transaction(LOCTEXT("SpawnLightCardFromTemplateTransactionMessage", "Spawn Light Card from Template"));
 	
-	check(InTemplate && InTemplate->LightCardActor);
+	FScopedTransaction Transaction(LOCTEXT("SpawnActorTransactionMessage", "Spawn Actor"));
 
-	ULevel* Level = InLevel ? InLevel : ActiveRootActor->GetWorld()->GetCurrentLevel();
+	const UDisplayClusterLightCardEditorProjectSettings* Settings = GetDefault<UDisplayClusterLightCardEditorProjectSettings>();
 	
-	FName UniqueName = *InTemplate->GetName().Replace(TEXT("Template"), TEXT(""));
-	if (StaticFindObjectFast(InTemplate->LightCardActor->GetClass(), Level, UniqueName))
-	{
-		UniqueName = MakeUniqueObjectName(Level, InTemplate->LightCardActor->GetClass(), UniqueName);
+	FDisplayClusterLightCardEditorHelper::FSpawnActorArgs SpawnArgs;
+	{
+		SpawnArgs.ActorClass = InActorClass;
+		SpawnArgs.ActorName = InActorName;
+		SpawnArgs.RootActor = ActiveRootActor.Get();
+		SpawnArgs.Template = InTemplate;
+		SpawnArgs.Level = InLevel;
+		SpawnArgs.AddLightCardArgs.bShowLabels = Settings->bDisplayLightCardLabels;
+		SpawnArgs.AddLightCardArgs.LabelScale = Settings->LightCardLabelScale;
+	
+		if (ViewportView.IsValid())
+		{
+			SpawnArgs.ProjectionMode = ViewportView->GetLightCardEditorViewportClient()->GetProjectionMode();
+		}
 	}
 	
-	// Duplicate, don't copy properties or spawn from a template. Doing so will copy component data incorrectly,
-	// specifically the static mesh override textures. They will be parented to the template, not the level instance
-	// and prevent the map from saving.
-	ADisplayClusterLightCardActor* NewLightCard = CastChecked<ADisplayClusterLightCardActor>(StaticDuplicateObject(InTemplate->LightCardActor.Get(), Level, UniqueName));
-	Level->AddLoadedActor(NewLightCard);
+	AActor* NewActor = FDisplayClusterLightCardEditorHelper::SpawnStageActor(MoveTemp(SpawnArgs));
+	check(NewActor);
 
-	TryConfigureLightCardForUV(NewLightCard);
+	RefreshPreviewActors(EDisplayClusterLightCardEditorProxyType::StageActor);
 	
 	if (!bIsPreview)
 	{
-		NewLightCard->SetActorLabel(NewLightCard->GetName());
-	
-		const TArray<ADisplayClusterLightCardActor*> LightCards { NewLightCard } ;
-		AddLightCardsToActor(LightCards);
-
-		FDisplayClusterLightCardEditorRecentItem RecentlyPlacedItem;
-		RecentlyPlacedItem.ObjectPath = InTemplate;
-		RecentlyPlacedItem.ItemType = FDisplayClusterLightCardEditorRecentItem::Type_LightCardTemplate;
+		if (InTemplate)
+		{
+			// Only template recent items need to be handled, new actor spawning is handled from caller otherwise.
+			FDisplayClusterLightCardEditorRecentItem RecentlyPlacedItem;
+			RecentlyPlacedItem.ObjectPath = InTemplate;
+			RecentlyPlacedItem.ItemType = FDisplayClusterLightCardEditorRecentItem::Type_LightCardTemplate;
 		
-		AddRecentlyPlacedItem(MoveTemp(RecentlyPlacedItem));
-		SelectActors({NewLightCard});
-	}
-
-#if WITH_EDITOR
-	// Need to call this if spawned from a template since this would normally be called in SpawnActor
-	if (GIsEditor)
-	{
-		GEditor->BroadcastLevelActorAdded(NewLightCard);
+			AddRecentlyPlacedItem(MoveTemp(RecentlyPlacedItem));
+		}
+		
+		SelectActors({ NewActor });
 	}
-#endif
 	
-	return NewLightCard;
+	return NewActor;
+}
+
+AActor* FDisplayClusterLightCardEditor::SpawnActor(const UDisplayClusterLightCardTemplate* InTemplate, ULevel* InLevel, bool bIsPreview)
+{
+	return SpawnActor(nullptr, NAME_None, InTemplate, InLevel, bIsPreview);
 }
 
 void FDisplayClusterLightCardEditor::AddNewLightCard()
@@ -1200,7 +1153,7 @@ TSharedRef<SWidget> FDisplayClusterLightCardEditor::GeneratePlaceActorsMenu()
 			{
 				if (TemplateWeakPtr.IsValid())
 				{
-					SpawnLightCardFromTemplate(TemplateWeakPtr.Get());
+					SpawnActor(TemplateWeakPtr.Get());
 				}
 			});
 			
@@ -1262,7 +1215,7 @@ TSharedRef<SWidget> FDisplayClusterLightCardEditor::GeneratePlaceActorsMenu()
 					{
 						if (const UDisplayClusterLightCardTemplate* Template = Cast<UDisplayClusterLightCardTemplate>(ObjectWeakPtr.Get()))
 						{
-							SpawnLightCardFromTemplate(Template);
+							SpawnActor(Template);
 						}
 					}
 					else if (RecentlyPlacedItem.ItemType == FDisplayClusterLightCardEditorRecentItem::Type_Dynamic)
@@ -1711,7 +1664,10 @@ bool FDisplayClusterLightCardEditor::IsOurObject(UObject* InObject,
 				if (LightCard.IsValid())
 				{
 					bIsOurActor = IsOurActor(LightCard->Actor.Get());
-					break;
+					if (bIsOurActor)
+					{
+						break;
+					}
 				}
 			}
 		}
@@ -1732,21 +1688,6 @@ bool FDisplayClusterLightCardEditor::IsOurObject(UObject* InObject,
 	return bIsOurActor;
 }
 
-void FDisplayClusterLightCardEditor::TryConfigureLightCardForUV(ADisplayClusterLightCardActor* InLightCard)
-{
-	check(InLightCard);
-	if (ViewportView.IsValid())
-	{
-		const EDisplayClusterMeshProjectionType ProjectionMode = ViewportView->GetLightCardEditorViewportClient()->GetProjectionMode();
-		if (ProjectionMode == EDisplayClusterMeshProjectionType::UV)
-		{
-			InLightCard->bIsUVLightCard = true;
-			InLightCard->Scale /= 4;
-			InLightCard->Feathering = 0.05; // Just enough to avoid jagged look on UV lightcards.
-		}
-	}
-}
-
 void FDisplayClusterLightCardEditor::BindCompileDelegates()
 {
 	if (LightCardOutliner.IsValid())
@@ -1856,6 +1797,30 @@ void FDisplayClusterLightCardEditor::OnObjectTransacted(UObject* Object,
 		// a LightCard actor from the level manually then undoes it.
 		RefreshPreviewActors();
 	}
+	else if (AActor* Actor = Cast<AActor>(Object))
+	{
+		if (TransactionObjectEvent.GetEventType() == ETransactionObjectEventType::Finalized &&
+			TransactionObjectEvent.HasPropertyChanges() && ViewportView.IsValid() && GEditor && GEditor->Trans &&
+			UE::DisplayClusterLightCardEditorUtils::IsManagedActor(Actor))
+		{
+			// Look for level snapshots restore transactions. Snapshots doesn't fire PostEditChange events so we need
+			// to listen for the transaction instead.
+			
+			const int32 TransactionIndex = GEditor->Trans->FindTransactionIndex(TransactionObjectEvent.GetTransactionId());
+			if (const FTransaction* Transaction = GEditor->Trans->GetTransaction(TransactionIndex))
+			{
+				const TOptional<FString> Namespace = FTextInspector::GetNamespace(Transaction->GetTitle());
+				if (Namespace.IsSet() && *Namespace == TEXT("LevelSnapshotsEditor"))
+				{
+					const TOptional<FString> Key = FTextInspector::GetKey(Transaction->GetTitle());
+					if (Key.IsSet() && *Key == TEXT("ApplyToWorldKey"))
+					{
+						ViewportView->GetLightCardEditorViewportClient()->UpdateProxyTransformFromLevelInstance(Actor);
+					}
+				}
+			}
+		}
+	}
 }
 
 void FDisplayClusterLightCardEditor::AddRecentlyPlacedItem(const FDisplayClusterLightCardEditorRecentItem& InItem)
diff --git a/Engine/Plugins/Runtime/nDisplay/Source/DisplayClusterLightCardEditor/Private/DisplayClusterLightCardEditor.h b/Engine/Plugins/Runtime/nDisplay/Source/DisplayClusterLightCardEditor/Private/DisplayClusterLightCardEditor.h
--- a/Engine/Plugins/Runtime/nDisplay/Source/DisplayClusterLightCardEditor/Private/DisplayClusterLightCardEditor.h
+++ b/Engine/Plugins/Runtime/nDisplay/Source/DisplayClusterLightCardEditor/Private/DisplayClusterLightCardEditor.h
@@ -88,7 +88,11 @@ public:
 	void CenterActorInView(AActor* Actor);
 
 	/** Spawns a new actor and adds it to the root actor if it is a light card */
-	AActor* SpawnActor(TSubclassOf<AActor> InActorClass, const FName& InActorName = NAME_None, const UDisplayClusterLightCardTemplate* InTemplate = nullptr);
+	AActor* SpawnActor(TSubclassOf<AActor> InActorClass, const FName& InActorName = NAME_None,
+		const UDisplayClusterLightCardTemplate* InTemplate = nullptr, ULevel* InLevel = nullptr, bool bIsPreview = false);
+
+	/** Spawn an actor from a template */
+	AActor* SpawnActor(const UDisplayClusterLightCardTemplate* InTemplate, ULevel* InLevel = nullptr, bool bIsPreview = false);
 
 	template<typename T>
 	T* SpawnActorAs(const FName& InActorName = NAME_None, const UDisplayClusterLightCardTemplate* InTemplate = nullptr)
@@ -96,9 +100,6 @@ public:
 		return Cast<T>(SpawnActor(T::StaticClass(), InActorName, InTemplate));
 	}
 	
-	/** Spawns a new light card from a light card template */
-	ADisplayClusterLightCardActor* SpawnLightCardFromTemplate(const UDisplayClusterLightCardTemplate* InTemplate, ULevel* InLevel = nullptr, bool bIsPreview = false);
-
 	/** Adds a new light card to the root actor and centers it in the viewport */
 	void AddNewLightCard();
 
@@ -273,9 +274,6 @@ private:
 	 * @return True if our object, false if not
 	 */
 	bool IsOurObject(UObject* InObject, EDisplayClusterLightCardEditorProxyType& OutProxyType) const;
-
-	/** Configure a light card actor for UV if possible */
-	void TryConfigureLightCardForUV(ADisplayClusterLightCardActor* InLightCard);
 	
 	/** Bind delegates to when a BP compiles */
 	void BindCompileDelegates();
diff --git a/Engine/Plugins/Runtime/nDisplay/Source/DisplayClusterLightCardEditor/Private/DisplayClusterLightCardEditorModule.cpp b/Engine/Plugins/Runtime/nDisplay/Source/DisplayClusterLightCardEditor/Private/DisplayClusterLightCardEditorModule.cpp
--- a/Engine/Plugins/Runtime/nDisplay/Source/DisplayClusterLightCardEditor/Private/DisplayClusterLightCardEditorModule.cpp
+++ b/Engine/Plugins/Runtime/nDisplay/Source/DisplayClusterLightCardEditor/Private/DisplayClusterLightCardEditorModule.cpp
@@ -3,10 +3,14 @@
 #include "DisplayClusterLightCardEditorModule.h"
 
 #include "DisplayClusterConfigurationTypes.h"
+
 #include "DisplayClusterLightCardActor.h"
-#include "DisplayClusterLightCardEditorCommands.h"
 #include "DisplayClusterRootActor.h"
+#include "Blueprints/DisplayClusterBlueprintLib.h"
+
 #include "DisplayClusterLightCardEditor.h"
+#include "DisplayClusterLightCardEditorCommands.h"
+
 #include "DetailCustomizations/DisplayClusterLightCardActorDetails.h"
 #include "LightCardTemplates/DisplayClusterLightCardTemplate.h"
 #include "Settings/DisplayClusterLightCardEditorSettings.h"
@@ -57,17 +61,13 @@ void FDisplayClusterLightCardEditorModule::ShowLabels(const FLabelArgs& InArgs)
 	ProjectSettings->bDisplayLightCardLabels = InArgs.bVisible;
 	ProjectSettings->LightCardLabelScale = InArgs.Scale;
 
-	if (UDisplayClusterConfigurationData* ConfigData = InArgs.RootActor->GetConfigData())
+	TSet<ADisplayClusterLightCardActor*> RootActorLightCardActors;
+	UDisplayClusterBlueprintLib::FindLightCardsForRootActor(InArgs.RootActor, RootActorLightCardActors);
+
+	for (ADisplayClusterLightCardActor* LightCardActor : RootActorLightCardActors)
 	{
-		FDisplayClusterConfigurationICVFX_VisibilityList& RootActorLightCards = ConfigData->StageSettings.Lightcard.ShowOnlyList;
-		for (TSoftObjectPtr<AActor> Actor : RootActorLightCards.Actors)
-		{
-			if (ADisplayClusterLightCardActor* LightCardActor = Cast<ADisplayClusterLightCardActor>(Actor.Get()))
-			{
-				LightCardActor->Modify(false);
-				LightCardActor->ShowLightCardLabel(InArgs.bVisible, InArgs.Scale, InArgs.RootActor);
-			}
-		}
+		LightCardActor->Modify(false);
+		LightCardActor->ShowLightCardLabel(InArgs.bVisible, InArgs.Scale, InArgs.RootActor);
 	}
 }
 
diff --git a/Engine/Plugins/Runtime/nDisplay/Source/DisplayClusterLightCardEditor/Private/LightCardTemplates/SDisplayClusterLightCardTemplateList.cpp b/Engine/Plugins/Runtime/nDisplay/Source/DisplayClusterLightCardEditor/Private/LightCardTemplates/SDisplayClusterLightCardTemplateList.cpp
--- a/Engine/Plugins/Runtime/nDisplay/Source/DisplayClusterLightCardEditor/Private/LightCardTemplates/SDisplayClusterLightCardTemplateList.cpp
+++ b/Engine/Plugins/Runtime/nDisplay/Source/DisplayClusterLightCardEditor/Private/LightCardTemplates/SDisplayClusterLightCardTemplateList.cpp
@@ -335,7 +335,7 @@ void SDisplayClusterLightCardTemplateList::OnSelectionChanged(TSharedPtr<FLightC
 {
 	if (bSpawnOnSelection && (SelectInfo == ESelectInfo::OnKeyPress || SelectInfo == ESelectInfo::OnMouseClick) && InItem.IsValid() && LightCardEditorPtr.IsValid())
 	{
-		LightCardEditorPtr.Pin()->SpawnLightCardFromTemplate(InItem->LightCardTemplate.Get());
+		LightCardEditorPtr.Pin()->SpawnActor(InItem->LightCardTemplate.Get());
 		// Close the menu containing this widget
 		FSlateApplication::Get().DismissAllMenus();
 	}
diff --git a/Engine/Plugins/Runtime/nDisplay/Source/DisplayClusterLightCardEditor/Private/Outliner/DisplayClusterLightCardOutlinerColumns.cpp b/Engine/Plugins/Runtime/nDisplay/Source/DisplayClusterLightCardEditor/Private/Outliner/DisplayClusterLightCardOutlinerColumns.cpp
--- a/Engine/Plugins/Runtime/nDisplay/Source/DisplayClusterLightCardEditor/Private/Outliner/DisplayClusterLightCardOutlinerColumns.cpp
+++ b/Engine/Plugins/Runtime/nDisplay/Source/DisplayClusterLightCardEditor/Private/Outliner/DisplayClusterLightCardOutlinerColumns.cpp
@@ -10,6 +10,8 @@
 #include "SortHelper.h"
 #include "Widgets/Input/SCheckBox.h"
 
+#define LOCTEXT_NAMESPACE "DisplayClusterLightCardOutlinerColumns"
+
 /**
  * Sets the 'Hidden in Game' field. Similar to SVisibilityWidget but is implemented as a checkbox
  * and can produce an undetermined state.
@@ -187,6 +189,11 @@ FName FDisplayClusterLightCardOutlinerHiddenInGameColumn::GetID()
 	return HiddenInGame;
 }
 
+FText FDisplayClusterLightCardOutlinerHiddenInGameColumn::GetDisplayText()
+{
+	return LOCTEXT("HiddenInGameColumnText", "Hidden in Game");
+}
+
 SHeaderRow::FColumn::FArguments FDisplayClusterLightCardOutlinerHiddenInGameColumn::ConstructHeaderRowColumn()
 {
 	return SHeaderRow::Column(GetColumnID())
@@ -236,4 +243,6 @@ void FDisplayClusterLightCardOutlinerHiddenInGameColumn::SortItems(TArray<FScene
 			return false;
 		}, SortMode)
 		.Sort(RootItems);
-}
\ No newline at end of file
+}
+
+#undef LOCTEXT_NAMESPACE
\ No newline at end of file
diff --git a/Engine/Plugins/Runtime/nDisplay/Source/DisplayClusterLightCardEditor/Private/Outliner/DisplayClusterLightCardOutlinerColumns.h b/Engine/Plugins/Runtime/nDisplay/Source/DisplayClusterLightCardEditor/Private/Outliner/DisplayClusterLightCardOutlinerColumns.h
--- a/Engine/Plugins/Runtime/nDisplay/Source/DisplayClusterLightCardEditor/Private/Outliner/DisplayClusterLightCardOutlinerColumns.h
+++ b/Engine/Plugins/Runtime/nDisplay/Source/DisplayClusterLightCardEditor/Private/Outliner/DisplayClusterLightCardOutlinerColumns.h
@@ -17,6 +17,7 @@ public:
 	FDisplayClusterLightCardOutlinerHiddenInGameColumn(ISceneOutliner& SceneOutliner) : FSceneOutlinerGutter(SceneOutliner) {}
 	virtual ~FDisplayClusterLightCardOutlinerHiddenInGameColumn() {}
 	static FName GetID();
+	static FText GetDisplayText();
 	
 	// Begin ISceneOutlinerColumn Implementation
 	virtual FName GetColumnID() override { return GetID(); }
diff --git a/Engine/Plugins/Runtime/nDisplay/Source/DisplayClusterLightCardEditor/Private/Outliner/SDisplayClusterLightCardOutliner.cpp b/Engine/Plugins/Runtime/nDisplay/Source/DisplayClusterLightCardEditor/Private/Outliner/SDisplayClusterLightCardOutliner.cpp
--- a/Engine/Plugins/Runtime/nDisplay/Source/DisplayClusterLightCardEditor/Private/Outliner/SDisplayClusterLightCardOutliner.cpp
+++ b/Engine/Plugins/Runtime/nDisplay/Source/DisplayClusterLightCardEditor/Private/Outliner/SDisplayClusterLightCardOutliner.cpp
@@ -242,7 +242,8 @@ void SDisplayClusterLightCardOutliner::CreateWorldOutliner()
 	SceneOutlinerOptions.ColumnMap.Add(FDisplayClusterLightCardOutlinerHiddenInGameColumn::GetID(),
 		FSceneOutlinerColumnInfo(FSceneOutlinerColumnInfo(ESceneOutlinerColumnVisibility::Visible,
 			1, FCreateSceneOutlinerColumn::CreateLambda([](ISceneOutliner& InSceneOutliner)
-			{ return MakeShareable(new FDisplayClusterLightCardOutlinerHiddenInGameColumn(InSceneOutliner)); }))));
+			{ return MakeShareable(new FDisplayClusterLightCardOutlinerHiddenInGameColumn(InSceneOutliner)); }), true,
+			TOptional<float>(), FDisplayClusterLightCardOutlinerHiddenInGameColumn::GetDisplayText())));
 	
 	SceneOutlinerOptions.ColumnMap.Add(FSceneOutlinerBuiltInColumnTypes::Label(),
 		FSceneOutlinerColumnInfo(ESceneOutlinerColumnVisibility::Visible, 10, FCreateSceneOutlinerColumn(), false, TOptional<float>(),
diff --git a/Engine/Plugins/Runtime/nDisplay/Source/DisplayClusterLightCardEditor/Private/Viewport/DisplayClusterLightCardEditorViewportClient.cpp b/Engine/Plugins/Runtime/nDisplay/Source/DisplayClusterLightCardEditor/Private/Viewport/DisplayClusterLightCardEditorViewportClient.cpp
--- a/Engine/Plugins/Runtime/nDisplay/Source/DisplayClusterLightCardEditor/Private/Viewport/DisplayClusterLightCardEditorViewportClient.cpp
+++ b/Engine/Plugins/Runtime/nDisplay/Source/DisplayClusterLightCardEditor/Private/Viewport/DisplayClusterLightCardEditorViewportClient.cpp
@@ -1039,6 +1039,39 @@ double FDisplayClusterLightCardEditorViewportClient::CalculateFinalActorDistance
 	return FMath::Max(Distance, 0);
 }
 
+const FDisplayClusterLightCardEditorViewportClient::FActorProxy* FDisplayClusterLightCardEditorViewportClient::FindActorProxyFromLevelInstance(
+	AActor* InLevelInstance) const
+{
+	if (InLevelInstance == nullptr)
+	{
+		return nullptr;
+	}
+	
+	return ActorProxies.FindByPredicate([InLevelInstance](const FActorProxy& Proxy)
+	{
+		return Proxy.LevelInstance.AsActor() == InLevelInstance;
+	});
+}
+
+void FDisplayClusterLightCardEditorViewportClient::UpdateProxyTransforms(const FActorProxy& InActorProxy)
+{
+	DECLARE_SCOPE_CYCLE_COUNTER(TEXT("FDisplayClusterLightCardEditorViewportClient::UpdateProxyTransforms"), STAT_UpdateProxyTransforms, STATGROUP_NDisplayLightCardEditor);
+	
+	if (InActorProxy.LevelInstance.IsValid() && InActorProxy.Proxy.IsValid())
+	{
+		const FTransform RALevelTransformNoScale(RootActorLevelInstance->GetActorRotation(), RootActorLevelInstance->GetActorLocation(), FVector::OneVector);
+		const FTransform LCLevelRelativeToRALevel = InActorProxy.LevelInstance.AsActorChecked()->GetTransform().GetRelativeTransform(RALevelTransformNoScale);
+		InActorProxy.Proxy.AsActorChecked()->SetActorTransform(LCLevelRelativeToRALevel);
+				
+		// When dealing with light card actors the transform scale can be set separately from the stage actor 2d scale. When updating the
+		// transform scale on the level instance it won't impact the positional params so the proxy transform scale won't update unless we manually set it
+		InActorProxy.Proxy.AsActorChecked()->SetActorScale3D(InActorProxy.LevelInstance.AsActorChecked()->GetActorScale3D());
+				
+		// Need to update these manually or the proxy's position will be out of sync next update
+		InActorProxy.Proxy->SetPositionalParams(InActorProxy.LevelInstance->GetPositionalParams());
+	}
+}
+
 void FDisplayClusterLightCardEditorViewportClient::ProcessClick(FSceneView& View, HHitProxy* HitProxy, FKey Key, EInputEvent Event, uint32 HitX, uint32 HitY)
 {
 	// Don't select light cards while drawing a new light card
@@ -1287,11 +1320,11 @@ bool FDisplayClusterLightCardEditorViewportClient::DropObjectsAtCoordinates(int3
 	
 	for (UObject* DroppedObject : DroppedObjects)
 	{
-		if (UDisplayClusterLightCardTemplate* Template = Cast<UDisplayClusterLightCardTemplate>(DroppedObject))
+		if (const UDisplayClusterLightCardTemplate* Template = Cast<UDisplayClusterLightCardTemplate>(DroppedObject))
 		{
-			ADisplayClusterLightCardActor* LightCardActor = LightCardEditorPtr.Pin()->SpawnLightCardFromTemplate(Template,
-				bCreateDropPreview ? PreviewWorld->GetCurrentLevel() : nullptr, bCreateDropPreview);
-			check(LightCardActor);
+			ADisplayClusterLightCardActor* LightCardActor =
+				CastChecked<ADisplayClusterLightCardActor>(LightCardEditorPtr.Pin()->SpawnActor(Template,
+					bCreateDropPreview ? PreviewWorld->GetCurrentLevel() : nullptr, bCreateDropPreview));
 
 			ProjectionHelper->VerifyAndFixActorOrigin(LightCardActor);
 			
@@ -1530,23 +1563,19 @@ void FDisplayClusterLightCardEditorViewportClient::UpdateProxyTransforms()
 		
 		for (const FActorProxy& ActorProxy : ActorProxies)
 		{
-			if (ActorProxy.LevelInstance.IsValid() && ActorProxy.Proxy.IsValid())
-			{
-				const FTransform RALevelTransformNoScale(RootActorLevelInstance->GetActorRotation(), RootActorLevelInstance->GetActorLocation(), FVector::OneVector);
-				const FTransform LCLevelRelativeToRALevel = ActorProxy.LevelInstance.AsActorChecked()->GetTransform().GetRelativeTransform(RALevelTransformNoScale);
-				ActorProxy.Proxy.AsActorChecked()->SetActorTransform(LCLevelRelativeToRALevel);
-				
-				// When dealing with light card actors the transform scale can be set separately from the stage actor 2d scale. When updating the
-				// transform scale on the level instance it won't impact the positional params so the proxy transform scale won't update unless we manually set it
-				ActorProxy.Proxy.AsActorChecked()->SetActorScale3D(ActorProxy.LevelInstance.AsActorChecked()->GetActorScale3D());
-				
-				// Need to update these manually or the proxy's position will be out of sync next update
-				ActorProxy.Proxy->SetPositionalParams(ActorProxy.LevelInstance->GetPositionalParams());
-			}
+			UpdateProxyTransforms(ActorProxy);
 		}
 	}
 }
 
+void FDisplayClusterLightCardEditorViewportClient::UpdateProxyTransformFromLevelInstance(AActor* InLevelInstance)
+{
+	if (const FActorProxy* ActorProxy = FindActorProxyFromLevelInstance(InLevelInstance))
+	{
+		UpdateProxyTransforms(*ActorProxy);
+	}
+}
+
 void FDisplayClusterLightCardEditorViewportClient::DestroyProxies(
 	EDisplayClusterLightCardEditorProxyType ProxyType)
 {
diff --git a/Engine/Plugins/Runtime/nDisplay/Source/DisplayClusterLightCardEditor/Private/Viewport/DisplayClusterLightCardEditorViewportClient.h b/Engine/Plugins/Runtime/nDisplay/Source/DisplayClusterLightCardEditor/Private/Viewport/DisplayClusterLightCardEditorViewportClient.h
--- a/Engine/Plugins/Runtime/nDisplay/Source/DisplayClusterLightCardEditor/Private/Viewport/DisplayClusterLightCardEditorViewportClient.h
+++ b/Engine/Plugins/Runtime/nDisplay/Source/DisplayClusterLightCardEditor/Private/Viewport/DisplayClusterLightCardEditorViewportClient.h
@@ -88,6 +88,9 @@ public:
 	/** Only update required transform values of proxies */
 	void UpdateProxyTransforms();
 
+	/** Update transforms value for a proxy from a level instance */
+	void UpdateProxyTransformFromLevelInstance(AActor* InLevelInstance);
+	
 	/** Remove proxies of the specified type */
 	void DestroyProxies(EDisplayClusterLightCardEditorProxyType ProxyType);
 	
@@ -241,19 +244,6 @@ private:
 	/** Calculates the final distance from the origin of an actor, given its flush distance and a desired offset */
 	double CalculateFinalActorDistance(double FlushDistance, double DesiredOffsetFromFlush = 0.) const;
 
-private:
-	TWeakPtr<FSceneViewport> SceneViewportPtr;
-	TWeakPtr<SDisplayClusterLightCardEditorViewport> LightCardEditorViewportPtr;
-	TWeakPtr<FDisplayClusterLightCardEditor> LightCardEditorPtr;
-	TWeakObjectPtr<ADisplayClusterRootActor> RootActorProxy;
-	TWeakObjectPtr<ADisplayClusterRootActor> RootActorLevelInstance;
-
-	/** Proxy types that are currently refreshing over 1-2 frames */
-	TSet<EDisplayClusterLightCardEditorProxyType> ProxyTypesRefreshing;
-	
-	/** The radius of the bounding sphere that entirely encapsulates the root actor */
-	float RootActorBoundingRadius = 0.0f;
-
 	struct FActorProxy
 	{
 		FDisplayClusterWeakStageActorPtr LevelInstance;
@@ -285,6 +275,25 @@ private:
 		/** Construct a sprite proxy from a billboard component */
 		static FSpriteProxy FromBillboard(const UBillboardComponent* InBillboardComponent);
 	};
+	
+	/** Find the actor proxy if it exists from a level instance */
+	const FActorProxy* FindActorProxyFromLevelInstance(AActor* InLevelInstance) const;
+
+	/** Update the transforms for the specific proxy */
+	void UpdateProxyTransforms(const FActorProxy& InActorProxy);
+	
+private:
+	TWeakPtr<FSceneViewport> SceneViewportPtr;
+	TWeakPtr<SDisplayClusterLightCardEditorViewport> LightCardEditorViewportPtr;
+	TWeakPtr<FDisplayClusterLightCardEditor> LightCardEditorPtr;
+	TWeakObjectPtr<ADisplayClusterRootActor> RootActorProxy;
+	TWeakObjectPtr<ADisplayClusterRootActor> RootActorLevelInstance;
+
+	/** Proxy types that are currently refreshing over 1-2 frames */
+	TSet<EDisplayClusterLightCardEditorProxyType> ProxyTypesRefreshing;
+	
+	/** The radius of the bounding sphere that entirely encapsulates the root actor */
+	float RootActorBoundingRadius = 0.0f;
 
 	/** Stores view configuration values (view transform and field of view) for the supported projection modes */
 	struct FProjectionViewConfiguration
diff --git a/Engine/Plugins/Runtime/nDisplay/Source/DisplayClusterLightCardEditor/Public/LightCardTemplates/DisplayClusterLightCardTemplate.h b/Engine/Plugins/Runtime/nDisplay/Source/DisplayClusterLightCardEditor/Public/LightCardTemplates/DisplayClusterLightCardTemplate.h
--- a/Engine/Plugins/Runtime/nDisplay/Source/DisplayClusterLightCardEditor/Public/LightCardTemplates/DisplayClusterLightCardTemplate.h
+++ b/Engine/Plugins/Runtime/nDisplay/Source/DisplayClusterLightCardEditor/Public/LightCardTemplates/DisplayClusterLightCardTemplate.h
@@ -3,23 +3,25 @@
 #pragma once
 
 #include "DisplayClusterLightCardActor.h"
+#include "StageActor/DisplayClusterStageActorTemplate.h"
 
 #include "DisplayClusterLightCardTemplate.generated.h"
 
 /**
  * A template asset to store appearance settings from Light Card actors.
  */
-UCLASS(NotBlueprintType, NotBlueprintable, NotPlaceable, PerObjectConfig, MinimalAPI, config=EditorPerProjectUserSettings)
-class UDisplayClusterLightCardTemplate : public UObject
+UCLASS(NotBlueprintType, NotBlueprintable, NotPlaceable, MinimalAPI)
+class UDisplayClusterLightCardTemplate : public UDisplayClusterStageActorTemplate
 {
 	GENERATED_BODY()
 
+public:
+	// UDisplayClusterStageActorTemplate
+	virtual AActor* GetTemplateActor() const override { return LightCardActor; }
+	// ~UDisplayClusterStageActorTemplate
+	
 public:
 	/** The instanced template object containing user settings for the light card. */
 	UPROPERTY(VisibleAnywhere, Instanced, Category = Template, NoClear, meta = (ShowOnlyInnerProperties))
 	TObjectPtr<ADisplayClusterLightCardActor> LightCardActor;
-
-	/** If the user has marked this a favorite template. */
-	UPROPERTY(Config)
-	bool bIsFavorite;
 };
diff --git a/Engine/Plugins/Runtime/nDisplay/Source/DisplayClusterScenePreview/Private/DisplayClusterLightCardEditorHelper.cpp b/Engine/Plugins/Runtime/nDisplay/Source/DisplayClusterScenePreview/Private/DisplayClusterLightCardEditorHelper.cpp
--- a/Engine/Plugins/Runtime/nDisplay/Source/DisplayClusterScenePreview/Private/DisplayClusterLightCardEditorHelper.cpp
+++ b/Engine/Plugins/Runtime/nDisplay/Source/DisplayClusterScenePreview/Private/DisplayClusterLightCardEditorHelper.cpp
@@ -2,17 +2,24 @@
 
 #include "DisplayClusterLightCardEditorHelper.h"
 
-#include "CanvasTypes.h"
-#include "Containers/ArrayView.h"
-#include "DisplayClusterRootActor.h"
-#include "DisplayClusterLightCardActor.h"
-#include "Engine/StaticMesh.h"
-#include "Engine/Texture2D.h"
 #include "IDisplayClusterScenePreview.h"
+
+#include "DisplayClusterConfigurationTypes.h"
+
+#include "DisplayClusterLightCardActor.h"
+#include "DisplayClusterRootActor.h"
+#include "Components/DisplayClusterCameraComponent.h"
+
+#include "StageActor/DisplayClusterStageActorTemplate.h"
+
+#include "CanvasTypes.h"
 #include "ImageUtils.h"
 #include "KismetProceduralMeshLibrary.h"
 #include "PreviewScene.h"
 #include "ProceduralMeshComponent.h"
+#include "Containers/ArrayView.h"
+#include "Engine/StaticMesh.h"
+#include "Engine/Texture2D.h"
 #include "UObject/UObjectGlobals.h"
 
 #if WITH_EDITOR
@@ -944,6 +951,150 @@ FDisplayClusterMeshProjectionPrimitiveFilter::FPrimitiveFilter FDisplayClusterLi
 		});
 }
 
+AActor* FDisplayClusterLightCardEditorHelper::SpawnStageActor(const FSpawnActorArgs& InSpawnArgs)
+{
+	ADisplayClusterRootActor* RootActor = InSpawnArgs.RootActor;
+	const TSubclassOf<AActor> ActorClass = InSpawnArgs.ActorClass;
+	const UDisplayClusterStageActorTemplate* Template = InSpawnArgs.Template;
+	
+	check(RootActor)
+	check(ActorClass || Template);
+	
+	const FName ActorName = InSpawnArgs.ActorName;
+	const EDisplayClusterMeshProjectionType ProjectionMode = InSpawnArgs.ProjectionMode;
+	ULevel* Level = InSpawnArgs.Level ? InSpawnArgs.Level : RootActor->GetWorld()->GetCurrentLevel();
+	const bool bIsPreview = InSpawnArgs.bIsPreview;
+
+	AActor* NewActor = nullptr;
+	
+	if (Template)
+	{
+		const AActor* TemplateActor = Template->GetTemplateActor();
+		check(TemplateActor);
+
+		// For now only light card templates are supported
+		check(TemplateActor->GetClass()->IsChildOf(ADisplayClusterLightCardActor::StaticClass()));
+		
+		FName UniqueName = *Template->GetName().Replace(TEXT("Template"), TEXT(""));
+		if (StaticFindObjectFast(TemplateActor->GetClass(), Level, UniqueName))
+		{
+			UniqueName = MakeUniqueObjectName(Level, TemplateActor->GetClass(), UniqueName);
+		}
+	
+		// Duplicate, don't copy properties or spawn from a template. Doing so will copy component data incorrectly,
+		// specifically the static mesh override textures. They will be parented to the template, not the level instance
+		// and prevent the map from saving.
+		ADisplayClusterLightCardActor* NewLightCard = CastChecked<ADisplayClusterLightCardActor>(StaticDuplicateObject(TemplateActor, Level, UniqueName));
+
+#if WITH_EDITOR
+		Level->AddLoadedActor(NewLightCard);
+#endif
+		
+		NewActor = NewLightCard;
+	}
+	else
+	{
+		const FVector SpawnLocation = RootActor->GetDefaultCamera()->GetComponentLocation();
+		FRotator SpawnRotation = RootActor->GetDefaultCamera()->GetComponentRotation();
+		SpawnRotation.Yaw -= 180.f;
+
+		FActorSpawnParameters SpawnParameters;
+		SpawnParameters.bNoFail = true;
+		SpawnParameters.SpawnCollisionHandlingOverride = ESpawnActorCollisionHandlingMethod::AdjustIfPossibleButAlwaysSpawn;
+		SpawnParameters.Name = ActorName;
+		SpawnParameters.NameMode = FActorSpawnParameters::ESpawnActorNameMode::Requested;
+		SpawnParameters.OverrideLevel = Level;
+
+		NewActor = CastChecked<AActor>(
+			RootActor->GetWorld()->SpawnActor(ActorClass,
+				&SpawnLocation, &SpawnRotation, MoveTemp(SpawnParameters)));
+	}
+
+	check(NewActor);
+
+#if WITH_EDITOR
+	if (!bIsPreview)
+	{
+		NewActor->SetActorLabel(NewActor->GetName());
+	}
+#endif
+	
+	if (ADisplayClusterLightCardActor* NewLightCard = Cast<ADisplayClusterLightCardActor>(NewActor))
+	{
+		if (ProjectionMode == EDisplayClusterMeshProjectionType::UV)
+		{
+			// If this already is a UV light card leave everything alone
+			if (!NewLightCard->bIsUVLightCard)
+			{
+				NewLightCard->bIsUVLightCard = true;
+				NewLightCard->Scale /= 4;
+
+				// Don't override feathering if this was spawned from a template
+				if (!Template)
+				{
+					NewLightCard->Feathering = 0.05; // Just enough to avoid jagged look on UV lightcards.
+				}
+			}
+		}
+
+		if (!bIsPreview)
+		{
+			AddLightCardsToRootActor({ NewLightCard }, RootActor, InSpawnArgs.AddLightCardArgs);
+		}
+	}
+
+#if WITH_EDITOR
+		// Need to call this if spawned from a template since this would normally be called in SpawnActor
+		if (Template && GIsEditor)
+		{
+			GEditor->BroadcastLevelActorAdded(NewActor);
+		}
+#endif
+
+	return NewActor;
+}
+
+void FDisplayClusterLightCardEditorHelper::AddLightCardsToRootActor(
+	const TArray<ADisplayClusterLightCardActor*>& LightCards, ADisplayClusterRootActor* RootActor,
+	const FAddLightCardArgs& AddLightCardArgs)
+{
+	if (LightCards.Num() == 0)
+	{
+		return;
+	}
+
+	check(RootActor);
+	
+	UDisplayClusterConfigurationData* ConfigData = RootActor->GetConfigData();
+	ConfigData->Modify();
+	FDisplayClusterConfigurationICVFX_VisibilityList& RootActorLightCards = ConfigData->StageSettings.Lightcard.ShowOnlyList;
+	
+	for (ADisplayClusterLightCardActor* LightCard : LightCards)
+	{
+		check(LightCard);
+
+		if (!RootActorLightCards.Actors.ContainsByPredicate([&](const TSoftObjectPtr<AActor>& Actor)
+			{
+				// Don't add if a loaded actor is already present.
+				return Actor.Get() == LightCard;
+			}))
+		{
+			LightCard->ShowLightCardLabel(AddLightCardArgs.bShowLabels, AddLightCardArgs.LabelScale, RootActor);
+				
+			const TSoftObjectPtr<AActor> LightCardSoftObject(LightCard);
+
+			// Remove any exact paths to this actor. It's possible invalid actors are present if a light card
+			// was force deleted from a level.
+			RootActorLightCards.Actors.RemoveAll([&](const TSoftObjectPtr<AActor>& Actor)
+				{
+					return Actor == LightCardSoftObject;
+				});
+
+			LightCard->AddToLightCardLayer(RootActor);
+		}
+	}
+}
+
 ADisplayClusterRootActor* FDisplayClusterLightCardEditorHelper::UpdateRootActor()
 {
 	ADisplayClusterRootActor* NewRootActor = IDisplayClusterScenePreview::Get().GetRendererRootActor(RendererId);
diff --git a/Engine/Plugins/Runtime/nDisplay/Source/DisplayClusterScenePreview/Public/DisplayClusterLightCardEditorHelper.h b/Engine/Plugins/Runtime/nDisplay/Source/DisplayClusterScenePreview/Public/DisplayClusterLightCardEditorHelper.h
--- a/Engine/Plugins/Runtime/nDisplay/Source/DisplayClusterScenePreview/Public/DisplayClusterLightCardEditorHelper.h
+++ b/Engine/Plugins/Runtime/nDisplay/Source/DisplayClusterScenePreview/Public/DisplayClusterLightCardEditorHelper.h
@@ -294,6 +294,60 @@ public:
 	 */
 	DISPLAYCLUSTERSCENEPREVIEW_API FDisplayClusterMeshProjectionPrimitiveFilter::FPrimitiveFilter CreateDefaultShouldApplyProjectionToPrimitiveFilter() const;
 
+	struct FAddLightCardArgs
+	{
+		bool bShowLabels;
+		float LabelScale;
+
+		FAddLightCardArgs():
+		bShowLabels(false),
+		LabelScale(1.f)
+		{}
+	};
+	
+	struct FSpawnActorArgs
+	{
+		/** [Required] The root actor managing the stage */
+		ADisplayClusterRootActor* RootActor;
+		
+		/** [Required if Template is null] The class to spawn */
+		TSubclassOf<AActor> ActorClass;
+
+		/** [Required if ActorClass is null] The template to use */
+		const class UDisplayClusterStageActorTemplate* Template;
+		
+		/** The name to set for the actor and label */
+		FName ActorName;
+		
+		/** The projection mode the actor should be spawned in under */
+		EDisplayClusterMeshProjectionType ProjectionMode;
+
+		/** An override level for the actor, otherwise the RootActor level is used */
+		ULevel* Level;
+
+		/** Args if this is a light card being added to the root actor */
+		FAddLightCardArgs AddLightCardArgs;
+
+		/** Notify this actor should be used as a preview (such as for drag & drop) */
+		bool bIsPreview;
+
+		FSpawnActorArgs():
+		RootActor(nullptr),
+		Template(nullptr),
+		ActorName(NAME_None),
+		ProjectionMode(),
+		Level(nullptr),
+		bIsPreview(false)
+		{}
+	};
+
+	/** Spawns a new actor for stage use and adds it to the root actor if it is a light card */
+	static DISPLAYCLUSTERSCENEPREVIEW_API AActor* SpawnStageActor(const FSpawnActorArgs& InSpawnArgs);
+	
+	/** Adds the given Light Card to the root actor */
+	static DISPLAYCLUSTERSCENEPREVIEW_API void AddLightCardsToRootActor(const TArray<ADisplayClusterLightCardActor*>& LightCards, ADisplayClusterRootActor* RootActor,
+		const FAddLightCardArgs& AddLightCardArgs = FAddLightCardArgs());
+	
 private:
 	/**
 	 * Update the cached pointer to the root actor from the associated preview renderer and, if the actor has changed, update related data.
diff --git a/Engine/Plugins/Runtime/nDisplayModularFeatures/Source/DisplayClusterLightCardExtender/Public/StageActor/DisplayClusterStageActorTemplate.h b/Engine/Plugins/Runtime/nDisplayModularFeatures/Source/DisplayClusterLightCardExtender/Public/StageActor/DisplayClusterStageActorTemplate.h
new file mode 100644
--- /dev/null
+++ b/Engine/Plugins/Runtime/nDisplayModularFeatures/Source/DisplayClusterLightCardExtender/Public/StageActor/DisplayClusterStageActorTemplate.h
@@ -0,0 +1,21 @@
+// Copyright Epic Games, Inc. All Rights Reserved.
+
+#pragma once
+
+#include "CoreMinimal.h"
+
+#include "DisplayClusterStageActorTemplate.generated.h"
+
+UCLASS(Abstract, PerObjectConfig, config=EditorPerProjectUserSettings)
+class DISPLAYCLUSTERLIGHTCARDEXTENDER_API UDisplayClusterStageActorTemplate : public UObject
+{
+	GENERATED_BODY()
+
+public:
+	virtual AActor* GetTemplateActor() const { return nullptr; }
+	
+public:
+	/** If the user has marked this a favorite template. */
+	UPROPERTY(Config)
+	bool bIsFavorite;
+};
\ No newline at end of file
