nDisplay: Fix jumps when moving light cards in icvfx panel orthographic map view in gaps between meshes.

This seems to be a regression caused by a change in behavior in procedural meshes. While before it was necessary to invert the line trace ray to get a hit, now we need to not invert it.

#skipci

diff --git a/Engine/Plugins/Runtime/nDisplay/Source/DisplayClusterScenePreview/Private/DisplayClusterLightCardEditorHelper.cpp b/Engine/Plugins/Runtime/nDisplay/Source/DisplayClusterScenePreview/Private/DisplayClusterLightCardEditorHelper.cpp
--- a/Engine/Plugins/Runtime/nDisplay/Source/DisplayClusterScenePreview/Private/DisplayClusterLightCardEditorHelper.cpp
+++ b/Engine/Plugins/Runtime/nDisplay/Source/DisplayClusterScenePreview/Private/DisplayClusterLightCardEditorHelper.cpp
@@ -1632,13 +1632,12 @@ FVector FDisplayClusterLightCardEditorHelper::TraceScreenRay(const FVector& Orth
 		}
 		else
 		{
-			// If we didn't hit any stage geometry, try to trace against the normal map mesh. Procedural meshes does not appear to handle inward pointing normals correctly,
-			// so we need to reverse the trace start and end locations to get a useful hit result
+			// If we didn't hit any stage geometry, try to trace against the normal map mesh.
 
 			FCollisionQueryParams TraceParams(SCENE_QUERY_STAT(DisplayClusterStageTrace), true);
 
 			FHitResult HitResult;
-			if (ensure(!bNormalMapInvalid) && NormalMapMeshComponent->LineTraceComponent(HitResult, RayEnd, RayStart, TraceParams))
+			if (ensure(!bNormalMapInvalid) && NormalMapMeshComponent->LineTraceComponent(HitResult, RayStart, RayEnd, TraceParams))
 			{
 				HitLocation = HitResult.Location;
 				Direction = (HitLocation - ViewOrigin).GetSafeNormal();
