[Backout] - CL22294726
Original CL Desc
-----------------------------------------------------------------
nDisplay:
- added copying alpha channel from scene color to the RenderTarget (for lightcard and chromakey viewports with OCIO, PP etc.)
- removed unused old code about OCIO lightcard viewports (LightCardColor and LightCardAlpha viewports)
- Added  parent TSharedFromThis for FDisplayClusterViewportManagerProxy,FDisplayClusterViewportProxy and FDisplayClusterViewportPostProcessManager. Added changes to support it.
- improved FDisplayClusterViewportManagerViewExtension. Added scene color capture feature. . Added changes to support it.
- added CVar 'nDisplay.render.icvfx.lightcard'

diff --git a/Engine/Plugins/Runtime/nDisplay/Source/DisplayCluster/Private/Render/Viewport/Configuration/DisplayClusterViewportConfigurationHelpers_ICVFX.cpp b/Engine/Plugins/Runtime/nDisplay/Source/DisplayCluster/Private/Render/Viewport/Configuration/DisplayClusterViewportConfigurationHelpers_ICVFX.cpp
--- a/Engine/Plugins/Runtime/nDisplay/Source/DisplayCluster/Private/Render/Viewport/Configuration/DisplayClusterViewportConfigurationHelpers_ICVFX.cpp
+++ b/Engine/Plugins/Runtime/nDisplay/Source/DisplayCluster/Private/Render/Viewport/Configuration/DisplayClusterViewportConfigurationHelpers_ICVFX.cpp
@@ -348,10 +348,10 @@ FDisplayClusterViewport* FDisplayClusterViewportConfigurationHelpers_ICVFX::GetO
 	return ChromakeyViewport;
 }
 
-FDisplayClusterViewport* FDisplayClusterViewportConfigurationHelpers_ICVFX::GetOrCreateLightcardViewport(FDisplayClusterViewport& BaseViewport, ADisplayClusterRootActor& RootActor)
+FDisplayClusterViewport* FDisplayClusterViewportConfigurationHelpers_ICVFX::GetOrCreateLightcardViewport(FDisplayClusterViewport& BaseViewport, ADisplayClusterRootActor& RootActor, bool bIsOpenColorIOViewportExist)
 {
 	// Create new lightcard viewport
-	const FString ResourceId = DisplayClusterViewportStrings::icvfx::lightcard;
+	const FString ResourceId = bIsOpenColorIOViewportExist ? DisplayClusterViewportStrings::icvfx::lightcard_OCIO : DisplayClusterViewportStrings::icvfx::lightcard;
 
 	FDisplayClusterViewport* LightcardViewport = ImplFindViewport(RootActor, BaseViewport.GetId(), ResourceId);
 	if (LightcardViewport == nullptr)
@@ -374,6 +374,10 @@ FDisplayClusterViewport* FDisplayClusterViewportConfigurationHelpers_ICVFX::GetO
 
 	// Add viewport ICVFX usage as Lightcard
 	LightcardViewport->RenderSettingsICVFX.RuntimeFlags |= ViewportRuntime_ICVFXLightcard;
+	if (bIsOpenColorIOViewportExist)
+	{
+		LightcardViewport->RenderSettingsICVFX.RuntimeFlags |= ViewportRuntime_ICVFXLightcardAlpha;
+	}
 
 	return LightcardViewport;
 }
@@ -714,7 +718,7 @@ bool FDisplayClusterViewportConfigurationHelpers_ICVFX::IsShouldUseLightcard(con
 	return FDisplayClusterViewportConfigurationHelpers_Visibility::IsValid(InLightcardSettings.ShowOnlyList);
 }
 
-void FDisplayClusterViewportConfigurationHelpers_ICVFX::UpdateLightcardViewportSetting(FDisplayClusterViewport& DstViewport, FDisplayClusterViewport& BaseViewport, ADisplayClusterRootActor& RootActor)
+void FDisplayClusterViewportConfigurationHelpers_ICVFX::UpdateLightcardViewportSetting(FDisplayClusterViewport& DstViewport, FDisplayClusterViewport& BaseViewport, ADisplayClusterRootActor& RootActor, bool bIsOpenColorIOViewportExist)
 {
 	const FDisplayClusterConfigurationICVFX_StageSettings& StageSettings = RootActor.GetStageSettings();
 	const FDisplayClusterConfigurationICVFX_LightcardSettings& LightcardSettings = StageSettings.Lightcard;
@@ -727,10 +731,32 @@ void FDisplayClusterViewportConfigurationHelpers_ICVFX::UpdateLightcardViewportS
 	// LIghtcard texture used as overlay
 	DstViewport.RenderSettings.bVisible = false;
 
-	FDisplayClusterViewportConfigurationHelpers_Postprocess::UpdateLightcardPostProcessSettings(DstViewport, BaseViewport, RootActor);
-	FDisplayClusterViewportConfigurationHelpers_OpenColorIO::UpdateLightcardViewport(DstViewport, BaseViewport, RootActor);
+	bool bIsLightcardUseResolvedScene = false;
+	if (!bIsOpenColorIOViewportExist)
+	{
+		if (FDisplayClusterViewportConfigurationHelpers_Postprocess::UpdateLightcardPostProcessSettings(DstViewport, BaseViewport, RootActor))
+		{
+			bIsLightcardUseResolvedScene = true;
+		}
 
-	DstViewport.RenderSettings.CaptureMode = EDisplayClusterViewportCaptureMode::Lightcard;
+		if (FDisplayClusterViewportConfigurationHelpers_OpenColorIO::UpdateLightcardViewport(DstViewport, BaseViewport, RootActor))
+		{
+			bIsLightcardUseResolvedScene = true;
+		}
+	}
+
+	if (bIsLightcardUseResolvedScene)
+	{
+		// ligthcard used OCIO - capture 2 vp, [1] color with OCIO and [2] alpha
+		// to support OCIO capture mode must be without alpha, so we need to render second vp with alpha
+		DstViewport.RenderSettings.CaptureMode = EDisplayClusterViewportCaptureMode::Lightcard_OCIO;
+		DstViewport.RenderSettingsICVFX.RuntimeFlags |= ViewportRuntime_ICVFXLightcardColor;
+	}
+	else
+	{
+		// no OCIO, capture 1 vp - color+alpha
+		DstViewport.RenderSettings.CaptureMode = EDisplayClusterViewportCaptureMode::Lightcard;
+	}
 
 	const FDisplayClusterConfigurationICVFX_LightcardRenderSettings& InRenderSettings = LightcardSettings.RenderSettings;
 	{
@@ -774,7 +800,14 @@ void FDisplayClusterViewportConfigurationHelpers_ICVFX::UpdateLightcardViewportS
 	}
 	else
 	{
-		BaseViewport.RenderSettingsICVFX.ICVFX.Lightcard.ViewportId = DstViewport.GetId();
+		if (bIsLightcardUseResolvedScene)
+		{
+			BaseViewport.RenderSettingsICVFX.ICVFX.Lightcard_OCIO.ViewportId = DstViewport.GetId();
+		}
+		else
+		{
+			BaseViewport.RenderSettingsICVFX.ICVFX.Lightcard.ViewportId = DstViewport.GetId();
+		}
 	}
 }
 
diff --git a/Engine/Plugins/Runtime/nDisplay/Source/DisplayCluster/Private/Render/Viewport/Configuration/DisplayClusterViewportConfigurationHelpers_ICVFX.h b/Engine/Plugins/Runtime/nDisplay/Source/DisplayCluster/Private/Render/Viewport/Configuration/DisplayClusterViewportConfigurationHelpers_ICVFX.h
--- a/Engine/Plugins/Runtime/nDisplay/Source/DisplayCluster/Private/Render/Viewport/Configuration/DisplayClusterViewportConfigurationHelpers_ICVFX.h
+++ b/Engine/Plugins/Runtime/nDisplay/Source/DisplayCluster/Private/Render/Viewport/Configuration/DisplayClusterViewportConfigurationHelpers_ICVFX.h
@@ -37,7 +37,7 @@ public:
 
 	static FDisplayClusterViewport* GetOrCreateCameraViewport(ADisplayClusterRootActor& RootActor, UDisplayClusterICVFXCameraComponent& InCameraComponent);
 	static FDisplayClusterViewport* GetOrCreateChromakeyViewport(ADisplayClusterRootActor& RootActor, UDisplayClusterICVFXCameraComponent& InCameraComponent);
-	static FDisplayClusterViewport* GetOrCreateLightcardViewport(FDisplayClusterViewport& BaseViewport, ADisplayClusterRootActor& RootActor);
+	static FDisplayClusterViewport* GetOrCreateLightcardViewport(FDisplayClusterViewport& BaseViewport, ADisplayClusterRootActor& RootActor, bool bIsOpenColorIOViewportExist);
 
 #if WITH_EDITOR
 	static TArray<FDisplayClusterViewport*> PreviewGetRenderedInCameraViewports(ADisplayClusterRootActor& RootActor, UDisplayClusterICVFXCameraComponent& InCameraComponent, bool bGetChromakey = false);
@@ -47,7 +47,7 @@ public:
 
 	static void UpdateCameraViewportSettings(FDisplayClusterViewport& DstViewport, ADisplayClusterRootActor& RootActor, UDisplayClusterICVFXCameraComponent& InCameraComponent);
 	static void UpdateChromakeyViewportSettings(FDisplayClusterViewport& DstViewport, FDisplayClusterViewport& InCameraViewport, ADisplayClusterRootActor& RootActor, UDisplayClusterICVFXCameraComponent& InCameraComponent);
-	static void UpdateLightcardViewportSetting(FDisplayClusterViewport& DstViewport, FDisplayClusterViewport& BaseViewport, ADisplayClusterRootActor& RootActor);
+	static void UpdateLightcardViewportSetting(FDisplayClusterViewport& DstViewport, FDisplayClusterViewport& BaseViewport, ADisplayClusterRootActor& RootActor, bool bIsOpenColorIOViewportExist);
 
 	static bool IsCameraUsed(UDisplayClusterICVFXCameraComponent& InCameraComponent);
 	static bool GetCameraContext(ADisplayClusterRootActor& RootActor, UDisplayClusterICVFXCameraComponent& InCameraComponent, FCameraContext_ICVFX& OutCameraContext);
diff --git a/Engine/Plugins/Runtime/nDisplay/Source/DisplayCluster/Private/Render/Viewport/Configuration/DisplayClusterViewportConfigurationICVFX.cpp b/Engine/Plugins/Runtime/nDisplay/Source/DisplayCluster/Private/Render/Viewport/Configuration/DisplayClusterViewportConfigurationICVFX.cpp
--- a/Engine/Plugins/Runtime/nDisplay/Source/DisplayCluster/Private/Render/Viewport/Configuration/DisplayClusterViewportConfigurationICVFX.cpp
+++ b/Engine/Plugins/Runtime/nDisplay/Source/DisplayCluster/Private/Render/Viewport/Configuration/DisplayClusterViewportConfigurationICVFX.cpp
@@ -279,15 +279,30 @@ FDisplayClusterViewportConfigurationICVFX::~FDisplayClusterViewportConfiguration
 
 bool FDisplayClusterViewportConfigurationICVFX::CreateLightcardViewport(FDisplayClusterViewport& BaseViewport)
 {
-	FDisplayClusterViewport* LightcardViewport = FDisplayClusterViewportConfigurationHelpers_ICVFX::GetOrCreateLightcardViewport(BaseViewport, RootActor);
+	FDisplayClusterViewport* LightcardViewport = FDisplayClusterViewportConfigurationHelpers_ICVFX::GetOrCreateLightcardViewport(BaseViewport, RootActor, false);
 	if (LightcardViewport)
 	{
 		// Update lightcard viewport settings
-		FDisplayClusterViewportConfigurationHelpers_ICVFX::UpdateLightcardViewportSetting(*LightcardViewport, BaseViewport, RootActor);
+		FDisplayClusterViewportConfigurationHelpers_ICVFX::UpdateLightcardViewportSetting(*LightcardViewport, BaseViewport, RootActor, false);
 
 		// Support projection policy update
 		FDisplayClusterViewportConfigurationHelpers::UpdateProjectionPolicy(*LightcardViewport);
 
+		// When lightcard use OCIO or PP for render we need resolve scene. Now get alpha channel by second render pass
+		if ((LightcardViewport->RenderSettingsICVFX.RuntimeFlags & ViewportRuntime_ICVFXLightcardColor) != 0)
+		{
+			// Now OCIO require second vp for alpha
+			FDisplayClusterViewport* LightcardViewportAlpha = FDisplayClusterViewportConfigurationHelpers_ICVFX::GetOrCreateLightcardViewport(BaseViewport, RootActor, true);
+			if (LightcardViewportAlpha)
+			{
+				// Update lightcard viewport settings
+				FDisplayClusterViewportConfigurationHelpers_ICVFX::UpdateLightcardViewportSetting(*LightcardViewportAlpha, BaseViewport, RootActor, true);
+
+				// Support projection policy update
+				FDisplayClusterViewportConfigurationHelpers::UpdateProjectionPolicy(*LightcardViewportAlpha);
+			}
+		}
+
 		return true;
 	}
 
@@ -396,7 +411,7 @@ void FDisplayClusterViewportConfigurationICVFX::Update()
 				// Freeze render for lightcards when outer viewports freezed
 				if (StageSettings.Lightcard.bIgnoreOuterViewportsFreezingForLightcards == false)
 				{
-					const EDisplayClusterViewportRuntimeICVFXFlags LightcardViewportMask = ViewportRuntime_ICVFXLightcard;
+					const EDisplayClusterViewportRuntimeICVFXFlags LightcardViewportMask = ViewportRuntime_ICVFXLightcard | ViewportRuntime_ICVFXLightcardColor | ViewportRuntime_ICVFXLightcardAlpha;
 
 					for (FDisplayClusterViewport* ViewportIt : ViewportManager->ImplGetViewports())
 					{
diff --git a/Engine/Plugins/Runtime/nDisplay/Source/DisplayCluster/Private/Render/Viewport/Containers/DisplayClusterViewportProxyData.h b/Engine/Plugins/Runtime/nDisplay/Source/DisplayCluster/Private/Render/Viewport/Containers/DisplayClusterViewportProxyData.h
--- a/Engine/Plugins/Runtime/nDisplay/Source/DisplayCluster/Private/Render/Viewport/Containers/DisplayClusterViewportProxyData.h
+++ b/Engine/Plugins/Runtime/nDisplay/Source/DisplayCluster/Private/Render/Viewport/Containers/DisplayClusterViewportProxyData.h
@@ -27,7 +27,7 @@ public:
 	void UpdateProxy_RenderThread() const;
 
 private:
-	TSharedPtr<FDisplayClusterViewportProxy, ESPMode::ThreadSafe> DstViewportProxy;
+	FDisplayClusterViewportProxy* DstViewportProxy;
 
 	TSharedPtr<IDisplayClusterRender_MeshComponent, ESPMode::ThreadSafe> RemapMesh;
 
diff --git a/Engine/Plugins/Runtime/nDisplay/Source/DisplayCluster/Private/Render/Viewport/DisplayClusterViewport.cpp b/Engine/Plugins/Runtime/nDisplay/Source/DisplayCluster/Private/Render/Viewport/DisplayClusterViewport.cpp
--- a/Engine/Plugins/Runtime/nDisplay/Source/DisplayCluster/Private/Render/Viewport/DisplayClusterViewport.cpp
+++ b/Engine/Plugins/Runtime/nDisplay/Source/DisplayCluster/Private/Render/Viewport/DisplayClusterViewport.cpp
@@ -4,7 +4,6 @@
 #include "Render/Viewport/DisplayClusterViewportHelpers.h"
 #include "Render/Viewport/DisplayClusterViewportManager.h"
 #include "Render/Viewport/DisplayClusterViewportProxy.h"
-#include "Render/Viewport/DisplayClusterViewportManagerProxy.h"
 
 #include "Render/Viewport/Configuration/DisplayClusterViewportConfiguration.h"
 
@@ -38,37 +37,14 @@ FDisplayClusterViewport::FDisplayClusterViewport(FDisplayClusterViewportManager&
 	check(UninitializedProjectionPolicy.IsValid());
 
 	// Create scene proxy pair with on game thread. Outside, in ViewportManager added to proxy array on render thread
-	ViewportProxy = MakeShared<FDisplayClusterViewportProxy, ESPMode::ThreadSafe>(*this);
-
-	// Add viewport proxy on renderthread
-	ENQUEUE_RENDER_COMMAND(CreateDisplayClusterViewportProxy)(
-		[ViewportManagerProxy = Owner.GetViewportManagerProxy(), ViewportProxy = ViewportProxy](FRHICommandListImmediate& RHICmdList)
-		{
-			ViewportManagerProxy->CreateViewport_RenderThread(ViewportProxy);
-		}
-	);
+	ViewportProxy = new FDisplayClusterViewportProxy(Owner.ImplGetProxy(), *this);
 }
 
 FDisplayClusterViewport::~FDisplayClusterViewport()
 {
-	// Handle projection policy event
-	ProjectionPolicy.Reset();
-	UninitializedProjectionPolicy.Reset();
-
-	// Remove viewport proxy on render_thread
-	ENQUEUE_RENDER_COMMAND(DeleteDisplayClusterViewportProxy)(
-		[ViewportManagerProxy = Owner.GetViewportManagerProxy(), ViewportProxy = ViewportProxy](FRHICommandListImmediate& RHICmdList)
-		{
-			ViewportManagerProxy->DeleteViewport_RenderThread(ViewportProxy);
-		}
-	);
-
-	ViewportProxy.Reset();
-
 	HandleEndScene();
 
-	// Reset RTT size after viewport delete
-	Owner.ResetSceneRenderTargetSize();
+	// ViewportProxy deleted on render thread from FDisplayClusterViewportManagerProxy::ImplDeleteViewport()
 }
 
 IDisplayClusterViewportManager& FDisplayClusterViewport::GetOwner() const
@@ -244,6 +220,7 @@ void FDisplayClusterViewport::SetupSceneView(uint32 ContextNum, class UWorld* Wo
 	{
 	case EDisplayClusterViewportCaptureMode::Chromakey:
 	case EDisplayClusterViewportCaptureMode::Lightcard:
+	case EDisplayClusterViewportCaptureMode::Lightcard_OCIO:
 
 		InOutView.bAllowRayTracing = false;
 		break;
diff --git a/Engine/Plugins/Runtime/nDisplay/Source/DisplayCluster/Private/Render/Viewport/DisplayClusterViewport.h b/Engine/Plugins/Runtime/nDisplay/Source/DisplayCluster/Private/Render/Viewport/DisplayClusterViewport.h
--- a/Engine/Plugins/Runtime/nDisplay/Source/DisplayCluster/Private/Render/Viewport/DisplayClusterViewport.h
+++ b/Engine/Plugins/Runtime/nDisplay/Source/DisplayCluster/Private/Render/Viewport/DisplayClusterViewport.h
@@ -122,11 +122,6 @@ public:
 
 	virtual IDisplayClusterViewportManager& GetOwner() const override;
 
-	FDisplayClusterViewportManager& ImplGetOwner() const
-	{
-		return Owner;
-	}
-
 	const FDisplayClusterRenderFrameSettings& GetRenderFrameSettings() const;
 
 	//////////////////////////////////////////////////////
@@ -251,7 +246,7 @@ protected:
 	friend FDisplayClusterViewportRemap;
 
 	// viewport render thread data
-	TSharedPtr<FDisplayClusterViewportProxy, ESPMode::ThreadSafe> ViewportProxy;
+	FDisplayClusterViewportProxy* ViewportProxy = nullptr;
 
 	// Unique viewport name
 	const FString ViewportId;
diff --git a/Engine/Plugins/Runtime/nDisplay/Source/DisplayCluster/Private/Render/Viewport/DisplayClusterViewportManager.cpp b/Engine/Plugins/Runtime/nDisplay/Source/DisplayCluster/Private/Render/Viewport/DisplayClusterViewportManager.cpp
--- a/Engine/Plugins/Runtime/nDisplay/Source/DisplayCluster/Private/Render/Viewport/DisplayClusterViewportManager.cpp
+++ b/Engine/Plugins/Runtime/nDisplay/Source/DisplayCluster/Private/Render/Viewport/DisplayClusterViewportManager.cpp
@@ -60,9 +60,9 @@ FDisplayClusterViewportManager::FDisplayClusterViewportManager()
 	Configuration      = MakeUnique<FDisplayClusterViewportConfiguration>(*this);
 	RenderFrameManager = MakeUnique<FDisplayClusterRenderFrameManager>();
 
-	ViewportManagerProxy = MakeShared<FDisplayClusterViewportManagerProxy, ESPMode::ThreadSafe>();
+	ViewportManagerProxy = new FDisplayClusterViewportManagerProxy();
 
-	RenderTargetManager = MakeShared<FDisplayClusterRenderTargetManager, ESPMode::ThreadSafe>(ViewportManagerProxy.Get());
+	RenderTargetManager = MakeShared<FDisplayClusterRenderTargetManager, ESPMode::ThreadSafe>(ViewportManagerProxy);
 	PostProcessManager  = MakeShared<FDisplayClusterViewportPostProcessManager, ESPMode::ThreadSafe>(*this);
 	LightCardManager = MakeShared<FDisplayClusterViewportLightCardManager, ESPMode::ThreadSafe>(*this);
 
@@ -76,41 +76,30 @@ FDisplayClusterViewportManager::FDisplayClusterViewportManager()
 FDisplayClusterViewportManager::~FDisplayClusterViewportManager()
 {
 	// Remove viewports
-	Viewports.Reset();
-	ClusterNodeViewports.Reset();
-
-	RenderTargetManager.Reset();
-	PostProcessManager.Reset();
-
-	if (LightCardManager.IsValid())
 	{
-		LightCardManager->Release();
-		LightCardManager.Reset();
+		TArray<FDisplayClusterViewport*> ExistViewports = Viewports;
+		for (FDisplayClusterViewport* Viewport : ExistViewports)
+		{
+			ImplDeleteViewport(Viewport);
+		}
+		ExistViewports.Empty();
 	}
 
-	Configuration.Reset();
-
-	if (ViewportManagerProxy.IsValid())
+	if (ViewportManagerProxy)
 	{
-		// Remove viewport manager proxy on render_thread
-		ENQUEUE_RENDER_COMMAND(DeleteDisplayClusterViewportManagerProxy)(
-			[ViewportManagerProxy = ViewportManagerProxy](FRHICommandListImmediate& RHICmdList)
-			{
-				ViewportManagerProxy->Release_RenderThread();
-			});
-
-		ViewportManagerProxy.Reset();
+		ViewportManagerProxy->ImplSafeRelease();
+		ViewportManagerProxy = nullptr;
 	}
 }
 
 const IDisplayClusterViewportManagerProxy* FDisplayClusterViewportManager::GetProxy() const
 {
-	return ViewportManagerProxy.Get();
+	return ViewportManagerProxy;
 }
 
 IDisplayClusterViewportManagerProxy* FDisplayClusterViewportManager::GetProxy()
 {
-	return ViewportManagerProxy.Get();
+	return ViewportManagerProxy;
 }
 
 UWorld* FDisplayClusterViewportManager::GetCurrentWorld() const
@@ -623,6 +612,13 @@ FSceneViewFamily::ConstructionValues FDisplayClusterViewportManager::CreateViewF
 	{
 		case EDisplayClusterViewportCaptureMode::Chromakey:
 		case EDisplayClusterViewportCaptureMode::Lightcard:
+		case EDisplayClusterViewportCaptureMode::Lightcard_OCIO:
+
+			if (InFrameTarget.CaptureMode != EDisplayClusterViewportCaptureMode::Lightcard_OCIO)
+			{
+				bResolveScene = false;
+				InEngineShowFlags.PostProcessing = 0;
+			}
 
 			InEngineShowFlags.SetAtmosphere(0);
 			InEngineShowFlags.SetFog(0);
@@ -653,6 +649,7 @@ FSceneViewFamily::ConstructionValues FDisplayClusterViewportManager::CreateViewF
 		break;
 
 	case EDisplayClusterViewportCaptureMode::Lightcard:
+	case EDisplayClusterViewportCaptureMode::Lightcard_OCIO:
 
 		if (!GDisplayClusterLightcardsAllowNanite)
 		{
@@ -820,7 +817,6 @@ bool FDisplayClusterViewportManager::DeleteViewport(const FString& ViewportId)
 	if (ExistViewport != nullptr)
 	{
 		ImplDeleteViewport(ExistViewport);
-
 		return true;
 	}
 
@@ -844,6 +840,9 @@ FDisplayClusterViewport* FDisplayClusterViewportManager::ImplCreateViewport(cons
 	Viewports.Add(NewViewport);
 	ClusterNodeViewports.Add(NewViewport);
 
+	// Add viewport proxy on renderthread
+	ViewportManagerProxy->ImplCreateViewport(NewViewport->ViewportProxy);
+
 	// Handle start scene for viewport
 	NewViewport->HandleStartScene();
 
@@ -856,8 +855,11 @@ void FDisplayClusterViewportManager::ImplDeleteViewport(FDisplayClusterViewport*
 	ExistViewport->ProjectionPolicy.Reset();
 	ExistViewport->UninitializedProjectionPolicy.Reset();
 
+	// Delete viewport proxy on render thread
+	ViewportManagerProxy->ImplDeleteViewport(ExistViewport->ViewportProxy);
+
 	{
-		// Remove viewport from the whole viewports list
+		// Remove viewport obj from manager
 		int32 ViewportIndex = Viewports.Find(ExistViewport);
 		if (ViewportIndex != INDEX_NONE)
 		{
@@ -867,7 +869,7 @@ void FDisplayClusterViewportManager::ImplDeleteViewport(FDisplayClusterViewport*
 	}
 
 	{
-		// Remove viewport from the cluster viewports list
+		// Remove cluster node viewport
 		int32 ViewportIndex = ClusterNodeViewports.Find(ExistViewport);
 		if (ViewportIndex != INDEX_NONE)
 		{
diff --git a/Engine/Plugins/Runtime/nDisplay/Source/DisplayCluster/Private/Render/Viewport/DisplayClusterViewportManager.h b/Engine/Plugins/Runtime/nDisplay/Source/DisplayCluster/Private/Render/Viewport/DisplayClusterViewportManager.h
--- a/Engine/Plugins/Runtime/nDisplay/Source/DisplayCluster/Private/Render/Viewport/DisplayClusterViewportManager.h
+++ b/Engine/Plugins/Runtime/nDisplay/Source/DisplayCluster/Private/Render/Viewport/DisplayClusterViewportManager.h
@@ -116,11 +116,12 @@ public:
 
 	static TSharedPtr<IDisplayClusterProjectionPolicy, ESPMode::ThreadSafe> CreateProjectionPolicy(const FString& InViewportId, const FDisplayClusterConfigurationProjection* InConfigurationProjectionPolicy);
 
-	TSharedPtr<FDisplayClusterViewportManagerProxy, ESPMode::ThreadSafe> GetViewportManagerProxy() const
+	const FDisplayClusterViewportManagerProxy& ImplGetProxy() const
 	{
+		check(ViewportManagerProxy);
 		check(IsInGameThread());
 
-		return ViewportManagerProxy;
+		return *ViewportManagerProxy;
 	}
 
 	const FDisplayClusterRenderFrameSettings& GetRenderFrameSettings() const;
@@ -132,10 +133,10 @@ public:
 	bool ShouldUseFullSizeFrameTargetableResource() const;
 
 	void SetViewportBufferRatio(FDisplayClusterViewport& DstViewport, float InBufferRatio);
-	void ResetSceneRenderTargetSize();
 
 private:
 	void UpdateDesiredNumberOfViews(FDisplayClusterRenderFrame& InOutRenderFrame);
+	void ResetSceneRenderTargetSize();
 	void UpdateSceneRenderTargetSize();
 	void HandleViewportRTTChanges(const TArray<FDisplayClusterViewport_Context>& PrevContexts, const TArray<FDisplayClusterViewport_Context>& Contexts);
 	void ImplUpdateClusterNodeViewports(const EDisplayClusterRenderFrameMode InRenderMode, const FString& InClusterNodeId);
@@ -158,7 +159,7 @@ private:
 	TArray<FDisplayClusterViewport*> ClusterNodeViewports;
 
 	/** Render thread proxy manager. Deleted on render thread */
-	TSharedPtr<FDisplayClusterViewportManagerProxy, ESPMode::ThreadSafe> ViewportManagerProxy;
+	FDisplayClusterViewportManagerProxy* ViewportManagerProxy = nullptr;
 
 	// Pointer to the current scene
 	TWeakObjectPtr<UWorld> CurrentWorldRef;
diff --git a/Engine/Plugins/Runtime/nDisplay/Source/DisplayCluster/Private/Render/Viewport/DisplayClusterViewportManagerProxy.cpp b/Engine/Plugins/Runtime/nDisplay/Source/DisplayCluster/Private/Render/Viewport/DisplayClusterViewportManagerProxy.cpp
--- a/Engine/Plugins/Runtime/nDisplay/Source/DisplayCluster/Private/Render/Viewport/DisplayClusterViewportManagerProxy.cpp
+++ b/Engine/Plugins/Runtime/nDisplay/Source/DisplayCluster/Private/Render/Viewport/DisplayClusterViewportManagerProxy.cpp
@@ -81,35 +81,24 @@ FDisplayClusterViewportManagerProxy::FDisplayClusterViewportManagerProxy()
 { }
 
 FDisplayClusterViewportManagerProxy::~FDisplayClusterViewportManagerProxy()
-{ }
-
-void FDisplayClusterViewportManagerProxy::Release_RenderThread()
 {
-	check(IsInRenderingThread());
+	Release();
+}
 
+void FDisplayClusterViewportManagerProxy::Release()
+{
 	// Delete viewport proxy objects
-	ViewportProxies.Empty();
-	ClusterNodeViewportProxies.Empty();
-
-	if (RenderTargetManager.IsValid())
+	for (FDisplayClusterViewportProxy* ViewportProxy : ViewportProxies)
 	{
-		RenderTargetManager->Release();
-		RenderTargetManager.Reset();
-	}
-
-	if (PostProcessManager.IsValid())
-	{
-		PostProcessManager->Release();
-		PostProcessManager.Reset();
-	}
-
-	if (LightCardManager.IsValid())
-	{
-		LightCardManager->Release();
-		LightCardManager.Reset();
+		if (ViewportProxy)
+		{
+			delete ViewportProxy;
+		}
 	}
+	ViewportProxies.Empty();
+	ClusterNodeViewportProxies.Empty();
 
-	ViewportManagerViewExtension.Reset();
+	RenderTargetManager.Reset();
 }
 
 void FDisplayClusterViewportManagerProxy::Initialize(FDisplayClusterViewportManager& InViewportManager)
@@ -117,7 +106,6 @@ void FDisplayClusterViewportManagerProxy::Initialize(FDisplayClusterViewportMana
 	RenderTargetManager = InViewportManager.RenderTargetManager;
 	PostProcessManager = InViewportManager.PostProcessManager;
 	LightCardManager = InViewportManager.LightCardManager;
-	ViewportManagerViewExtension = InViewportManager.ViewportManagerViewExtension;
 }
 
 void FDisplayClusterViewportManagerProxy::DeleteResource_RenderThread(FDisplayClusterViewportResource* InDeletedResourcePtr)
@@ -125,9 +113,9 @@ void FDisplayClusterViewportManagerProxy::DeleteResource_RenderThread(FDisplayCl
 	if (InDeletedResourcePtr)
 	{
 		// Handle resource refs must be removed from all viewports
-		for (TSharedPtr<FDisplayClusterViewportProxy, ESPMode::ThreadSafe>& ViewportProxyIt : ViewportProxies)
+		for (FDisplayClusterViewportProxy* ViewportProxyIt : ViewportProxies)
 		{
-			if (ViewportProxyIt.IsValid())
+			if (ViewportProxyIt)
 			{
 				ViewportProxyIt->HandleResourceDelete_RenderThread(InDeletedResourcePtr);
 			}
@@ -142,43 +130,77 @@ void FDisplayClusterViewportManagerProxy::ImplUpdateClusterNodeViewportProxies()
 {
 	ClusterNodeViewportProxies.Empty();
 	// Collect viewport proxies for rendered cluster node
-	for (TSharedPtr<FDisplayClusterViewportProxy, ESPMode::ThreadSafe>& ViewportProxyIt : ViewportProxies)
+	for (FDisplayClusterViewportProxy* ViewportProxyIt : ViewportProxies)
 	{
-		if (ViewportProxyIt.IsValid() && (RenderFrameSettings.ClusterNodeId.IsEmpty() || ViewportProxyIt->GetClusterNodeId() == RenderFrameSettings.ClusterNodeId))
+		if (ViewportProxyIt && (RenderFrameSettings.ClusterNodeId.IsEmpty() || ViewportProxyIt->GetClusterNodeId() == RenderFrameSettings.ClusterNodeId))
 		{
 			ClusterNodeViewportProxies.Add(ViewportProxyIt);
 		}
 	}
 }
 
-void FDisplayClusterViewportManagerProxy::CreateViewport_RenderThread(const TSharedPtr<FDisplayClusterViewportProxy, ESPMode::ThreadSafe>& InViewportProxy)
+void FDisplayClusterViewportManagerProxy::ImplSafeRelease()
 {
-	check(IsInRenderingThread());
+	check(IsInGameThread());
 
-	ViewportProxies.Add(InViewportProxy);
-	ImplUpdateClusterNodeViewportProxies();
+	if (RenderTargetManager.IsValid())
+	{
+		RenderTargetManager->Release();
+	}
+	
+	if (PostProcessManager.IsValid())
+	{
+		PostProcessManager->Release();
+	}
+
+	if (LightCardManager.IsValid())
+	{
+		LightCardManager->Release();
+	}
+
+	// Remove viewport manager proxy on render_thread
+	ENQUEUE_RENDER_COMMAND(DeleteDisplayClusterViewportManagerProxy)(
+		[ViewportManagerProxy = this](FRHICommandListImmediate& RHICmdList)
+	{
+		delete ViewportManagerProxy;
+	});
 }
 
-void FDisplayClusterViewportManagerProxy::DeleteViewport_RenderThread(const TSharedPtr<FDisplayClusterViewportProxy, ESPMode::ThreadSafe>& InViewportProxy)
+void FDisplayClusterViewportManagerProxy::ImplCreateViewport(FDisplayClusterViewportProxy* InViewportProxy)
 {
-	check(IsInRenderingThread());
+	check(IsInGameThread());
 
-	// Remove viewport obj from manager
-	int32 ViewportProxyIndex = ViewportProxies.Find(InViewportProxy);
-	if (ViewportProxyIndex != INDEX_NONE)
+	if (InViewportProxy)
 	{
-		ViewportProxies[ViewportProxyIndex].Reset();
-		ViewportProxies.RemoveAt(ViewportProxyIndex);
+		ENQUEUE_RENDER_COMMAND(CreateDisplayClusterViewportProxy)(
+			[ViewportManagerProxy = this, ViewportProxy = InViewportProxy](FRHICommandListImmediate& RHICmdList)
+		{
+			ViewportManagerProxy->ViewportProxies.Add(ViewportProxy);
+			ViewportManagerProxy->ImplUpdateClusterNodeViewportProxies();
+		});
 	}
+}
 
-	int32 ClusterViewportProxyIndex = ClusterNodeViewportProxies.Find(InViewportProxy);
-	if (ClusterViewportProxyIndex != INDEX_NONE)
+void FDisplayClusterViewportManagerProxy::ImplDeleteViewport(FDisplayClusterViewportProxy* InViewportProxy)
+{
+	check(IsInGameThread());
+
+	// Remove viewport sceneproxy on renderthread
+	ENQUEUE_RENDER_COMMAND(DeleteDisplayClusterViewportProxy)(
+		[ViewportManagerProxy = this, ViewportProxy = InViewportProxy](FRHICommandListImmediate& RHICmdList)
 	{
-		ClusterNodeViewportProxies[ClusterViewportProxyIndex].Reset();
-		ClusterNodeViewportProxies.RemoveAt(ClusterViewportProxyIndex);
-	}
-}
+		// Remove viewport obj from manager
+		int32 ViewportProxyIndex = ViewportManagerProxy->ViewportProxies.Find(ViewportProxy);
+		if (ViewportProxyIndex != INDEX_NONE)
+		{
+			ViewportManagerProxy->ViewportProxies[ViewportProxyIndex] = nullptr;
+			ViewportManagerProxy->ViewportProxies.RemoveAt(ViewportProxyIndex);
+			ViewportManagerProxy->ImplUpdateClusterNodeViewportProxies();
+		}
 
+		delete ViewportProxy;
+	});
+}
 void FDisplayClusterViewportManagerProxy::ImplUpdateRenderFrameSettings(const FDisplayClusterRenderFrameSettings& InRenderFrameSettings)
 {
 	check(IsInGameThread());
@@ -187,7 +209,7 @@ void FDisplayClusterViewportManagerProxy::ImplUpdateRenderFrameSettings(const FD
 
 	// Send frame settings to renderthread
 	ENQUEUE_RENDER_COMMAND(DeleteDisplayClusterViewportProxy)(
-		[ViewportManagerProxy = SharedThis(this), Settings](FRHICommandListImmediate& RHICmdList)
+		[ViewportManagerProxy = this, Settings](FRHICommandListImmediate& RHICmdList)
 	{
 		ViewportManagerProxy->RenderFrameSettings = *Settings;
 		delete Settings;
@@ -225,13 +247,11 @@ DECLARE_GPU_STAT_NAMED(nDisplay_ViewportManager_RenderFrame, TEXT("nDisplay View
 void FDisplayClusterViewportManagerProxy::ImplRenderFrame(FViewport* InViewport)
 {
 	ENQUEUE_RENDER_COMMAND(DeleteDisplayClusterViewportProxy)(
-		[InViewportManagerProxy = SharedThis(this), InViewport](FRHICommandListImmediate& RHICmdList)
+		[ViewportManagerProxy = this, InViewport](FRHICommandListImmediate& RHICmdList)
 	{
 		SCOPED_GPU_STAT(RHICmdList, nDisplay_ViewportManager_RenderFrame);
 		SCOPED_DRAW_EVENT(RHICmdList, nDisplay_ViewportManager_RenderFrame);
 
-		const FDisplayClusterViewportManagerProxy* ViewportManagerProxy = &InViewportManagerProxy.Get();
-
 		// Handle render setup
 		if (ViewportManagerProxy->PostProcessManager.IsValid())
 		{
@@ -302,10 +322,10 @@ void FDisplayClusterViewportManagerProxy::UpdateDeferredResources_RenderThread(F
 {
 	check(IsInRenderingThread());
 
-	TArray<TSharedPtr<FDisplayClusterViewportProxy, ESPMode::ThreadSafe>> OverriddenViewports;
+	TArray<FDisplayClusterViewportProxy*> OverriddenViewports;
 	OverriddenViewports.Reserve(ClusterNodeViewportProxies.Num());
 
-	for (const TSharedPtr<FDisplayClusterViewportProxy, ESPMode::ThreadSafe>& ViewportProxy : ClusterNodeViewportProxies)
+	for (FDisplayClusterViewportProxy* ViewportProxy : ClusterNodeViewportProxies)
 	{
 		if (ViewportProxy->RenderSettings.OverrideViewportId.IsEmpty())
 		{
@@ -319,7 +339,7 @@ void FDisplayClusterViewportManagerProxy::UpdateDeferredResources_RenderThread(F
 	}
 
 	// Update deferred viewports after all
-	for (TSharedPtr<FDisplayClusterViewportProxy, ESPMode::ThreadSafe>& ViewportProxy : OverriddenViewports)
+	for (FDisplayClusterViewportProxy* ViewportProxy : OverriddenViewports)
 	{
 		ViewportProxy->UpdateDeferredResources(RHICmdList);
 	}
@@ -372,11 +392,11 @@ void FDisplayClusterViewportManagerProxy::UpdateFrameResources_RenderThread(FRHI
 	}
 
 	// Support viewport overlap order sorting:
-	TArray<TSharedPtr<FDisplayClusterViewportProxy, ESPMode::ThreadSafe>> SortedViewportProxy = ClusterNodeViewportProxies;
+	TArray<FDisplayClusterViewportProxy*> SortedViewportProxy = ClusterNodeViewportProxies;
 	SortedViewportProxy.Sort(
-		[](const TSharedPtr<FDisplayClusterViewportProxy, ESPMode::ThreadSafe>& VP1, const TSharedPtr<FDisplayClusterViewportProxy, ESPMode::ThreadSafe>& VP2)
+		[](const FDisplayClusterViewportProxy& VP1, const FDisplayClusterViewportProxy& VP2)
 		{
-			return  VP1->GetRenderSettings_RenderThread().OverlapOrder < VP2->GetRenderSettings_RenderThread().OverlapOrder;
+			return  VP1.GetRenderSettings_RenderThread().OverlapOrder < VP2.GetRenderSettings_RenderThread().OverlapOrder;
 		}
 	);
 
@@ -393,28 +413,28 @@ void FDisplayClusterViewportManagerProxy::UpdateFrameResources_RenderThread(FRHI
 	for (uint8 WarpPass = 0; WarpPass < (uint8)EWarpPass::COUNT; WarpPass++)
 	{
 		// Update deferred resources for viewports
-		for (TSharedPtr<FDisplayClusterViewportProxy, ESPMode::ThreadSafe>& ViewportProxyIt : SortedViewportProxy)
+		for (FDisplayClusterViewportProxy* ViewportProxy : SortedViewportProxy)
 		{
 			// Iterate over visible viewports:
-			if (ViewportProxyIt.IsValid() && ViewportProxyIt->GetRenderSettings_RenderThread().bVisible)
+			if (ViewportProxy && ViewportProxy->GetRenderSettings_RenderThread().bVisible)
 			{
-				if (bWarpBlendEnabled && ShouldApplyWarpBlend(ViewportProxyIt.Get()))
+				if (bWarpBlendEnabled && ShouldApplyWarpBlend(ViewportProxy))
 				{
-					const TSharedPtr<IDisplayClusterProjectionPolicy, ESPMode::ThreadSafe>& PrjPolicy = ViewportProxyIt->GetProjectionPolicy_RenderThread();
+					const TSharedPtr<IDisplayClusterProjectionPolicy, ESPMode::ThreadSafe>& PrjPolicy = ViewportProxy->GetProjectionPolicy_RenderThread();
 					switch ((EWarpPass)WarpPass)
 					{
 					case EWarpPass::Begin:
-						IDisplayCluster::Get().GetCallbacks().OnDisplayClusterPreWarpViewport_RenderThread().Broadcast(RHICmdList, ViewportProxyIt.Get());
-						PrjPolicy->BeginWarpBlend_RenderThread(RHICmdList, ViewportProxyIt.Get());
+						IDisplayCluster::Get().GetCallbacks().OnDisplayClusterPreWarpViewport_RenderThread().Broadcast(RHICmdList, ViewportProxy);
+						PrjPolicy->BeginWarpBlend_RenderThread(RHICmdList, ViewportProxy);
 						break;
 
 					case EWarpPass::Render:
-						PrjPolicy->ApplyWarpBlend_RenderThread(RHICmdList, ViewportProxyIt.Get());
+						PrjPolicy->ApplyWarpBlend_RenderThread(RHICmdList, ViewportProxy);
 						break;
 
 					case EWarpPass::End:
-						PrjPolicy->EndWarpBlend_RenderThread(RHICmdList, ViewportProxyIt.Get());
-						IDisplayCluster::Get().GetCallbacks().OnDisplayClusterPostWarpViewport_RenderThread().Broadcast(RHICmdList, ViewportProxyIt.Get());
+						PrjPolicy->EndWarpBlend_RenderThread(RHICmdList, ViewportProxy);
+						IDisplayCluster::Get().GetCallbacks().OnDisplayClusterPostWarpViewport_RenderThread().Broadcast(RHICmdList, ViewportProxy);
 						break;
 
 					default:
@@ -436,15 +456,15 @@ void FDisplayClusterViewportManagerProxy::UpdateFrameResources_RenderThread(FRHI
 
 	// Post resolve to Frame RTT
 	// All warp&blend results are now inside AdditionalTargetableResource. Viewport images of other projection policies are still stored in the InputShaderResource.
-	for (TSharedPtr<FDisplayClusterViewportProxy, ESPMode::ThreadSafe>& ViewportProxyIt : SortedViewportProxy)
+	for (FDisplayClusterViewportProxy* ViewportProxy : SortedViewportProxy)
 	{
 		// Iterate over visible viewports:
-		if (ViewportProxyIt.IsValid() && ViewportProxyIt->GetRenderSettings_RenderThread().bVisible)
+		if (ViewportProxy && ViewportProxy->GetRenderSettings_RenderThread().bVisible)
 		{
 			EDisplayClusterViewportResourceType ViewportSource = EDisplayClusterViewportResourceType::InputShaderResource;
-			if (bWarpBlendEnabled && ShouldApplyWarpBlend(ViewportProxyIt.Get()))
+			if (bWarpBlendEnabled && ShouldApplyWarpBlend(ViewportProxy))
 			{
-				const TSharedPtr<IDisplayClusterProjectionPolicy, ESPMode::ThreadSafe>& PrjPolicy = ViewportProxyIt->GetProjectionPolicy_RenderThread();
+				const TSharedPtr<IDisplayClusterProjectionPolicy, ESPMode::ThreadSafe>& PrjPolicy = ViewportProxy->GetProjectionPolicy_RenderThread();
 				if (PrjPolicy->ShouldUseAdditionalTargetableResource())
 				{
 					ViewportSource = EDisplayClusterViewportResourceType::AdditionalTargetableResource;
@@ -452,10 +472,10 @@ void FDisplayClusterViewportManagerProxy::UpdateFrameResources_RenderThread(FRHI
 			}
 
 			// resolve viewports to the frame target texture
-			ViewportProxyIt->ResolveResources_RenderThread(RHICmdList, ViewportSource, ViewportProxyIt->GetOutputResourceType_RenderThread());
+			ViewportProxy->ResolveResources_RenderThread(RHICmdList, ViewportSource, ViewportProxy->GetOutputResourceType_RenderThread());
 
 			// Apply post-warp (viewport remap, etc)
-			ViewportProxyIt->PostResolveViewport_RenderThread(RHICmdList);
+			ViewportProxy->PostResolveViewport_RenderThread(RHICmdList);
 		}
 	}
 
@@ -481,27 +501,27 @@ void FDisplayClusterViewportManagerProxy::DoCrossGPUTransfers_RenderThread(FRHIC
 	// Copy the view render results to all GPUs that are native to the viewport.
 	TArray<FTransferResourceParams> TransferResources;
 
-	for (const TSharedPtr<FDisplayClusterViewportProxy, ESPMode::ThreadSafe>& ViewportProxyIt : ClusterNodeViewportProxies)
+	for (FDisplayClusterViewportProxy* ViewportProxy : ClusterNodeViewportProxies)
 	{
 		bool bShouldCrossGPUTransfersViewportRenderResources = true;
 
 		// Skip a frozen viewport that has already been transferred between GPUs
 		// The first time freezing should do the transfer (RenderTargets must be assigned on the first pass)
-		if (ViewportProxyIt->RenderSettings.bFreezeRendering && ViewportProxyIt->RenderTargets.Num() == 0)
+		if (ViewportProxy->RenderSettings.bFreezeRendering && ViewportProxy->RenderTargets.Num() == 0)
 		{
 			bShouldCrossGPUTransfersViewportRenderResources = false;
 		}
 
 		if (bShouldCrossGPUTransfersViewportRenderResources)
 		{
-			for (FDisplayClusterViewport_Context& ViewportContext : ViewportProxyIt->Contexts)
+			for (FDisplayClusterViewport_Context& ViewportContext : ViewportProxy->Contexts)
 			{
 				if (ViewportContext.bAllowGPUTransferOptimization && ViewportContext.GPUIndex >= 0)
 				{
 					// Use optimized cross GPU transfer for this context
 
-					FRenderTarget* RenderTarget = ViewportProxyIt->RenderTargets[ViewportContext.ContextNum];
-					FRHITexture2D* TextureRHI = ViewportProxyIt->RenderTargets[ViewportContext.ContextNum]->GetViewportRenderTargetResourceRHI();
+					FRenderTarget* RenderTarget = ViewportProxy->RenderTargets[ViewportContext.ContextNum];
+					FRHITexture2D* TextureRHI = ViewportProxy->RenderTargets[ViewportContext.ContextNum]->GetViewportRenderTargetResourceRHI();
 
 					FRHIGPUMask RenderTargetGPUMask = (GNumExplicitGPUsForRendering > 1 && RenderTarget) ? RenderTarget->GetGPUMask(RHICmdList) : FRHIGPUMask::GPU0();
 					FRHIGPUMask ContextGPUMask = FRHIGPUMask::FromIndex(ViewportContext.GPUIndex);
@@ -541,12 +561,12 @@ bool FDisplayClusterViewportManagerProxy::GetFrameTargets_RenderThread(TArray<FR
 	check(IsInRenderingThread());
 
 	// Get any defined frame targets from first visible viewport
-	for (const TSharedPtr<FDisplayClusterViewportProxy, ESPMode::ThreadSafe>& ViewportProxyIt : ClusterNodeViewportProxies)
+	for (FDisplayClusterViewportProxy* ViewportProxy : ClusterNodeViewportProxies)
 	{
-		if (ViewportProxyIt.IsValid())
+		if (ViewportProxy)
 		{
-			const TArray<FDisplayClusterViewportTextureResource*>& Frames = ViewportProxyIt->OutputFrameTargetableResources;
-			const TArray<FDisplayClusterViewportTextureResource*>& AdditionalFrames = ViewportProxyIt->AdditionalFrameTargetableResources;
+			const TArray<FDisplayClusterViewportTextureResource*>& Frames = ViewportProxy->OutputFrameTargetableResources;
+			const TArray<FDisplayClusterViewportTextureResource*>& AdditionalFrames = ViewportProxy->AdditionalFrameTargetableResources;
 
 			if (Frames.Num() > 0)
 			{
@@ -636,23 +656,23 @@ FDisplayClusterViewportProxy* FDisplayClusterViewportManagerProxy::ImplFindViewp
 	check(IsInRenderingThread());
 
 	// Ok, we have a request for a particular viewport. Let's find it.
-	TSharedPtr<FDisplayClusterViewportProxy, ESPMode::ThreadSafe> const* DesiredViewport = ViewportProxies.FindByPredicate([ViewportId](const TSharedPtr<FDisplayClusterViewportProxy, ESPMode::ThreadSafe>& ItemViewport)
+	FDisplayClusterViewportProxy* const* DesiredViewport = ViewportProxies.FindByPredicate([ViewportId](const FDisplayClusterViewportProxy* ItemViewport)
 	{
 		return ViewportId.Equals(ItemViewport->GetId(), ESearchCase::IgnoreCase);
 	});
 
-	return (DesiredViewport != nullptr) ? DesiredViewport->Get() : nullptr;
+	return (DesiredViewport != nullptr) ? *DesiredViewport : nullptr;
 }
 
 IDisplayClusterViewportProxy* FDisplayClusterViewportManagerProxy::FindViewport_RenderThread(const int32 StereoViewIndex, uint32* OutContextNum) const
 {
 	check(IsInRenderingThread());
 
-	for (const TSharedPtr<FDisplayClusterViewportProxy, ESPMode::ThreadSafe>&  ViewportProxyIt : ViewportProxies)
+	for (FDisplayClusterViewportProxy* ViewportProxy : ViewportProxies)
 	{
-		if (ViewportProxyIt.IsValid() && ViewportProxyIt->FindContext_RenderThread(StereoViewIndex, OutContextNum))
+		if (ViewportProxy && ViewportProxy->FindContext_RenderThread(StereoViewIndex, OutContextNum))
 		{
-			return ViewportProxyIt.Get();
+			return ViewportProxy;
 		}
 	}
 
diff --git a/Engine/Plugins/Runtime/nDisplay/Source/DisplayCluster/Private/Render/Viewport/DisplayClusterViewportManagerProxy.h b/Engine/Plugins/Runtime/nDisplay/Source/DisplayCluster/Private/Render/Viewport/DisplayClusterViewportManagerProxy.h
--- a/Engine/Plugins/Runtime/nDisplay/Source/DisplayCluster/Private/Render/Viewport/DisplayClusterViewportManagerProxy.h
+++ b/Engine/Plugins/Runtime/nDisplay/Source/DisplayCluster/Private/Render/Viewport/DisplayClusterViewportManagerProxy.h
@@ -8,7 +8,6 @@
 #include "Render/Viewport/DisplayClusterViewportProxy.h"
 #include "Render/Viewport/Containers/DisplayClusterViewport_Enums.h"
 #include "Render/Viewport/RenderFrame/DisplayClusterRenderFrameSettings.h"
-#include "Templates/SharedPointer.h"
 
 class FDisplayClusterRenderTargetManager;
 class FDisplayClusterViewportPostProcessManager;
@@ -23,13 +22,12 @@ class FViewport;
 
 class FDisplayClusterViewportManagerProxy
 	: public IDisplayClusterViewportManagerProxy
-	, public TSharedFromThis<FDisplayClusterViewportManagerProxy, ESPMode::ThreadSafe>
 {
 public:
 	FDisplayClusterViewportManagerProxy();
 	virtual ~FDisplayClusterViewportManagerProxy();
 
-	void Release_RenderThread();
+	void Release();
 
 public:
 	void DoCrossGPUTransfers_RenderThread(FRHICommandListImmediate& RHICmdList) const;
@@ -44,9 +42,9 @@ public:
 
 	virtual IDisplayClusterViewportProxy* FindViewport_RenderThread(const int32 StereoViewIndex, uint32* OutContextNum = nullptr) const override;
 
-	virtual const TArrayView<TSharedPtr<IDisplayClusterViewportProxy, ESPMode::ThreadSafe>> GetViewports_RenderThread() const override
+	virtual const TArrayView<IDisplayClusterViewportProxy*> GetViewports_RenderThread() const override
 	{
-		return TArrayView<TSharedPtr<IDisplayClusterViewportProxy, ESPMode::ThreadSafe>>((TSharedPtr<IDisplayClusterViewportProxy, ESPMode::ThreadSafe>*)(ClusterNodeViewportProxies.GetData()), ClusterNodeViewportProxies.Num());
+		return TArrayView<IDisplayClusterViewportProxy*>((IDisplayClusterViewportProxy**)(ClusterNodeViewportProxies.GetData()), ClusterNodeViewportProxies.Num());
 	}
 
 	virtual bool GetFrameTargets_RenderThread(TArray<FRHITexture2D*>& OutFrameResources, TArray<FIntPoint>& OutTargetOffsets, TArray<FRHITexture2D*>* OutAdditionalFrameResources = nullptr) const override;
@@ -55,9 +53,6 @@ public:
 	virtual TSharedPtr<IDisplayClusterViewportLightCardManager, ESPMode::ThreadSafe> GetLightCardManager_RenderThread() const override;
 
 	// internal use only
-	void CreateViewport_RenderThread(const TSharedPtr<FDisplayClusterViewportProxy, ESPMode::ThreadSafe>& InViewportProxy);
-	void DeleteViewport_RenderThread(const TSharedPtr<FDisplayClusterViewportProxy, ESPMode::ThreadSafe>& InViewportProxy);
-
 	void DeleteResource_RenderThread(FDisplayClusterViewportResource* InDeletedResourcePtr);
 	void Initialize(FDisplayClusterViewportManager& InViewportManager);
 
@@ -68,7 +63,7 @@ public:
 		return RenderFrameSettings;
 	}
 
-	const TArray<TSharedPtr<FDisplayClusterViewportProxy, ESPMode::ThreadSafe>>& ImplGetViewportProxies_RenderThread() const
+	const TArray<FDisplayClusterViewportProxy*>& ImplGetViewportProxies_RenderThread() const
 	{
 		check(IsInRenderingThread());
 
@@ -77,8 +72,11 @@ public:
 
 	FDisplayClusterViewportProxy* ImplFindViewport_RenderThread(const FString& InViewportId) const;
 
+	void ImplCreateViewport(FDisplayClusterViewportProxy* InViewportProxy);
+	void ImplDeleteViewport(FDisplayClusterViewportProxy* InViewportProxy);
 	void ImplUpdateRenderFrameSettings(const FDisplayClusterRenderFrameSettings& InRenderFrameSettings);
 	void ImplUpdateViewports(const TArray<FDisplayClusterViewport*>& InViewports);
+	void ImplSafeRelease();
 
 	void ImplRenderFrame(FViewport* InViewport);
 	void ImplClearFrameTargets_RenderThread(FRHICommandListImmediate& RHICmdList) const;
@@ -93,9 +91,6 @@ private:
 
 	FDisplayClusterRenderFrameSettings RenderFrameSettings;
 
-	TArray<TSharedPtr<FDisplayClusterViewportProxy, ESPMode::ThreadSafe>> ViewportProxies;
-	TArray<TSharedPtr<FDisplayClusterViewportProxy, ESPMode::ThreadSafe>> ClusterNodeViewportProxies;
-
-	// Handle special features
-	TSharedPtr<class FDisplayClusterViewportManagerViewExtension, ESPMode::ThreadSafe> ViewportManagerViewExtension;
+	TArray<FDisplayClusterViewportProxy*> ViewportProxies;
+	TArray<FDisplayClusterViewportProxy*> ClusterNodeViewportProxies;
 };
diff --git a/Engine/Plugins/Runtime/nDisplay/Source/DisplayCluster/Private/Render/Viewport/DisplayClusterViewportManagerViewExtension.cpp b/Engine/Plugins/Runtime/nDisplay/Source/DisplayCluster/Private/Render/Viewport/DisplayClusterViewportManagerViewExtension.cpp
--- a/Engine/Plugins/Runtime/nDisplay/Source/DisplayCluster/Private/Render/Viewport/DisplayClusterViewportManagerViewExtension.cpp
+++ b/Engine/Plugins/Runtime/nDisplay/Source/DisplayCluster/Private/Render/Viewport/DisplayClusterViewportManagerViewExtension.cpp
@@ -10,27 +10,10 @@
 #include "IDisplayClusterCallbacks.h"
 
 
-namespace DisplayClusterViewportManagerViewExtensionHelpers
-{
-	static const FName RendererModuleName(TEXT("Renderer"));
-};
-using namespace DisplayClusterViewportManagerViewExtensionHelpers;
-
-///////////////////////////////////////////////////////////////////////////////////////
-// FDisplayClusterViewportManagerViewExtension
-///////////////////////////////////////////////////////////////////////////////////////
 FDisplayClusterViewportManagerViewExtension::FDisplayClusterViewportManagerViewExtension(const FAutoRegister& AutoRegister, const IDisplayClusterViewportManager* InViewportManager)
 	: FSceneViewExtensionBase(AutoRegister)
 	, ViewportManager(InViewportManager)
-	, ViewportManagerProxy(InViewportManager->GetProxy())
-{
-	RegisterCallbacks();
-}
-
-FDisplayClusterViewportManagerViewExtension::~FDisplayClusterViewportManagerViewExtension()
-{
-	UnregisterCallbacks();
-}
+{ }
 
 void FDisplayClusterViewportManagerViewExtension::PreRenderView_RenderThread(FRDGBuilder& GraphBuilder, FSceneView& InView)
 {
@@ -59,71 +42,21 @@ bool FDisplayClusterViewportManagerViewExtension::IsActiveThisFrame_Internal(con
 	return false;
 }
 
-void FDisplayClusterViewportManagerViewExtension::PreRenderViewFamily_RenderThread(FRDGBuilder& GraphBuilder, FSceneViewFamily& InViewFamily)
+void FDisplayClusterViewportManagerViewExtension::PostRenderViewFamily_RenderThread(FRDGBuilder& GraphBuilder, FSceneViewFamily& InViewFamily)
 {
-	Viewports.Empty();
-
-	// Get all viewport proxies that are uses in this view family
-	for (int32 ViewIndex = 0; ViewIndex < InViewFamily.Views.Num(); ViewIndex++)
+	// FSceneViewFamily::ProfileDescription is used currently as a viewport ID holder, so we use it to find a proper viewport proxy.
+	if (ViewportManager)
 	{
-		if (const FSceneView* SceneView = InViewFamily.Views[ViewIndex])
+		if (const IDisplayClusterViewportManagerProxy* const ViewportMgrPoxy = ViewportManager->GetProxy())
 		{
-			FViewportProxy NewViewportProxy;
-			NewViewportProxy.ViewportProxy = ViewportManagerProxy->FindViewport_RenderThread(SceneView->StereoViewIndex, &NewViewportProxy.ViewportProxyContext.ContextNum);
-			NewViewportProxy.ViewportProxyContext.ViewFamilyProfileDescription = InViewFamily.ProfileDescription;
-
-			NewViewportProxy.ViewIndex = ViewIndex;
-
-			if (NewViewportProxy.IsEnabled())
+			if (!InViewFamily.ProfileDescription.IsEmpty())
 			{
-				Viewports.Add(NewViewportProxy);
+				if (IDisplayClusterViewportProxy* ViewportProxy = ViewportMgrPoxy->FindViewport_RenderThread(InViewFamily.ProfileDescription))
+				{
+					// Notify about viewport rendering finish
+					IDisplayCluster::Get().GetCallbacks().OnDisplayClusterPostRenderViewFamily_RenderThread().Broadcast(GraphBuilder, InViewFamily, ViewportProxy);
+				}
 			}
 		}
 	}
 }
-
-void FDisplayClusterViewportManagerViewExtension::PostRenderViewFamily_RenderThread(FRDGBuilder& GraphBuilder, FSceneViewFamily& InViewFamily)
-{
-	for (FViewportProxy& ViewportIt : Viewports)
-	{
-		if (const FSceneView* SceneView = InViewFamily.Views[ViewportIt.ViewIndex])
-		{
-			ViewportIt.ViewportProxy->PostRenderViewFamily_RenderThread(GraphBuilder, InViewFamily, *SceneView, ViewportIt.ViewportProxyContext);
-		}
-	}
-}
-
-//////////////////////////////////////////////////////////////////////////////////////////////
-void FDisplayClusterViewportManagerViewExtension::OnResolvedSceneColor_RenderThread(FRDGBuilder& GraphBuilder, const FSceneTextures& SceneTextures)
-{
-	check(IsInRenderingThread());
-
-	for (FViewportProxy& ViewportIt : Viewports)
-	{
-		ViewportIt.ViewportProxy->OnResolvedSceneColor_RenderThread(GraphBuilder, SceneTextures, ViewportIt.ViewportProxyContext);
-	}
-}
-
-void FDisplayClusterViewportManagerViewExtension::RegisterCallbacks()
-{
-	if (!ResolvedSceneColorCallbackHandle.IsValid())
-	{
-		if (IRendererModule* RendererModule = FModuleManager::GetModulePtr<IRendererModule>(RendererModuleName))
-		{
-			ResolvedSceneColorCallbackHandle = RendererModule->GetResolvedSceneColorCallbacks().AddRaw(this, &FDisplayClusterViewportManagerViewExtension::OnResolvedSceneColor_RenderThread);
-		}
-	}
-}
-
-void FDisplayClusterViewportManagerViewExtension::UnregisterCallbacks()
-{
-	if (ResolvedSceneColorCallbackHandle.IsValid())
-	{
-		if (IRendererModule* RendererModule = FModuleManager::GetModulePtr<IRendererModule>(RendererModuleName))
-		{
-			RendererModule->GetResolvedSceneColorCallbacks().Remove(ResolvedSceneColorCallbackHandle);
-		}
-
-		ResolvedSceneColorCallbackHandle.Reset();
-	}
-}
diff --git a/Engine/Plugins/Runtime/nDisplay/Source/DisplayCluster/Private/Render/Viewport/DisplayClusterViewportManagerViewExtension.h b/Engine/Plugins/Runtime/nDisplay/Source/DisplayCluster/Private/Render/Viewport/DisplayClusterViewportManagerViewExtension.h
--- a/Engine/Plugins/Runtime/nDisplay/Source/DisplayCluster/Private/Render/Viewport/DisplayClusterViewportManagerViewExtension.h
+++ b/Engine/Plugins/Runtime/nDisplay/Source/DisplayCluster/Private/Render/Viewport/DisplayClusterViewportManagerViewExtension.h
@@ -2,10 +2,9 @@
 
 #pragma once
 #include "SceneViewExtension.h"
-#include "Render/Viewport/Containers/DisplayClusterViewportProxy_Context.h"
 
 class IDisplayClusterViewportManager;
-class IDisplayClusterViewportManagerProxy;
+
 
 /**
  * View extension applying an DC Viewport features
@@ -14,7 +13,6 @@ class FDisplayClusterViewportManagerViewExtension : public FSceneViewExtensionBa
 {
 public:
 	FDisplayClusterViewportManagerViewExtension(const FAutoRegister& AutoRegister, const IDisplayClusterViewportManager* InViewportManager);
-	virtual ~FDisplayClusterViewportManagerViewExtension();
 
 public:
 	//~ Begin ISceneViewExtension interface
@@ -36,37 +34,9 @@ public:
 
 	virtual bool IsActiveThisFrame_Internal(const FSceneViewExtensionContext& Context) const override;
 
-	virtual void PreRenderViewFamily_RenderThread(FRDGBuilder& GraphBuilder, FSceneViewFamily& InViewFamily) override;
 	virtual void PostRenderViewFamily_RenderThread(FRDGBuilder& GraphBuilder, FSceneViewFamily& InViewFamily) override;
 	//~End ISceneVIewExtension interface
 
-private:
-	/** Rendered callback (get scene textures to share) */
-	void OnResolvedSceneColor_RenderThread(FRDGBuilder& GraphBuilder, const FSceneTextures& SceneTextures);
-
-	void RegisterCallbacks();
-	void UnregisterCallbacks();
-
-private:
-	FDelegateHandle ResolvedSceneColorCallbackHandle;
-
 private:
 	const IDisplayClusterViewportManager* ViewportManager;
-	const IDisplayClusterViewportManagerProxy* ViewportManagerProxy;
-
-	struct FViewportProxy
-	{
-		inline bool IsEnabled() const
-		{
-			return ViewportProxy != nullptr;
-		}
-
-		class IDisplayClusterViewportProxy* ViewportProxy = nullptr;
-		FDisplayClusterViewportProxy_Context ViewportProxyContext;
-
-		int32 ViewIndex = 0;
-	};
-
-	TArray<FViewportProxy> Viewports;
-	
 };
diff --git a/Engine/Plugins/Runtime/nDisplay/Source/DisplayCluster/Private/Render/Viewport/DisplayClusterViewportProxy.cpp b/Engine/Plugins/Runtime/nDisplay/Source/DisplayCluster/Private/Render/Viewport/DisplayClusterViewportProxy.cpp
--- a/Engine/Plugins/Runtime/nDisplay/Source/DisplayCluster/Private/Render/Viewport/DisplayClusterViewportProxy.cpp
+++ b/Engine/Plugins/Runtime/nDisplay/Source/DisplayCluster/Private/Render/Viewport/DisplayClusterViewportProxy.cpp
@@ -4,7 +4,6 @@
 
 #include "Render/Viewport/DisplayClusterViewportManagerProxy.h"
 #include "Render/Viewport/DisplayClusterViewport.h"
-#include "Render/Viewport/DisplayClusterViewportManager.h"
 
 #include "Render/Viewport/Containers/DisplayClusterViewport_PostRenderSettings.h"
 #include "Render/Viewport/RenderTarget/DisplayClusterRenderTargetResource.h"
@@ -16,10 +15,6 @@
 #include "Render/Containers/IDisplayClusterRender_MeshComponent.h"
 #include "Render/Containers/IDisplayClusterRender_MeshComponentProxy.h"
 
-#include "IDisplayCluster.h"
-#include "IDisplayClusterCallbacks.h"
-#include "IDisplayClusterShaders.h"
-
 #if WITH_EDITOR
 #include "DisplayClusterRootActor.h"
 #include "Render/Viewport/Containers/DisplayClusterViewportReadPixels.h"
@@ -38,9 +33,8 @@
 
 #include "ScreenRendering.h"
 #include "PostProcess/SceneFilterRendering.h"
-#include "SceneRendering.h"
-#include "RenderGraphUtils.h"
-#include "RenderGraphResources.h"
+
+#include "IDisplayClusterShaders.h"
 
 static TAutoConsoleVariable<int32> CVarDisplayClusterRenderOverscanResolve(
 	TEXT("nDisplay.render.overscan.resolve"),
@@ -154,12 +148,12 @@ namespace DisplayClusterViewportProxyHelpers
 using namespace DisplayClusterViewportProxyHelpers;
 
 ///////////////////////////////////////////////////////////////////////////////////////
-FDisplayClusterViewportProxy::FDisplayClusterViewportProxy(const FDisplayClusterViewport& RenderViewport)
+FDisplayClusterViewportProxy::FDisplayClusterViewportProxy(const FDisplayClusterViewportManagerProxy& InOwner, const FDisplayClusterViewport& RenderViewport)
 	: ViewportId(RenderViewport.ViewportId)
 	, ClusterNodeId(RenderViewport.ClusterNodeId)
 	, RenderSettings(RenderViewport.RenderSettings)
 	, ProjectionPolicy(RenderViewport.UninitializedProjectionPolicy)
-	, Owner(RenderViewport.ImplGetOwner().GetViewportManagerProxy().ToSharedRef())
+	, Owner(InOwner)
 	, ShadersAPI(IDisplayClusterShaders::Get())
 {
 	check(ProjectionPolicy.IsValid());
@@ -173,7 +167,7 @@ const IDisplayClusterViewportManagerProxy& FDisplayClusterViewportProxy::GetOwne
 {
 	check(IsInRenderingThread());
 
-	return Owner.Get();
+	return Owner;
 }
 
 void FDisplayClusterViewportProxy::HandleResourceDelete_RenderThread(FDisplayClusterViewportResource* InDeletedResourcePtr)
@@ -226,7 +220,7 @@ bool FDisplayClusterViewportProxy::ImplGetResources_RenderThread(const EDisplayC
 	{
 		if (InRecursionDepth < DisplayClusterViewportProxyResourcesOverrideRecursionDepthMax)
 		{
-			if (FDisplayClusterViewportProxy const* OverrideViewportProxy = Owner->ImplFindViewport_RenderThread(RenderSettings.OverrideViewportId))
+			if (FDisplayClusterViewportProxy const* OverrideViewportProxy = Owner.ImplFindViewport_RenderThread(RenderSettings.OverrideViewportId))
 			{
 				return OverrideViewportProxy->ImplGetResources_RenderThread(InResourceType, OutResources, InRecursionDepth + 1);
 			}
@@ -326,7 +320,7 @@ void FDisplayClusterViewportProxy::PostResolveViewport_RenderThread(FRHICommandL
 void FDisplayClusterViewportProxy::ImplViewportRemap_RenderThread(FRHICommandListImmediate& RHICmdList) const
 {
 #if WITH_EDITOR
-	switch(Owner->GetRenderFrameSettings_RenderThread().RenderMode)
+	switch(Owner.GetRenderFrameSettings_RenderThread().RenderMode)
 	{
 	case EDisplayClusterRenderFrameMode::PreviewInScene:
 		// Preview in editor not support this feature
@@ -375,7 +369,7 @@ void FDisplayClusterViewportProxy::ImplPreviewReadPixels_RenderThread(FRHIComman
 	// Now try to read until success:
 	if (bPreviewReadPixels && Contexts.Num() > 0 && OutputPreviewTargetableResource.IsValid())
 	{
-		check(Owner->GetRenderFrameSettings_RenderThread().RenderMode == EDisplayClusterRenderFrameMode::PreviewInScene);
+		check(Owner.GetRenderFrameSettings_RenderThread().RenderMode == EDisplayClusterRenderFrameMode::PreviewInScene);
 
 		// We should synchronize thread for preview read
 		FScopeLock Lock(&PreviewPixelsCSGuard);
@@ -433,7 +427,7 @@ EDisplayClusterViewportResourceType FDisplayClusterViewportProxy::GetOutputResou
 	check(IsInRenderingThread());
 
 #if WITH_EDITOR
-	switch(Owner->GetRenderFrameSettings_RenderThread().RenderMode)
+	switch(Owner.GetRenderFrameSettings_RenderThread().RenderMode)
 	{
 	case EDisplayClusterRenderFrameMode::PreviewInScene:
 		return EDisplayClusterViewportResourceType::OutputPreviewTargetableResource;
@@ -470,7 +464,7 @@ bool FDisplayClusterViewportProxy::ImplGetResourcesWithRects_RenderThread(const
 	{
 		if (InRecursionDepth < DisplayClusterViewportProxyResourcesOverrideRecursionDepthMax)
 		{
-			if (FDisplayClusterViewportProxy const* OverrideViewportProxy = Owner->ImplFindViewport_RenderThread(RenderSettings.OverrideViewportId))
+			if (FDisplayClusterViewportProxy const* OverrideViewportProxy = Owner.ImplFindViewport_RenderThread(RenderSettings.OverrideViewportId))
 			{
 				return OverrideViewportProxy->ImplGetResourcesWithRects_RenderThread(InResourceType, OutResources, OutResourceRects, InRecursionDepth + 1);
 			}
@@ -577,14 +571,14 @@ void FDisplayClusterViewportProxy::UpdateDeferredResources(FRHICommandListImmedi
 }
 
 template<class TScreenPixelShader>
-void ResampleCopyTextureImpl_RenderThread(FRHICommandListImmediate& RHICmdList, FRHITexture* SrcTexture, FRHITexture* DstTexture, const FIntRect& SrcRect, const FIntRect& DstRect, const bool bCopyOnlyAlphaChannel)
+void ResampleCopyTextureImpl_RenderThread(FRHICommandListImmediate& RHICmdList, FRHITexture* SrcTexture, FRHITexture* DstTexture, const FIntRect& SrcRect, const FIntRect& DstRect)
 {
 	// Texture format mismatch, use a shader to do the copy.
 	// #todo-renderpasses there's no explicit resolve here? Do we need one?
 	FRHIRenderPassInfo RPInfo(DstTexture, ERenderTargetActions::Load_Store);
 	RHICmdList.Transition(FRHITransitionInfo(DstTexture, ERHIAccess::Unknown, ERHIAccess::RTV));
 
-	RHICmdList.BeginRenderPass(RPInfo, TEXT("DisplayCluster_ResampleCopyTexture"));
+	RHICmdList.BeginRenderPass(RPInfo, TEXT("DisaplyClusterRender_ResampleTexture"));
 	{
 		FIntVector SrcSizeXYZ = SrcTexture->GetSizeXYZ();
 		FIntVector DstSizeXYZ = DstTexture->GetSizeXYZ();
@@ -597,23 +591,11 @@ void ResampleCopyTextureImpl_RenderThread(FRHICommandListImmediate& RHICmdList,
 		FGraphicsPipelineStateInitializer GraphicsPSOInit;
 		RHICmdList.ApplyCachedRenderTargets(GraphicsPSOInit);
 
-		if(bCopyOnlyAlphaChannel)
-		{
-			// Copy alpha channel from source to dest
-			// RT0ColorWriteMask,RT0ColorBlendOp,RT0ColorSrcBlend,RT0ColorDestBlend,RT0AlphaBlendOp,RT0AlphaSrcBlend,RT0AlphaDestBlend,
-			GraphicsPSOInit.BlendState = TStaticBlendState <CW_ALPHA, BO_Add, BF_One, BF_Zero, BO_Add, BF_One, BF_Zero>::GetRHI();
-		}
-		else
-		{
-			GraphicsPSOInit.BlendState = TStaticBlendState<>::GetRHI();
-		}
-
+		GraphicsPSOInit.BlendState = TStaticBlendState<>::GetRHI();
 		GraphicsPSOInit.RasterizerState = TStaticRasterizerState<>::GetRHI();
 		GraphicsPSOInit.DepthStencilState = TStaticDepthStencilState<false, CF_Always>::GetRHI();
 
 
-
-
 		FGlobalShaderMap* ShaderMap = GetGlobalShaderMap(GMaxRHIFeatureLevel);
 		TShaderMapRef<FScreenVS> VertexShader(ShaderMap);
 		TShaderMapRef<TScreenPixelShader> PixelShader(ShaderMap);
@@ -651,11 +633,6 @@ void ResampleCopyTextureImpl_RenderThread(FRHICommandListImmediate& RHICmdList,
 	RHICmdList.Transition(FRHITransitionInfo(DstTexture, ERHIAccess::Unknown, ERHIAccess::SRVMask));
 }
 
-void ImplResolveResourceAlpha(FRHICommandListImmediate& RHICmdList, FRHITexture2D* InputResource, const FIntRect& InputRect, FRHITexture2D* OutputResource, const FIntRect& OutputRect)
-{
-	ResampleCopyTextureImpl_RenderThread<FScreenPS>(RHICmdList, InputResource, OutputResource, InputRect, OutputRect, true);
-}
-
 void ImplResolveResource(FRHICommandListImmediate& RHICmdList, FRHITexture2D* InputResource, const FIntRect& InputRect, FRHITexture2D* OutputResource, const FIntRect& OutputRect, const bool bOutputIsMipsResource, const bool bOutputIsPreviewResource)
 {
 	check(InputResource);
@@ -664,7 +641,7 @@ void ImplResolveResource(FRHICommandListImmediate& RHICmdList, FRHITexture2D* In
 	if (bOutputIsPreviewResource)
 	{
 		// Preview require a normal alpha (re-invert)
-		ResampleCopyTextureImpl_RenderThread<FScreenPSInvertAlpha>(RHICmdList, InputResource, OutputResource, InputRect, OutputRect, false);
+		ResampleCopyTextureImpl_RenderThread<FScreenPSInvertAlpha>(RHICmdList, InputResource, OutputResource, InputRect, OutputRect);
 	}
 	else if (InputRect.Size() == OutputRect.Size() && InputResource->GetFormat() == OutputResource->GetFormat())
 	{
@@ -679,36 +656,17 @@ void ImplResolveResource(FRHICommandListImmediate& RHICmdList, FRHITexture2D* In
 	}
 	else
 	{
-		ResampleCopyTextureImpl_RenderThread<FScreenPS>(RHICmdList, InputResource, OutputResource, InputRect, OutputRect, false);
-	}
-}
-
-FIntRect FDisplayClusterViewportProxy::GetFinalContextRect(const EDisplayClusterViewportResourceType InputResourceType, const FIntRect& InRect) const
-{
-	// When resolving without warp, apply overscan
-	switch (InputResourceType)
-	{
-	case EDisplayClusterViewportResourceType::InternalRenderTargetResource:
-		if (OverscanSettings.bIsEnabled && CVarDisplayClusterRenderOverscanResolve.GetValueOnRenderThread() != 0)
-		{
-			// Support overscan crop
-			return OverscanSettings.OverscanPixels.GetInnerRect(InRect);
-		}
-		break;
-	default:
-		break;
+		ResampleCopyTextureImpl_RenderThread<FScreenPS>(RHICmdList, InputResource, OutputResource, InputRect, OutputRect);
 	}
-
-	return InRect;
 }
 
 // Resolve resource contexts
-bool FDisplayClusterViewportProxy::ResolveResources_RenderThread(FRHICommandListImmediate& RHICmdList, const EDisplayClusterViewportResourceType InputResourceType, const EDisplayClusterViewportResourceType OutputResourceType, const int32 InContextNum) const
+bool FDisplayClusterViewportProxy::ResolveResources_RenderThread(FRHICommandListImmediate& RHICmdList, const EDisplayClusterViewportResourceType InputResourceType, const EDisplayClusterViewportResourceType OutputResourceType) const
 {
-	return ImplResolveResources_RenderThread(RHICmdList, this, InputResourceType, OutputResourceType, InContextNum);
+	return ImplResolveResources_RenderThread(RHICmdList, this, InputResourceType, OutputResourceType);
 }
 
-bool FDisplayClusterViewportProxy::ImplResolveResources_RenderThread(FRHICommandListImmediate& RHICmdList, FDisplayClusterViewportProxy const* SourceProxy, const EDisplayClusterViewportResourceType InputResourceType, const EDisplayClusterViewportResourceType OutputResourceType, const int32 InContextNum) const
+bool FDisplayClusterViewportProxy::ImplResolveResources_RenderThread(FRHICommandListImmediate& RHICmdList, FDisplayClusterViewportProxy const* SourceProxy, const EDisplayClusterViewportResourceType InputResourceType, const EDisplayClusterViewportResourceType OutputResourceType) const
 {
 	check(IsInRenderingThread());
 
@@ -726,163 +684,46 @@ bool FDisplayClusterViewportProxy::ImplResolveResources_RenderThread(FRHICommand
 	const bool bOutputIsPreviewResource = false;
 #endif
 
-	TArray<FRHITexture2D*> SrcResources, DestResources;
-	TArray<FIntRect> SrcResourcesRect, DestResourcesRect;
-	if (SourceProxy->GetResourcesWithRects_RenderThread(InputResourceType, SrcResources, SrcResourcesRect) && GetResourcesWithRects_RenderThread(OutputResourceType, DestResources, DestResourcesRect))
+	TArray<FRHITexture2D*> InputResources, OutputResources;
+	TArray<FIntRect> InputResourcesRect, OutputResourcesRect;
+	if (SourceProxy->GetResourcesWithRects_RenderThread(InputResourceType, InputResources, InputResourcesRect) && GetResourcesWithRects_RenderThread(OutputResourceType, OutputResources, OutputResourcesRect))
 	{
-		const int32 SrcResourcesAmount = FMath::Min(SrcResources.Num(), DestResources.Num());
-		for (int32 SrcResourceContextIndex = 0; SrcResourceContextIndex < SrcResourcesAmount; SrcResourceContextIndex++)
+		const int32 InputAmount = FMath::Min(InputResources.Num(), OutputResources.Num());
+		for (int32 InputIt = 0; InputIt < InputAmount; InputIt++)
 		{
-			// Get context num to copy
-			const int32 SrcContextNum = (InContextNum == INDEX_NONE) ? SrcResourceContextIndex : InContextNum;
-			const FIntRect SrcRect = GetFinalContextRect(InputResourceType, SrcResourcesRect[SrcContextNum]);
+			FIntRect SourceRect = InputResourcesRect[InputIt];
 
-			if ((SrcContextNum + 1) == SrcResourcesAmount)
+			// When resolving without warp, apply overscan
+			switch (InputResourceType)
 			{
-				// last input mono -> stereo outputs
-				for (int32 DestResourceContextIndex = SrcContextNum; DestResourceContextIndex < DestResources.Num(); DestResourceContextIndex++)
+			case EDisplayClusterViewportResourceType::InternalRenderTargetResource:
+				if (OverscanSettings.bIsEnabled && CVarDisplayClusterRenderOverscanResolve.GetValueOnRenderThread() != 0)
 				{
-					const FIntRect DestRect = GetFinalContextRect(OutputResourceType, DestResourcesRect[DestResourceContextIndex]);
-					ImplResolveResource(RHICmdList, SrcResources[SrcContextNum], SrcRect, DestResources[DestResourceContextIndex], DestRect, bOutputIsMipsResource, bOutputIsPreviewResource);
+					// Support overscan crop
+					SourceRect = OverscanSettings.OverscanPixels.GetInnerRect(SourceRect);
 				}
 				break;
-			}
-			else
-			{
-				const FIntRect DestRect = GetFinalContextRect(OutputResourceType, DestResourcesRect[SrcContextNum]);
-				ImplResolveResource(RHICmdList, SrcResources[SrcContextNum], SrcRect, DestResources[SrcContextNum], DestRect, bOutputIsMipsResource, bOutputIsPreviewResource);
-			}
-
-			if (InContextNum != INDEX_NONE)
-			{
-				// Copy only one texture
+			default:
 				break;
 			}
-		}
-
-		return true;
-	}
-
-	return false;
-}
-
-bool FDisplayClusterViewportProxy::ShouldCopyAlphaChannel() const
-{
-	if ((RenderSettingsICVFX.RuntimeFlags & (ViewportRuntime_ICVFXLightcard | ViewportRuntime_ICVFXChromakey)) != 0)
-	{
-		return true;
-	}
 
-	return false;
-}
-
-void FDisplayClusterViewportProxy::CopyRenderTargetAlphaChannel_RehderThread(FRHICommandListImmediate& RHICmdList, const int32 InContextNum, FRHITexture* InSrcTexture, bool bInvertCopyDirection)
-{
-	if (InputShaderResources.IsValidIndex(InContextNum) && InputShaderResources[InContextNum] != nullptr && InSrcTexture)
-	{
-		if(FRHITexture* DestTexture = InputShaderResources[InContextNum]->GetTextureRHI())
-		{
-			const FIntRect SrcRect = GetFinalContextRect(EDisplayClusterViewportResourceType::InternalRenderTargetResource, Contexts[InContextNum].RenderTargetRect);
-			const FIntRect DestRect = FIntRect(FIntPoint(0, 0), DestTexture->GetSizeXY());
-
-			if (bInvertCopyDirection)
+			if ((InputIt + 1) == InputAmount)
 			{
-
-				ImplResolveResourceAlpha(RHICmdList, DestTexture, DestRect, InSrcTexture, SrcRect);
+				// last input mono -> stereo outputs
+				for (int32 OutputIt = InputIt; OutputIt < OutputResources.Num(); OutputIt++)
+				{
+					ImplResolveResource(RHICmdList, InputResources[InputIt], SourceRect, OutputResources[OutputIt], OutputResourcesRect[OutputIt], bOutputIsMipsResource, bOutputIsPreviewResource);
+				}
+				break;
 			}
 			else
 			{
-				ImplResolveResourceAlpha(RHICmdList, InSrcTexture, SrcRect, DestTexture, DestRect);
-				SceneColorResourcesState |= (1 << InContextNum);
+				ImplResolveResource(RHICmdList, InputResources[InputIt], SourceRect, OutputResources[InputIt], OutputResourcesRect[InputIt], bOutputIsMipsResource, bOutputIsPreviewResource);
 			}
 		}
-	}
-}
 
-BEGIN_SHADER_PARAMETER_STRUCT(FDisplayClusterCopyTextureParameters, )
-RDG_TEXTURE_ACCESS(SrcTexture, ERHIAccess::CopySrc)
-END_SHADER_PARAMETER_STRUCT()
-
-void FDisplayClusterViewportProxy::OnResolvedSceneColor_RenderThread(FRDGBuilder& GraphBuilder, const FSceneTextures& SceneTextures, const FDisplayClusterViewportProxy_Context& InProxyContext)
-{
-	const uint32 InContextNum = InProxyContext.ContextNum;
-
-	check(InContextNum < sizeof(SceneColorResourcesState) * 8);
-
-	SceneColorResourcesState &= ~(1 << InContextNum);
-
-	if (ShouldCopyAlphaChannel())
-	{
-		// Copy scene color with alpha channel to InputShaderResource texture
-		FRDGTextureRef InTextureRef = SceneTextures.Color.Resolve;
-		if (HasBeenProduced(InTextureRef))
-		{
-			check(InputShaderResources.IsValidIndex(InContextNum));
-
-			FDisplayClusterCopyTextureParameters* PassParameters = GraphBuilder.AllocParameters<FDisplayClusterCopyTextureParameters>();
-			PassParameters->SrcTexture = InTextureRef;
-
-			GraphBuilder.AddPass(
-				RDG_EVENT_NAME("DisplayCluster_SaveViewportAlphaChannel_%s", *GetId()),
-				PassParameters,
-				ERDGPassFlags::Copy | ERDGPassFlags::NeverCull,
-				[ViewportProxy = SharedThis(this), InTextureRef, InContextNum](FRHICommandListImmediate& RHICmdList)
-				{
-					ViewportProxy->CopyRenderTargetAlphaChannel_RehderThread(RHICmdList, InContextNum, InTextureRef->GetRHI());
-				});
-		}
-	}
-}
-
-BEGIN_SHADER_PARAMETER_STRUCT(FDisplayClusterEmptyParameters, )
-END_SHADER_PARAMETER_STRUCT()
-
-void FDisplayClusterViewportProxy::PostRenderViewFamily_RenderThread(FRDGBuilder& GraphBuilder, FSceneViewFamily& InViewFamily, const FSceneView& InSceneView, const FDisplayClusterViewportProxy_Context& InProxyContext)
-{
-	const uint32 InContextNum = InProxyContext.ContextNum;
-
-	check(InContextNum < sizeof(SceneColorResourcesState) * 8);
-
-	// Save GPUMask for each viewport context
-	if (Contexts.IsValidIndex(InContextNum))
-	{
-		if (GPUMask.Num() < Contexts.Num())
-		{
-			GPUMask.AddDefaulted(Contexts.Num() - GPUMask.Num());
-		}
-
-		checkSlow(InSceneView.bIsViewInfo);
-		const FViewInfo& ViewInfo = static_cast<const FViewInfo&>(InSceneView);
-
-		GPUMask[InContextNum] = ViewInfo.GPUMask;
-	}
-
-	if (ShouldCopyAlphaChannel())
-	{
-		FTextureRHIRef DestRenderTarget = InViewFamily.RenderTarget->GetRenderTargetTexture();
-		
-		FDisplayClusterEmptyParameters* PassParameters = GraphBuilder.AllocParameters<FDisplayClusterEmptyParameters>();
-		GraphBuilder.AddPass(
-			RDG_EVENT_NAME("DisplayCluster_LoadViewportAlphaChannel_%s", *GetId()),
-			PassParameters,
-			ERDGPassFlags::Copy | ERDGPassFlags::NeverCull,
-			[ViewportProxy = SharedThis(this), DestRenderTarget, InContextNum](FRHICommandListImmediate& RHICmdList)
-			{
-				// Copy alpha channel back
-				if ((ViewportProxy->SceneColorResourcesState & (1 << InContextNum)) != 0)
-				{
-					ViewportProxy->CopyRenderTargetAlphaChannel_RehderThread(RHICmdList, InContextNum, DestRenderTarget.GetReference(), true);
-				}
-			});
+		return true;
 	}
 
-	if (!InProxyContext.ViewFamilyProfileDescription.IsEmpty())
-	{
-		static IDisplayClusterCallbacks& DCCallbacksAPI = IDisplayCluster::Get().GetCallbacks();
-		if (DCCallbacksAPI.OnDisplayClusterPostRenderViewFamily_RenderThread().IsBound())
-		{
-			// Now we can perform viewport notification
-			DCCallbacksAPI.OnDisplayClusterPostRenderViewFamily_RenderThread().Broadcast(GraphBuilder, InViewFamily, this);
-		}
-	}
+	return false;
 }
diff --git a/Engine/Plugins/Runtime/nDisplay/Source/DisplayCluster/Private/Render/Viewport/DisplayClusterViewportProxy.h b/Engine/Plugins/Runtime/nDisplay/Source/DisplayCluster/Private/Render/Viewport/DisplayClusterViewportProxy.h
--- a/Engine/Plugins/Runtime/nDisplay/Source/DisplayCluster/Private/Render/Viewport/DisplayClusterViewportProxy.h
+++ b/Engine/Plugins/Runtime/nDisplay/Source/DisplayCluster/Private/Render/Viewport/DisplayClusterViewportProxy.h
@@ -2,13 +2,6 @@
 
 #pragma once
 
-#include "CoreMinimal.h"
-
-#include "EngineUtils.h"
-#include "ScreenRendering.h"
-#include "SceneView.h"
-#include "Templates/SharedPointer.h"
-
 #include "Render/Viewport/IDisplayClusterViewportProxy.h"
 #include "Render/Viewport/Containers/DisplayClusterViewport_OverscanSettings.h"
 
@@ -25,10 +18,9 @@ class FDisplayClusterViewportReadPixelsData;
 
 class FDisplayClusterViewportProxy
 	: public IDisplayClusterViewportProxy
-	, public TSharedFromThis<FDisplayClusterViewportProxy, ESPMode::ThreadSafe>
 {
 public:
-	FDisplayClusterViewportProxy(const FDisplayClusterViewport& RenderViewport);
+	FDisplayClusterViewportProxy(const FDisplayClusterViewportManagerProxy& Owner, const FDisplayClusterViewport& RenderViewport);
 	virtual ~FDisplayClusterViewportProxy();
 
 public:
@@ -98,14 +90,12 @@ public:
 	virtual bool GetResourcesWithRects_RenderThread(const EDisplayClusterViewportResourceType InResourceType, TArray<FRHITexture2D*>& OutResources, TArray<FIntRect>& OutRects) const override;
 
 	// Resolve resource contexts
-	virtual bool ResolveResources_RenderThread(FRHICommandListImmediate& RHICmdList, const EDisplayClusterViewportResourceType InputResourceType, const EDisplayClusterViewportResourceType OutputResourceType, const int32 InContextNum = INDEX_NONE) const override;
+	virtual bool ResolveResources_RenderThread(FRHICommandListImmediate& RHICmdList, const EDisplayClusterViewportResourceType InputResourceType, const EDisplayClusterViewportResourceType OutputResourceType) const override;
 
 	virtual EDisplayClusterViewportResourceType GetOutputResourceType_RenderThread() const override;
 
 	virtual const IDisplayClusterViewportManagerProxy& GetOwner_RenderThread() const override;
 
-	virtual void OnResolvedSceneColor_RenderThread(FRDGBuilder& GraphBuilder, const FSceneTextures& SceneTextures, const FDisplayClusterViewportProxy_Context& InProxyContext) override;
-	virtual void PostRenderViewFamily_RenderThread(FRDGBuilder& InGraphBuilder, class FSceneViewFamily& InViewFamily, const class FSceneView& InSceneView, const FDisplayClusterViewportProxy_Context& InProxyContext) override;
 
 	///////////////////////////////
 	// ~IDisplayClusterViewportProxy
@@ -116,8 +106,6 @@ public:
 
 	void PostResolveViewport_RenderThread(FRHICommandListImmediate& RHICmdList) const;
 
-	bool ShouldCopyAlphaChannel() const;
-
 #if WITH_EDITOR
 	bool GetPreviewPixels_GameThread(TSharedPtr<FDisplayClusterViewportReadPixelsData, ESPMode::ThreadSafe>& OutPixelsData) const;
 #endif
@@ -142,8 +130,6 @@ public:
 		return false;
 	}
 
-	FIntRect GetFinalContextRect(const EDisplayClusterViewportResourceType InputResourceType, const FIntRect& InRect) const;
-
 private:
 	bool ImplGetResourcesWithRects_RenderThread(const EDisplayClusterViewportResourceType InResourceType, TArray<FRHITexture2D*>& OutResources, TArray<FIntRect>& OutResourceRects, const int32 InRecursionDepth) const;
 	bool ImplGetResources_RenderThread(const EDisplayClusterViewportResourceType InResourceType, TArray<FRHITexture2D*>& OutResources, const int32 InRecursionDepth) const;
@@ -151,13 +137,10 @@ private:
 	void ImplViewportRemap_RenderThread(FRHICommandListImmediate& RHICmdList) const;
 	void ImplPreviewReadPixels_RenderThread(FRHICommandListImmediate& RHICmdList) const;
 
-	bool ImplResolveResources_RenderThread(FRHICommandListImmediate& RHICmdList, FDisplayClusterViewportProxy const* SourceProxy, const EDisplayClusterViewportResourceType InputResourceType, const EDisplayClusterViewportResourceType OutputResourceType, const int32 InContextNum) const;
+	bool ImplResolveResources_RenderThread(FRHICommandListImmediate& RHICmdList, FDisplayClusterViewportProxy const* SourceProxy, const EDisplayClusterViewportResourceType InputResourceType, const EDisplayClusterViewportResourceType OutputResourceType) const;
 
 	bool IsShouldOverrideViewportResource(const EDisplayClusterViewportResourceType InResourceType) const;
 
-	// The InputShaderResources[] uses as temporary storage for alpha channel data
-	void CopyRenderTargetAlphaChannel_RehderThread(FRHICommandListImmediate& RHICmdList, const int32 InContextNum, FRHITexture* InSrcTexture, bool bInvertCopyDirection = false);
-
 protected:
 	friend FDisplayClusterViewportProxyData;
 	friend FDisplayClusterViewportManagerProxy;
@@ -183,10 +166,6 @@ protected:
 	// Viewport contexts (left/center/right eyes)
 	mutable TArray<FDisplayClusterViewport_Context> Contexts;
 
-	/** The GPU nodes on which to render this view. */
-	TArray<FRHIGPUMask> GPUMask;
-	uint32 SceneColorResourcesState = 0;
-
 	// View family render to this resources
 	TArray<FDisplayClusterViewportRenderTargetResource*> RenderTargets;
 
@@ -207,7 +186,7 @@ protected:
 	TArray<FDisplayClusterViewportTextureResource*> AdditionalTargetableResources;
 	TArray<FDisplayClusterViewportTextureResource*> MipsShaderResources;
 
-	const TSharedRef<FDisplayClusterViewportManagerProxy, ESPMode::ThreadSafe> Owner;
+	const FDisplayClusterViewportManagerProxy& Owner;
 	IDisplayClusterShaders& ShadersAPI;
 };
 
diff --git a/Engine/Plugins/Runtime/nDisplay/Source/DisplayCluster/Private/Render/Viewport/Postprocess/DisplayClusterViewportPostProcessManager.cpp b/Engine/Plugins/Runtime/nDisplay/Source/DisplayCluster/Private/Render/Viewport/Postprocess/DisplayClusterViewportPostProcessManager.cpp
--- a/Engine/Plugins/Runtime/nDisplay/Source/DisplayCluster/Private/Render/Viewport/Postprocess/DisplayClusterViewportPostProcessManager.cpp
+++ b/Engine/Plugins/Runtime/nDisplay/Source/DisplayCluster/Private/Render/Viewport/Postprocess/DisplayClusterViewportPostProcessManager.cpp
@@ -369,7 +369,7 @@ void FDisplayClusterViewportPostProcessManager::FinalizeNewFrame()
 	check(IsInGameThread());
 
 	ENQUEUE_RENDER_COMMAND(DisplayClusterViewportPostProcessManager_FinalizeNewFrame)(
-		[PostprocessManager = SharedThis(this), PostprocessData = Postprocess](FRHICommandListImmediate& RHICmdList)
+		[PostprocessManager = this, PostprocessData = Postprocess](FRHICommandListImmediate& RHICmdList)
 	{
 		PostprocessManager->PostprocessProxy.Empty();
 		PostprocessManager->PostprocessProxy.Append(PostprocessData);
@@ -485,12 +485,12 @@ void FDisplayClusterViewportPostProcessManager::ImplPerformPostProcessViewBefore
 				// send nullptr viewport first to handle begin
 				It->PerformPostProcessViewBeforeWarpBlend_RenderThread(RHICmdList, nullptr);
 
-				for (TSharedPtr<IDisplayClusterViewportProxy, ESPMode::ThreadSafe>& ViewportProxyIt : InViewportManagerProxy->GetViewports_RenderThread())
+				for (IDisplayClusterViewportProxy* ViewportProxyIt : InViewportManagerProxy->GetViewports_RenderThread())
 				{
-					if (ViewportProxyIt.IsValid())
+					if (ViewportProxyIt)
 					{
 						UE_LOG(LogDisplayClusterRender, VeryVerbose, TEXT("Postprocess VIEW before WarpBlend - Viewport '%s'"), *ViewportProxyIt->GetId());
-						It->PerformPostProcessViewBeforeWarpBlend_RenderThread(RHICmdList, ViewportProxyIt.Get());
+						It->PerformPostProcessViewBeforeWarpBlend_RenderThread(RHICmdList, ViewportProxyIt);
 					}
 				}
 			}
@@ -515,12 +515,12 @@ void FDisplayClusterViewportPostProcessManager::ImplPerformPostProcessViewAfterW
 				// send nullptr viewport first to handle begin
 				It->PerformPostProcessViewAfterWarpBlend_RenderThread(RHICmdList, nullptr);
 
-				for (const TSharedPtr<IDisplayClusterViewportProxy, ESPMode::ThreadSafe>& ViewportProxyIt : InViewportManagerProxy->GetViewports_RenderThread())
+				for (IDisplayClusterViewportProxy* ViewportProxyIt : InViewportManagerProxy->GetViewports_RenderThread())
 				{
-					if (ViewportProxyIt.IsValid())
+					if (ViewportProxyIt)
 					{
 						UE_LOG(LogDisplayClusterRender, VeryVerbose, TEXT("Postprocess VIEW after WarpBlend - Viewport '%s'"), *ViewportProxyIt->GetId());
-						It->PerformPostProcessViewAfterWarpBlend_RenderThread(RHICmdList, ViewportProxyIt.Get());
+						It->PerformPostProcessViewAfterWarpBlend_RenderThread(RHICmdList, ViewportProxyIt);
 					}
 				}
 			}
diff --git a/Engine/Plugins/Runtime/nDisplay/Source/DisplayCluster/Private/Render/Viewport/Postprocess/DisplayClusterViewportPostProcessManager.h b/Engine/Plugins/Runtime/nDisplay/Source/DisplayCluster/Private/Render/Viewport/Postprocess/DisplayClusterViewportPostProcessManager.h
--- a/Engine/Plugins/Runtime/nDisplay/Source/DisplayCluster/Private/Render/Viewport/Postprocess/DisplayClusterViewportPostProcessManager.h
+++ b/Engine/Plugins/Runtime/nDisplay/Source/DisplayCluster/Private/Render/Viewport/Postprocess/DisplayClusterViewportPostProcessManager.h
@@ -4,7 +4,6 @@
 
 #include "Render/PostProcess/IDisplayClusterPostProcess.h"
 #include "Render/IDisplayClusterRenderManager.h"
-#include "Templates/SharedPointer.h"
 
 class FDisplayClusterViewportManager;
 class FDisplayClusterViewportManagerProxy;
@@ -16,7 +15,6 @@ struct FDisplayClusterConfigurationPostprocess;
  */
 
 class FDisplayClusterViewportPostProcessManager
-	: public TSharedFromThis<FDisplayClusterViewportPostProcessManager, ESPMode::ThreadSafe>
 {
 public:
 	FDisplayClusterViewportPostProcessManager(FDisplayClusterViewportManager& InViewportManager);
diff --git a/Engine/Plugins/Runtime/nDisplay/Source/DisplayCluster/Private/Render/Viewport/RenderTarget/DisplayClusterRenderTargetResourcesPool.cpp b/Engine/Plugins/Runtime/nDisplay/Source/DisplayCluster/Private/Render/Viewport/RenderTarget/DisplayClusterRenderTargetResourcesPool.cpp
--- a/Engine/Plugins/Runtime/nDisplay/Source/DisplayCluster/Private/Render/Viewport/RenderTarget/DisplayClusterRenderTargetResourcesPool.cpp
+++ b/Engine/Plugins/Runtime/nDisplay/Source/DisplayCluster/Private/Render/Viewport/RenderTarget/DisplayClusterRenderTargetResourcesPool.cpp
@@ -28,19 +28,19 @@ namespace DisplayClusterRenderTargetResourcesPool
 
 	static void ImplReleaseViewportResourcesRHI(FDisplayClusterViewportManagerProxy* InViewportManagerProxy, const TArrayView<FDisplayClusterViewportResource*>& InResources)
 	{
-		if (InViewportManagerProxy != nullptr && !InResources.IsEmpty())
-		{
-			TArray<FDisplayClusterViewportResource*> ResourcesForRenderThread(InResources);
+		TArray<FDisplayClusterViewportResource*> ResourcesForRenderThread(InResources);
 
-			ENQUEUE_RENDER_COMMAND(DisplayCluster_ReleaseViewportResourcesRHI)(
-				[InViewportManagerProxy = InViewportManagerProxy->AsShared(), ReleasedResources = std::move(ResourcesForRenderThread)](FRHICommandListImmediate& RHICmdList)
+		ENQUEUE_RENDER_COMMAND(DisplayCluster_ReleaseViewportResourcesRHI)(
+			[InViewportManagerProxy, ReleasedResources = std::move(ResourcesForRenderThread)](FRHICommandListImmediate& RHICmdList)
+		{
+			for (FDisplayClusterViewportResource* DeletedResourcePtrIt : ReleasedResources)
+			{
+				if (InViewportManagerProxy)
 				{
-					for (FDisplayClusterViewportResource* DeletedResourcePtrIt : ReleasedResources)
-					{
-						InViewportManagerProxy->DeleteResource_RenderThread(DeletedResourcePtrIt);
-					}
-				});
-		}
+					InViewportManagerProxy->DeleteResource_RenderThread(DeletedResourcePtrIt);
+				}
+			}
+		});
 	}
 
 	static FDisplayClusterViewportResourceSettings* ImplCreateViewportResourceSettings(const FDisplayClusterRenderFrameSettings& InRenderFrameSettings, FViewport* InViewport)
diff --git a/Engine/Plugins/Runtime/nDisplay/Source/DisplayCluster/Public/Render/Viewport/Containers/DisplayClusterViewportProxy_Context.h b/Engine/Plugins/Runtime/nDisplay/Source/DisplayCluster/Public/Render/Viewport/Containers/DisplayClusterViewportProxy_Context.h
deleted file mode 100644
diff --git a/Engine/Plugins/Runtime/nDisplay/Source/DisplayCluster/Public/Render/Viewport/Containers/DisplayClusterViewport_RenderSettings.h b/Engine/Plugins/Runtime/nDisplay/Source/DisplayCluster/Public/Render/Viewport/Containers/DisplayClusterViewport_RenderSettings.h
--- a/Engine/Plugins/Runtime/nDisplay/Source/DisplayCluster/Public/Render/Viewport/Containers/DisplayClusterViewport_RenderSettings.h
+++ b/Engine/Plugins/Runtime/nDisplay/Source/DisplayCluster/Public/Render/Viewport/Containers/DisplayClusterViewport_RenderSettings.h
@@ -14,6 +14,7 @@ enum class EDisplayClusterViewportCaptureMode : uint8
 
 	// use hi-res float texture with alpha for compisiting
 	Lightcard,
+	Lightcard_OCIO,
 
 	// Special hi-res mode for movie pipeline
 	MoviePipeline,
diff --git a/Engine/Plugins/Runtime/nDisplay/Source/DisplayCluster/Public/Render/Viewport/Containers/DisplayClusterViewport_RenderSettingsICVFX.h b/Engine/Plugins/Runtime/nDisplay/Source/DisplayCluster/Public/Render/Viewport/Containers/DisplayClusterViewport_RenderSettingsICVFX.h
--- a/Engine/Plugins/Runtime/nDisplay/Source/DisplayCluster/Public/Render/Viewport/Containers/DisplayClusterViewport_RenderSettingsICVFX.h
+++ b/Engine/Plugins/Runtime/nDisplay/Source/DisplayCluster/Public/Render/Viewport/Containers/DisplayClusterViewport_RenderSettingsICVFX.h
@@ -39,6 +39,8 @@ enum EDisplayClusterViewportRuntimeICVFXFlags
 	ViewportRuntime_ICVFXIncamera       = 1 << 1,
 	ViewportRuntime_ICVFXChromakey      = 1 << 2,
 	ViewportRuntime_ICVFXLightcard      = 1 << 3,
+	ViewportRuntime_ICVFXLightcardAlpha = 1 << 4,
+	ViewportRuntime_ICVFXLightcardColor = 1 << 5,
 
 	// This viewport used as internal icvfx composing resource (created and deleted inside icvfx logic)
 	ViewportRuntime_InternalResource = 1 << 30,
diff --git a/Engine/Plugins/Runtime/nDisplay/Source/DisplayCluster/Public/Render/Viewport/IDisplayClusterViewportManagerProxy.h b/Engine/Plugins/Runtime/nDisplay/Source/DisplayCluster/Public/Render/Viewport/IDisplayClusterViewportManagerProxy.h
--- a/Engine/Plugins/Runtime/nDisplay/Source/DisplayCluster/Public/Render/Viewport/IDisplayClusterViewportManagerProxy.h
+++ b/Engine/Plugins/Runtime/nDisplay/Source/DisplayCluster/Public/Render/Viewport/IDisplayClusterViewportManagerProxy.h
@@ -41,7 +41,7 @@ public:
 	*
 	* @return - arrays with viewport render thread proxy objects refs
 	*/
-	virtual const TArrayView<TSharedPtr<IDisplayClusterViewportProxy, ESPMode::ThreadSafe>> GetViewports_RenderThread() const = 0;
+	virtual const TArrayView<IDisplayClusterViewportProxy*> GetViewports_RenderThread() const = 0;
 
 	/**
 	* Return render frame targets for current frame
diff --git a/Engine/Plugins/Runtime/nDisplay/Source/DisplayCluster/Public/Render/Viewport/IDisplayClusterViewportProxy.h b/Engine/Plugins/Runtime/nDisplay/Source/DisplayCluster/Public/Render/Viewport/IDisplayClusterViewportProxy.h
--- a/Engine/Plugins/Runtime/nDisplay/Source/DisplayCluster/Public/Render/Viewport/IDisplayClusterViewportProxy.h
+++ b/Engine/Plugins/Runtime/nDisplay/Source/DisplayCluster/Public/Render/Viewport/IDisplayClusterViewportProxy.h
@@ -13,7 +13,7 @@
 #include "Render/Viewport/Containers/DisplayClusterViewport_RenderSettings.h"
 #include "Render/Viewport/Containers/DisplayClusterViewport_RenderSettingsICVFX.h"
 #include "Render/Viewport/Containers/DisplayClusterViewport_PostRenderSettings.h"
-#include "Render/Viewport/Containers/DisplayClusterViewportProxy_Context.h"
+
 
 class DISPLAYCLUSTER_API IDisplayClusterViewportProxy
 {
@@ -37,7 +37,7 @@ public:
 	virtual bool GetResourcesWithRects_RenderThread(const EDisplayClusterViewportResourceType InResourceType, TArray<FRHITexture2D*>& OutResources, TArray<FIntRect>& OutRects) const = 0;
 
 	// Resolve resource contexts
-	virtual bool ResolveResources_RenderThread(FRHICommandListImmediate& RHICmdList, const EDisplayClusterViewportResourceType InputResourceType, const EDisplayClusterViewportResourceType OutputResourceType, const int32 InContextNum = INDEX_NONE) const = 0;
+	virtual bool ResolveResources_RenderThread(FRHICommandListImmediate& RHICmdList, const EDisplayClusterViewportResourceType InputResourceType, const EDisplayClusterViewportResourceType OutputResourceType) const = 0;
 
 	virtual EDisplayClusterViewportResourceType   GetOutputResourceType_RenderThread() const = 0;
 
@@ -45,9 +45,4 @@ public:
 
 	virtual void SetRenderSettings_RenderThread(const FDisplayClusterViewport_RenderSettings& InRenderSettings) const = 0;
 	virtual void SetContexts_RenderThread(const TArray<FDisplayClusterViewport_Context>& InContexts) const = 0;
-
-	virtual void OnResolvedSceneColor_RenderThread(class FRDGBuilder& GraphBuilder, const struct FSceneTextures& SceneTextures, const FDisplayClusterViewportProxy_Context& InProxyContext) = 0;
-
-	// Callback. After the viewfamily has finished rendering
-	virtual void PostRenderViewFamily_RenderThread(class FRDGBuilder& InGraphBuilder, class FSceneViewFamily& InViewFamily, const class FSceneView& InSceneView, const FDisplayClusterViewportProxy_Context& InProxyContext) = 0;
 };
diff --git a/Engine/Plugins/Runtime/nDisplay/Source/DisplayClusterMedia/Private/Synchronization/DisplayClusterFrameQueue.cpp b/Engine/Plugins/Runtime/nDisplay/Source/DisplayClusterMedia/Private/Synchronization/DisplayClusterFrameQueue.cpp
--- a/Engine/Plugins/Runtime/nDisplay/Source/DisplayClusterMedia/Private/Synchronization/DisplayClusterFrameQueue.cpp
+++ b/Engine/Plugins/Runtime/nDisplay/Source/DisplayClusterMedia/Private/Synchronization/DisplayClusterFrameQueue.cpp
@@ -108,18 +108,18 @@ void FDisplayClusterFrameQueue::HandleProcessLatency_RenderThread(FRHICommandLis
 			{
 				if (IDisplayClusterViewportManagerProxy* ViewportMgrProxy = ViewportMgr->GetProxy())
 				{
-					TArrayView<TSharedPtr<IDisplayClusterViewportProxy, ESPMode::ThreadSafe>> ViewportProxies = ViewportMgrProxy->GetViewports_RenderThread();
-					for (TSharedPtr<IDisplayClusterViewportProxy, ESPMode::ThreadSafe>& ViewportProxyIt : ViewportProxies)
+					TArrayView<IDisplayClusterViewportProxy*> ViewportProxies = ViewportMgrProxy->GetViewports_RenderThread();
+					for (IDisplayClusterViewportProxy* ViewportProxy : ViewportProxies)
 					{
 						// Handle textures of all viewports that don't have any media input attached. Those don't get rendered.
-						if (ViewportProxyIt.IsValid() && !ViewportProxyIt->GetRenderSettings_RenderThread().bSkipSceneRenderingButLeaveResourcesAvailable)
+						if (!ViewportProxy->GetRenderSettings_RenderThread().bSkipSceneRenderingButLeaveResourcesAvailable)
 						{
 							// Get viewport texture
 							TArray<FRHITexture*> Textures;
-							ViewportProxyIt->GetResources_RenderThread(EDisplayClusterViewportResourceType::InternalRenderTargetResource, Textures);
+							ViewportProxy->GetResources_RenderThread(EDisplayClusterViewportResourceType::InternalRenderTargetResource, Textures);
 							check(Textures.Num() > 0);
 
-							const FString ViewportId(ViewportProxyIt->GetId());
+							const FString ViewportId(ViewportProxy->GetId());
 
 							// Enqueue this view (head)
 							Frames[IdxHead].SaveView(RHICmdList, ViewportId, Textures[0]);
diff --git a/Engine/Plugins/Runtime/nDisplay/Source/DisplayClusterShaders/Private/Shaders/DisplayClusterShadersPreprocess_UVLightCards.cpp b/Engine/Plugins/Runtime/nDisplay/Source/DisplayClusterShaders/Private/Shaders/DisplayClusterShadersPreprocess_UVLightCards.cpp
--- a/Engine/Plugins/Runtime/nDisplay/Source/DisplayClusterShaders/Private/Shaders/DisplayClusterShadersPreprocess_UVLightCards.cpp
+++ b/Engine/Plugins/Runtime/nDisplay/Source/DisplayClusterShaders/Private/Shaders/DisplayClusterShadersPreprocess_UVLightCards.cpp
@@ -146,22 +146,6 @@ bool FDisplayClusterShadersPreprocess_UVLightCards::RenderPreprocess_UVLightCard
 
 	FEngineShowFlags EngineShowFlags(ESFIM_Game);
 
-	// LightCard settings from the FDisplayClusterViewportManager::ConfigureViewFamily
-	{
-		EngineShowFlags.PostProcessing = 0;
-		EngineShowFlags.SetAtmosphere(0);
-		EngineShowFlags.SetFog(0);
-		EngineShowFlags.SetVolumetricFog(0);
-		EngineShowFlags.SetMotionBlur(0); // motion blur doesn't work correctly with scene captures.
-		EngineShowFlags.SetSeparateTranslucency(0);
-		EngineShowFlags.SetHMDDistortion(0);
-		EngineShowFlags.SetOnScreenDebug(0);
-
-		EngineShowFlags.SetLumenReflections(0);
-		EngineShowFlags.SetLumenGlobalIllumination(0);
-		EngineShowFlags.SetGlobalIllumination(0);
-	}
-
 	FSceneViewFamilyContext ViewFamily(FSceneViewFamily::ConstructionValues(
 		InRenderTarget,
 		InScene,
@@ -209,7 +193,7 @@ bool FDisplayClusterShadersPreprocess_UVLightCards::RenderPreprocess_UVLightCard
 		);
 	}
 
-	ViewInitOptions.BackgroundColor = FLinearColor::Black;
+	ViewInitOptions.BackgroundColor = FLinearColor::Transparent;
 
 	GetRendererModule().CreateAndInitSingleView(RHICmdList, &ViewFamily, &ViewInitOptions);
 	const FSceneView* View = ViewFamily.Views[0];
@@ -255,4 +239,4 @@ bool FDisplayClusterShadersPreprocess_UVLightCards::RenderPreprocess_UVLightCard
 	GraphBuilder.Execute();
 
 	return true;
-}
+}
\ No newline at end of file
diff --git a/Engine/Plugins/Runtime/nDisplay/Source/DisplayClusterShaders/Private/Shaders/DisplayClusterShadersWarpblend_ICVFX.cpp b/Engine/Plugins/Runtime/nDisplay/Source/DisplayClusterShaders/Private/Shaders/DisplayClusterShadersWarpblend_ICVFX.cpp
--- a/Engine/Plugins/Runtime/nDisplay/Source/DisplayClusterShaders/Private/Shaders/DisplayClusterShadersWarpblend_ICVFX.cpp
+++ b/Engine/Plugins/Runtime/nDisplay/Source/DisplayClusterShaders/Private/Shaders/DisplayClusterShadersWarpblend_ICVFX.cpp
@@ -45,14 +45,6 @@ static TAutoConsoleVariable<int32> CVarIcvfxMPCDIShaderType(
 	ECVF_RenderThreadSafe
 );
 
-int32 GDisplayClusterShadersICVFXEnableLightCard = 1;
-static FAutoConsoleVariableRef CVarDisplayClusterShadersICVFXEnableLightCard(
-	TEXT("nDisplay.render.icvfx.lightcard"),
-	GDisplayClusterShadersICVFXEnableLightCard,
-	TEXT("Enable ICVFX lightcards render (0 - disable).\n"),
-	ECVF_RenderThreadSafe
-);
-
 enum class EIcvfxShaderType : uint8
 {
 	Passthrough,  // Viewport frustum (no warpblend, only frustum math)
@@ -329,7 +321,7 @@ private:
 			// 2 - second camera additive pass
 			// N - camera N
 			// N+1 - overlayOver (optional)
-			return (ICVFXParameters.IsLightcardOverUsed() && GDisplayClusterShadersICVFXEnableLightCard) ? ( CameraIndex == (TotalUsedCameras + 1) ) : ( CameraIndex == TotalUsedCameras );
+			return (ICVFXParameters.IsLightcardOverUsed()) ? ( CameraIndex == (TotalUsedCameras + 1) ) : ( CameraIndex == TotalUsedCameras );
 		}
 
 		if (TotalUsedCameras>1)
@@ -339,7 +331,7 @@ private:
 			// 1   - camera2
 			// N-1 - camera N
 			// N   - overlayOver (optional)
-			return (ICVFXParameters.IsLightcardOverUsed() && GDisplayClusterShadersICVFXEnableLightCard) ? ( CameraIndex == TotalUsedCameras ) : ( CameraIndex == (TotalUsedCameras-1) );
+			return (ICVFXParameters.IsLightcardOverUsed()) ? ( CameraIndex == TotalUsedCameras ) : ( CameraIndex == (TotalUsedCameras-1) );
 		}
 
 		// By default render single camera in one pass
@@ -349,7 +341,7 @@ private:
 	inline bool IsOverlayUnderUsed() const
 	{
 		// Render OverlayUnder only on first pass
-		return IsFirstPass() && ICVFXParameters.IsLightcardUnderUsed() && GDisplayClusterShadersICVFXEnableLightCard;
+		return IsFirstPass() && ICVFXParameters.IsLightcardUnderUsed();
 	}
 
 	inline bool IsOverlayOverUsed() const
@@ -357,10 +349,10 @@ private:
 		if (IsMultiPassRender())
 		{
 			// Render OverlayOver only on last additive pass
-			return IsLastPass() && ICVFXParameters.IsLightcardOverUsed() && GDisplayClusterShadersICVFXEnableLightCard;
+			return IsLastPass() && ICVFXParameters.IsLightcardOverUsed();
 		}
 		
-		return ICVFXParameters.IsLightcardOverUsed() && GDisplayClusterShadersICVFXEnableLightCard;
+		return ICVFXParameters.IsLightcardOverUsed();
 	}
 
 	inline int32 GetUsedCameraIndex() const
@@ -529,9 +521,22 @@ public:
 		if (IsOverlayUnderUsed() || IsOverlayOverUsed())
 		{
 
-			RenderPassData.PSParameters.OverlaySampler = TStaticSamplerState<SF_Trilinear, AM_Clamp, AM_Clamp, AM_Clamp>::GetRHI();
-			RenderPassData.PSParameters.OverlayTexture = ICVFXParameters.Lightcard.Texture;
+			if (ICVFXParameters.Lightcard_OCIO.IsValid())
+			{
+				// use OCIO with separated alpha lightmap texture
+				RenderPassData.PSParameters.OverlaySampler = TStaticSamplerState<SF_Trilinear, AM_Clamp, AM_Clamp, AM_Clamp>::GetRHI();
+				RenderPassData.PSParameters.OverlayTexture = ICVFXParameters.Lightcard_OCIO.Texture;
+
+				RenderPassData.PSParameters.OverlayAlphaSampler = TStaticSamplerState<SF_Trilinear, AM_Clamp, AM_Clamp, AM_Clamp>::GetRHI();
+				RenderPassData.PSParameters.OverlayAlphaTexture = ICVFXParameters.Lightcard.Texture;
 
+				RenderPassData.PSPermutationVector.Set<IcvfxShaderPermutation::FIcvfxShaderOverlayAlpha>(true);
+			}
+			else
+			{
+				RenderPassData.PSParameters.OverlaySampler = TStaticSamplerState<SF_Trilinear, AM_Clamp, AM_Clamp, AM_Clamp>::GetRHI();
+				RenderPassData.PSParameters.OverlayTexture = ICVFXParameters.Lightcard.Texture;
+			}
 			if (IsOverlayUnderUsed())
 			{
 				RenderPassData.PSPermutationVector.Set<IcvfxShaderPermutation::FIcvfxShaderOverlayUnder>(true);
@@ -641,7 +646,7 @@ public:
 
 	bool GetUVLightCardParameters(FIcvfxRenderPassData& RenderPassData)
 	{
-		if (ICVFXParameters.UVLightCardMap.IsValid() && GDisplayClusterShadersICVFXEnableLightCard && RenderPassData.PSPermutationVector.Get<IcvfxShaderPermutation::FIcvfxShaderMeshWarp>())
+		if (ICVFXParameters.UVLightCardMap.IsValid() && RenderPassData.PSPermutationVector.Get<IcvfxShaderPermutation::FIcvfxShaderMeshWarp>())
 		{
 			RenderPassData.PSParameters.UVLightCardMapSampler = TStaticSamplerState<SF_Trilinear, AM_Clamp, AM_Clamp, AM_Clamp>::GetRHI();
 			RenderPassData.PSParameters.UVLightCardMapTexture = ICVFXParameters.UVLightCardMap;
diff --git a/Engine/Plugins/Runtime/nDisplay/Source/DisplayClusterShaders/Public/ShaderParameters/DisplayClusterShaderParameters_ICVFX.h b/Engine/Plugins/Runtime/nDisplay/Source/DisplayClusterShaders/Public/ShaderParameters/DisplayClusterShaderParameters_ICVFX.h
--- a/Engine/Plugins/Runtime/nDisplay/Source/DisplayClusterShaders/Public/ShaderParameters/DisplayClusterShaderParameters_ICVFX.h
+++ b/Engine/Plugins/Runtime/nDisplay/Source/DisplayClusterShaders/Public/ShaderParameters/DisplayClusterShaderParameters_ICVFX.h
@@ -113,6 +113,7 @@ public:
 	{
 		Cameras.Empty();
 		Lightcard.Reset();
+		Lightcard_OCIO.Reset();
 
 		UVLightCardMap = nullptr;
 	}
@@ -125,6 +126,7 @@ public:
 		Cameras = InParameters.Cameras;
 
 		Lightcard      = InParameters.Lightcard;
+		Lightcard_OCIO = InParameters.Lightcard_OCIO;
 		LightcardMode = InParameters.LightcardMode;
 
 		UVLightCardMap = InParameters.UVLightCardMap;
@@ -137,6 +139,11 @@ public:
 			Dst.Add(&Lightcard);
 		}
 
+		if (Lightcard_OCIO.IsDefined())
+		{
+			Dst.Add(&Lightcard_OCIO);
+		}
+
 		for (FCameraSettings& CameraIt : Cameras)
 		{
 			CameraIt.CollectRefViewports(Dst);
@@ -242,6 +249,7 @@ public:
 
 	// Lightcard settings
 	FDisplayClusterShaderParametersICVFX_ViewportResource    Lightcard;
+	FDisplayClusterShaderParametersICVFX_ViewportResource    Lightcard_OCIO;
 	EDisplayClusterShaderParametersICVFX_LightcardRenderMode LightcardMode = EDisplayClusterShaderParametersICVFX_LightcardRenderMode::Under;
 
 	/** Texture containing a UV map of the rendered UV light cards */
diff --git a/Engine/Plugins/VirtualProduction/TextureShare/Source/TextureShareDisplayCluster/Private/PostProcess/TextureSharePostprocess_FrameProxy.cpp b/Engine/Plugins/VirtualProduction/TextureShare/Source/TextureShareDisplayCluster/Private/PostProcess/TextureSharePostprocess_FrameProxy.cpp
--- a/Engine/Plugins/VirtualProduction/TextureShare/Source/TextureShareDisplayCluster/Private/PostProcess/TextureSharePostprocess_FrameProxy.cpp
+++ b/Engine/Plugins/VirtualProduction/TextureShare/Source/TextureShareDisplayCluster/Private/PostProcess/TextureSharePostprocess_FrameProxy.cpp
@@ -42,13 +42,13 @@ void FTextureSharePostprocess::ShareViewport_RenderThread(FRHICommandListImmedia
 {
 	if (InViewportManagerProxy)
 	{
-		for (TSharedPtr<IDisplayClusterViewportProxy, ESPMode::ThreadSafe>& ViewportProxyIt : InViewportManagerProxy->GetViewports_RenderThread())
+		for (IDisplayClusterViewportProxy* ViewportProxyIt : InViewportManagerProxy->GetViewports_RenderThread())
 		{
-			if (ViewportProxyIt.IsValid())
+			if (ViewportProxyIt)
 			{
 				// Get viewport resource type
 				EDisplayClusterViewportResourceType ResourceType = InResourceType;
-				if (bAfterWarpBlend && ShouldApplyWarpBlend(ViewportProxyIt.Get()))
+				if (bAfterWarpBlend && ShouldApplyWarpBlend(ViewportProxyIt))
 				{
 					const TSharedPtr<IDisplayClusterProjectionPolicy, ESPMode::ThreadSafe>& PrjPolicy = ViewportProxyIt->GetProjectionPolicy_RenderThread();
 					if (PrjPolicy->ShouldUseAdditionalTargetableResource())
