nDisplay MRQ -  Fixed render issues:
* Added use of component visibility list from viewport for MRQ rendering.
* Rendering of the DCRA preview in the scene is turned off while the MRQ is running.
* Changed nDisplay viewport initialization.

diff --git a/Engine/Plugins/Runtime/nDisplay/Source/DisplayCluster/Private/Game/EngineClasses/Scene/DisplayClusterRootActorEditor.cpp b/Engine/Plugins/Runtime/nDisplay/Source/DisplayCluster/Private/Game/EngineClasses/Scene/DisplayClusterRootActorEditor.cpp
--- a/Engine/Plugins/Runtime/nDisplay/Source/DisplayCluster/Private/Game/EngineClasses/Scene/DisplayClusterRootActorEditor.cpp
+++ b/Engine/Plugins/Runtime/nDisplay/Source/DisplayCluster/Private/Game/EngineClasses/Scene/DisplayClusterRootActorEditor.cpp
@@ -197,6 +197,12 @@ void ADisplayClusterRootActor::EnableEditorRender(bool bValue)
 
 bool ADisplayClusterRootActor::IsPreviewEnabled() const
 {
+	if (bMoviePipelineRenderPass)
+	{
+		// Disable preview rendering for MRQ
+		return false;
+	}
+
 	//@todo: (GUI) Scene preview can be disabled when the configuration window with internal preview is open.
 #if 0
 	bool bIsScenePreview = true; //@todo: handle GUI logic
diff --git a/Engine/Plugins/Runtime/nDisplay/Source/DisplayCluster/Private/Render/Viewport/DisplayClusterViewport.cpp b/Engine/Plugins/Runtime/nDisplay/Source/DisplayCluster/Private/Render/Viewport/DisplayClusterViewport.cpp
--- a/Engine/Plugins/Runtime/nDisplay/Source/DisplayCluster/Private/Render/Viewport/DisplayClusterViewport.cpp
+++ b/Engine/Plugins/Runtime/nDisplay/Source/DisplayCluster/Private/Render/Viewport/DisplayClusterViewport.cpp
@@ -204,7 +204,16 @@ bool FDisplayClusterViewport::ShouldUseFullSizeFrameTargetableResource() const
 void FDisplayClusterViewport::SetupSceneView(uint32 ContextNum, class UWorld* World, FSceneViewFamily& InOutViewFamily, FSceneView& InOutView) const
 {
 	check(IsInGameThread());
-	check(ContextNum < (uint32)Contexts.Num());
+	check(Contexts.IsValidIndex(ContextNum));
+
+	// MRQ only uses viewport visibility settings
+	if(RenderSettings.CaptureMode == EDisplayClusterViewportCaptureMode::MoviePipeline)
+	{
+		// Apply visibility settigns to view
+		VisibilitySettings.SetupSceneView(World, InOutView);
+
+		return;
+	}
 
 	// Setup MGPU features:
 	if(Contexts[ContextNum].GPUIndex >= 0)
diff --git a/Engine/Plugins/Runtime/nDisplay/Source/DisplayCluster/Public/DisplayClusterRootActor.h b/Engine/Plugins/Runtime/nDisplay/Source/DisplayCluster/Public/DisplayClusterRootActor.h
--- a/Engine/Plugins/Runtime/nDisplay/Source/DisplayCluster/Public/DisplayClusterRootActor.h
+++ b/Engine/Plugins/Runtime/nDisplay/Source/DisplayCluster/Public/DisplayClusterRootActor.h
@@ -304,6 +304,10 @@ public:
 //////////////////////////////////////////////////////////////////////////////////////////////
 #if WITH_EDITORONLY_DATA
 public:
+	/** When the MRQ is rendered, this flag is raised. */
+	UPROPERTY()
+	bool bMoviePipelineRenderPass = false;
+
 	/** Render the scene and display it as a preview on the nDisplay root actor in the editor.  This will impact editor performance. */
 	UPROPERTY(EditAnywhere, BlueprintReadWrite, Category = "Editor Preview", meta = (DisplayName = "Enable Editor Preview"))
 	bool bPreviewEnable = true;
diff --git a/Engine/Plugins/Runtime/nDisplay/Source/DisplayClusterMoviePipeline/Private/DisplayClusterMoviePipelineViewportPass.cpp b/Engine/Plugins/Runtime/nDisplay/Source/DisplayClusterMoviePipeline/Private/DisplayClusterMoviePipelineViewportPass.cpp
--- a/Engine/Plugins/Runtime/nDisplay/Source/DisplayClusterMoviePipeline/Private/DisplayClusterMoviePipelineViewportPass.cpp
+++ b/Engine/Plugins/Runtime/nDisplay/Source/DisplayClusterMoviePipeline/Private/DisplayClusterMoviePipelineViewportPass.cpp
@@ -94,13 +94,29 @@ UDisplayClusterMoviePipelineViewportPassBase::UDisplayClusterMoviePipelineViewpo
 
 void UDisplayClusterMoviePipelineViewportPassBase::SetupImpl(const MoviePipeline::FMoviePipelineRenderPassInitSettings& InPassInitSettings)
 {
-	InitializeDisplayCluster();
+	if (InitializeDisplayCluster())
+	{
+		check(DCRootActor);
+
+#if WITH_EDITORONLY_DATA
+		// Save current DCRA preview settings, and disable preview rendering until MRQ render is done
+		DCRootActor->bMoviePipelineRenderPass = true;
+#endif
+	}
 
 	Super::SetupImpl(InPassInitSettings);
 }
 
 void UDisplayClusterMoviePipelineViewportPassBase::TeardownImpl()
 {
+	// Restore original preview value for DCRA actor
+	if (DCRootActor)
+	{
+#if WITH_EDITORONLY_DATA
+		DCRootActor->bMoviePipelineRenderPass = false;
+#endif
+	}
+
 	ReleaseDisplayCluster();
 
 	Super::TeardownImpl();
@@ -127,7 +143,7 @@ UE::MoviePipeline::FImagePassCameraViewData UDisplayClusterMoviePipelineViewport
 	UE::MoviePipeline::FImagePassCameraViewData OutCameraData;
 
 	// This has to come from the main camera for consistency's sake, and it's not a per-camera setting in the editor.
-	OutCameraData.ViewActor = GetPipeline()->GetWorld()->GetFirstPlayerController()->GetViewTarget();
+	OutCameraData.ViewActor = GetWorld()->GetFirstPlayerController()->GetViewTarget();
 
 	const int32 InViewportIndex = InOutSampleState.OutputState.CameraIndex;
 	FString ViewportId;
@@ -272,9 +288,11 @@ void UDisplayClusterMoviePipelineViewportPassBase::ModifyProjectionMatrixForTili
 
 FSceneView* UDisplayClusterMoviePipelineViewportPassBase::GetSceneViewForSampleState(FSceneViewFamily* ViewFamily, FMoviePipelineRenderPassMetrics& InOutSampleState, IViewCalcPayload* OptPayload)
 {
+	IDisplayClusterViewport* DCViewport = nullptr;
+
 	const int32 InViewportIndex = InOutSampleState.OutputState.CameraIndex;
 	FString ViewportId;
-	if (GetViewportId(InViewportIndex, ViewportId))
+	if (DCRootActor && GetViewportId(InViewportIndex, ViewportId))
 	{
 		UMoviePipelineExecutorShot* CurrentShot = GetPipeline()->GetActiveShotList()[GetPipeline()->GetCurrentShotIndex()];
 		UMoviePipelineHighResSetting* HighResSettings = GetPipeline()->FindOrAddSettingForShot<UMoviePipelineHighResSetting>(CurrentShot);
@@ -295,7 +313,8 @@ FSceneView* UDisplayClusterMoviePipelineViewportPassBase::GetSceneViewForSampleS
 		const FIntPoint DestSize = FullDestSize - OffsetMin - OffsetMax;
 
 		FDisplayClusterViewInfo NewView, PrevView;
-		if (CalculateDisplayClusterViewport(InOutSampleState, *ViewportId, DestSize, 0, NewView))
+		DCViewport = GetAndCalculateDisplayClusterViewport(InOutSampleState, *ViewportId, DestSize, 0, NewView);
+		if(DCViewport)
 		{
 			// Update view cache
 			if (!DCViews.Contains(InViewportIndex))
@@ -316,7 +335,15 @@ FSceneView* UDisplayClusterMoviePipelineViewportPassBase::GetSceneViewForSampleS
 		UE_LOG(LogMovieRenderPipeline, Error, TEXT("DCRA render failed: undefined viewport index, skipped"));
 	}
 
-	return Super::GetSceneViewForSampleState(ViewFamily, InOutSampleState, OptPayload);
+	FSceneView* SceneView = Super::GetSceneViewForSampleState(ViewFamily, InOutSampleState, OptPayload);
+
+	if (DCViewport && SceneView)
+	{
+		const uint32 ContextNum = 0;
+		DCViewport->SetupSceneView(ContextNum, GetWorld(), *ViewFamily, *SceneView);
+	}
+
+	return SceneView;
 }
 
 void UDisplayClusterMoviePipelineViewportPassBase::PostRendererSubmission(const FMoviePipelineRenderPassMetrics& InSampleState, const FMoviePipelinePassIdentifier InPassIdentifier, const int32 InSortingOrder, FCanvas& InCanvas)
@@ -325,11 +352,11 @@ void UDisplayClusterMoviePipelineViewportPassBase::PostRendererSubmission(const
 	{
 		const int32 InViewportIndex = InSampleState.OutputState.CameraIndex;
 		FString ViewportId;
-		if (GetViewportId(InViewportIndex, ViewportId))
+		if (DCRootActor && GetViewportId(InViewportIndex, ViewportId))
 		{
 			if (FRenderTarget* RenderTarget = InCanvas.GetRenderTarget())
 			{
-				if (IDisplayClusterViewportManager* ViewportManager = DisplayClusterRootActorPtr->GetViewportManager())
+				if (IDisplayClusterViewportManager* ViewportManager = DCRootActor->GetViewportManager())
 				{
 					FIntPoint OffsetMin = FIntPoint::ZeroValue;
 					FIntPoint OffsetMax = FIntPoint::ZeroValue;
@@ -358,7 +385,7 @@ void UDisplayClusterMoviePipelineViewportPassBase::GetViewportCutOffset(const FM
 	OutOffsetMin = FIntPoint::ZeroValue;
 	OutOffsetMax = FIntPoint::ZeroValue;
 
-	APlayerCameraManager* PlayerCameraManager = GetPipeline()->GetWorld()->GetFirstPlayerController()->PlayerCameraManager;
+	APlayerCameraManager* PlayerCameraManager = GetWorld()->GetFirstPlayerController()->PlayerCameraManager;
 	if (PlayerCameraManager && PlayerCameraManager->GetCameraCacheView().bConstrainAspectRatio)
 	{
 		const FMinimalViewInfo CameraCache = PlayerCameraManager->GetCameraCacheView();
@@ -399,68 +426,72 @@ void UDisplayClusterMoviePipelineViewportPassBase::GetViewportCutOffset(const FM
 	}
 }
 
-ADisplayClusterRootActor* UDisplayClusterMoviePipelineViewportPassBase::GetRootActor(const UDisplayClusterMoviePipelineSettings& InDCSettings)
+void UDisplayClusterMoviePipelineViewportPassBase::ReleaseDisplayCluster()
 {
-	if (DisplayClusterRootActorPtr)
-	{
-		const bool bWorldEqual = DisplayClusterRootActorPtr->GetWorld() == GetWorld();
-		if (!bWorldEqual)
-		{
-			DisplayClusterRootActorPtr = nullptr;
-		}
-	}
+	// Reset all runtime values
+	DCRootActor = nullptr;
 
-	if (DisplayClusterRootActorPtr == nullptr)
-	{
-		DisplayClusterRootActorPtr = InDCSettings.GetRootActor(GetWorld());
-	}
+	bFrameWarpBlend = false;
 
-	return DisplayClusterRootActorPtr;
+	DisplayClusterViewports.Empty();
+	DCViews.Empty();
+	DCPrevViews.Empty();
 }
 
 bool UDisplayClusterMoviePipelineViewportPassBase::InitializeDisplayCluster()
 {
-	// Reset all runtime values
+	ReleaseDisplayCluster();
+
+	UWorld* CurrentWorld = GetWorld();
+	if (CurrentWorld == nullptr)
 	{
-		DCViews.Empty();
-		DCPrevViews.Empty();
+		UE_LOG(LogMovieRenderPipeline, Error, TEXT("Failed to find current world for nDisplay render"));
 
-		bFrameInitialized = false;
-		bFrameWarpBlend = false;
+		return false;
+	}
 
-		DisplayClusterRootActorPtr = nullptr;
+	if (UMoviePipelineExecutorShot* CurrentShot = GetPipeline()->GetActiveShotList()[GetPipeline()->GetCurrentShotIndex()])
+	{
+		if (UDisplayClusterMoviePipelineSettings* DCMoviePipelineSettings = GetPipeline()->FindOrAddSettingForShot<UDisplayClusterMoviePipelineSettings>(CurrentShot))
+		{
+			DCRootActor = DCMoviePipelineSettings->GetRootActor(CurrentWorld);
+			if (DCRootActor == nullptr)
+			{
+				UE_LOG(LogMovieRenderPipeline, Error, TEXT("Failed to find nDisplay RootActor"));
 
-		bDisplayClusterInitialized = false;
+				return false;
+			}
 
-		DisplayClusterViewports.Empty();
-	}
+			if (!DCMoviePipelineSettings->GetViewports(CurrentWorld, DisplayClusterViewports, DisplayClusterViewportSizes))
+			{
+				UE_LOG(LogMovieRenderPipeline, Error, TEXT("Failed to find nDisplay viewports"));
 
-	UMoviePipelineExecutorShot* CurrentShot = GetPipeline()->GetActiveShotList()[GetPipeline()->GetCurrentShotIndex()];
-	UDisplayClusterMoviePipelineSettings* DisplayClusterCSettings = GetPipeline()->FindOrAddSettingForShot<UDisplayClusterMoviePipelineSettings>(CurrentShot);
+				return false;
+			}
 
-	// Initialize
-	if (DisplayClusterCSettings && GetRootActor(*DisplayClusterCSettings) && DisplayClusterCSettings->GetViewports(GetWorld(), DisplayClusterViewports, DisplayClusterViewportSizes))
-	{
-		bDisplayClusterInitialized = true;
+			IDisplayClusterViewportManager* ViewportManager = DCRootActor->GetViewportManager();
+			check(ViewportManager);
 
-		return true;
+			const EDisplayClusterRenderFrameMode RenderFrameMode = EDisplayClusterRenderFrameMode::Mono;
+
+			// Update local node viewports (update\create\delete) and build new render frame
+			if (ViewportManager->UpdateCustomConfiguration(RenderFrameMode, DisplayClusterViewports, DCRootActor))
+			{
+				return true;
+			}
+		}
 	}
 
-	UE_LOG(LogMovieRenderPipeline, Error, TEXT("Failed to find nDisplay viewports"));
+	UE_LOG(LogMovieRenderPipeline, Error, TEXT("Failed to initialize nDisplay render"));
+	ReleaseDisplayCluster();
 
 	return false;
 }
 
-void UDisplayClusterMoviePipelineViewportPassBase::ReleaseDisplayCluster()
-{
-	bDisplayClusterInitialized = false;
-	DisplayClusterViewports.Empty();
-	DCViews.Empty();
-}
 
 bool UDisplayClusterMoviePipelineViewportPassBase::GetViewportId(int32 InViewportIndex, FString& OutViewportId) const
 {
-	if (bDisplayClusterInitialized && InViewportIndex >= 0 && InViewportIndex < DisplayClusterViewports.Num())
+	if (DCRootActor && InViewportIndex >= 0 && InViewportIndex < DisplayClusterViewports.Num())
 	{
 		OutViewportId = DisplayClusterViewports[InViewportIndex];
 
@@ -470,160 +501,162 @@ bool UDisplayClusterMoviePipelineViewportPassBase::GetViewportId(int32 InViewpor
 	return false;
 }
 
-bool UDisplayClusterMoviePipelineViewportPassBase::CalculateDisplayClusterViewport(const FMoviePipelineRenderPassMetrics& InSampleState, const FString& InViewportId, const FIntPoint& InDestSize, const uint32 InContextNum, FDisplayClusterViewInfo& OutView)
+IDisplayClusterViewport* UDisplayClusterMoviePipelineViewportPassBase::GetAndCalculateDisplayClusterViewport(const FMoviePipelineRenderPassMetrics& InSampleState, const FString& InViewportId, const FIntPoint& InDestSize, const uint32 InContextNum, FDisplayClusterViewInfo& OutView)
 {
-	const EDisplayClusterRenderFrameMode RenderFrameMode = EDisplayClusterRenderFrameMode::Mono;
+	check(DCRootActor);
+	
+	UWorld* CurrentWorld = GetWorld();
+	check(CurrentWorld);
+
+	UGameViewportClient* GameViewportClient = CurrentWorld->GetGameViewport();
+	check(GameViewportClient);
+
+	IDisplayClusterViewportManager* ViewportManager = DCRootActor->GetViewportManager();
+	check(ViewportManager);
 
-	bFrameInitialized = false;
-	if (IDisplayClusterViewportManager* ViewportManager = DisplayClusterRootActorPtr->GetViewportManager())
+	if (IDisplayClusterViewport* DCViewport = ViewportManager->FindViewport(InViewportId))
 	{
-		if (UWorld* InWorld = GetWorld())
+		// Force custom viewport size and settings:
 		{
-			const AWorldSettings* WorldSettings = InWorld->GetWorldSettings();
-			const float InWorldToMeters = (WorldSettings) ? WorldSettings->WorldToMeters : 100.f;
+			FDisplayClusterViewport_RenderSettings RenderSettings;
+			RenderSettings.Rect = FIntRect(FIntPoint(0, 0), InDestSize);
+			RenderSettings.CaptureMode = EDisplayClusterViewportCaptureMode::MoviePipeline;
 
-			if (UGameViewportClient* GameViewportClient = InWorld->GetGameViewport())
+			if (!bFrameWarpBlend)
 			{
-				// Update local node viewports (update\create\delete) and build new render frame
-				if (ViewportManager->UpdateCustomConfiguration(RenderFrameMode, DisplayClusterViewports, DisplayClusterRootActorPtr))
-				{
-					// Force custom viewport size and settings:
-					if (IDisplayClusterViewport* ViewportPtr = ViewportManager->FindViewport(InViewportId))
-					{
-						FDisplayClusterViewport_RenderSettings RenderSettings;
-						RenderSettings.Rect = FIntRect(FIntPoint(0, 0), InDestSize);
-						RenderSettings.CaptureMode = EDisplayClusterViewportCaptureMode::MoviePipeline;
+				// Create internal resources only when warpblend will be used
+				RenderSettings.bSkipRendering = true;
+			}
 
-						if (!bFrameWarpBlend)
-						{
-							// Create internal resources only when warpblend will be used
-							RenderSettings.bSkipRendering = true;
-						}
+			// Disable DCRA frustum modifier features
+			RenderSettings.bDisableCustomFrustumFeature = true;
+			RenderSettings.bDisableFrustumOverscanFeature = true;
 
-						// Disable DCRA frustum modifier features
-						RenderSettings.bDisableCustomFrustumFeature = true;
-						RenderSettings.bDisableFrustumOverscanFeature = true;
+			DCViewport->SetRenderSettings(RenderSettings);
+		}
 
-						ViewportPtr->SetRenderSettings(RenderSettings);
+		FDisplayClusterRenderFrame RenderFrame;
+		if (ViewportManager->BeginNewFrame(GameViewportClient->Viewport, CurrentWorld, RenderFrame))
+		{
+			if (!DCViewport->GetProjectionPolicy().IsValid())
+			{
+				// ignore viewports with uninitialized prj policy
+				return nullptr;
+			}
 
-						FDisplayClusterRenderFrame RenderFrame;
-						if (ViewportManager->BeginNewFrame(GameViewportClient->Viewport, InWorld, RenderFrame))
-						{
-							// InDestSize
-							bFrameInitialized = true;
+			const TArray<FDisplayClusterViewport_Context>& ViewportContexts = DCViewport->GetContexts();
+			check(ViewportContexts.IsValidIndex(InContextNum));
 
-							if (bFrameWarpBlend)
-							{
-								// Initialize proxy only when warpblend will be used
-								ViewportManager->FinalizeNewFrame();
-							}
+			const FDisplayClusterViewport_Context& ViewportContext = ViewportContexts[InContextNum];
 
-							if (!ViewportPtr->GetProjectionPolicy().IsValid())
-							{
-								// ignore viewports with uninitialized prj policy
-								return false;
-							}
+			UDisplayClusterCameraComponent* ViewCamera = nullptr;
+			{
+				// Get camera ID assigned to the viewport
+				const FString& CameraId = DCViewport->GetRenderSettings().CameraId;
 
-							const TArray<FDisplayClusterViewport_Context>& ViewportContexts = ViewportPtr->GetContexts();
-							const FDisplayClusterViewport_Context& ViewportContext = ViewportContexts[InContextNum];
+				// Get camera component assigned to the viewport (or default camera if nothing assigned)
+				ViewCamera = (CameraId.IsEmpty() ?
+					DCRootActor->GetDefaultCamera() :
+					DCRootActor->GetComponentByName<UDisplayClusterCameraComponent>(CameraId));
+			}
 
-							UDisplayClusterCameraComponent* ViewCamera = nullptr;
-							{
-								// Get camera ID assigned to the viewport
-								const FString& CameraId = ViewportPtr->GetRenderSettings().CameraId;
+			if (ViewCamera == nullptr)
+			{
+				UE_LOG(LogMovieRenderPipeline, Error, TEXT("No camera found for viewport '%s'"), *DCViewport->GetId());
+				return nullptr;
+			}
 
-								// Get camera component assigned to the viewport (or default camera if nothing assigned)
-								ViewCamera = (CameraId.IsEmpty() ?
-									DisplayClusterRootActorPtr->GetDefaultCamera() :
-									DisplayClusterRootActorPtr->GetComponentByName<UDisplayClusterCameraComponent>(CameraId));
-							}
+			FVector ViewOffset = FVector::ZeroVector;
+			const float CfgNCP = GNearClippingPlane;
+			{
+				// Get the actual camera settings
+				const float CfgEyeDist = ViewCamera->GetInterpupillaryDistance();
+				const bool  CfgEyeSwap = ViewCamera->GetSwapEyes();
+				const EDisplayClusterEyeStereoOffset CfgEyeOffset = ViewCamera->GetStereoOffset();
 
-							if (!ViewCamera)
-							{
-								UE_LOG(LogMovieRenderPipeline, Error, TEXT("No camera found for viewport '%s'"), *ViewportPtr->GetId());
-								return false;
-							}
+				// Calculate eye offset considering the world scale
+				const float EyeOffset = CfgEyeDist / 2.f;
+				const float EyeOffsetValues[] = { -EyeOffset, 0.f, EyeOffset };
 
-							// Get the actual camera settings
-							const float CfgEyeDist = ViewCamera->GetInterpupillaryDistance();
-							const bool  CfgEyeSwap = ViewCamera->GetSwapEyes();
-							const float CfgNCP = GNearClippingPlane;
-							const EDisplayClusterEyeStereoOffset CfgEyeOffset = ViewCamera->GetStereoOffset();
+				enum class EDisplayClusterEyeType : uint8
+				{
+					StereoLeft = 0,
+					Mono = 1,
+					StereoRight = 2,
+					COUNT
+				};
 
-							// Calculate eye offset considering the world scale
-							const float EyeOffset = CfgEyeDist / 2.f;
-							const float EyeOffsetValues[] = { -EyeOffset, 0.f, EyeOffset };
+				// Decode current eye type
+				const EDisplayClusterEyeType DCEyeType = (ViewportContexts.Num() > 1)
+					? ((InContextNum == 0) ? EDisplayClusterEyeType::StereoLeft : EDisplayClusterEyeType::StereoRight)
+					: EDisplayClusterEyeType::Mono;
 
-							enum class EDisplayClusterEyeType : uint8
-							{
-								StereoLeft = 0,
-								Mono = 1,
-								StereoRight = 2,
-								COUNT
-							};
+				const int32 EyeIndex = (uint8)DCEyeType;
+
+				float PassOffset = 0.f;
+				float PassOffsetSwap = 0.f;
+
+				switch (DCEyeType)
+				{
+				case EDisplayClusterEyeType::Mono:
+				{
+					// For monoscopic camera let's check if the "force offset" feature is used
+					// * Force left (-1) ==> 0 left eye
+					// * Force right (1) ==> 2 right eye
+					// * Default (0) ==> 1 mono
+					const int32 EyeOffsetIdx = (CfgEyeOffset == EDisplayClusterEyeStereoOffset::None ? 0 : (CfgEyeOffset == EDisplayClusterEyeStereoOffset::Left ? -1 : 1));
+
+					PassOffset = EyeOffsetValues[EyeOffsetIdx + 1];
+					// Eye swap is not available for monoscopic so just save the value
+					PassOffsetSwap = PassOffset;
+				}
+				break;
 
-							// Decode current eye type
-							const EDisplayClusterEyeType DCEyeType = (ViewportContexts.Num() > 1)
-								? ((InContextNum == 0) ? EDisplayClusterEyeType::StereoLeft : EDisplayClusterEyeType::StereoRight)
-								: EDisplayClusterEyeType::Mono;
+				default:
+				{
+					// For stereo camera we can only swap eyes if required (no "force offset" allowed)
+					PassOffset = EyeOffsetValues[EyeIndex];
+					PassOffsetSwap = (CfgEyeSwap ? -PassOffset : PassOffset);
+				}
+				break;
+				}
 
-							const int32 EyeIndex = (uint8)DCEyeType;
+				// View base location
+				OutView.ViewLocation = ViewCamera->GetComponentLocation();
+				OutView.ViewRotation = ViewCamera->GetComponentRotation();
 
-							float PassOffset = 0.f;
-							float PassOffsetSwap = 0.f;
+				// Apply computed offset to the view location
+				const FQuat EyeQuat = OutView.ViewRotation.Quaternion();
+				ViewOffset = EyeQuat.RotateVector(FVector(0.0f, PassOffsetSwap, 0.0f));
+				OutView.ViewLocation += ViewOffset;
+			}
 
-							switch (DCEyeType)
-							{
-							case EDisplayClusterEyeType::Mono:
-							{
-								// For monoscopic camera let's check if the "force offset" feature is used
-								// * Force left (-1) ==> 0 left eye
-								// * Force right (1) ==> 2 right eye
-								// * Default (0) ==> 1 mono
-								const int32 EyeOffsetIdx = (CfgEyeOffset == EDisplayClusterEyeStereoOffset::None ? 0 : (CfgEyeOffset == EDisplayClusterEyeStereoOffset::Left ? -1 : 1));
-
-								PassOffset = EyeOffsetValues[EyeOffsetIdx + 1];
-								// Eye swap is not available for monoscopic so just save the value
-								PassOffsetSwap = PassOffset;
-							}
-							break;
+			const AWorldSettings* WorldSettings = CurrentWorld->GetWorldSettings();
+			const float InWorldToMeters = (WorldSettings) ? WorldSettings->WorldToMeters : 100.f;
 
-							default:
-							{
-								// For stereo camera we can only swap eyes if required (no "force offset" allowed)
-								PassOffset = EyeOffsetValues[EyeIndex];
-								PassOffsetSwap = (CfgEyeSwap ? -PassOffset : PassOffset);
-							}
-							break;
-							}
+			bool bResult = false;
+			// Perform view calculations on a policy side
+			if (DCViewport->CalculateView(InContextNum, OutView.ViewLocation, OutView.ViewRotation, ViewOffset, InWorldToMeters, CfgNCP, CfgNCP))
+			{
+				OutView.ProjectionMatrix = FMatrix::Identity;
 
-							FVector ViewOffset = FVector::ZeroVector;
-							if (ViewCamera)
-							{
-								// View base location
-								OutView.ViewLocation = ViewCamera->GetComponentLocation();
-								OutView.ViewRotation = ViewCamera->GetComponentRotation();
-								// Apply computed offset to the view location
-								const FQuat EyeQuat = OutView.ViewRotation.Quaternion();
-								ViewOffset = EyeQuat.RotateVector(FVector(0.0f, PassOffsetSwap, 0.0f));
-								OutView.ViewLocation += ViewOffset;
-							}
+				bResult = ViewportManager->IsSceneOpened() && DCViewport->GetProjectionMatrix(InContextNum, OutView.ProjectionMatrix);
+			}
 
-							// Perform view calculations on a policy side
-							if (ViewportPtr->CalculateView(InContextNum, OutView.ViewLocation, OutView.ViewRotation, ViewOffset, InWorldToMeters, CfgNCP, CfgNCP))
-							{
-								OutView.ProjectionMatrix = FMatrix::Identity;
+			if (bFrameWarpBlend)
+			{
+				// Initialize proxy only when warpblend will be used
+				ViewportManager->FinalizeNewFrame();
+			}
 
-								return ViewportPtr->GetOwner().IsSceneOpened() && ViewportPtr->GetProjectionMatrix(InContextNum, OutView.ProjectionMatrix);
-							}
-						}
-					}
-				}
+			if (bResult)
+			{
+				return DCViewport;
 			}
 		}
 	}
 
-	return false;
+	return nullptr;
 }
 
 /////////////////////////////////////////////////////////////////////////////////////////////////////////////////////////////////
diff --git a/Engine/Plugins/Runtime/nDisplay/Source/DisplayClusterMoviePipeline/Public/DisplayClusterMoviePipelineViewportPass.h b/Engine/Plugins/Runtime/nDisplay/Source/DisplayClusterMoviePipeline/Public/DisplayClusterMoviePipelineViewportPass.h
--- a/Engine/Plugins/Runtime/nDisplay/Source/DisplayClusterMoviePipeline/Public/DisplayClusterMoviePipelineViewportPass.h
+++ b/Engine/Plugins/Runtime/nDisplay/Source/DisplayClusterMoviePipeline/Public/DisplayClusterMoviePipelineViewportPass.h
@@ -55,9 +55,8 @@ public:
 
 protected:
 	bool GetViewportId(int32 InViewportIndex, FString& OutViewportId) const;
-	ADisplayClusterRootActor* GetRootActor(const class UDisplayClusterMoviePipelineSettings& InDCSettings);
 
-	bool CalculateDisplayClusterViewport(const FMoviePipelineRenderPassMetrics& InSampleState, const FString& InViewportId, const FIntPoint& InDestSize, const uint32 InContextNum, FDisplayClusterViewInfo& OutView);
+	class IDisplayClusterViewport* GetAndCalculateDisplayClusterViewport(const FMoviePipelineRenderPassMetrics& InSampleState, const FString& InViewportId, const FIntPoint& InDestSize, const uint32 InContextNum, FDisplayClusterViewInfo& OutView);
 
 	bool InitializeDisplayCluster();
 	void ReleaseDisplayCluster();
@@ -74,7 +73,7 @@ private:
 	const FString RenderPassName;
 
 	// DCRA for current shot
-	ADisplayClusterRootActor* DisplayClusterRootActorPtr = nullptr;
+	ADisplayClusterRootActor* DCRootActor = nullptr;
 
 	// Names of viewports to render in the current shot
 	TArray<FString> DisplayClusterViewports;
@@ -85,9 +84,7 @@ private:
 	TMap<int32, FDisplayClusterViewInfo> DCPrevViews;
 
 	// Runtime flags
-	bool bFrameInitialized = true;
 	bool bFrameWarpBlend = false;
-	bool bDisplayClusterInitialized = false;
 };
 
 /**
