nDisplay: Fix incomplete preview when ICVFX panel is open and NDC has even number of nodes.

diff --git a/Engine/Plugins/Runtime/nDisplay/Source/DisplayCluster/Private/Game/EngineClasses/Scene/DisplayClusterRootActorEditor.cpp b/Engine/Plugins/Runtime/nDisplay/Source/DisplayCluster/Private/Game/EngineClasses/Scene/DisplayClusterRootActorEditor.cpp
--- a/Engine/Plugins/Runtime/nDisplay/Source/DisplayCluster/Private/Game/EngineClasses/Scene/DisplayClusterRootActorEditor.cpp
+++ b/Engine/Plugins/Runtime/nDisplay/Source/DisplayCluster/Private/Game/EngineClasses/Scene/DisplayClusterRootActorEditor.cpp
@@ -132,19 +132,6 @@ void ADisplayClusterRootActor::Tick_Editor(float DeltaSeconds)
 			// Render viewport for preview material RTTs
 			if (ViewportManager.IsValid())
 			{
-				// Either render to the preview components' standard RenderTarget or the RenderTargetPostProcess.
-				// This is needed in case preview post process is disabled but this preview should also be displayed
-				// in another editor window that requires post process completed.
-				
-				// TODO: Optimize the render routine to cache the pre PP preview then apply PostProcess.
-				
-				if (!PreviewRenderFrame.IsValid()) // Frame has finished.
-				{
-					bOutputFrameToPostProcessRenderTarget = bOutputFrameToPostProcessRenderTarget ?
-						bPreviewEnablePostProcess :
-						bPreviewEnablePostProcess || DoObserversNeedPostProcessRenderTarget();
-				}
-
 				ImplRenderPreview_Editor();
 			}
 		}
@@ -378,9 +365,17 @@ bool ADisplayClusterRootActor::ImplUpdatePreviewRenderFrame_Editor(const FString
 			{
 				PreviewRenderFrame = MakeUnique<FDisplayClusterRenderFrame>();
 
+				const bool bSceneNeedsStarting = !ViewportManager->GetCurrentWorld();
+				
 				// Update preview viewports from settings
 				if (ViewportManager->BeginNewFrame(nullptr, PreviewWorld, *PreviewRenderFrame))
 				{
+					if (bSceneNeedsStarting)
+					{
+						// Fix inner frustum possibly using the wrong post process render target
+						ImplUpdatePreviewConfiguration_Editor(InClusterNodeId);
+					}
+					
 					PreviewViewportIndex = 0;
 					return true;
 				}
@@ -498,6 +493,10 @@ void ADisplayClusterRootActor::ImplRenderPreview_Editor()
 		if (PreviewClusterNodeIndex >= ExistClusterNodesIDs.Num())
 		{
 			PreviewClusterNodeIndex = bFreezePreviewRender ? -1 : 0;
+
+			bOutputFrameToPostProcessRenderTarget = bOutputFrameToPostProcessRenderTarget ?
+				bPreviewEnablePostProcess :
+				bPreviewEnablePostProcess || DoObserversNeedPostProcessRenderTarget();
 		}
 
 		if (PreviewClusterNodeIndex < 0)
