nDisplay.MRQ: Fixed crash:
FRenderTarget creates the RenderTargetTextureRHI later on the render thread, so the GetRenderTargetTexture() function has been moved to the render thread

diff --git a/Engine/Plugins/Runtime/nDisplay/Source/DisplayClusterMoviePipeline/Private/DisplayClusterMoviePipelineViewportPass.cpp b/Engine/Plugins/Runtime/nDisplay/Source/DisplayClusterMoviePipeline/Private/DisplayClusterMoviePipelineViewportPass.cpp
--- a/Engine/Plugins/Runtime/nDisplay/Source/DisplayClusterMoviePipeline/Private/DisplayClusterMoviePipelineViewportPass.cpp
+++ b/Engine/Plugins/Runtime/nDisplay/Source/DisplayClusterMoviePipeline/Private/DisplayClusterMoviePipelineViewportPass.cpp
@@ -363,12 +363,16 @@ void UDisplayClusterMoviePipelineViewportPassBase::PostRendererSubmission(const
 					GetViewportCutOffset(InSampleState, OffsetMin, OffsetMax);
 
 					ENQUEUE_RENDER_COMMAND(DisplayClusterMoviePipelineWarp)([
-						InTextureRHIRef = RenderTarget->GetRenderTargetTexture(),
+						RenderTarget,
 						ViewportManagerProxy = ViewportManager->GetProxy(),
 						InViewportId = ViewportId,
 						OffsetMin, OffsetMax](FRHICommandListImmediate& RHICmdList) mutable
 					{
-						DisplayClusterWarpBlendImpl_RenderThread(RHICmdList, InTextureRHIRef, ViewportManagerProxy, InViewportId, OffsetMin, OffsetMax);
+						FTextureRHIRef RenderTargetRHIRef = RenderTarget->GetRenderTargetTexture();
+						if (RenderTargetRHIRef.IsValid())
+						{
+							DisplayClusterWarpBlendImpl_RenderThread(RHICmdList, RenderTargetRHIRef, ViewportManagerProxy, InViewportId, OffsetMin, OffsetMax);
+						}
 					});
 				}
 			}
