nDisplay: Fix CCWs being repositioned when opening the ICVFX panel.

diff --git a/Engine/Plugins/Experimental/ColorCorrectRegions/Source/ColorCorrectRegions/Private/ColorCorrectWindow.cpp b/Engine/Plugins/Experimental/ColorCorrectRegions/Source/ColorCorrectRegions/Private/ColorCorrectWindow.cpp
--- a/Engine/Plugins/Experimental/ColorCorrectRegions/Source/ColorCorrectRegions/Private/ColorCorrectWindow.cpp
+++ b/Engine/Plugins/Experimental/ColorCorrectRegions/Source/ColorCorrectRegions/Private/ColorCorrectWindow.cpp
@@ -257,7 +257,6 @@ FVector2D AColorCorrectionWindow::GetScale() const
 void AColorCorrectionWindow::SetOrigin(const FTransform& InOrigin)
 {
 	Origin = InOrigin;
-	NOTIFY_PARAM_SETTER()
 }
 
 FTransform AColorCorrectionWindow::GetOrigin() const
diff --git a/Engine/Plugins/Runtime/nDisplay/Source/DisplayClusterLightCardEditor/Private/Viewport/DisplayClusterLightCardEditorViewportClient.cpp b/Engine/Plugins/Runtime/nDisplay/Source/DisplayClusterLightCardEditor/Private/Viewport/DisplayClusterLightCardEditorViewportClient.cpp
--- a/Engine/Plugins/Runtime/nDisplay/Source/DisplayClusterLightCardEditor/Private/Viewport/DisplayClusterLightCardEditorViewportClient.cpp
+++ b/Engine/Plugins/Runtime/nDisplay/Source/DisplayClusterLightCardEditor/Private/Viewport/DisplayClusterLightCardEditorViewportClient.cpp
@@ -1513,10 +1513,13 @@ void FDisplayClusterLightCardEditorViewportClient::UpdatePreviewActor(ADisplayCl
 
 							BillboardComponentProxies.Append(BillboardComponents);
 						}
-						
-						ActorProxies.Add(FActorProxy(Actor, ActorProxy));
+
+						FActorProxy ActorProxyStruct(Actor, ActorProxy);
+						ActorProxies.Add(ActorProxyStruct);
 						ProjectionHelper->VerifyAndFixActorOrigin(Actor);
 						ProjectionHelper->VerifyAndFixActorOrigin(ActorProxy);
+
+						UpdateProxyTransforms(ActorProxyStruct);
 					}
 				}
 
