nDisplay: Fix UV light cards not being restored with level snapshots.

Fix dragging a CCW in the level by gizmo applying incorrect values. Regression from 22516709.

diff --git a/Engine/Plugins/Runtime/nDisplay/Source/DisplayClusterLightCardEditor/Private/Viewport/DisplayClusterLightCardEditorViewportClient.cpp b/Engine/Plugins/Runtime/nDisplay/Source/DisplayClusterLightCardEditor/Private/Viewport/DisplayClusterLightCardEditorViewportClient.cpp
--- a/Engine/Plugins/Runtime/nDisplay/Source/DisplayClusterLightCardEditor/Private/Viewport/DisplayClusterLightCardEditorViewportClient.cpp
+++ b/Engine/Plugins/Runtime/nDisplay/Source/DisplayClusterLightCardEditor/Private/Viewport/DisplayClusterLightCardEditorViewportClient.cpp
@@ -1069,6 +1069,9 @@ void FDisplayClusterLightCardEditorViewportClient::UpdateProxyTransforms(const F
 				
 		// Need to update these manually or the proxy's position will be out of sync next update
 		InActorProxy.Proxy->SetPositionalParams(InActorProxy.LevelInstance->GetPositionalParams());
+
+		// Update UV -- only valid for UV light cards
+		InActorProxy.Proxy->SetUVCoordinates(InActorProxy.LevelInstance->GetUVCoordinates());
 	}
 }
 
@@ -2205,6 +2208,11 @@ void FDisplayClusterLightCardEditorViewportClient::PropagateActorTransform(const
 
 		// Allows MU to receive the update in real-time.
 		LevelInstance->PostEditMove(false);
+
+		if (IDisplayClusterStageActor* StageActor = Cast<IDisplayClusterStageActor>(LevelInstance))
+		{
+			StageActor->UpdateEditorGizmos();
+		}
 	}
 }
 
diff --git a/Engine/Plugins/Runtime/nDisplay/Source/DisplayClusterScenePreview/Private/DisplayClusterLightCardEditorHelper.cpp b/Engine/Plugins/Runtime/nDisplay/Source/DisplayClusterScenePreview/Private/DisplayClusterLightCardEditorHelper.cpp
--- a/Engine/Plugins/Runtime/nDisplay/Source/DisplayClusterScenePreview/Private/DisplayClusterLightCardEditorHelper.cpp
+++ b/Engine/Plugins/Runtime/nDisplay/Source/DisplayClusterScenePreview/Private/DisplayClusterLightCardEditorHelper.cpp
@@ -1375,9 +1375,14 @@ void FDisplayClusterLightCardEditorHelper::InternalMoveActorTo(
 	Actor->SetYaw(Rotation.Yaw);
 
 #if WITH_EDITOR
-	if (bIsFinalChange && !Actor->IsProxy())
+	if (!Actor->IsProxy())
 	{
-		PostEditChangePropertiesForMovedActor(Actor);
+		Actor->UpdateEditorGizmos();
+		
+		if (bIsFinalChange)
+		{
+			PostEditChangePropertiesForMovedActor(Actor);
+		}
 	}
 #endif
 }
diff --git a/Engine/Plugins/Runtime/nDisplayModularFeatures/Source/DisplayClusterLightCardExtender/Private/StageActor/IDisplayClusterStageActor.cpp b/Engine/Plugins/Runtime/nDisplayModularFeatures/Source/DisplayClusterLightCardExtender/Private/StageActor/IDisplayClusterStageActor.cpp
--- a/Engine/Plugins/Runtime/nDisplayModularFeatures/Source/DisplayClusterLightCardExtender/Private/StageActor/IDisplayClusterStageActor.cpp
+++ b/Engine/Plugins/Runtime/nDisplayModularFeatures/Source/DisplayClusterLightCardExtender/Private/StageActor/IDisplayClusterStageActor.cpp
@@ -137,10 +137,6 @@ void IDisplayClusterStageActor::UpdatePositionalParamsFromTransform()
 		Actor->GetTransform(), GetOrigin(), GetRadialOffset());
 
 		SetPositionalParams(PositionalParams);
-
-#if WITH_EDITOR
-		UpdateEditorGizmos();
-#endif
 	}
 }
 
diff --git a/Engine/Plugins/Runtime/nDisplayModularFeatures/Source/DisplayClusterLightCardExtender/Public/StageActor/IDisplayClusterStageActor.h b/Engine/Plugins/Runtime/nDisplayModularFeatures/Source/DisplayClusterLightCardExtender/Public/StageActor/IDisplayClusterStageActor.h
--- a/Engine/Plugins/Runtime/nDisplayModularFeatures/Source/DisplayClusterLightCardExtender/Public/StageActor/IDisplayClusterStageActor.h
+++ b/Engine/Plugins/Runtime/nDisplayModularFeatures/Source/DisplayClusterLightCardExtender/Public/StageActor/IDisplayClusterStageActor.h
@@ -98,7 +98,6 @@ public:
 	static FDisplayClusterPositionalParams TransformToPositionalParams(const FTransform& InTransform, const FTransform& InOrigin, double InRadialOffset);
 
 #if WITH_EDITOR
-protected:
 	/** Update global selection gizmos */
 	void UpdateEditorGizmos();
 #endif
