Potential fix for compile issue found by the CI.

diff --git a/Engine/Plugins/Runtime/nDisplay/Source/DisplayClusterLightCardEditor/DisplayClusterLightCardEditor.Build.cs b/Engine/Plugins/Runtime/nDisplay/Source/DisplayClusterLightCardEditor/DisplayClusterLightCardEditor.Build.cs
--- a/Engine/Plugins/Runtime/nDisplay/Source/DisplayClusterLightCardEditor/DisplayClusterLightCardEditor.Build.cs
+++ b/Engine/Plugins/Runtime/nDisplay/Source/DisplayClusterLightCardEditor/DisplayClusterLightCardEditor.Build.cs
@@ -39,5 +39,11 @@ public class DisplayClusterLightCardEditor : ModuleRules
 				"ToolWidgets",
 				"WorkspaceMenuStructure"
 			});
+
+		PrivateIncludePaths.AddRange(
+			new string[] {
+				System.IO.Path.Combine(GetModuleDirectory("Renderer"), "Private"),
+			}
+		);
 	}
 }
