nDisplay: ICVFX - Optimize how proxies are updated and reduce number of times a full refresh is required. Proxies will refresh only their transforms in more situations or only recreate single proxies rather than all proxies.
Add thread safety for billboard components.
Fix deleting a level actor not refreshing the proxies as well as a possible MU related crash.

diff --git a/Engine/Plugins/Runtime/nDisplay/Source/DisplayClusterLightCardEditor/Private/DisplayClusterLightCardEditor.cpp b/Engine/Plugins/Runtime/nDisplay/Source/DisplayClusterLightCardEditor/Private/DisplayClusterLightCardEditor.cpp
--- a/Engine/Plugins/Runtime/nDisplay/Source/DisplayClusterLightCardEditor/Private/DisplayClusterLightCardEditor.cpp
+++ b/Engine/Plugins/Runtime/nDisplay/Source/DisplayClusterLightCardEditor/Private/DisplayClusterLightCardEditor.cpp
@@ -225,7 +225,7 @@ AActor* FDisplayClusterLightCardEditor::SpawnActor(TSubclassOf<AActor> InActorCl
 	AActor* NewActor = FDisplayClusterLightCardEditorHelper::SpawnStageActor(MoveTemp(SpawnArgs));
 	check(NewActor);
 
-	RefreshPreviewActors(EDisplayClusterLightCardEditorProxyType::StageActor);
+	RefreshPreviewStageActor(NewActor);
 	
 	if (!bIsPreview)
 	{
@@ -462,9 +462,9 @@ void FDisplayClusterLightCardEditor::AddLightCardsToActor(const TArray<ADisplayC
 
 				LightCard->AddToLightCardLayer(ActiveRootActor.Get());
 			}
-		}
 
-		RefreshPreviewActors(EDisplayClusterLightCardEditorProxyType::StageActor);
+			RefreshPreviewStageActor(LightCard);
+		}
 	}
 }
 
@@ -1626,6 +1626,23 @@ void FDisplayClusterLightCardEditor::RefreshPreviewActors(EDisplayClusterLightCa
 	BindCompileDelegates();
 }
 
+void FDisplayClusterLightCardEditor::RefreshPreviewStageActor(AActor* Actor)
+{
+	if (ADisplayClusterRootActor* RootActor = GetActiveRootActor().Get())
+	{
+		if (LightCardOutliner.IsValid())
+		{
+			LightCardOutliner->SetRootActor(RootActor);
+		}
+
+		if (ViewportView.IsValid())
+		{
+			ViewportView->GetLightCardEditorViewportClient()->UpdatePreviewActor(RootActor, /*bForce*/ true,
+				EDisplayClusterLightCardEditorProxyType::StageActor, Actor);
+		}
+	}
+}
+
 void FDisplayClusterLightCardEditor::RefreshLabels()
 {
 	ShowLightCardLabels(ShouldShowLightCardLabels());
@@ -1728,24 +1745,56 @@ void FDisplayClusterLightCardEditor::OnActorPropertyChanged(UObject* ObjectBeing
 	EDisplayClusterLightCardEditorProxyType ProxyType;
 	if (IsOurObject(ObjectBeingModified, ProxyType))
 	{
+		IDisplayClusterStageActor* StageActor = Cast<IDisplayClusterStageActor>(ObjectBeingModified);
+		
 		const FName PropertyName = PropertyChangedEvent.GetPropertyName();
 		const bool bTransformationChanged =
 				(PropertyName == USceneComponent::GetRelativeLocationPropertyName() ||
 					PropertyName == USceneComponent::GetRelativeRotationPropertyName() ||
-					PropertyName == USceneComponent::GetRelativeScale3DPropertyName());
+					PropertyName == USceneComponent::GetRelativeScale3DPropertyName() ||
+					(StageActor && StageActor->GetPositionalPropertyNames().Contains(PropertyName)));
 		if (PropertyChangedEvent.ChangeType == EPropertyChangeType::Interactive ||
 			(bTransformationChanged && !ObjectBeingModified->IsA<ADisplayClusterRootActor>()))
 		{
-			// Real-time & efficient update when dragging a slider.
+			// Real-time & efficient update when dragging a slider or when we know the property type doesn't
+			// require a full refresh
 			if (ViewportView.IsValid())
 			{
-				ViewportView->GetLightCardEditorViewportClient()->UpdateProxyTransforms();
+				if (AActor* Actor = Cast<AActor>(ObjectBeingModified))
+				{
+					ViewportView->GetLightCardEditorViewportClient()->UpdateProxyTransformFromLevelInstance(Actor);
+				}
+				else
+				{
+					ViewportView->GetLightCardEditorViewportClient()->UpdateProxyTransforms();
+				}
 			}
 		}
 		else
 		{
-			// Full destroy and refresh.
-			RefreshPreviewActors(ProxyType);
+			if (StageActor == nullptr)
+			{
+				// Check if this object belongs to a stage actor
+				StageActor = Cast<IDisplayClusterStageActor>(ObjectBeingModified->GetTypedOuter(UDisplayClusterStageActor::StaticClass()));
+			}
+			
+			if (AActor* Actor = Cast<AActor>(StageActor))
+			{
+				// Full destroy and refresh of a single proxy
+				RefreshPreviewStageActor(Actor);
+			}
+			else if (ObjectBeingModified->IsA<UActorComponent>() &&
+				CastChecked<UActorComponent>(ObjectBeingModified)->GetOwner() == GetActiveRootActor().Get())
+			{
+				// Avoid recreating all stage proxies if we can determine it's just a component that has changed
+				// Especially helps improve ICVFX property changes
+				RefreshPreviewActors(EDisplayClusterLightCardEditorProxyType::RootActor);
+			}
+			else
+			{
+				// Full destroy and refresh of all relevant proxies
+				RefreshPreviewActors(ProxyType);
+			}
 		}
 	}
 }
@@ -1754,24 +1803,25 @@ void FDisplayClusterLightCardEditor::OnLevelActorAdded(AActor* Actor)
 {
 	if (UE::DisplayClusterLightCardEditorUtils::IsManagedActor(Actor))
 	{
-		RefreshPreviewActors(EDisplayClusterLightCardEditorProxyType::StageActor);
+		RefreshPreviewStageActor(Actor);
 	}
 }
 
 void FDisplayClusterLightCardEditor::OnLevelActorDeleted(AActor* Actor)
 {
-	if (UE::DisplayClusterLightCardEditorUtils::IsManagedActor(Actor))
+	if (Actor && UE::DisplayClusterLightCardEditorUtils::IsManagedActor(Actor) &&
+		Actor->GetPackage() != GetTransientPackage() /* Don't trigger if this is a proxy actor being destroyed */)
 	{
-		if (Actor && Actor->GetClass()->HasAnyClassFlags(CLASS_NewerVersionExists))
+		if (Actor->GetClass()->HasAnyClassFlags(CLASS_NewerVersionExists))
 		{
 			// When a blueprint class is regenerated instances are deleted and replaced.
 			// In this case the OnCompiled() delegate will fire and refresh the actor.
 			return;
 		}
 		
-		if (GEditor && GEditor->GetWorld())
+		if (Actor->GetWorld())
 		{
-			GEditor->GetWorld()->GetTimerManager().SetTimerForNextTick([=]()
+			Actor->GetWorld()->GetTimerManager().SetTimerForNextTick([=]()
 			{
 				// Schedule for next tick so available selections are properly updated once the
 				// actor is fully deleted.
diff --git a/Engine/Plugins/Runtime/nDisplay/Source/DisplayClusterLightCardEditor/Private/DisplayClusterLightCardEditor.h b/Engine/Plugins/Runtime/nDisplay/Source/DisplayClusterLightCardEditor/Private/DisplayClusterLightCardEditor.h
--- a/Engine/Plugins/Runtime/nDisplay/Source/DisplayClusterLightCardEditor/Private/DisplayClusterLightCardEditor.h
+++ b/Engine/Plugins/Runtime/nDisplay/Source/DisplayClusterLightCardEditor/Private/DisplayClusterLightCardEditor.h
@@ -264,6 +264,9 @@ private:
 	/** Refresh all preview actors */
 	void RefreshPreviewActors(EDisplayClusterLightCardEditorProxyType ProxyType = EDisplayClusterLightCardEditorProxyType::All);
 
+	/** Refresh a specific preview stage actor */
+	void RefreshPreviewStageActor(AActor* Actor);
+	
 	/** Reapply the current label settings */
 	void RefreshLabels();
 	
diff --git a/Engine/Plugins/Runtime/nDisplay/Source/DisplayClusterLightCardEditor/Private/Viewport/DisplayClusterLightCardEditorViewportClient.cpp b/Engine/Plugins/Runtime/nDisplay/Source/DisplayClusterLightCardEditor/Private/Viewport/DisplayClusterLightCardEditorViewportClient.cpp
--- a/Engine/Plugins/Runtime/nDisplay/Source/DisplayClusterLightCardEditor/Private/Viewport/DisplayClusterLightCardEditorViewportClient.cpp
+++ b/Engine/Plugins/Runtime/nDisplay/Source/DisplayClusterLightCardEditor/Private/Viewport/DisplayClusterLightCardEditorViewportClient.cpp
@@ -41,6 +41,7 @@
 #include "KismetProceduralMeshLibrary.h"
 #include "LegacyScreenPercentageDriver.h"
 #include "Math/UnrealMathUtility.h"
+#include "Misc/ScopeLock.h"
 #include "PreviewScene.h"
 #include "ProceduralMeshComponent.h"
 #include "RayTracingDebugVisualizationMenuCommands.h"
@@ -219,13 +220,6 @@ void FDisplayClusterLightCardEditorViewportClient::Tick(float DeltaSeconds)
 			DesiredLookAtLocation.Reset();
 		}
 	}
-
-	if (BillboardComponentProxies.Num() > 0)
-	{
-		FSceneViewInitOptions SceneViewInitOptions;
-		GetSceneViewInitOptions(SceneViewInitOptions);
-		BillboardViewMatrices = FViewMatrices(SceneViewInitOptions);
-	}
 }
 
 void FDisplayClusterLightCardEditorViewportClient::Draw(FViewport* InViewport, FCanvas* Canvas)
@@ -326,6 +320,8 @@ void FDisplayClusterLightCardEditorViewportClient::Draw(FViewport* InViewport, F
 
 	Canvas->Clear(FLinearColor::Black);
 
+	FSceneViewInitOptions BillboardSceneViewInitOptions;
+	
 	if (!bDisableCustomRenderer)
 	{
 		FDisplayClusterMeshProjectionRenderSettings RenderSettings;
@@ -338,13 +334,26 @@ void FDisplayClusterLightCardEditorViewportClient::Draw(FViewport* InViewport, F
 
 		GetSceneViewInitOptions(RenderSettings.ViewInitOptions);
 
+		BillboardSceneViewInitOptions = RenderSettings.ViewInitOptions;
+		
 		IDisplayClusterScenePreview::Get().Render(PreviewRendererId, RenderSettings, *Canvas);
 	}
 	else
 	{
+		if (BillboardComponentProxies.Num() > 0)
+		{
+			GetSceneViewInitOptions(BillboardSceneViewInitOptions);
+		}
+		
 		GetRendererModule().BeginRenderingViewFamily(Canvas, &ViewFamily);
 	}
 
+	if (BillboardComponentProxies.Num() > 0)
+	{
+		FScopeLock Lock(&BillboardComponentCS);
+		BillboardViewMatrices = FViewMatrices(BillboardSceneViewInitOptions);
+	}
+
 	if (View)
 	{
 		DrawCanvas(*Viewport, *View, *Canvas);
@@ -453,54 +462,59 @@ void FDisplayClusterLightCardEditorViewportClient::Draw(const FSceneView* View,
 {
 	const bool bIsUVProjection = ProjectionMode == EDisplayClusterMeshProjectionType::UV;
 
-	// Draw any sprites
-	for (const TWeakObjectPtr<UBillboardComponent>& BillboardComponent : BillboardComponentProxies)
+	// Draw any sprites - this would normally go under the CPU draw which would avoid needing critical regions, but
+	// we can't render the gizmo over the icons in that case. Canvas items drawn onto the viewport will draw onto the
+	// back buffer after the entire custom render pipeline is drawn.
 	{
-		if (!BillboardComponent.IsValid())
+		FScopeLock Lock(&BillboardComponentCS);
+		for (const TWeakObjectPtr<UBillboardComponent>& BillboardComponent : BillboardComponentProxies)
 		{
-			continue;
-		}
-		FSpriteProxy SpriteProxy = FSpriteProxy::FromBillboard(BillboardComponent.Get());
+			if (!BillboardComponent.IsValid())
+			{
+				continue;
+			}
+			FSpriteProxy SpriteProxy = FSpriteProxy::FromBillboard(BillboardComponent.Get());
 		
-		FVector ProjectedLocation = ProjectWorldPosition(SpriteProxy.WorldPosition, BillboardViewMatrices);
+			FVector ProjectedLocation = ProjectWorldPosition(SpriteProxy.WorldPosition, BillboardViewMatrices);
 		
-		if (const FTexture* TextureResource = SpriteProxy.Sprite ? SpriteProxy.Sprite->GetResource() : nullptr)
-		{
-			const UDisplayClusterLightCardEditorSettings* Settings = GetDefault<UDisplayClusterLightCardEditorSettings>();
-			
 			PDI->SetHitProxy(new HActor(BillboardComponent->GetOwner(), BillboardComponent.Get()));
-
-			if (Settings->bDisplayIcons)
+		
+			if (const FTexture* TextureResource = SpriteProxy.Sprite ? SpriteProxy.Sprite->GetResource() : nullptr)
 			{
-				const float UL = SpriteProxy.UL == 0.0f ? TextureResource->GetSizeX() : SpriteProxy.UL;
-				const float VL = SpriteProxy.VL == 0.0f ? TextureResource->GetSizeY() : SpriteProxy.VL;
-			
-				// Calculate the view-dependent scaling factor.
-				float ViewedSizeX = SpriteProxy.SpriteScale * UL;
-				float ViewedSizeY = SpriteProxy.SpriteScale * VL;
+				const UDisplayClusterLightCardEditorSettings* Settings = GetDefault<UDisplayClusterLightCardEditorSettings>();
 
-				if (SpriteProxy.bIsScreenSizeScaled && (View->ViewMatrices.GetProjectionMatrix().M[3][3] != 1.0f))
+				if (Settings->bDisplayIcons)
 				{
-					const float ZoomFactor = FMath::Min<float>(View->ViewMatrices.GetProjectionMatrix().M[0][0], View->ViewMatrices.GetProjectionMatrix().M[1][1]);
-					if(ZoomFactor != 0.0f)
-					{
-						const float Radius = View->WorldToScreen(ProjectedLocation).W * (SpriteProxy.ScreenSize / ZoomFactor);
+					const float UL = SpriteProxy.UL == 0.0f ? TextureResource->GetSizeX() : SpriteProxy.UL;
+					const float VL = SpriteProxy.VL == 0.0f ? TextureResource->GetSizeY() : SpriteProxy.VL;
+			
+					// Calculate the view-dependent scaling factor.
+					float ViewedSizeX = SpriteProxy.SpriteScale * UL;
+					float ViewedSizeY = SpriteProxy.SpriteScale * VL;
 
-						if (Radius < 1.0f)
+					if (SpriteProxy.bIsScreenSizeScaled && (View->ViewMatrices.GetProjectionMatrix().M[3][3] != 1.0f))
+					{
+						const float ZoomFactor = FMath::Min<float>(View->ViewMatrices.GetProjectionMatrix().M[0][0], View->ViewMatrices.GetProjectionMatrix().M[1][1]);
+						if(ZoomFactor != 0.0f)
 						{
-							ViewedSizeX *= Radius;
-							ViewedSizeY *= Radius;
-						}					
+							const float Radius = View->WorldToScreen(ProjectedLocation).W * (SpriteProxy.ScreenSize / ZoomFactor);
+
+							if (Radius < 1.0f)
+							{
+								ViewedSizeX *= Radius;
+								ViewedSizeY *= Radius;
+							}					
+						}
 					}
-				}
 				
-				ViewedSizeX *= Settings->IconScale;
-				ViewedSizeY *= Settings->IconScale;
+					ViewedSizeX *= Settings->IconScale;
+					ViewedSizeY *= Settings->IconScale;
 			
-				PDI->DrawSprite(ProjectedLocation, ViewedSizeX, ViewedSizeY, TextureResource, FLinearColor::White, SDPG_World,
-					SpriteProxy.U, UL, SpriteProxy.V, VL, SE_BLEND_Masked, SpriteProxy.OpacityMaskRefVal);
+					PDI->DrawSprite(ProjectedLocation, ViewedSizeX, ViewedSizeY, TextureResource, FLinearColor::White, SDPG_World,
+						SpriteProxy.U, UL, SpriteProxy.V, VL, SE_BLEND_Masked, SpriteProxy.OpacityMaskRefVal);
+				}
 			}
-			
+
 			PDI->SetHitProxy(nullptr);
 		}
 	}
@@ -1053,6 +1067,85 @@ const FDisplayClusterLightCardEditorViewportClient::FActorProxy* FDisplayCluster
 	});
 }
 
+const FDisplayClusterLightCardEditorViewportClient::FActorProxy* FDisplayClusterLightCardEditorViewportClient::FindActorProxyFromActor(AActor* InActor) const
+{
+	if (InActor == nullptr)
+	{
+		return nullptr;
+	}
+	
+	return ActorProxies.FindByPredicate([InActor](const FActorProxy& Proxy)
+	{
+		return Proxy.LevelInstance.AsActor() == InActor || Proxy.Proxy.AsActor() == InActor;
+	});
+}
+
+void FDisplayClusterLightCardEditorViewportClient::CreateStageActorProxy(AActor* InLevelInstance)
+{
+	if (!IsValid(InLevelInstance))
+	{
+		// Can happen if the level actor was destroyed the tick prior to proxy creation.
+		return;
+	}
+	
+	const FTransform RALevelTransformNoScale(RootActorLevelInstance->GetActorRotation(), RootActorLevelInstance->GetActorLocation(), FVector::OneVector);
+
+	const UWorld* PreviewWorld = PreviewScene->GetWorld();
+	check(PreviewWorld);
+	
+	// Create proxy
+	FObjectDuplicationParameters DupeActorParameters(InLevelInstance, PreviewWorld->GetCurrentLevel());
+	DupeActorParameters.FlagMask = RF_AllFlags & ~(RF_ArchetypeObject | RF_Transactional);
+	DupeActorParameters.PortFlags = PPF_DuplicateVerbatim;
+
+	AActor* ActorProxy = CastChecked<AActor>(StaticDuplicateObjectEx(DupeActorParameters));
+	ActorProxy->SetFlags(RF_Transient); // This signals to the stage actor it is a proxy
+	PreviewWorld->GetCurrentLevel()->AddLoadedActor(ActorProxy);
+
+	const FTransform LCLevelRelativeToRALevel = ActorProxy->GetTransform().GetRelativeTransform(RALevelTransformNoScale);
+	ActorProxy->SetActorTransform(LCLevelRelativeToRALevel);
+	
+	if (ADisplayClusterLightCardActor* LightCard = Cast<ADisplayClusterLightCardActor>(InLevelInstance))
+	{
+		ADisplayClusterLightCardActor* LightCardProxy = CastChecked<ADisplayClusterLightCardActor>(ActorProxy);
+		LightCardProxy->PolygonMask = LightCard->PolygonMask;
+
+		// Change mesh to proxy mesh with more vertices
+		if (const UStaticMesh* LightCardMesh = LightCardProxy->GetStaticMesh())
+		{
+			// Only change the mesh if we are using the default one.
+
+			const FString DefaultPlanePath = TEXT("/nDisplay/LightCard/SM_LightCardPlane.SM_LightCardPlane");
+			const UStaticMesh* DefaultPlane = Cast<UStaticMesh>(FSoftObjectPath(DefaultPlanePath).TryLoad());
+
+			if (DefaultPlane == LightCardMesh)
+			{
+				const FString LightCardPlanePath = TEXT("/nDisplay/LightCard/SM_LightCardPlaneSubdivided.SM_LightCardPlaneSubdivided");
+				if (UStaticMesh* LightCardPlaneMesh = Cast<UStaticMesh>(FSoftObjectPath(LightCardPlanePath).TryLoad()))
+				{
+					LightCardProxy->SetStaticMesh(LightCardPlaneMesh);
+				}
+			}
+		}
+	}
+	else
+	{
+		// Misc actor proxies
+		TArray<UBillboardComponent*> BillboardComponents;
+		ActorProxy->GetComponents<UBillboardComponent>(BillboardComponents);
+
+		FScopeLock Lock(&BillboardComponentCS);
+		BillboardComponentProxies.Append(BillboardComponents);
+	}
+
+	FActorProxy ActorProxyStruct(InLevelInstance, ActorProxy);
+	ActorProxies.Add(ActorProxyStruct);
+	ProjectionHelper->VerifyAndFixActorOrigin(InLevelInstance);
+	ProjectionHelper->VerifyAndFixActorOrigin(ActorProxy);
+
+	UpdateProxyTransforms(ActorProxyStruct);
+}
+
 void FDisplayClusterLightCardEditorViewportClient::UpdateProxyTransforms(const FActorProxy& InActorProxy)
 {
 	DECLARE_SCOPE_CYCLE_COUNTER(TEXT("FDisplayClusterLightCardEditorViewportClient::UpdateProxyTransforms"), STAT_UpdateProxyTransforms, STATGROUP_NDisplayLightCardEditor);
@@ -1366,24 +1459,32 @@ bool FDisplayClusterLightCardEditorViewportClient::DropObjectsAtCoordinates(int3
 }
 
 void FDisplayClusterLightCardEditorViewportClient::UpdatePreviewActor(ADisplayClusterRootActor* RootActor, bool bForce,
-                                                                      EDisplayClusterLightCardEditorProxyType ProxyType)
+                                                                      EDisplayClusterLightCardEditorProxyType ProxyType,
+                                                                      AActor* StageActor)
 {
 	DECLARE_SCOPE_CYCLE_COUNTER(TEXT("FDisplayClusterLightCardEditorViewportClient::UpdatePreviewActor"), STAT_UpdatePreviewActor, STATGROUP_NDisplayLightCardEditor);
 	
-	if ((!bForce && RootActor == RootActorLevelInstance.GetEvenIfUnreachable()) ||
-		(ProxyTypesRefreshing.Contains(ProxyType) || ProxyTypesRefreshing.Contains(EDisplayClusterLightCardEditorProxyType::All)))
+	if (!StageActor && ((!bForce && RootActor == RootActorLevelInstance.GetEvenIfUnreachable()) ||
+		(ProxyTypesRefreshing.Contains(ProxyType) || ProxyTypesRefreshing.Contains(EDisplayClusterLightCardEditorProxyType::All))))
 	{
 		return;
 	}
 
-	ProxyTypesRefreshing.Add(ProxyType);
-
-	auto Finalize = [this, ProxyType]()
+	if (!StageActor)
+	{
+		ProxyTypesRefreshing.Add(ProxyType);
+	}
+	
+	auto Finalize = [this, ProxyType, StageActor]()
 	{
 		Viewport->InvalidateHitProxy();
 		bShouldCheckHitProxy = true;
 
-		ProxyTypesRefreshing.Remove(ProxyType);
+		if (!StageActor)
+		{
+			ProxyTypesRefreshing.Remove(ProxyType);
+		}
+		
 		OnNextSceneRefreshDelegate.Broadcast();
 		OnNextSceneRefreshDelegate.Clear();
 	};
@@ -1395,7 +1496,7 @@ void FDisplayClusterLightCardEditorViewportClient::UpdatePreviewActor(ADisplayCl
 	}
 	else
 	{
-		UWorld* PreviewWorld = PreviewScene->GetWorld();
+		const UWorld* PreviewWorld = PreviewScene->GetWorld();
 		check(PreviewWorld);
 
 		const TWeakObjectPtr<ADisplayClusterRootActor> RootActorPtr (RootActor);
@@ -1405,13 +1506,22 @@ void FDisplayClusterLightCardEditorViewportClient::UpdatePreviewActor(ADisplayCl
 		// regeneration & compiles. nDisplay's custom propagation may have issues if the archetype isn't correct.
 		PreviewWorld->GetTimerManager().SetTimerForNextTick([=]()
 		{
+			DECLARE_SCOPE_CYCLE_COUNTER(TEXT("FDisplayClusterLightCardEditorViewportClient::UpdatePreviewActorImpl"), STAT_UpdatePreviewActorImpl, STATGROUP_NDisplayLightCardEditor);
+			
 			if (!WeakPtrThis.IsValid())
 			{
 				return;
 			}
-			
-			DestroyProxies(ProxyType);
 
+			if (StageActor)
+			{
+				DestroyProxy(StageActor);
+			}
+			else
+			{
+				DestroyProxies(ProxyType);
+			}
+			
 			if (!RootActorPtr.IsValid())
 			{
 				Finalize();
@@ -1465,64 +1575,20 @@ void FDisplayClusterLightCardEditorViewportClient::UpdatePreviewActor(ADisplayCl
 				ProxyType == EDisplayClusterLightCardEditorProxyType::StageActor)
 			{
 				SelectActor(nullptr);
-				const FTransform RALevelTransformNoScale(RootActorLevelInstance->GetActorRotation(), RootActorLevelInstance->GetActorLocation(), FVector::OneVector);
 
 				if (LightCardEditorPtr.IsValid())
 				{
-					TArray<AActor*> ManagedActors = LightCardEditorPtr.Pin()->FindAllManagedActors();
-					
-					for (AActor* Actor : ManagedActors)
+					if (StageActor)
 					{
-						// Create proxy
-						FObjectDuplicationParameters DupeActorParameters(Actor, PreviewWorld->GetCurrentLevel());
-						DupeActorParameters.FlagMask = RF_AllFlags & ~(RF_ArchetypeObject | RF_Transactional);
-						DupeActorParameters.PortFlags = PPF_DuplicateVerbatim;
-
-						AActor* ActorProxy = CastChecked<AActor>(StaticDuplicateObjectEx(DupeActorParameters));
-						ActorProxy->SetFlags(RF_Transient); // This signals to the stage actor it is a proxy
-						PreviewWorld->GetCurrentLevel()->AddLoadedActor(ActorProxy);
-
-						const FTransform LCLevelRelativeToRALevel = ActorProxy->GetTransform().GetRelativeTransform(RALevelTransformNoScale);
-						ActorProxy->SetActorTransform(LCLevelRelativeToRALevel);
-						
-						if (ADisplayClusterLightCardActor* LightCard = Cast<ADisplayClusterLightCardActor>(Actor))
-						{
-							ADisplayClusterLightCardActor* LightCardProxy = CastChecked<ADisplayClusterLightCardActor>(ActorProxy);
-							LightCardProxy->PolygonMask = LightCard->PolygonMask;
-
-							// Change mesh to proxy mesh with more vertices
-							if (const UStaticMesh* LightCardMesh = LightCardProxy->GetStaticMesh())
-							{
-								// Only change the mesh if we are using the default one.
-					
-								const FString DefaultPlanePath = TEXT("/nDisplay/LightCard/SM_LightCardPlane.SM_LightCardPlane");
-								const UStaticMesh* DefaultPlane = Cast<UStaticMesh>(FSoftObjectPath(DefaultPlanePath).TryLoad());
-
-								if (DefaultPlane == LightCardMesh)
-								{
-									const FString LightCardPlanePath = TEXT("/nDisplay/LightCard/SM_LightCardPlaneSubdivided.SM_LightCardPlaneSubdivided");
-									if (UStaticMesh* LightCardPlaneMesh = Cast<UStaticMesh>(FSoftObjectPath(LightCardPlanePath).TryLoad()))
-									{
-										LightCardProxy->SetStaticMesh(LightCardPlaneMesh);
-									}
-								}
-							}
-						}
-						else
+						CreateStageActorProxy(StageActor);
+					}
+					else
+					{
+						TArray<AActor*> ManagedActors = LightCardEditorPtr.Pin()->FindAllManagedActors();
+						for (AActor* Actor : ManagedActors)
 						{
-							// Misc actor proxies
-							TArray<UBillboardComponent*> BillboardComponents;
-							ActorProxy->GetComponents<UBillboardComponent>(BillboardComponents);
-
-							BillboardComponentProxies.Append(BillboardComponents);
+							CreateStageActorProxy(Actor);
 						}
-
-						FActorProxy ActorProxyStruct(Actor, ActorProxy);
-						ActorProxies.Add(ActorProxyStruct);
-						ProjectionHelper->VerifyAndFixActorOrigin(Actor);
-						ProjectionHelper->VerifyAndFixActorOrigin(ActorProxy);
-
-						UpdateProxyTransforms(ActorProxyStruct);
 					}
 				}
 
@@ -1536,6 +1602,7 @@ void FDisplayClusterLightCardEditorViewportClient::UpdatePreviewActor(ADisplayCl
 				SelectActors(CurrentlySelectedActors);
 			}
 
+			// Make sure proxies are added to the renderer. Necessary for selections to render even if stage actors were not modified but root actor was updated
 			for (const FActorProxy& ActorProxy : ActorProxies)
 			{
 				IDisplayClusterScenePreview::Get().AddActorToRenderer(PreviewRendererId, ActorProxy.Proxy.AsActor(), [this, ActorProxy](const UPrimitiveComponent* PrimitiveComponent)
@@ -1585,6 +1652,8 @@ void FDisplayClusterLightCardEditorViewportClient::UpdateProxyTransformFromLevel
 void FDisplayClusterLightCardEditorViewportClient::DestroyProxies(
 	EDisplayClusterLightCardEditorProxyType ProxyType)
 {
+	FScopeLock Lock(&BillboardComponentCS);
+
 	// Clear the primitives from the scene renderer based on the type of proxy that is being destroyed
 	switch (ProxyType)
 	{
@@ -1648,6 +1717,32 @@ void FDisplayClusterLightCardEditorViewportClient::DestroyProxies(
 	}
 }
 
+void FDisplayClusterLightCardEditorViewportClient::DestroyProxy(AActor* Actor)
+{
+	if (const FActorProxy* ActorProxy = FindActorProxyFromActor(Actor))
+	{
+		if (AActor* ActualProxy = ActorProxy->Proxy.AsActor())
+		{
+			ActorProxies.RemoveAll([ActualProxy](const FActorProxy& OtherProxy)
+			{
+				return OtherProxy == ActualProxy;
+			});
+		
+			FScopeLock Lock(&BillboardComponentCS);
+		
+			UWorld* PreviewWorld = PreviewScene->GetWorld();
+			check(PreviewWorld);
+			IDisplayClusterScenePreview::Get().RemoveActorFromRenderer(PreviewRendererId, ActualProxy);
+			PreviewWorld->EditorDestroyActor(ActualProxy, false);
+			
+			BillboardComponentProxies.RemoveAll([ActualProxy](const TWeakObjectPtr<UBillboardComponent>& BillboardComponent)
+			{
+				return !BillboardComponent.IsValid() || BillboardComponent->GetOwner() == ActualProxy;
+			});
+		}
+	}
+}
+
 void FDisplayClusterLightCardEditorViewportClient::SelectActors(const TArray<AActor*>& ActorsToSelect)
 {
 	SelectActor(nullptr);
diff --git a/Engine/Plugins/Runtime/nDisplay/Source/DisplayClusterLightCardEditor/Private/Viewport/DisplayClusterLightCardEditorViewportClient.h b/Engine/Plugins/Runtime/nDisplay/Source/DisplayClusterLightCardEditor/Private/Viewport/DisplayClusterLightCardEditorViewportClient.h
--- a/Engine/Plugins/Runtime/nDisplay/Source/DisplayClusterLightCardEditor/Private/Viewport/DisplayClusterLightCardEditorViewportClient.h
+++ b/Engine/Plugins/Runtime/nDisplay/Source/DisplayClusterLightCardEditor/Private/Viewport/DisplayClusterLightCardEditorViewportClient.h
@@ -81,10 +81,12 @@ public:
 	 * @param RootActor The new root actor to use. Accepts nullptr
 	 * @param bForce Force the update even if the RootActor hasn't changed
 	 * @param ProxyType The proxy type to destroy and update
+	 * @param StageActor Limit the update to a single stage actor
 	 */
 	void UpdatePreviewActor(ADisplayClusterRootActor* RootActor, bool bForce = false,
-		EDisplayClusterLightCardEditorProxyType ProxyType = EDisplayClusterLightCardEditorProxyType::All);
-
+		EDisplayClusterLightCardEditorProxyType ProxyType = EDisplayClusterLightCardEditorProxyType::All,
+		AActor* StageActor = nullptr);
+	
 	/** Only update required transform values of proxies */
 	void UpdateProxyTransforms();
 
@@ -93,6 +95,9 @@ public:
 	
 	/** Remove proxies of the specified type */
 	void DestroyProxies(EDisplayClusterLightCardEditorProxyType ProxyType);
+
+	/** Remove a proxy that is or belongs to the given actor */
+	void DestroyProxy(AActor* Actor);
 	
 	/** Selects the actor proxies that correspond to the specified actors */
 	void SelectActors(const TArray<AActor*>& ActorsToSelect);
@@ -279,6 +284,12 @@ private:
 	/** Find the actor proxy if it exists from a level instance */
 	const FActorProxy* FindActorProxyFromLevelInstance(AActor* InLevelInstance) const;
 
+	/** Find the actor proxy from either the proxy itself or the actor the proxy belongs to */
+	const FActorProxy* FindActorProxyFromActor(AActor* InActor) const;
+
+	/** Create a stage actor proxy */
+	void CreateStageActorProxy(AActor* InLevelInstance);
+	
 	/** Update the transforms for the specific proxy */
 	void UpdateProxyTransforms(const FActorProxy& InActorProxy);
 	
@@ -309,6 +320,9 @@ private:
 	TArray<TWeakObjectPtr<UBillboardComponent>> BillboardComponentProxies;
 	TArray<FDisplayClusterWeakStageActorPtr> SelectedActors;
 
+	/** Billboards are accessed from the render thread and game thread */
+	FCriticalSection BillboardComponentCS;
+
 	/** View matrices set per tick if there are any billboard components */
 	FViewMatrices BillboardViewMatrices;
 
