nDisplay: Add blueprint setting to enable construction scripts on interactive change, defaulted to false.
Limit ICVFX panel subscriptions to the root actor to prevent the preview from being recreated every time a proxy is recreated.
Fix an issue where a stale NDC from a reinstancing would be held on to by the viewport client and cause the preview to not render.

diff --git a/Engine/Plugins/Runtime/nDisplay/Source/DisplayCluster/Private/Blueprints/DisplayClusterBlueprint.cpp b/Engine/Plugins/Runtime/nDisplay/Source/DisplayCluster/Private/Blueprints/DisplayClusterBlueprint.cpp
--- a/Engine/Plugins/Runtime/nDisplay/Source/DisplayCluster/Private/Blueprints/DisplayClusterBlueprint.cpp
+++ b/Engine/Plugins/Runtime/nDisplay/Source/DisplayCluster/Private/Blueprints/DisplayClusterBlueprint.cpp
@@ -20,6 +20,9 @@ UDisplayClusterBlueprint::UDisplayClusterBlueprint()
 	: ConfigData(nullptr), AssetVersion(0)
 {
 	BlueprintType = BPTYPE_Normal;
+#if WITH_EDITORONLY_DATA
+	bRunConstructionScriptOnInteractiveChange = false;
+#endif
 }
 
 #if WITH_EDITOR
diff --git a/Engine/Plugins/Runtime/nDisplay/Source/DisplayCluster/Private/Game/EngineClasses/Scene/DisplayClusterRootActorEditor.cpp b/Engine/Plugins/Runtime/nDisplay/Source/DisplayCluster/Private/Game/EngineClasses/Scene/DisplayClusterRootActorEditor.cpp
--- a/Engine/Plugins/Runtime/nDisplay/Source/DisplayCluster/Private/Game/EngineClasses/Scene/DisplayClusterRootActorEditor.cpp
+++ b/Engine/Plugins/Runtime/nDisplay/Source/DisplayCluster/Private/Game/EngineClasses/Scene/DisplayClusterRootActorEditor.cpp
@@ -2,6 +2,7 @@
 
 #include "DisplayClusterRootActor.h"
 
+#include "Blueprints/DisplayClusterBlueprint.h"
 #include "Components/SceneComponent.h"
 #include "Components/DisplayClusterOriginComponent.h"
 #include "Components/DisplayClusterCameraComponent.h"
@@ -321,6 +322,8 @@ IDisplayClusterViewport* ADisplayClusterRootActor::FindPreviewViewport(const FSt
 
 bool ADisplayClusterRootActor::ImplUpdatePreviewConfiguration_Editor(const FString& InClusterNodeId)
 {
+	DECLARE_SCOPE_CYCLE_COUNTER(TEXT("ADisplayClusterRootActor::ImplUpdatePreviewConfiguration_Editor"), STAT_ImplUpdatePreviewConfiguration_Editor, STATGROUP_NDisplay);
+	
 	PreviewRenderFrame.Reset();
 
 	if (IsPreviewEnabled() && ViewportManager.IsValid())
@@ -694,9 +697,6 @@ static FName Name_RelativeScale3D = USceneComponent::GetRelativeScale3DPropertyN
 
 void ADisplayClusterRootActor::PostEditChangeProperty(FPropertyChangedEvent& PropertyChangedEvent)
 {
-	// any property update causes a DCRA rebuild. Restore preview material show
-	ResetPreviewComponents_Editor(true);
-
 	const FName PropertyName = (PropertyChangedEvent.Property != nullptr) ? PropertyChangedEvent.Property->GetFName() : NAME_None;
 	
 	// The AActor method, simplified and modified to skip construction scripts.
@@ -731,18 +731,23 @@ void ADisplayClusterRootActor::PostEditChangeProperty(FPropertyChangedEvent& Pro
 
 	bool bReinitializeActor = true;
 	bool bCanSkipConstructionScripts = false;
-	if (const UBlueprint* Blueprint = UBlueprint::GetBlueprintFromClass(GetClass()))
+	bool bResetPreviewComponents = false;
+	if (const UDisplayClusterBlueprint* Blueprint = Cast<UDisplayClusterBlueprint>(UBlueprint::GetBlueprintFromClass(GetClass())))
 	{
-		bCanSkipConstructionScripts = !Blueprint->bRunConstructionScriptOnDrag;
+		bCanSkipConstructionScripts = !Blueprint->bRunConstructionScriptOnInteractiveChange;
 	}
 	
 	if (bCanSkipConstructionScripts && PropertyChangedEvent.ChangeType == EPropertyChangeType::Interactive && !CurrentTransactionAnnotation.IsValid())
 	{
 		// Avoid calling construction scripts when the change occurs while the user is dragging a slider.
 		SuperCallWithoutConstructionScripts();
+		bReinitializeActor = false;
 	}
 	else
 	{
+		// any property update causes a DCRA rebuild. Restore preview material show
+		ResetPreviewComponents_Editor(true);
+		bResetPreviewComponents = true;
 		Super::PostEditChangeProperty(PropertyChangedEvent);
 	}
 	
@@ -780,7 +785,10 @@ void ADisplayClusterRootActor::PostEditChangeProperty(FPropertyChangedEvent& Pro
 		InitializeRootActor();
 	}
 
-	ResetPreviewComponents_Editor(false);
+	if (bResetPreviewComponents)
+	{
+		ResetPreviewComponents_Editor(false);
+	}
 }
 
 void ADisplayClusterRootActor::PostEditMove(bool bFinished)
@@ -818,6 +826,8 @@ void ADisplayClusterRootActor::ResetPreviewComponents_Editor(bool bInRestoreScen
 
 void ADisplayClusterRootActor::UpdatePreviewComponents()
 {
+	DECLARE_SCOPE_CYCLE_COUNTER(TEXT("ADisplayClusterRootActor::UpdatePreviewComponents"), STAT_UpdatePreviewComponents, STATGROUP_NDisplay);
+	
 	if (IsTemplate() || bDeferPreviewGeneration)
 	{
 		return;
diff --git a/Engine/Plugins/Runtime/nDisplay/Source/DisplayCluster/Public/Blueprints/DisplayClusterBlueprint.h b/Engine/Plugins/Runtime/nDisplay/Source/DisplayCluster/Public/Blueprints/DisplayClusterBlueprint.h
--- a/Engine/Plugins/Runtime/nDisplay/Source/DisplayCluster/Public/Blueprints/DisplayClusterBlueprint.h
+++ b/Engine/Plugins/Runtime/nDisplay/Source/DisplayCluster/Public/Blueprints/DisplayClusterBlueprint.h
@@ -15,7 +15,13 @@ class DISPLAYCLUSTER_API UDisplayClusterBlueprint : public UBlueprint
 
 public:
 	UDisplayClusterBlueprint();
-	
+
+#if WITH_EDITORONLY_DATA
+	/** Whether or not you want to continuously rerun the construction script while dragging a slider */
+	UPROPERTY(EditAnywhere, Category=BlueprintOptions)
+	uint8 bRunConstructionScriptOnInteractiveChange : 1;
+#endif
+
 #if WITH_EDITOR
 	// UBlueprint
 	virtual bool SupportedByDefaultBlueprintFactory() const override { return false; }
diff --git a/Engine/Plugins/Runtime/nDisplay/Source/DisplayClusterLightCardEditor/Private/Viewport/DisplayClusterLightCardEditorViewportClient.cpp b/Engine/Plugins/Runtime/nDisplay/Source/DisplayClusterLightCardEditor/Private/Viewport/DisplayClusterLightCardEditorViewportClient.cpp
--- a/Engine/Plugins/Runtime/nDisplay/Source/DisplayClusterLightCardEditor/Private/Viewport/DisplayClusterLightCardEditorViewportClient.cpp
+++ b/Engine/Plugins/Runtime/nDisplay/Source/DisplayClusterLightCardEditor/Private/Viewport/DisplayClusterLightCardEditorViewportClient.cpp
@@ -132,10 +132,7 @@ FDisplayClusterLightCardEditorViewportClient::~FDisplayClusterLightCardEditorVie
 	IDisplayClusterScenePreview::Get().DestroyRenderer(PreviewRendererId);
 
 	EndTransaction();
-	if (RootActorLevelInstance.IsValid())
-	{
-		RootActorLevelInstance->UnsubscribeFromPostProcessRenderTarget(reinterpret_cast<uint8*>(this));
-	}
+	UnsubscribeFromRootActor();
 }
 
 FLinearColor FDisplayClusterLightCardEditorViewportClient::GetBackgroundColor() const
@@ -220,6 +217,13 @@ void FDisplayClusterLightCardEditorViewportClient::Tick(float DeltaSeconds)
 			DesiredLookAtLocation.Reset();
 		}
 	}
+
+	// The root actor could be stale if it was set during a re-instancing from a compile, but the editor should have the up-to-date actor.
+	if (RootActorLevelInstance.IsStale() && LightCardEditorPtr.IsValid() &&
+		(LightCardEditorPtr.Pin()->GetActiveRootActor().IsValid() || LightCardEditorPtr.Pin()->GetActiveRootActor().IsExplicitlyNull()))
+	{
+		UpdatePreviewActor(LightCardEditorPtr.Pin()->GetActiveRootActor().Get(), true);
+	}
 }
 
 void FDisplayClusterLightCardEditorViewportClient::Draw(FViewport* InViewport, FCanvas* Canvas)
@@ -1168,6 +1172,26 @@ void FDisplayClusterLightCardEditorViewportClient::UpdateProxyTransforms(const F
 	}
 }
 
+void FDisplayClusterLightCardEditorViewportClient::SubscribeToRootActor()
+{
+	if (RootActorLevelInstance.IsValid())
+	{
+		const uint8* GenericThis = reinterpret_cast<uint8*>(this);
+		RootActorLevelInstance->SubscribeToPostProcessRenderTarget(GenericThis);
+		RootActorLevelInstance->AddPreviewEnableOverride(GenericThis);
+	}
+}
+
+void FDisplayClusterLightCardEditorViewportClient::UnsubscribeFromRootActor()
+{
+	if (RootActorLevelInstance.IsValid())
+	{
+		const uint8* GenericThis = reinterpret_cast<uint8*>(this);
+		RootActorLevelInstance->UnsubscribeFromPostProcessRenderTarget(GenericThis);
+		RootActorLevelInstance->RemovePreviewEnableOverride(GenericThis);
+	}
+}
+
 void FDisplayClusterLightCardEditorViewportClient::ProcessClick(FSceneView& View, HHitProxy* HitProxy, FKey Key, EInputEvent Event, uint32 HitX, uint32 HitY)
 {
 	// Don't select light cards while drawing a new light card
@@ -1492,6 +1516,9 @@ void FDisplayClusterLightCardEditorViewportClient::UpdatePreviewActor(ADisplayCl
 	if (RootActor == nullptr)
 	{
 		DestroyProxies(ProxyType);
+		UnsubscribeFromRootActor();
+		RootActorLevelInstance.Reset();
+		
 		Finalize();
 	}
 	else
@@ -1528,10 +1555,14 @@ void FDisplayClusterLightCardEditorViewportClient::UpdatePreviewActor(ADisplayCl
 				return;
 			}
 
-			const uint8* GenericThis = reinterpret_cast<uint8*>(this);
-			RootActorPtr->SubscribeToPostProcessRenderTarget(GenericThis);
-			RootActorPtr->AddPreviewEnableOverride(GenericThis);
+			if (RootActorLevelInstance.IsValid() && RootActorLevelInstance.Get() != RootActor)
+			{
+				UnsubscribeFromRootActor();
+			}
+			
 			RootActorLevelInstance = RootActorPtr;
+
+			SubscribeToRootActor();
 			
 			if (ProxyType == EDisplayClusterLightCardEditorProxyType::All ||
 				ProxyType == EDisplayClusterLightCardEditorProxyType::RootActor)
@@ -1691,14 +1722,6 @@ void FDisplayClusterLightCardEditorViewportClient::DestroyProxies(
 			PreviewWorld->EditorDestroyActor(RootActorProxy.Get(), false);
 			RootActorProxy.Reset();
 		}
-
-		if (RootActorLevelInstance.IsValid())
-		{
-			const uint8* GenericThis = reinterpret_cast<uint8*>(this);
-			RootActorLevelInstance->UnsubscribeFromPostProcessRenderTarget(GenericThis);
-			RootActorLevelInstance->RemovePreviewEnableOverride(GenericThis);
-			RootActorLevelInstance.Reset();
-		}
 	}
 	
 	if (ProxyType == EDisplayClusterLightCardEditorProxyType::All ||
diff --git a/Engine/Plugins/Runtime/nDisplay/Source/DisplayClusterLightCardEditor/Private/Viewport/DisplayClusterLightCardEditorViewportClient.h b/Engine/Plugins/Runtime/nDisplay/Source/DisplayClusterLightCardEditor/Private/Viewport/DisplayClusterLightCardEditorViewportClient.h
--- a/Engine/Plugins/Runtime/nDisplay/Source/DisplayClusterLightCardEditor/Private/Viewport/DisplayClusterLightCardEditorViewportClient.h
+++ b/Engine/Plugins/Runtime/nDisplay/Source/DisplayClusterLightCardEditor/Private/Viewport/DisplayClusterLightCardEditorViewportClient.h
@@ -292,6 +292,12 @@ private:
 	
 	/** Update the transforms for the specific proxy */
 	void UpdateProxyTransforms(const FActorProxy& InActorProxy);
+
+	/** Subscribe post process and preview hooks to the root actor */
+	void SubscribeToRootActor();
+	
+	/** Unsubscribe post process and preview hooks from the root actor */
+	void UnsubscribeFromRootActor();
 	
 private:
 	TWeakPtr<FSceneViewport> SceneViewportPtr;
