[nDisplay] Fix crash when multiple nDisplay previews are running at the same time

diff --git a/Engine/Plugins/Runtime/nDisplay/Source/DisplayClusterScenePreview/Private/DisplayClusterLightCardEditorHelper.cpp b/Engine/Plugins/Runtime/nDisplay/Source/DisplayClusterScenePreview/Private/DisplayClusterLightCardEditorHelper.cpp
--- a/Engine/Plugins/Runtime/nDisplay/Source/DisplayClusterScenePreview/Private/DisplayClusterLightCardEditorHelper.cpp
+++ b/Engine/Plugins/Runtime/nDisplay/Source/DisplayClusterScenePreview/Private/DisplayClusterLightCardEditorHelper.cpp
@@ -1674,7 +1674,7 @@ void FDisplayClusterLightCardEditorHelper::InvalidateNormalMap()
 
 bool FDisplayClusterLightCardEditorHelper::UpdateNormalMaps()
 {
-	if (!bNormalMapInvalid)
+	if (!bNormalMapInvalid && NormalMapMeshComponent.IsValid())
 	{
 		return true;
 	}
@@ -1718,7 +1718,8 @@ void FDisplayClusterLightCardEditorHelper::UpdateNormalMapMesh()
 
 	UpdateProjectionOriginComponent();
 
-	NormalMapMeshComponent = NewObject<UProceduralMeshComponent>(GetTransientPackage(), TEXT("NormalMapMesh"));
+	const FName UniqueName = MakeUniqueObjectName(GetTransientPackage(), UProceduralMeshComponent::StaticClass(), TEXT("NormalMapMesh"));
+	NormalMapMeshComponent = NewObject<UProceduralMeshComponent>(GetTransientPackage(), UniqueName);
 	NormalMapMeshComponent->SetCollisionEnabled(ECollisionEnabled::QueryOnly);
 	NormalMeshScene->AddComponent(NormalMapMeshComponent.Get(), FTransform(GetProjectionOrigin()));
 
