nDisplay: fixed CVar issue

diff --git a/Engine/Plugins/Runtime/nDisplay/Source/DisplayCluster/Private/Render/Viewport/Configuration/DisplayClusterViewportConfigurationBase.cpp b/Engine/Plugins/Runtime/nDisplay/Source/DisplayCluster/Private/Render/Viewport/Configuration/DisplayClusterViewportConfigurationBase.cpp
--- a/Engine/Plugins/Runtime/nDisplay/Source/DisplayCluster/Private/Render/Viewport/Configuration/DisplayClusterViewportConfigurationBase.cpp
+++ b/Engine/Plugins/Runtime/nDisplay/Source/DisplayCluster/Private/Render/Viewport/Configuration/DisplayClusterViewportConfigurationBase.cpp
@@ -154,7 +154,7 @@ void FDisplayClusterViewportConfigurationBase::UpdateClusterNodePostProcess(cons
 		if (PPManager.IsValid())
 		{
 			static const FString TextureShareID(TEXT("TextureShare"));
-			IConsoleVariable* const CVarTextureShareEnabled = IConsoleManager::Get().FindConsoleVariable(TEXT("nDisplay.render.texturesharing"));
+			static IConsoleVariable* const CVarTextureShareEnabled = IConsoleManager::Get().FindConsoleVariable(TEXT("nDisplay.render.texturesharing"));
 
 			// TextureShare with nDisplay is only supported for PIE and Runtime.
 			// For PIE in nDisplay, the node must be selected in the preview options.
