nDisplay: Fix ICVFX panel creating proxies for snapshot temporary actors.

diff --git a/Engine/Plugins/Runtime/nDisplay/Source/DisplayClusterLightCardEditor/Private/DisplayClusterLightCardEditor.cpp b/Engine/Plugins/Runtime/nDisplay/Source/DisplayClusterLightCardEditor/Private/DisplayClusterLightCardEditor.cpp
--- a/Engine/Plugins/Runtime/nDisplay/Source/DisplayClusterLightCardEditor/Private/DisplayClusterLightCardEditor.cpp
+++ b/Engine/Plugins/Runtime/nDisplay/Source/DisplayClusterLightCardEditor/Private/DisplayClusterLightCardEditor.cpp
@@ -1802,16 +1802,20 @@ void FDisplayClusterLightCardEditor::OnActorPropertyChanged(UObject* ObjectBeing
 
 void FDisplayClusterLightCardEditor::OnLevelActorAdded(AActor* Actor)
 {
-	if (UE::DisplayClusterLightCardEditorUtils::IsManagedActor(Actor))
+	if (Actor && UE::DisplayClusterLightCardEditorUtils::IsManagedActor(Actor))
 	{
-		RefreshPreviewStageActor(Actor);
+		const UPackage* Package = Actor->GetPackage();
+		if (Package && !Package->HasAnyFlags(RF_Transient) /** Snapshots can create a temporary preview world which we need to avoid */)
+		{
+			RefreshPreviewStageActor(Actor);
+		}
 	}
 }
 
 void FDisplayClusterLightCardEditor::OnLevelActorDeleted(AActor* Actor)
 {
 	if (Actor && UE::DisplayClusterLightCardEditorUtils::IsManagedActor(Actor) &&
-		Actor->GetPackage() != GetTransientPackage() /* Don't trigger if this is a proxy actor being destroyed */)
+		Actor->GetPackage() && !Actor->GetPackage()->HasAnyFlags(RF_Transient) /* Don't trigger if this is a proxy actor being destroyed */)
 	{
 		if (Actor->GetClass()->HasAnyClassFlags(CLASS_NewerVersionExists))
 		{
