Fix use of deprecated function.

diff --git a/Engine/Plugins/Runtime/nDisplay/Source/DisplayCluster/Private/Game/DisplayClusterGameManager.cpp b/Engine/Plugins/Runtime/nDisplay/Source/DisplayCluster/Private/Game/DisplayClusterGameManager.cpp
--- a/Engine/Plugins/Runtime/nDisplay/Source/DisplayCluster/Private/Game/DisplayClusterGameManager.cpp
+++ b/Engine/Plugins/Runtime/nDisplay/Source/DisplayCluster/Private/Game/DisplayClusterGameManager.cpp
@@ -197,7 +197,7 @@ ADisplayClusterRootActor* FDisplayClusterGameManager::FindRootActor(UWorld* InWo
 	const TArray<ULevelStreaming*>& StreamingLevels = InWorld->GetStreamingLevels();
 	for (const ULevelStreaming* const StreamingLevel : StreamingLevels)
 	{
-		if (StreamingLevel && StreamingLevel->GetCurrentState() == ELevelStreamingState::LoadedVisible)
+		if (StreamingLevel && StreamingLevel->GetLevelStreamingState() == ELevelStreamingState::LoadedVisible)
 		{
 			// Look for the actor in those sub-levels that have been loaded already
 			const TSoftObjectPtr<UWorld>& SubWorldAsset = StreamingLevel->GetWorldAsset();
