Fix use of deprecated function.

diff --git a/Engine/Plugins/Runtime/nDisplay/Source/DisplayClusterEditor/Private/DisplayClusterEditorEngine.cpp b/Engine/Plugins/Runtime/nDisplay/Source/DisplayClusterEditor/Private/DisplayClusterEditorEngine.cpp
--- a/Engine/Plugins/Runtime/nDisplay/Source/DisplayClusterEditor/Private/DisplayClusterEditorEngine.cpp
+++ b/Engine/Plugins/Runtime/nDisplay/Source/DisplayClusterEditor/Private/DisplayClusterEditorEngine.cpp
@@ -92,7 +92,7 @@ void UDisplayClusterEditorEngine::StartPlayInEditorSession(FRequestPlaySessionPa
 			const TArray<ULevelStreaming*>& StreamingLevels = EditorWorldPreDup->GetStreamingLevels();
 			for (const ULevelStreaming* const StreamingLevel : StreamingLevels)
 			{
-				if (StreamingLevel && StreamingLevel->GetCurrentState() == ELevelStreamingState::LoadedVisible)
+				if (StreamingLevel && StreamingLevel->GetLevelStreamingState() == ELevelStreamingState::LoadedVisible)
 				{
 					// Look for the actor in those sub-levels that have been loaded already
 					const TSoftObjectPtr<UWorld>& SubWorldAsset = StreamingLevel->GetWorldAsset();
