nDisplay: Optimized the ICVFX panel's color grading drawer to improve performance when selecting objects, specifically nDisplay config actors and ICVFX camera components, as well as improving performance when the drawer is refreshed, such as when an ICVFX camera component's properties are changed.

diff --git a/Engine/Plugins/Runtime/nDisplay/Source/DisplayClusterColorGrading/Private/DataModelGenerators/DisplayClusterColorGradingGenerator_ColorCorrectRegion.cpp b/Engine/Plugins/Runtime/nDisplay/Source/DisplayClusterColorGrading/Private/DataModelGenerators/DisplayClusterColorGradingGenerator_ColorCorrectRegion.cpp
--- a/Engine/Plugins/Runtime/nDisplay/Source/DisplayClusterColorGrading/Private/DataModelGenerators/DisplayClusterColorGradingGenerator_ColorCorrectRegion.cpp
+++ b/Engine/Plugins/Runtime/nDisplay/Source/DisplayClusterColorGrading/Private/DataModelGenerators/DisplayClusterColorGradingGenerator_ColorCorrectRegion.cpp
@@ -4,6 +4,9 @@
 
 #include "ClassIconFinder.h"
 #include "ColorCorrectRegion.h"
+#include "DetailCategoryBuilder.h"
+#include "DetailLayoutBuilder.h"
+#include "IDetailCustomization.h"
 #include "IDetailTreeNode.h"
 #include "IPropertyRowGenerator.h"
 #include "PropertyHandle.h"
@@ -20,6 +23,52 @@ TSharedRef<IDisplayClusterColorGradingDataModelGenerator> FDisplayClusterColorGr
 	return MakeShareable(new FDisplayClusterColorGradingGenerator_ColorCorrectRegion());
 }
 
+class FColorCorrectRegionCustomization : public IDetailCustomization
+{
+public:
+	virtual void CustomizeDetails(IDetailLayoutBuilder& DetailBuilder) override
+	{
+		TArray<FName> Categories;
+		DetailBuilder.GetCategoryNames(Categories);
+
+		for (const FName& Category : Categories)
+		{
+			if (Category != TEXT("Color Correction"))
+			{
+				DetailBuilder.HideCategory(Category);
+			}
+		}
+
+		// TransformCommon is a custom category that doesn't get returned by GetCategoryNames that also needs to be hidden
+		DetailBuilder.HideCategory(TEXT("TransformCommon"));
+
+		IDetailCategoryBuilder& CCCategoryBuilder = DetailBuilder.EditCategory(TEXT("Color Correction"));
+
+		TArray<TSharedRef<IPropertyHandle>> CCPropertyHandles;
+		CCCategoryBuilder.GetDefaultProperties(CCPropertyHandles);
+
+		IDetailCategoryBuilder& ColorGradingElementsCategory = DetailBuilder.EditCategory(FName("ColorGradingElements"));
+
+		ColorGradingElementsCategory.AddProperty(DetailBuilder.GetProperty((GET_MEMBER_NAME_CHECKED(AColorCorrectionRegion, ColorGradingSettings.Global))));
+		ColorGradingElementsCategory.AddProperty(DetailBuilder.GetProperty((GET_MEMBER_NAME_CHECKED(AColorCorrectionRegion, ColorGradingSettings.Shadows))));
+		ColorGradingElementsCategory.AddProperty(DetailBuilder.GetProperty((GET_MEMBER_NAME_CHECKED(AColorCorrectionRegion, ColorGradingSettings.Midtones))));
+		ColorGradingElementsCategory.AddProperty(DetailBuilder.GetProperty((GET_MEMBER_NAME_CHECKED(AColorCorrectionRegion, ColorGradingSettings.Highlights))));
+	}
+};
+
+void FDisplayClusterColorGradingGenerator_ColorCorrectRegion::Initialize(const TSharedRef<class FDisplayClusterColorGradingDataModel>& ColorGradingDataModel, const TSharedRef<IPropertyRowGenerator>& PropertyRowGenerator)
+{
+	PropertyRowGenerator->RegisterInstancedCustomPropertyLayout(AColorCorrectRegion::StaticClass(), FOnGetDetailCustomizationInstance::CreateLambda([]
+	{
+		return MakeShared<FColorCorrectRegionCustomization>();
+	}));
+}
+
+void FDisplayClusterColorGradingGenerator_ColorCorrectRegion::Destroy(const TSharedRef<class FDisplayClusterColorGradingDataModel>& ColorGradingDataModel, const TSharedRef<IPropertyRowGenerator>& PropertyRowGenerator)
+{
+	PropertyRowGenerator->UnregisterInstancedCustomPropertyLayout(AColorCorrectRegion::StaticClass());
+}
+
 void FDisplayClusterColorGradingGenerator_ColorCorrectRegion::GenerateDataModel(IPropertyRowGenerator& PropertyRowGenerator, FDisplayClusterColorGradingDataModel& OutColorGradingDataModel)
 {
 	TArray<TWeakObjectPtr<AColorCorrectRegion>> SelectedCCRs;
@@ -40,48 +89,27 @@ void FDisplayClusterColorGradingGenerator_ColorCorrectRegion::GenerateDataModel(
 
 	const TArray<TSharedRef<IDetailTreeNode>>& RootNodes = PropertyRowGenerator.GetRootTreeNodes();
 
-	const TSharedRef<IDetailTreeNode>* ColorCorrectionNodePtr = RootNodes.FindByPredicate([](const TSharedRef<IDetailTreeNode>& Node)
+	const TSharedRef<IDetailTreeNode>* ColorGradingElementsPtr = RootNodes.FindByPredicate([](const TSharedRef<IDetailTreeNode>& Node)
 	{
-		return Node->GetNodeName() == TEXT("Color Correction");
+		return Node->GetNodeName() == TEXT("ColorGradingElements");
 	});
 
-	if (ColorCorrectionNodePtr)
+	if (ColorGradingElementsPtr)
 	{
-		const TSharedRef<IDetailTreeNode> ColorCorrectionNode = *ColorCorrectionNodePtr;
+		const TSharedRef<IDetailTreeNode> ColorGradingElements = *ColorGradingElementsPtr;
 		FDisplayClusterColorGradingDataModel::FColorGradingGroup ColorGradingGroup;
 
-		AddPropertiesToDetailsView(ColorCorrectionNode, ColorGradingGroup);
+		ColorGradingGroup.DetailsViewCategories.Add(TEXT("Color Correction"));
 
-		TArray<TSharedRef<IDetailTreeNode>> ColorCorrectionPropertyNodes;
-		ColorCorrectionNode->GetChildren(ColorCorrectionPropertyNodes);
+		TArray<TSharedRef<IDetailTreeNode>> ColorGradingPropertyNodes;
+		ColorGradingElements->GetChildren(ColorGradingPropertyNodes);
 
-		const TSharedRef<IDetailTreeNode>* ColorGradingSettingsNodePtr = ColorCorrectionPropertyNodes.FindByPredicate([](const TSharedRef<IDetailTreeNode>& Node)
+		for (const TSharedRef<IDetailTreeNode>& PropertyNode : ColorGradingPropertyNodes)
 		{
-			return Node->GetNodeName() == GET_MEMBER_NAME_CHECKED(AColorCorrectRegion, ColorGradingSettings);
-		});
+			TSharedPtr<IPropertyHandle> PropertyHandle = PropertyNode->CreatePropertyHandle();
 
-		if (ColorGradingSettingsNodePtr)
-		{
-			const TSharedRef<IDetailTreeNode> ColorGradingSettingsNode = *ColorGradingSettingsNodePtr;
-
-			TArray<TSharedRef<IDetailTreeNode>> ColorGradingPropertyNodes;
-			ColorGradingSettingsNode->GetChildren(ColorGradingPropertyNodes);
-
-			for (const TSharedRef<IDetailTreeNode>& PropertyNode : ColorGradingPropertyNodes)
-			{
-				FName NodeName = PropertyNode->GetNodeName();
-				if (NodeName == GET_MEMBER_NAME_CHECKED(FColorGradingSettings, Global) ||
-					NodeName == GET_MEMBER_NAME_CHECKED(FColorGradingSettings, Shadows) ||
-					NodeName == GET_MEMBER_NAME_CHECKED(FColorGradingSettings, Midtones) ||
-					NodeName == GET_MEMBER_NAME_CHECKED(FColorGradingSettings, Highlights))
-				{
-					TSharedPtr<IPropertyHandle> PropertyHandle = PropertyNode->CreatePropertyHandle();
-					PropertyHandle->MarkHiddenByCustomization();
-
-					FDisplayClusterColorGradingDataModel::FColorGradingElement ColorGradingElement = CreateColorGradingElement(PropertyNode, FText::FromName(PropertyNode->GetNodeName()));
-					ColorGradingGroup.ColorGradingElements.Add(ColorGradingElement);
-				}
-			}
+			FDisplayClusterColorGradingDataModel::FColorGradingElement ColorGradingElement = CreateColorGradingElement(PropertyNode, FText::FromName(PropertyNode->GetNodeName()));
+			ColorGradingGroup.ColorGradingElements.Add(ColorGradingElement);
 		}
 
 		ColorGradingGroup.GroupHeaderWidget = SNew(SHorizontalBox)
@@ -158,24 +186,15 @@ FDisplayClusterColorGradingDataModel::FColorGradingElement FDisplayClusterColorG
 	return ColorGradingElement;
 }
 
-void FDisplayClusterColorGradingGenerator_ColorCorrectRegion::AddPropertiesToDetailsView(const TSharedRef<IDetailTreeNode>& GroupNode, FDisplayClusterColorGradingDataModel::FColorGradingGroup& ColorGradingGroup)
+bool FDisplayClusterColorGradingGenerator_ColorCorrectRegion::FilterDetailsViewProperties(const TSharedRef<IDetailTreeNode>& InDetailTreeNode)
 {
-	TArray<TSharedRef<IDetailTreeNode>> ChildNodes;
-	GroupNode->GetChildren(ChildNodes);
-
-	for (const TSharedRef<IDetailTreeNode>& ChildNode : ChildNodes)
+	if (InDetailTreeNode->GetNodeType() == EDetailNodeType::Category)
 	{
-		TSharedPtr<IPropertyHandle> PropertyHandle = ChildNode->CreatePropertyHandle();
-		if (PropertyHandle.IsValid() && PropertyHandle->IsValidHandle())
-		{
-			const bool bIsColorProperty = PropertyHandle->GetProperty()->HasMetaData(TEXT("ColorGradingMode"));
-
-			if (!bIsColorProperty)
-			{
-				PropertyHandle->SetInstanceMetaData(TEXT("CategoryOverride"), GroupNode->GetNodeName().ToString());
-				ColorGradingGroup.DetailsViewPropertyHandles.Add(PropertyHandle);
-			}
-		}
+		return InDetailTreeNode->GetNodeName() == TEXT("Color Correction");
+	}
+	else
+	{
+		return true;
 	}
 }
 
diff --git a/Engine/Plugins/Runtime/nDisplay/Source/DisplayClusterColorGrading/Private/DataModelGenerators/DisplayClusterColorGradingGenerator_ColorCorrectRegion.h b/Engine/Plugins/Runtime/nDisplay/Source/DisplayClusterColorGrading/Private/DataModelGenerators/DisplayClusterColorGradingGenerator_ColorCorrectRegion.h
--- a/Engine/Plugins/Runtime/nDisplay/Source/DisplayClusterColorGrading/Private/DataModelGenerators/DisplayClusterColorGradingGenerator_ColorCorrectRegion.h
+++ b/Engine/Plugins/Runtime/nDisplay/Source/DisplayClusterColorGrading/Private/DataModelGenerators/DisplayClusterColorGradingGenerator_ColorCorrectRegion.h
@@ -14,6 +14,8 @@ public:
 	static TSharedRef<IDisplayClusterColorGradingDataModelGenerator> MakeInstance();
 
 	//~ IDisplayClusterColorGradingDataModelGenerator interface
+	virtual void Initialize(const TSharedRef<class FDisplayClusterColorGradingDataModel>& ColorGradingDataModel, const TSharedRef<IPropertyRowGenerator>& PropertyRowGenerator) override;
+	virtual void Destroy(const TSharedRef<class FDisplayClusterColorGradingDataModel>& ColorGradingDataModel, const TSharedRef<IPropertyRowGenerator>& PropertyRowGenerator) override;
 	virtual void GenerateDataModel(IPropertyRowGenerator& PropertyRowGenerator, FDisplayClusterColorGradingDataModel& OutColorGradingDataModel) override;
 	//~ End IDisplayClusterColorGradingDataModelGenerator interface
 
@@ -21,6 +23,5 @@ private:
 	/** Creates a new color grading element structure for the specified detail tree node, which is expected to have child color properties with the ColorGradingMode metadata set */
 	FDisplayClusterColorGradingDataModel::FColorGradingElement CreateColorGradingElement(const TSharedRef<IDetailTreeNode>& GroupNode, FText ElementLabel);
 
-	/** Adds all child properties of the specified detail tree node to the color grading group's list of properties to display in the details view*/
-	void AddPropertiesToDetailsView(const TSharedRef<IDetailTreeNode>& GroupNode, FDisplayClusterColorGradingDataModel::FColorGradingGroup& ColorGradingGroup);
+	bool FilterDetailsViewProperties(const TSharedRef<IDetailTreeNode>& InDetailTreeNode);
 };
\ No newline at end of file
diff --git a/Engine/Plugins/Runtime/nDisplay/Source/DisplayClusterColorGrading/Private/DataModelGenerators/DisplayClusterColorGradingGenerator_PostProcessVolume.cpp b/Engine/Plugins/Runtime/nDisplay/Source/DisplayClusterColorGrading/Private/DataModelGenerators/DisplayClusterColorGradingGenerator_PostProcessVolume.cpp
--- a/Engine/Plugins/Runtime/nDisplay/Source/DisplayClusterColorGrading/Private/DataModelGenerators/DisplayClusterColorGradingGenerator_PostProcessVolume.cpp
+++ b/Engine/Plugins/Runtime/nDisplay/Source/DisplayClusterColorGrading/Private/DataModelGenerators/DisplayClusterColorGradingGenerator_PostProcessVolume.cpp
@@ -4,6 +4,9 @@
 
 #include "ClassIconFinder.h"
 #include "Engine/PostProcessVolume.h"
+#include "DetailCategoryBuilder.h"
+#include "DetailLayoutBuilder.h"
+#include "IDetailCustomization.h"
 #include "IDetailTreeNode.h"
 #include "IPropertyRowGenerator.h"
 #include "PropertyHandle.h"
@@ -20,6 +23,64 @@ TSharedRef<IDisplayClusterColorGradingDataModelGenerator> FDisplayClusterColorGr
 	return MakeShareable(new FDisplayClusterColorGradingGenerator_PostProcessVolume());
 }
 
+class FPostProcessVolumeCustomization : public IDetailCustomization
+{
+public:
+	virtual void CustomizeDetails(IDetailLayoutBuilder& DetailBuilder) override
+	{
+		TArray<FName> Categories;
+		DetailBuilder.GetCategoryNames(Categories);
+
+		for (const FName& Category : Categories)
+		{
+			DetailBuilder.HideCategory(Category);
+		}
+
+		// TransformCommon is a custom category that doesn't get returned by GetCategoryNames that also needs to be hidden
+		DetailBuilder.HideCategory(TEXT("TransformCommon"));
+
+		TSharedRef<IPropertyHandle> SettingsHandle = DetailBuilder.GetProperty(GET_MEMBER_NAME_CHECKED(APostProcessVolume, Settings));
+
+		uint32 NumChildren;
+		SettingsHandle->GetNumChildren(NumChildren);
+		for (uint32 Index = 0; Index < NumChildren; ++Index)
+		{
+			TSharedPtr<IPropertyHandle> ChildHandle = SettingsHandle->GetChildHandle(Index);
+
+			FString CategoryName = TEXT("");
+			FString GroupName = TEXT("");
+			ChildHandle->GetDefaultCategoryName().ToString().Split(TEXT("|"), &CategoryName, &GroupName);
+
+			if (CategoryName == "Color Grading")
+			{
+				if (ChildHandle->HasMetaData(TEXT("ColorGradingMode")))
+				{
+					IDetailCategoryBuilder& ColorGradingCategory = DetailBuilder.EditCategory(FName("ColorGradingElements"));
+					ColorGradingCategory.AddProperty(ChildHandle);
+				}
+				else
+				{
+					IDetailCategoryBuilder& GroupCategory = DetailBuilder.EditCategory(FName("DetailView_" + GroupName), FText::FromString(GroupName));
+					GroupCategory.AddProperty(ChildHandle);
+				}
+			}
+		}
+	}
+};
+
+void FDisplayClusterColorGradingGenerator_PostProcessVolume::Initialize(const TSharedRef<class FDisplayClusterColorGradingDataModel>& ColorGradingDataModel, const TSharedRef<IPropertyRowGenerator>& PropertyRowGenerator)
+{
+	PropertyRowGenerator->RegisterInstancedCustomPropertyLayout(APostProcessVolume::StaticClass(), FOnGetDetailCustomizationInstance::CreateLambda([]
+	{
+		return MakeShared<FPostProcessVolumeCustomization>();
+	}));
+}
+
+void FDisplayClusterColorGradingGenerator_PostProcessVolume::Destroy(const TSharedRef<class FDisplayClusterColorGradingDataModel>& ColorGradingDataModel, const TSharedRef<IPropertyRowGenerator>& PropertyRowGenerator)
+{
+	PropertyRowGenerator->UnregisterInstancedCustomPropertyLayout(APostProcessVolume::StaticClass());
+}
+
 void FDisplayClusterColorGradingGenerator_PostProcessVolume::GenerateDataModel(IPropertyRowGenerator& PropertyRowGenerator, FDisplayClusterColorGradingDataModel& OutColorGradingDataModel)
 {
 	TArray<TWeakObjectPtr<APostProcessVolume>> SelectedPPVs;
@@ -40,32 +101,51 @@ void FDisplayClusterColorGradingGenerator_PostProcessVolume::GenerateDataModel(I
 
 	const TArray<TSharedRef<IDetailTreeNode>>& RootNodes = PropertyRowGenerator.GetRootTreeNodes();
 
-	const TSharedRef<IDetailTreeNode>* ColorGradingNodePtr = RootNodes.FindByPredicate([](const TSharedRef<IDetailTreeNode>& Node)
+	const TSharedRef<IDetailTreeNode>* ColorGradingElementsNodePtr = RootNodes.FindByPredicate([](const TSharedRef<IDetailTreeNode>& Node)
 	{
-		return Node->GetNodeName() == TEXT("Color Grading");
+		return Node->GetNodeName() == TEXT("ColorGradingElements");
 	});
 
-	if (ColorGradingNodePtr)
+	if (ColorGradingElementsNodePtr)
 	{
-		const TSharedRef<IDetailTreeNode> ColorGradingNode = *ColorGradingNodePtr;
+		const TSharedRef<IDetailTreeNode> ColorGradingElementsNode = *ColorGradingElementsNodePtr;
 		FDisplayClusterColorGradingDataModel::FColorGradingGroup ColorGradingGroup;
 
 		TArray<TSharedRef<IDetailTreeNode>> PropertyGroupNodes;
-		ColorGradingNode->GetChildren(PropertyGroupNodes);
+		ColorGradingElementsNode->GetChildren(PropertyGroupNodes);
+
+		TMap<FString, FDisplayClusterColorGradingDataModel::FColorGradingElement> ColorGradingElements;
 
 		for (const TSharedRef<IDetailTreeNode>& PropertyGroupNode : PropertyGroupNodes)
 		{
-			FString CategoryName = TEXT(""); 
-			FString GroupName = TEXT("");
-			PropertyGroupNode->GetNodeName().ToString().Split(TEXT("|"), &CategoryName, &GroupName);
+			TSharedPtr<IPropertyHandle> PropertyHandle = PropertyGroupNode->CreatePropertyHandle();
 
-			if (GroupName == TEXT("Global") || GroupName == TEXT("Shadows") || GroupName == TEXT("Midtones") || GroupName == TEXT("Highlights"))
+			if (PropertyHandle.IsValid() && PropertyHandle->IsValidHandle())
 			{
-				FDisplayClusterColorGradingDataModel::FColorGradingElement ColorGradingElement = CreateColorGradingElement(PropertyGroupNode, FText::FromString(GroupName));
-				ColorGradingGroup.ColorGradingElements.Add(ColorGradingElement);
+				FString CategoryName = TEXT("");
+				FString GroupName = TEXT("");
+				PropertyHandle->GetDefaultCategoryName().ToString().Split(TEXT("|"), &CategoryName, &GroupName);
+
+				if (!ColorGradingElements.Contains(GroupName))
+				{
+					FDisplayClusterColorGradingDataModel::FColorGradingElement& ColorGradingElement = ColorGradingElements.Add(GroupName);
+					ColorGradingElement.DisplayName = FText::FromString(GroupName);
+				}
+
+				AddPropertyToColorGradingElement(PropertyHandle, ColorGradingElements[GroupName]);
 			}
+		}
 
-			AddPropertiesToDetailsView(PropertyGroupNode, ColorGradingGroup);
+		ColorGradingElements.GenerateValueArray(ColorGradingGroup.ColorGradingElements);
+
+		// Add all categories that are not the color grading elements category to the list of categories to display in the detail view
+		for (const TSharedRef<IDetailTreeNode>& Node : RootNodes)
+		{
+			FName NodeName = Node->GetNodeName();
+			if (NodeName != TEXT("ColorGradingElements"))
+			{
+				ColorGradingGroup.DetailsViewCategories.Add(NodeName);
+			}
 		}
 
 		ColorGradingGroup.GroupHeaderWidget = SNew(SHorizontalBox)
@@ -98,71 +178,31 @@ void FDisplayClusterColorGradingGenerator_PostProcessVolume::GenerateDataModel(I
 	}
 }
 
-FDisplayClusterColorGradingDataModel::FColorGradingElement FDisplayClusterColorGradingGenerator_PostProcessVolume::CreateColorGradingElement(const TSharedRef<IDetailTreeNode>& GroupNode, FText ElementLabel)
+void FDisplayClusterColorGradingGenerator_PostProcessVolume::AddPropertyToColorGradingElement(const TSharedPtr<IPropertyHandle>& PropertyHandle, FDisplayClusterColorGradingDataModel::FColorGradingElement& ColorGradingElement)
 {
-	FDisplayClusterColorGradingDataModel::FColorGradingElement ColorGradingElement;
-	ColorGradingElement.DisplayName = ElementLabel;
-
-	TArray<TSharedRef<IDetailTreeNode>> ChildNodes;
-	GroupNode->GetChildren(ChildNodes);
+	const FString ColorGradingModeString = PropertyHandle->GetProperty()->GetMetaData(TEXT("ColorGradingMode")).ToLower();
 
-	for (const TSharedRef<IDetailTreeNode>& ChildNode : ChildNodes)
+	if (!ColorGradingModeString.IsEmpty())
 	{
-		TSharedPtr<IPropertyHandle> PropertyHandle = ChildNode->CreatePropertyHandle();
-		if (PropertyHandle.IsValid() && PropertyHandle->IsValidHandle())
+		if (ColorGradingModeString.Compare(TEXT("saturation")) == 0)
 		{
-			const FString ColorGradingModeString = PropertyHandle->GetProperty()->GetMetaData(TEXT("ColorGradingMode")).ToLower();
-
-			if (!ColorGradingModeString.IsEmpty())
-			{
-				if (ColorGradingModeString.Compare(TEXT("saturation")) == 0)
-				{
-					ColorGradingElement.SaturationPropertyHandle = PropertyHandle;
-				}
-				else if (ColorGradingModeString.Compare(TEXT("contrast")) == 0)
-				{
-					ColorGradingElement.ContrastPropertyHandle = PropertyHandle;
-				}
-				else if (ColorGradingModeString.Compare(TEXT("gamma")) == 0)
-				{
-					ColorGradingElement.GammaPropertyHandle = PropertyHandle;
-				}
-				else if (ColorGradingModeString.Compare(TEXT("gain")) == 0)
-				{
-					ColorGradingElement.GainPropertyHandle = PropertyHandle;
-				}
-				else if (ColorGradingModeString.Compare(TEXT("offset")) == 0)
-				{
-					ColorGradingElement.OffsetPropertyHandle = PropertyHandle;
-				}
-			}
+			ColorGradingElement.SaturationPropertyHandle = PropertyHandle;
 		}
-	}
-
-	return ColorGradingElement;
-}
-
-void FDisplayClusterColorGradingGenerator_PostProcessVolume::AddPropertiesToDetailsView(const TSharedRef<IDetailTreeNode>& GroupNode, FDisplayClusterColorGradingDataModel::FColorGradingGroup& ColorGradingGroup)
-{
-	FString CategoryName = TEXT("");
-	FString GroupName = TEXT("");
-	GroupNode->GetNodeName().ToString().Split(TEXT("|"), &CategoryName, &GroupName);
-
-	TArray<TSharedRef<IDetailTreeNode>> ChildNodes;
-	GroupNode->GetChildren(ChildNodes);
-
-	for (const TSharedRef<IDetailTreeNode>& ChildNode : ChildNodes)
-	{
-		TSharedPtr<IPropertyHandle> PropertyHandle = ChildNode->CreatePropertyHandle();
-		if (PropertyHandle.IsValid() && PropertyHandle->IsValidHandle())
+		else if (ColorGradingModeString.Compare(TEXT("contrast")) == 0)
 		{
-			const bool bIsColorProperty = PropertyHandle->GetProperty()->HasMetaData(TEXT("ColorGradingMode"));
-
-			if (!bIsColorProperty)
-			{
-				PropertyHandle->SetInstanceMetaData(TEXT("CategoryOverride"), GroupName);
-				ColorGradingGroup.DetailsViewPropertyHandles.Add(PropertyHandle);
-			}
+			ColorGradingElement.ContrastPropertyHandle = PropertyHandle;
+		}
+		else if (ColorGradingModeString.Compare(TEXT("gamma")) == 0)
+		{
+			ColorGradingElement.GammaPropertyHandle = PropertyHandle;
+		}
+		else if (ColorGradingModeString.Compare(TEXT("gain")) == 0)
+		{
+			ColorGradingElement.GainPropertyHandle = PropertyHandle;
+		}
+		else if (ColorGradingModeString.Compare(TEXT("offset")) == 0)
+		{
+			ColorGradingElement.OffsetPropertyHandle = PropertyHandle;
 		}
 	}
 }
diff --git a/Engine/Plugins/Runtime/nDisplay/Source/DisplayClusterColorGrading/Private/DataModelGenerators/DisplayClusterColorGradingGenerator_PostProcessVolume.h b/Engine/Plugins/Runtime/nDisplay/Source/DisplayClusterColorGrading/Private/DataModelGenerators/DisplayClusterColorGradingGenerator_PostProcessVolume.h
--- a/Engine/Plugins/Runtime/nDisplay/Source/DisplayClusterColorGrading/Private/DataModelGenerators/DisplayClusterColorGradingGenerator_PostProcessVolume.h
+++ b/Engine/Plugins/Runtime/nDisplay/Source/DisplayClusterColorGrading/Private/DataModelGenerators/DisplayClusterColorGradingGenerator_PostProcessVolume.h
@@ -14,13 +14,11 @@ public:
 	static TSharedRef<IDisplayClusterColorGradingDataModelGenerator> MakeInstance();
 
 	//~ IDisplayClusterColorGradingDataModelGenerator interface
+	virtual void Initialize(const TSharedRef<class FDisplayClusterColorGradingDataModel>& ColorGradingDataModel, const TSharedRef<IPropertyRowGenerator>& PropertyRowGenerator) override;
+	virtual void Destroy(const TSharedRef<class FDisplayClusterColorGradingDataModel>& ColorGradingDataModel, const TSharedRef<IPropertyRowGenerator>& PropertyRowGenerator) override;
 	virtual void GenerateDataModel(IPropertyRowGenerator& PropertyRowGenerator, FDisplayClusterColorGradingDataModel& OutColorGradingDataModel) override;
 	//~ End IDisplayClusterColorGradingDataModelGenerator interface
 
 private:
-	/** Creates a new color grading element structure for the specified detail tree node, which is expected to have child color properties with the ColorGradingMode metadata set */
-	FDisplayClusterColorGradingDataModel::FColorGradingElement CreateColorGradingElement(const TSharedRef<IDetailTreeNode>& GroupNode, FText ElementLabel);
-
-	/** Adds all child properties of the specified detail tree node to the color grading group's list of properties to display in the details view*/
-	void AddPropertiesToDetailsView(const TSharedRef<IDetailTreeNode>& GroupNode, FDisplayClusterColorGradingDataModel::FColorGradingGroup& ColorGradingGroup);
+	void AddPropertyToColorGradingElement(const TSharedPtr<IPropertyHandle>& PropertyHandle, FDisplayClusterColorGradingDataModel::FColorGradingElement& ColorGradingElement);
 };
\ No newline at end of file
diff --git a/Engine/Plugins/Runtime/nDisplay/Source/DisplayClusterColorGrading/Private/DataModelGenerators/DisplayClusterColorGradingGenerator_RootActor.cpp b/Engine/Plugins/Runtime/nDisplay/Source/DisplayClusterColorGrading/Private/DataModelGenerators/DisplayClusterColorGradingGenerator_RootActor.cpp
--- a/Engine/Plugins/Runtime/nDisplay/Source/DisplayClusterColorGrading/Private/DataModelGenerators/DisplayClusterColorGradingGenerator_RootActor.cpp
+++ b/Engine/Plugins/Runtime/nDisplay/Source/DisplayClusterColorGrading/Private/DataModelGenerators/DisplayClusterColorGradingGenerator_RootActor.cpp
@@ -15,21 +15,23 @@
 
 #include "DetailCategoryBuilder.h"
 #include "DetailLayoutBuilder.h"
+#include "DetailWidgetRow.h"
 #include "Framework/MultiBox/MultiBoxBuilder.h"
+#include "IDetailChildrenBuilder.h"
 #include "IDetailCustomization.h"
 #include "IDetailTreeNode.h"
 #include "IPropertyRowGenerator.h"
 #include "PropertyHandle.h"
 #include "PropertyPathHelpers.h"
 #include "ScopedTransaction.h"
-#include "Widgets/Input/SButton.h"
 #include "Widgets/Input/SComboButton.h"
+#include "Widgets/Input/SButton.h"
 #include "Widgets/Text/SInlineEditableTextBlock.h"
 #include "Widgets/Views/SListView.h"
 
 #define LOCTEXT_NAMESPACE "DisplayClusterColorGrading"
 
-#define CATEGORY_OVERRIDE_NAME(Category) TEXT(Category "_Override")
+#define GET_MEMBER_NAME_ARRAY_CHECKED(ClassName, MemberName, Index) FName(FString(GET_MEMBER_NAME_STRING_CHECKED(ClassName, MemberName)).Replace(TEXT(#Index), *FString::FromInt(Index), ESearchCase::CaseSensitive))
 
 FDisplayClusterColorGradingDataModel::FColorGradingGroup FDisplayClusterColorGradingGenerator_ColorGradingRenderingSettings::CreateColorGradingGroup(const TSharedPtr<IPropertyHandle>& GroupPropertyHandle)
 {
@@ -42,24 +44,12 @@ FDisplayClusterColorGradingDataModel::FColorGradingGroup FDisplayClusterColorGra
 	ColorGradingGroup.ColorGradingElements.Add(CreateColorGradingElement(GroupPropertyHandle, TEXT("Midtones"), LOCTEXT("ColorGrading_MidtonesLabel", "Midtones")));
 	ColorGradingGroup.ColorGradingElements.Add(CreateColorGradingElement(GroupPropertyHandle, TEXT("Highlights"), LOCTEXT("ColorGrading_HighlightsLabel", "Highlights")));
 
-	AddDetailsViewPropertyToGroup(GroupPropertyHandle, ColorGradingGroup, GET_MEMBER_NAME_CHECKED(FDisplayClusterConfigurationViewport_ColorGradingRenderingSettings, AutoExposureBias), TEXT("Exposure"));
-	AddDetailsViewPropertyToGroup(GroupPropertyHandle, ColorGradingGroup, GET_MEMBER_NAME_CHECKED(FDisplayClusterConfigurationViewport_ColorGradingRenderingSettings, ColorCorrectionShadowsMax), TEXT("Color Grading"));
-	AddDetailsViewPropertyToGroup(GroupPropertyHandle, ColorGradingGroup, GET_MEMBER_NAME_CHECKED(FDisplayClusterConfigurationViewport_ColorGradingRenderingSettings, ColorCorrectionHighlightsMin), TEXT("Color Grading"));
-	AddDetailsViewPropertyToGroup(GroupPropertyHandle, ColorGradingGroup, GET_MEMBER_NAME_CHECKED(FDisplayClusterConfigurationViewport_ColorGradingRenderingSettings, ColorCorrectionHighlightsMax), TEXT("Color Grading"));
-	AddDetailsViewPropertyToGroup(GroupPropertyHandle, ColorGradingGroup, GET_MEMBER_NAME_CHECKED(FDisplayClusterConfigurationViewport_ColorGradingWhiteBalanceSettings, TemperatureType), TEXT("White Balance"));
-	AddDetailsViewPropertyToGroup(GroupPropertyHandle, ColorGradingGroup, GET_MEMBER_NAME_CHECKED(FDisplayClusterConfigurationViewport_ColorGradingWhiteBalanceSettings, WhiteTemp), TEXT("White Balance"));
-	AddDetailsViewPropertyToGroup(GroupPropertyHandle, ColorGradingGroup, GET_MEMBER_NAME_CHECKED(FDisplayClusterConfigurationViewport_ColorGradingWhiteBalanceSettings, WhiteTint), TEXT("White Balance"));
-	AddDetailsViewPropertyToGroup(GroupPropertyHandle, ColorGradingGroup, GET_MEMBER_NAME_CHECKED(FDisplayClusterConfigurationViewport_ColorGradingMiscSettings, BlueCorrection), TEXT("Misc"));
-	AddDetailsViewPropertyToGroup(GroupPropertyHandle, ColorGradingGroup, GET_MEMBER_NAME_CHECKED(FDisplayClusterConfigurationViewport_ColorGradingMiscSettings, ExpandGamut), TEXT("Misc"));
-	AddDetailsViewPropertyToGroup(GroupPropertyHandle, ColorGradingGroup, GET_MEMBER_NAME_CHECKED(FDisplayClusterConfigurationViewport_ColorGradingMiscSettings, SceneColorTint), TEXT("Misc"));
-
-	// Ensure the details view categories are always displayed in a specific order
-	ColorGradingGroup.DetailsViewCategorySortOrder =
+	ColorGradingGroup.DetailsViewCategories =
 	{
-		{ CATEGORY_OVERRIDE_NAME("Exposure"), 0},
-		{ CATEGORY_OVERRIDE_NAME("Color Grading"), 1},
-		{ CATEGORY_OVERRIDE_NAME("White Balance"), 2},
-		{ CATEGORY_OVERRIDE_NAME("Misc"), 3}
+		TEXT("DetailView_Exposure"),
+		TEXT("DetailView_ColorGrading"),
+		TEXT("DetailView_WhiteBalance"),
+		TEXT("DetailView_Misc")
 	};
 
 	return ColorGradingGroup;
@@ -86,26 +76,10 @@ FDisplayClusterColorGradingDataModel::FColorGradingElement FDisplayClusterColorG
 	return ColorGradingElement;
 }
 
-void FDisplayClusterColorGradingGenerator_ColorGradingRenderingSettings::AddDetailsViewPropertyToGroup(
-	const TSharedPtr<IPropertyHandle>& GroupPropertyHandle,
-	FDisplayClusterColorGradingDataModel::FColorGradingGroup& Group,
-	FName PropertyName,
-	const FString& CategoryOverride)
-{
-	TSharedPtr<IPropertyHandle> PropertyHandle = GroupPropertyHandle->GetChildHandle(PropertyName);
-	if (PropertyHandle.IsValid() && PropertyHandle->IsValidHandle())
-	{
-		PropertyHandle->SetInstanceMetaData(TEXT("CategoryOverride"), CategoryOverride);
-
-		Group.DetailsViewPropertyHandles.Add(PropertyHandle);
-	}
-}
-
 TSharedPtr<IDetailTreeNode> FDisplayClusterColorGradingGenerator_ColorGradingRenderingSettings::FindPropertyTreeNode(const TSharedRef<IDetailTreeNode>& Node, const FCachedPropertyPath& PropertyPath)
 {
 	if (Node->GetNodeType() == EDetailNodeType::Item)
 	{
-		FName NodeName = Node->GetNodeName();
 		if (Node->GetNodeName() == PropertyPath.GetLastSegment().GetName())
 		{
 			TSharedPtr<IPropertyHandle> FoundPropertyHandle = Node->CreatePropertyHandle();
@@ -157,20 +131,15 @@ TSharedRef<IDisplayClusterColorGradingDataModelGenerator> FDisplayClusterColorGr
 	return MakeShareable(new FDisplayClusterColorGradingGenerator_RootActor());
 }
 
-class FRootActorDetailsSectionCustomization : public IDetailCustomization
+/**
+ * A detail customization that picks out only the necessary properties needed to display a root actor in the color grading drawer and hides all other properties
+ * Also organizes the properties into custom categories that can be easily displayed in the color grading drawer
+ */
+class FRootActorColorGradingCustomization : public IDetailCustomization
 {
 public:
-	enum EDetailsSectionType
-	{
-		Viewports,
-		InnerFrustum,
-		OCIO_AllViewports,
-		OCIO_PerViewport
-	};
-
-public:
-	FRootActorDetailsSectionCustomization(EDetailsSectionType InSectionType)
-		: SectionType(InSectionType)
+	FRootActorColorGradingCustomization(const TSharedRef<FDisplayClusterColorGradingDataModel>& InColorGradingDataModel)
+		: ColorGradingDataModel(InColorGradingDataModel)
 	{ }
 
 	virtual void CustomizeDetails(IDetailLayoutBuilder& DetailBuilder) override
@@ -186,47 +155,190 @@ public:
 		// TransformCommon is a custom category that doesn't get returned by GetCategoryNames that also needs to be hidden
 		DetailBuilder.HideCategory(TEXT("TransformCommon"));
 
-		switch (SectionType)
+		UClass* ConfigDataClass = UDisplayClusterConfigurationData::StaticClass();
+
+		IDetailCategoryBuilder& ColorGradingCategoryBuilder = DetailBuilder.EditCategory(TEXT("ColorGradingCategory"));
+		ColorGradingCategoryBuilder.AddProperty(DetailBuilder.GetProperty(GET_MEMBER_NAME_CHECKED(UDisplayClusterConfigurationData, StageSettings.EntireClusterColorGrading), ConfigDataClass));
+		ColorGradingCategoryBuilder.AddProperty(DetailBuilder.GetProperty(GET_MEMBER_NAME_CHECKED(UDisplayClusterConfigurationData, StageSettings.PerViewportColorGrading), ConfigDataClass));
+		ColorGradingCategoryBuilder.AddProperty(DetailBuilder.GetProperty(GET_MEMBER_NAME_CHECKED(UDisplayClusterConfigurationData, StageSettings.bEnableInnerFrustums), ConfigDataClass));
+		ColorGradingCategoryBuilder.AddProperty(DetailBuilder.GetProperty(GET_MEMBER_NAME_CHECKED(UDisplayClusterConfigurationData, StageSettings.bUseOverallClusterOCIOConfiguration), ConfigDataClass));
+
+		AddColorGradingDetailProperties(DetailBuilder);
+		AddDetailsPanelProperties(DetailBuilder);
+
+		DetailBuilder.SortCategories([](const TMap<FName, IDetailCategoryBuilder*>& CategoryMap)
 		{
-		case EDetailsSectionType::Viewports:
+			const TMap<FName, int32> SortOrder =
 			{
-				IDetailCategoryBuilder& ViewportsCategoryBuilder = DetailBuilder.EditCategory(TEXT("CustomViewportsCategory"), LOCTEXT("CustomViewportsCategoryLabel", "Viewports"));
-
-				ViewportsCategoryBuilder.AddProperty(TEXT("ViewportScreenPercentageMultiplierRef"), ADisplayClusterRootActor::StaticClass());
-				ViewportsCategoryBuilder.AddProperty(TEXT("FreezeRenderOuterViewportsRef"), ADisplayClusterRootActor::StaticClass());
+				{ TEXT("DetailView_PerViewport"), 0},
+				{ TEXT("DetailView_Exposure"), 1},
+				{ TEXT("DetailView_ColorGrading"), 2},
+				{ TEXT("DetailView_WhiteBalance"), 3},
+				{ TEXT("DetailView_Misc"), 4}
+			};
+
+			for (const TPair<FName, int32>& SortPair : SortOrder)
+			{
+				if (CategoryMap.Contains(SortPair.Key))
+				{
+					CategoryMap[SortPair.Key]->SetSortOrder(SortPair.Value);
+				}
 			}
-			break;
+		});
+	}
 
-		case EDetailsSectionType::InnerFrustum:
-			{
-				IDetailCategoryBuilder& InnerFrustumCategoryBuilder = DetailBuilder.EditCategory(TEXT("CustomICVFXCategory"), LOCTEXT("CustomICVFXCategoryLabel", "In-Camera VFX"));
+private:
+	void AddColorGradingDetailProperties(IDetailLayoutBuilder& DetailBuilder)
+	{
+		auto AddColorGradingSettings = [&DetailBuilder](const TSharedRef<IPropertyHandle>& ColorGradingSettingsHandle)
+		{
+			IDetailCategoryBuilder& DetailExposureCategoryBuilder = DetailBuilder.EditCategory(TEXT("DetailView_Exposure"), LOCTEXT("DetailView_ExposureDisplayName", "Exposure"));
+			DetailExposureCategoryBuilder.AddProperty(ColorGradingSettingsHandle->GetChildHandle(GET_MEMBER_NAME_CHECKED(FDisplayClusterConfigurationViewport_ColorGradingRenderingSettings, AutoExposureBias)));
 
-				InnerFrustumCategoryBuilder.AddProperty(GET_MEMBER_NAME_CHECKED(ADisplayClusterRootActor, InnerFrustumPriority), ADisplayClusterRootActor::StaticClass());
-			}
-			break;
+			IDetailCategoryBuilder& DetailColorGradingCategoryBuilder = DetailBuilder.EditCategory(TEXT("DetailView_ColorGrading"), LOCTEXT("DetailView_ColorGradingDisplayName", "Color Grading"));
+			DetailColorGradingCategoryBuilder.AddProperty(ColorGradingSettingsHandle->GetChildHandle(GET_MEMBER_NAME_CHECKED(FDisplayClusterConfigurationViewport_ColorGradingRenderingSettings, ColorCorrectionShadowsMax)));
+			DetailColorGradingCategoryBuilder.AddProperty(ColorGradingSettingsHandle->GetChildHandle(GET_MEMBER_NAME_CHECKED(FDisplayClusterConfigurationViewport_ColorGradingRenderingSettings, ColorCorrectionHighlightsMin)));
+			DetailColorGradingCategoryBuilder.AddProperty(ColorGradingSettingsHandle->GetChildHandle(GET_MEMBER_NAME_CHECKED(FDisplayClusterConfigurationViewport_ColorGradingRenderingSettings, ColorCorrectionHighlightsMax)));
 
-		case EDetailsSectionType::OCIO_AllViewports:
-			{
-				IDetailCategoryBuilder& AllViewportsOCIOCategoryBuilder = DetailBuilder.EditCategory(TEXT("CustomAllViewportsOCIOCategory"), LOCTEXT("CustomAllViewportsOCIOCategoryLabel", "All Viewports"));
+			IDetailCategoryBuilder& DetailWhiteBalanceCategoryBuilder = DetailBuilder.EditCategory(TEXT("DetailView_WhiteBalance"), LOCTEXT("DetailView_WhiteBalanceDisplayName", "White Balance"));
+			DetailWhiteBalanceCategoryBuilder.AddProperty(ColorGradingSettingsHandle->GetChildHandle(GET_MEMBER_NAME_CHECKED(FDisplayClusterConfigurationViewport_ColorGradingWhiteBalanceSettings, TemperatureType)));
+			DetailWhiteBalanceCategoryBuilder.AddProperty(ColorGradingSettingsHandle->GetChildHandle(GET_MEMBER_NAME_CHECKED(FDisplayClusterConfigurationViewport_ColorGradingWhiteBalanceSettings, WhiteTemp)));
+			DetailWhiteBalanceCategoryBuilder.AddProperty(ColorGradingSettingsHandle->GetChildHandle(GET_MEMBER_NAME_CHECKED(FDisplayClusterConfigurationViewport_ColorGradingWhiteBalanceSettings, WhiteTint)));
 
-				AllViewportsOCIOCategoryBuilder.AddProperty(TEXT("ClusterOCIOColorConfigurationRef"), ADisplayClusterRootActor::StaticClass());
-			}
-			break;
+			IDetailCategoryBuilder& DetailMiscCategoryBuilder = DetailBuilder.EditCategory(TEXT("DetailView_Misc"), LOCTEXT("DetailView_MiscDisplayName", "Misc"));
+			DetailMiscCategoryBuilder.AddProperty(ColorGradingSettingsHandle->GetChildHandle(GET_MEMBER_NAME_CHECKED(FDisplayClusterConfigurationViewport_ColorGradingMiscSettings, BlueCorrection)));
+			DetailMiscCategoryBuilder.AddProperty(ColorGradingSettingsHandle->GetChildHandle(GET_MEMBER_NAME_CHECKED(FDisplayClusterConfigurationViewport_ColorGradingMiscSettings, ExpandGamut)));
+			DetailMiscCategoryBuilder.AddProperty(ColorGradingSettingsHandle->GetChildHandle(GET_MEMBER_NAME_CHECKED(FDisplayClusterConfigurationViewport_ColorGradingMiscSettings, SceneColorTint)));
+		};
 
-		case EDetailsSectionType::OCIO_PerViewport:
-			{
-				IDetailCategoryBuilder& PerViewportOCIOCategoryBuilder = DetailBuilder.EditCategory(TEXT("CustomPerViewportOCIOCategory"), LOCTEXT("CustomPerViewportOCIOCategoryLabel", "Per-Viewport"));
+		UClass* ConfigDataClass = UDisplayClusterConfigurationData::StaticClass();
+		const int32 GroupIndex = ColorGradingDataModel.IsValid() ? ColorGradingDataModel.Pin()->GetSelectedColorGradingGroupIndex() : INDEX_NONE;
 
-				PerViewportOCIOCategoryBuilder.AddProperty(TEXT("PerViewportOCIOProfilesRef"), ADisplayClusterRootActor::StaticClass());
-			}
-			break;
+		if (GroupIndex > 0)
+		{
+			const int32 Index = GroupIndex - 1;
+			IDetailCategoryBuilder& PerNodeSettingsCategoryBuilder = DetailBuilder.EditCategory(TEXT("DetailView_PerViewport"), LOCTEXT("DetailView_PerViewportDisplayName", "Per-Viewport Settings"));
+			PerNodeSettingsCategoryBuilder.AddProperty(DetailBuilder.GetProperty(GET_MEMBER_NAME_ARRAY_CHECKED(UDisplayClusterConfigurationData, StageSettings.PerViewportColorGrading[Index].bIsEntireClusterEnabled, Index), ConfigDataClass));
+
+			AddColorGradingSettings(DetailBuilder.GetProperty(GET_MEMBER_NAME_ARRAY_CHECKED(UDisplayClusterConfigurationData, StageSettings.PerViewportColorGrading[Index].ColorGradingSettings, Index), ConfigDataClass));
+		}
+		else
+		{
+			AddColorGradingSettings(DetailBuilder.GetProperty(GET_MEMBER_NAME_CHECKED(UDisplayClusterConfigurationData, StageSettings.EntireClusterColorGrading.ColorGradingSettings), ConfigDataClass));
 		}
 	}
 
+	void AddDetailsPanelProperties(IDetailLayoutBuilder& DetailBuilder)
+	{
+		IDetailCategoryBuilder& ViewportsCategoryBuilder = DetailBuilder.EditCategory(TEXT("CustomViewportsCategory"), LOCTEXT("CustomViewportsCategoryLabel", "Viewports"));
+		ViewportsCategoryBuilder.AddProperty(TEXT("ViewportScreenPercentageMultiplierRef"), ADisplayClusterRootActor::StaticClass());
+		ViewportsCategoryBuilder.AddProperty(TEXT("FreezeRenderOuterViewportsRef"), ADisplayClusterRootActor::StaticClass());
+
+		IDetailCategoryBuilder& InnerFrustumCategoryBuilder = DetailBuilder.EditCategory(TEXT("CustomICVFXCategory"), LOCTEXT("CustomICVFXCategoryLabel", "In-Camera VFX"));
+		InnerFrustumCategoryBuilder.AddProperty(GET_MEMBER_NAME_CHECKED(ADisplayClusterRootActor, InnerFrustumPriority), ADisplayClusterRootActor::StaticClass());
+
+		IDetailCategoryBuilder& AllViewportsOCIOCategoryBuilder = DetailBuilder.EditCategory(TEXT("CustomAllViewportsOCIOCategory"), LOCTEXT("CustomAllViewportsOCIOCategoryLabel", "All Viewports"));
+		AllViewportsOCIOCategoryBuilder.AddProperty(TEXT("ClusterOCIOColorConfigurationRef"), ADisplayClusterRootActor::StaticClass());
+
+		IDetailCategoryBuilder& PerViewportOCIOCategoryBuilder = DetailBuilder.EditCategory(TEXT("CustomPerViewportOCIOCategory"), LOCTEXT("CustomPerViewportOCIOCategoryLabel", "Per-Viewport"));
+		PerViewportOCIOCategoryBuilder.AddProperty(TEXT("PerViewportOCIOProfilesRef"), ADisplayClusterRootActor::StaticClass());
+	}
+
 private:
-	EDetailsSectionType SectionType;
+	TWeakPtr<FDisplayClusterColorGradingDataModel> ColorGradingDataModel;
 };
 
+/** A property customizer that culls unneeded properties from the FDisplayClusterConfigurationViewport_EntireClusterColorGrading struct to help speed up property node tree generation */
+class FFastEntireClusterColorGradingCustomization : public IPropertyTypeCustomization
+{
+public:
+	FFastEntireClusterColorGradingCustomization(const TSharedRef<FDisplayClusterColorGradingDataModel>& InColorGradingDataModel)
+		: ColorGradingDataModel(InColorGradingDataModel)
+	{ }
+
+	virtual void CustomizeHeader(TSharedRef<IPropertyHandle> StructPropertyHandle, FDetailWidgetRow& HeaderRow, IPropertyTypeCustomizationUtils& StructCustomizationUtils) override
+	{
+		HeaderRow.NameContent()
+		[
+			StructPropertyHandle->CreatePropertyNameWidget()
+		];
+	}
+
+	virtual void CustomizeChildren(TSharedRef<IPropertyHandle> StructPropertyHandle, IDetailChildrenBuilder& StructBuilder, IPropertyTypeCustomizationUtils& StructCustomizationUtils) override
+	{
+		const int32 GroupIndex = ColorGradingDataModel.IsValid() ? ColorGradingDataModel.Pin()->GetSelectedColorGradingGroupIndex() : INDEX_NONE;
+
+		StructBuilder.AddProperty(StructPropertyHandle->GetChildHandle(GET_MEMBER_NAME_CHECKED(FDisplayClusterConfigurationViewport_EntireClusterColorGrading, bEnableEntireClusterColorGrading)).ToSharedRef());
+
+		if (GroupIndex < 1)
+		{
+			StructBuilder.AddProperty(StructPropertyHandle->GetChildHandle(GET_MEMBER_NAME_CHECKED(FDisplayClusterConfigurationViewport_EntireClusterColorGrading, ColorGradingSettings)).ToSharedRef());
+		}
+	}
+
+private:
+	TWeakPtr<FDisplayClusterColorGradingDataModel> ColorGradingDataModel;
+};
+
+/** A property customizer that culls unneeded properties from the FDisplayClusterConfigurationViewport_PerViewportColorGrading struct to help speed up property node tree generation */
+class FFastPerViewportColorGradingCustomization : public IPropertyTypeCustomization
+{
+public:
+	FFastPerViewportColorGradingCustomization(const TSharedRef<FDisplayClusterColorGradingDataModel>& InColorGradingDataModel)
+		: ColorGradingDataModel(InColorGradingDataModel)
+	{ }
+
+	virtual void CustomizeHeader(TSharedRef<IPropertyHandle> StructPropertyHandle, FDetailWidgetRow& HeaderRow, IPropertyTypeCustomizationUtils& StructCustomizationUtils) override
+	{
+		HeaderRow.NameContent()
+		[
+			StructPropertyHandle->CreatePropertyNameWidget()
+		];
+	}
+
+	virtual void CustomizeChildren(TSharedRef<IPropertyHandle> StructPropertyHandle, IDetailChildrenBuilder& StructBuilder, IPropertyTypeCustomizationUtils& StructCustomizationUtils) override
+	{
+		const int32 ArrayIndex = StructPropertyHandle->GetIndexInArray();
+		const int32 GroupIndex = ColorGradingDataModel.IsValid() ? ColorGradingDataModel.Pin()->GetSelectedColorGradingGroupIndex() : INDEX_NONE;
+
+		StructBuilder.AddProperty(StructPropertyHandle->GetChildHandle(GET_MEMBER_NAME_CHECKED(FDisplayClusterConfigurationViewport_PerViewportColorGrading, bIsEnabled)).ToSharedRef());
+		StructBuilder.AddProperty(StructPropertyHandle->GetChildHandle(GET_MEMBER_NAME_CHECKED(FDisplayClusterConfigurationViewport_PerViewportColorGrading, Name)).ToSharedRef());
+
+		if (GroupIndex == ArrayIndex + 1)
+		{
+			StructBuilder.AddProperty(StructPropertyHandle->GetChildHandle(GET_MEMBER_NAME_CHECKED(FDisplayClusterConfigurationViewport_PerViewportColorGrading, ColorGradingSettings)).ToSharedRef());
+		}
+	}
+
+private:
+	TWeakPtr<FDisplayClusterColorGradingDataModel> ColorGradingDataModel;
+};
+
+void FDisplayClusterColorGradingGenerator_RootActor::Initialize(const TSharedRef<class FDisplayClusterColorGradingDataModel>& ColorGradingDataModel, const TSharedRef<IPropertyRowGenerator>& PropertyRowGenerator)
+{
+	PropertyRowGenerator->RegisterInstancedCustomPropertyTypeLayout(FDisplayClusterConfigurationViewport_EntireClusterColorGrading::StaticStruct()->GetFName(), FOnGetPropertyTypeCustomizationInstance::CreateLambda([ColorGradingDataModel]
+	{
+		return MakeShared<FFastEntireClusterColorGradingCustomization>(ColorGradingDataModel);
+	}));
+
+	PropertyRowGenerator->RegisterInstancedCustomPropertyTypeLayout(FDisplayClusterConfigurationViewport_PerViewportColorGrading::StaticStruct()->GetFName(), FOnGetPropertyTypeCustomizationInstance::CreateLambda([ColorGradingDataModel]
+	{
+		return MakeShared<FFastPerViewportColorGradingCustomization>(ColorGradingDataModel);
+	}));
+
+	PropertyRowGenerator->RegisterInstancedCustomPropertyLayout(ADisplayClusterRootActor::StaticClass(), FOnGetDetailCustomizationInstance::CreateLambda([ColorGradingDataModel]
+	{
+		return MakeShared<FRootActorColorGradingCustomization>(ColorGradingDataModel);
+	}));
+}
+
+void FDisplayClusterColorGradingGenerator_RootActor::Destroy(const TSharedRef<class FDisplayClusterColorGradingDataModel>& ColorGradingDataModel, const TSharedRef<IPropertyRowGenerator>& PropertyRowGenerator)
+{
+	PropertyRowGenerator->UnregisterInstancedCustomPropertyTypeLayout(FDisplayClusterConfigurationViewport_EntireClusterColorGrading::StaticStruct()->GetFName());
+	PropertyRowGenerator->UnregisterInstancedCustomPropertyTypeLayout(FDisplayClusterConfigurationViewport_PerViewportColorGrading::StaticStruct()->GetFName());
+	PropertyRowGenerator->UnregisterInstancedCustomPropertyLayout(ADisplayClusterRootActor::StaticClass());
+}
+
 void FDisplayClusterColorGradingGenerator_RootActor::GenerateDataModel(IPropertyRowGenerator& PropertyRowGenerator, FDisplayClusterColorGradingDataModel& OutColorGradingDataModel)
 {
 	RootActors.Empty();
@@ -283,13 +395,7 @@ void FDisplayClusterColorGradingGenerator_RootActor::GenerateDataModel(IProperty
 				FDisplayClusterColorGradingDataModel::FColorGradingGroup PerViewportGroup = CreateColorGradingGroup(PerViewportElementHandle);
 				PerViewportGroup.EditConditionPropertyHandle = PerViewportElementHandle->GetChildHandle(GET_MEMBER_NAME_CHECKED(FDisplayClusterConfigurationViewport_PerViewportColorGrading, bIsEnabled));
 
-				// Add per-viewport group specific properties and force the category of these properties to be at the top
-				AddDetailsViewPropertyToGroup(PerViewportElementHandle, 
-					PerViewportGroup, 
-					GET_MEMBER_NAME_CHECKED(FDisplayClusterConfigurationViewport_PerViewportColorGrading, bIsEntireClusterEnabled),
-					TEXT("Per-Viewport Settings"));
-
-				PerViewportGroup.DetailsViewCategorySortOrder.Add(CATEGORY_OVERRIDE_NAME("Per-Viewport Settings"), -1);
+				PerViewportGroup.DetailsViewCategories.Add(TEXT("DetailView_PerViewport"));
 
 				TSharedPtr<IPropertyHandle> NamePropertyHandle = PerViewportElementHandle->GetChildHandle(GET_MEMBER_NAME_CHECKED(FDisplayClusterConfigurationViewport_PerViewportColorGrading, Name));
 				if (NamePropertyHandle.IsValid() && NamePropertyHandle->IsValidHandle())
@@ -353,20 +459,14 @@ void FDisplayClusterColorGradingGenerator_RootActor::GenerateDataModel(IProperty
 
 	FDisplayClusterColorGradingDataModel::FDetailsSection ViewportsDetailsSection;
 	ViewportsDetailsSection.DisplayName = LOCTEXT("ViewportsDetailsSectionLabel", "Viewports");
-	ViewportsDetailsSection.DetailCustomizationDelegate = FOnGetDetailCustomizationInstance::CreateLambda([]
-	{
-		return MakeShared<FRootActorDetailsSectionCustomization>(FRootActorDetailsSectionCustomization::EDetailsSectionType::Viewports);
-	});
+	ViewportsDetailsSection.Categories.Add(TEXT("CustomViewportsCategory"));
 
 	OutColorGradingDataModel.DetailsSections.Add(ViewportsDetailsSection);
 
 	FDisplayClusterColorGradingDataModel::FDetailsSection InnerFrustumDetailsSection;
 	InnerFrustumDetailsSection.DisplayName = LOCTEXT("InnerFrustumDetailsSectionLabel", "Inner Frustum");
+	InnerFrustumDetailsSection.Categories.Add(TEXT("CustomICVFXCategory"));
 	InnerFrustumDetailsSection.EditConditionPropertyHandle = FindPropertyHandle(PropertyRowGenerator, CREATE_PROPERTY_PATH(UDisplayClusterConfigurationData, StageSettings.bEnableInnerFrustums));
-	InnerFrustumDetailsSection.DetailCustomizationDelegate = FOnGetDetailCustomizationInstance::CreateLambda([]
-	{
-		return MakeShared<FRootActorDetailsSectionCustomization>(FRootActorDetailsSectionCustomization::EDetailsSectionType::InnerFrustum);
-	});
 
 	OutColorGradingDataModel.DetailsSections.Add(InnerFrustumDetailsSection);
 
@@ -376,19 +476,13 @@ void FDisplayClusterColorGradingGenerator_RootActor::GenerateDataModel(IProperty
 
 	FDisplayClusterColorGradingDataModel::FDetailsSubsection AllViewportsOCIODetailsSubsection;
 	AllViewportsOCIODetailsSubsection.DisplayName = LOCTEXT("AllViewportsOCIOSubsectionLabel", "All Viewports");
-	AllViewportsOCIODetailsSubsection.DetailCustomizationDelegate = FOnGetDetailCustomizationInstance::CreateLambda([]
-	{
-		return MakeShared<FRootActorDetailsSectionCustomization>(FRootActorDetailsSectionCustomization::EDetailsSectionType::OCIO_AllViewports);
-	});
+	AllViewportsOCIODetailsSubsection.Categories.Add(TEXT("CustomAllViewportsOCIOCategory"));
 
 	OCIODetailsSection.Subsections.Add(AllViewportsOCIODetailsSubsection);
 
 	FDisplayClusterColorGradingDataModel::FDetailsSubsection PerViewportOCIODetailsSubsection;
 	PerViewportOCIODetailsSubsection.DisplayName = LOCTEXT("PerViewportOCIOSubsectionLabel", "Per-Viewport");
-	PerViewportOCIODetailsSubsection.DetailCustomizationDelegate = FOnGetDetailCustomizationInstance::CreateLambda([]
-	{
-		return MakeShared<FRootActorDetailsSectionCustomization>(FRootActorDetailsSectionCustomization::EDetailsSectionType::OCIO_PerViewport);
-	});
+	PerViewportOCIODetailsSubsection.Categories.Add(TEXT("CustomPerViewportOCIOCategory"));
 
 	OCIODetailsSection.Subsections.Add(PerViewportOCIODetailsSubsection);
 
@@ -584,22 +678,15 @@ TSharedRef<IDisplayClusterColorGradingDataModelGenerator> FDisplayClusterColorGr
 	return MakeShareable(new FDisplayClusterColorGradingGenerator_ICVFXCamera());
 }
 
-class FICVFXCameraDetailsSectionCustomization : public IDetailCustomization
+/**
+ * A detail customization that picks out only the necessary properties needed to display a ICVFX camera component in the color grading drawer and hides all other properties
+ * Also organizes the properties into custom categories that can be easily displayed in the color grading drawer
+ */
+class FICVFXCameraColorGradingCustomization : public IDetailCustomization
 {
 public:
-	enum EDetailsSectionType
-	{
-		ICVFX,
-		Overscan,
-		Chromakey_Markers,
-		Chromakey_Custom,
-		OCIO_AllNodes,
-		OCIO_PerNode
-	};
-
-public:
-	FICVFXCameraDetailsSectionCustomization(EDetailsSectionType InSectionType)
-		: SectionType(InSectionType)
+	FICVFXCameraColorGradingCustomization(const TSharedRef<FDisplayClusterColorGradingDataModel>& InColorGradingDataModel)
+		: ColorGradingDataModel(InColorGradingDataModel)
 	{ }
 
 	virtual void CustomizeDetails(IDetailLayoutBuilder& DetailBuilder) override
@@ -615,83 +702,218 @@ public:
 		// TransformCommon is a custom category that doesn't get returned by GetCategoryNames that also needs to be hidden
 		DetailBuilder.HideCategory(TEXT("TransformCommon"));
 
-		switch (SectionType)
-		{
-		case EDetailsSectionType::ICVFX:
-			{
-				IDetailCategoryBuilder& ICVFXCategoryBuilder = DetailBuilder.EditCategory(TEXT("CustomICVFXCategory"), LOCTEXT("CustomICVFXCategoryLabel", "In-Camera VFX"));
-				AddProperty(DetailBuilder, ICVFXCategoryBuilder, TEXT("BufferRatioRef"));
-				AddProperty(DetailBuilder, ICVFXCategoryBuilder, TEXT("ExternalCameraActorRef"));
-				AddProperty(DetailBuilder, ICVFXCategoryBuilder, TEXT("HiddenICVFXViewportsRef"));
+		IDetailCategoryBuilder& CategoryBuilder = DetailBuilder.EditCategory(TEXT("ColorGradingCategory"));
 
-				IDetailCategoryBuilder& SoftEdgeCategoryBuilder = DetailBuilder.EditCategory(TEXT("CustomSoftEdgeCategory"), LOCTEXT("CustomSoftEdgeCategoryLabel", "Soft Edge"));
-				AddProperty(DetailBuilder, SoftEdgeCategoryBuilder, TEXT("SoftEdgeRef"), true);
+		CategoryBuilder.AddProperty(DetailBuilder.GetProperty(GET_MEMBER_NAME_CHECKED(UDisplayClusterICVFXCameraComponent, CameraSettings.bEnable)));
+		CategoryBuilder.AddProperty(DetailBuilder.GetProperty(GET_MEMBER_NAME_CHECKED(UDisplayClusterICVFXCameraComponent, CameraSettings.Chromakey.bEnable)));
+		CategoryBuilder.AddProperty(DetailBuilder.GetProperty(GET_MEMBER_NAME_CHECKED(UDisplayClusterICVFXCameraComponent, CameraSettings.AllNodesOCIOConfiguration.bIsEnabled)));
+		CategoryBuilder.AddProperty(DetailBuilder.GetProperty(GET_MEMBER_NAME_CHECKED(UDisplayClusterICVFXCameraComponent, CameraSettings.AllNodesColorGrading)));
+		CategoryBuilder.AddProperty(DetailBuilder.GetProperty(GET_MEMBER_NAME_CHECKED(UDisplayClusterICVFXCameraComponent, CameraSettings.PerNodeColorGrading)));
 
-				IDetailCategoryBuilder& BorderCategoryBuilder = DetailBuilder.EditCategory(TEXT("CustomBorderCategory"), LOCTEXT("CustomBorderCategoryLabel", "Border"));
-				AddProperty(DetailBuilder, BorderCategoryBuilder, TEXT("BorderRef"), true);
-			}
-			break;
+		AddColorGradingDetailProperties(DetailBuilder);
+		AddDetailsPanelProperties(DetailBuilder);
 
-		case EDetailsSectionType::Overscan:
+		DetailBuilder.SortCategories([](const TMap<FName, IDetailCategoryBuilder*>& CategoryMap)
+		{
+			const TMap<FName, int32> SortOrder =
+			{
+				{ TEXT("DetailView_PerNode"), 0},
+				{ TEXT("DetailView_Exposure"), 1},
+				{ TEXT("DetailView_ColorGrading"), 2},
+				{ TEXT("DetailView_WhiteBalance"), 3},
+				{ TEXT("DetailView_Misc"), 4}
+			};
+
+			for (const TPair<FName, int32>& SortPair : SortOrder)
 			{
-				IDetailCategoryBuilder& OverscanCategoryBuilder = DetailBuilder.EditCategory(TEXT("CustomOverscanCategory"), LOCTEXT("CustomOverscanCategoryLabel", "Inner Frustum Overscan"));
-				AddProperty(DetailBuilder, OverscanCategoryBuilder, TEXT("CustomFrustumRef"), true);
+				if (CategoryMap.Contains(SortPair.Key))
+				{
+					CategoryMap[SortPair.Key]->SetSortOrder(SortPair.Value);
+				}
 			}
-			break;
+		});
+	}
 
-		case EDetailsSectionType::Chromakey_Markers:
-			{
-				IDetailCategoryBuilder& ChromakeyCategoryBuilder = DetailBuilder.EditCategory(TEXT("CustomChromakeyCategory"), LOCTEXT("CustomChromakeyCategoryLabel", "Chromakey"));
-				AddProperty(DetailBuilder, ChromakeyCategoryBuilder, TEXT("ChromakeyColorRef"));
+private:
+	void AddColorGradingDetailProperties(IDetailLayoutBuilder& DetailBuilder)
+	{
+		auto AddColorGradingSettings = [&DetailBuilder](const TSharedRef<IPropertyHandle>& ColorGradingSettingsHandle)
+		{
+			IDetailCategoryBuilder& DetailExposureCategoryBuilder = DetailBuilder.EditCategory(TEXT("DetailView_Exposure"), LOCTEXT("DetailView_ExposureDisplayName", "Exposure"));
+			DetailExposureCategoryBuilder.AddProperty(ColorGradingSettingsHandle->GetChildHandle(GET_MEMBER_NAME_CHECKED(FDisplayClusterConfigurationViewport_ColorGradingRenderingSettings, AutoExposureBias)));
 
-				IDetailCategoryBuilder& ChromakeyMarkersCategoryBuilder = DetailBuilder.EditCategory(TEXT("CustomChromakeyMarkersCategory"), LOCTEXT("CustomChromakeyMarkersCategoryLabel", "ChromakeyMarkers"));
-				AddProperty(DetailBuilder, ChromakeyMarkersCategoryBuilder, TEXT("ChromakeyMarkersRef"), true);
-			}
-			break;
+			IDetailCategoryBuilder& DetailColorGradingCategoryBuilder = DetailBuilder.EditCategory(TEXT("DetailView_ColorGrading"), LOCTEXT("DetailView_ColorGradingDisplayName", "Color Grading"));
+			DetailColorGradingCategoryBuilder.AddProperty(ColorGradingSettingsHandle->GetChildHandle(GET_MEMBER_NAME_CHECKED(FDisplayClusterConfigurationViewport_ColorGradingRenderingSettings, ColorCorrectionShadowsMax)));
+			DetailColorGradingCategoryBuilder.AddProperty(ColorGradingSettingsHandle->GetChildHandle(GET_MEMBER_NAME_CHECKED(FDisplayClusterConfigurationViewport_ColorGradingRenderingSettings, ColorCorrectionHighlightsMin)));
+			DetailColorGradingCategoryBuilder.AddProperty(ColorGradingSettingsHandle->GetChildHandle(GET_MEMBER_NAME_CHECKED(FDisplayClusterConfigurationViewport_ColorGradingRenderingSettings, ColorCorrectionHighlightsMax)));
 
-		case EDetailsSectionType::Chromakey_Custom:
-			{
-				IDetailCategoryBuilder& ChromakeyCategoryBuilder = DetailBuilder.EditCategory(TEXT("CustomChromakeyCategory"), LOCTEXT("CustomChromakeyCategoryLabel", "Chromakey"));
-				AddProperty(DetailBuilder, ChromakeyCategoryBuilder, TEXT("ChromakeyColorRef"));
+			IDetailCategoryBuilder& DetailWhiteBalanceCategoryBuilder = DetailBuilder.EditCategory(TEXT("DetailView_WhiteBalance"), LOCTEXT("DetailView_WhiteBalanceDisplayName", "White Balance"));
+			DetailWhiteBalanceCategoryBuilder.AddProperty(ColorGradingSettingsHandle->GetChildHandle(GET_MEMBER_NAME_CHECKED(FDisplayClusterConfigurationViewport_ColorGradingWhiteBalanceSettings, TemperatureType)));
+			DetailWhiteBalanceCategoryBuilder.AddProperty(ColorGradingSettingsHandle->GetChildHandle(GET_MEMBER_NAME_CHECKED(FDisplayClusterConfigurationViewport_ColorGradingWhiteBalanceSettings, WhiteTemp)));
+			DetailWhiteBalanceCategoryBuilder.AddProperty(ColorGradingSettingsHandle->GetChildHandle(GET_MEMBER_NAME_CHECKED(FDisplayClusterConfigurationViewport_ColorGradingWhiteBalanceSettings, WhiteTint)));
 
-				IDetailCategoryBuilder& ChromakeyCustomCategoryBuilder = DetailBuilder.EditCategory(TEXT("CustomChromakeyCustomCategory"), LOCTEXT("CustomChromakeyCustomCategoryLabel", "Custom Chromakey"));
-				AddProperty(DetailBuilder, ChromakeyCustomCategoryBuilder, TEXT("ChromakeyRenderTextureRef"), true);
-			}
-			break;
+			IDetailCategoryBuilder& DetailMiscCategoryBuilder = DetailBuilder.EditCategory(TEXT("DetailView_Misc"), LOCTEXT("DetailView_MiscDisplayName", "Misc"));
+			DetailMiscCategoryBuilder.AddProperty(ColorGradingSettingsHandle->GetChildHandle(GET_MEMBER_NAME_CHECKED(FDisplayClusterConfigurationViewport_ColorGradingMiscSettings, BlueCorrection)));
+			DetailMiscCategoryBuilder.AddProperty(ColorGradingSettingsHandle->GetChildHandle(GET_MEMBER_NAME_CHECKED(FDisplayClusterConfigurationViewport_ColorGradingMiscSettings, ExpandGamut)));
+			DetailMiscCategoryBuilder.AddProperty(ColorGradingSettingsHandle->GetChildHandle(GET_MEMBER_NAME_CHECKED(FDisplayClusterConfigurationViewport_ColorGradingMiscSettings, SceneColorTint)));
+		};
 
-		case EDetailsSectionType::OCIO_AllNodes:
-			{
-				IDetailCategoryBuilder& OCIOAllNodesCategoryBuilder = DetailBuilder.EditCategory(TEXT("CustomAllNodesOCIOCategory"), LOCTEXT("CustomAllNodesOCIOCategoryLabel", "All Nodes"));
-				AddProperty(DetailBuilder, OCIOAllNodesCategoryBuilder, TEXT("OCIOColorConfiguratonRef"));
-			}
-			break;
+		const int32 GroupIndex = ColorGradingDataModel.IsValid() ? ColorGradingDataModel.Pin()->GetSelectedColorGradingGroupIndex() : INDEX_NONE;
+
+		if (GroupIndex > 0)
+		{
+			const int32 Index = GroupIndex - 1;
+			IDetailCategoryBuilder& PerNodeSettingsCategoryBuilder = DetailBuilder.EditCategory(TEXT("DetailView_PerNode"), LOCTEXT("DetailView_PerNodeDisplayName", "Per-Node Settings"));
+			PerNodeSettingsCategoryBuilder.AddProperty(DetailBuilder.GetProperty(GET_MEMBER_NAME_ARRAY_CHECKED(UDisplayClusterICVFXCameraComponent, CameraSettings.PerNodeColorGrading[Index].bEntireClusterColorGrading, Index)));
+			PerNodeSettingsCategoryBuilder.AddProperty(DetailBuilder.GetProperty(GET_MEMBER_NAME_ARRAY_CHECKED(UDisplayClusterICVFXCameraComponent, CameraSettings.PerNodeColorGrading[Index].bAllNodesColorGrading, Index)));
 
-		case EDetailsSectionType::OCIO_PerNode:
+			AddColorGradingSettings(DetailBuilder.GetProperty(GET_MEMBER_NAME_ARRAY_CHECKED(UDisplayClusterICVFXCameraComponent, CameraSettings.PerNodeColorGrading[Index].ColorGradingSettings, Index)));
+		}
+		else
+		{
+			AddColorGradingSettings(DetailBuilder.GetProperty(GET_MEMBER_NAME_CHECKED(UDisplayClusterICVFXCameraComponent, CameraSettings.AllNodesColorGrading.ColorGradingSettings)));
+		}
+	}
+
+	void AddDetailsPanelProperties(IDetailLayoutBuilder& DetailBuilder)
+	{
+		auto AddProperty = [&DetailBuilder](IDetailCategoryBuilder& Category, FName PropertyName, bool bExpandChildProperties = false)
+		{
+			TSharedRef<IPropertyHandle> PropertyHandle = DetailBuilder.GetProperty(PropertyName, UDisplayClusterICVFXCameraComponent::StaticClass());
+
+			if (bExpandChildProperties)
 			{
-				IDetailCategoryBuilder& OCIOPerNodeCategoryBuilder = DetailBuilder.EditCategory(TEXT("CustomPerNodeOCIOCategory"), LOCTEXT("CustomPerNodeOCIOCategoryLabel", "Per-Node"));
-				AddProperty(DetailBuilder, OCIOPerNodeCategoryBuilder, TEXT("PerNodeOCIOProfilesRef"));
+				PropertyHandle->SetInstanceMetaData(TEXT("ShowOnlyInnerProperties"), TEXT("1"));
 			}
-			break;
-		}
+
+			Category.AddProperty(PropertyHandle);
+		};
+
+		IDetailCategoryBuilder& ICVFXCategoryBuilder = DetailBuilder.EditCategory(TEXT("CustomICVFXCategory"), LOCTEXT("CustomICVFXCategoryLabel", "In-Camera VFX"));
+		AddProperty(ICVFXCategoryBuilder, TEXT("BufferRatioRef"));
+		AddProperty(ICVFXCategoryBuilder, TEXT("ExternalCameraActorRef"));
+		AddProperty(ICVFXCategoryBuilder, TEXT("HiddenICVFXViewportsRef"));
+
+		IDetailCategoryBuilder& SoftEdgeCategoryBuilder = DetailBuilder.EditCategory(TEXT("CustomSoftEdgeCategory"), LOCTEXT("CustomSoftEdgeCategoryLabel", "Soft Edge"));
+		AddProperty(SoftEdgeCategoryBuilder, TEXT("SoftEdgeRef"), true);
+
+		IDetailCategoryBuilder& BorderCategoryBuilder = DetailBuilder.EditCategory(TEXT("CustomBorderCategory"), LOCTEXT("CustomBorderCategoryLabel", "Border"));
+		AddProperty(BorderCategoryBuilder, TEXT("BorderRef"), true);
+
+		IDetailCategoryBuilder& OverscanCategoryBuilder = DetailBuilder.EditCategory(TEXT("CustomOverscanCategory"), LOCTEXT("CustomOverscanCategoryLabel", "Inner Frustum Overscan"));
+		AddProperty(OverscanCategoryBuilder, TEXT("CustomFrustumRef"), true);
+
+		IDetailCategoryBuilder& ChromakeyCategoryBuilder = DetailBuilder.EditCategory(TEXT("CustomChromakeyCategory"), LOCTEXT("CustomChromakeyCategoryLabel", "Chromakey"));
+		AddProperty(ChromakeyCategoryBuilder, TEXT("ChromakeyColorRef"));
+
+		IDetailCategoryBuilder& ChromakeyMarkersCategoryBuilder = DetailBuilder.EditCategory(TEXT("CustomChromakeyMarkersCategory"), LOCTEXT("CustomChromakeyMarkersCategoryLabel", "ChromakeyMarkers"));
+		AddProperty(ChromakeyMarkersCategoryBuilder, TEXT("ChromakeyMarkersRef"), true);
+
+		IDetailCategoryBuilder& ChromakeyCustomCategoryBuilder = DetailBuilder.EditCategory(TEXT("CustomChromakeyCustomCategory"), LOCTEXT("CustomChromakeyCustomCategoryLabel", "Custom Chromakey"));
+		AddProperty(ChromakeyCustomCategoryBuilder, TEXT("ChromakeyRenderTextureRef"), true);
+
+		IDetailCategoryBuilder& OCIOAllNodesCategoryBuilder = DetailBuilder.EditCategory(TEXT("CustomAllNodesOCIOCategory"), LOCTEXT("CustomAllNodesOCIOCategoryLabel", "All Nodes"));
+		AddProperty(OCIOAllNodesCategoryBuilder, TEXT("OCIOColorConfiguratonRef"));
+
+		IDetailCategoryBuilder& OCIOPerNodeCategoryBuilder = DetailBuilder.EditCategory(TEXT("CustomPerNodeOCIOCategory"), LOCTEXT("CustomPerNodeOCIOCategoryLabel", "Per-Node"));
+		AddProperty(OCIOPerNodeCategoryBuilder, TEXT("PerNodeOCIOProfilesRef"));
 	}
 
 private:
-	void AddProperty(IDetailLayoutBuilder& DetailBuilder, IDetailCategoryBuilder& Category, FName PropertyName, bool bExpandChildProperties = false)
+	TWeakPtr<FDisplayClusterColorGradingDataModel> ColorGradingDataModel;
+};
+
+/** A property customizer that culls unneeded properties from the FDisplayClusterConfigurationViewport_AllNodesColorGrading struct to help speed up property node tree generation */
+class FFastAllNodesColorGradingCustomization : public IPropertyTypeCustomization
+{
+public:
+	FFastAllNodesColorGradingCustomization(const TSharedRef<FDisplayClusterColorGradingDataModel>& InColorGradingDataModel)
+		: ColorGradingDataModel(InColorGradingDataModel)
+	{ }
+
+	virtual void CustomizeHeader(TSharedRef<IPropertyHandle> StructPropertyHandle, FDetailWidgetRow& HeaderRow, IPropertyTypeCustomizationUtils& StructCustomizationUtils) override
 	{
-		TSharedRef<IPropertyHandle> PropertyHandle = DetailBuilder.GetProperty(PropertyName, UDisplayClusterICVFXCameraComponent::StaticClass());
+		HeaderRow.NameContent()
+		[
+			StructPropertyHandle->CreatePropertyNameWidget()
+		];
+	}
+
+	virtual void CustomizeChildren(TSharedRef<IPropertyHandle> StructPropertyHandle, IDetailChildrenBuilder& StructBuilder, IPropertyTypeCustomizationUtils& StructCustomizationUtils) override
+	{
+		const int32 GroupIndex = ColorGradingDataModel.IsValid() ? ColorGradingDataModel.Pin()->GetSelectedColorGradingGroupIndex() : INDEX_NONE;
 
-		if (bExpandChildProperties)
+		StructBuilder.AddProperty(StructPropertyHandle->GetChildHandle(GET_MEMBER_NAME_CHECKED(FDisplayClusterConfigurationViewport_AllNodesColorGrading, bEnableInnerFrustumAllNodesColorGrading)).ToSharedRef());
+
+		if (GroupIndex < 1)
 		{
-			PropertyHandle->SetInstanceMetaData(TEXT("ShowOnlyInnerProperties"), TEXT("1"));
+			StructBuilder.AddProperty(StructPropertyHandle->GetChildHandle(GET_MEMBER_NAME_CHECKED(FDisplayClusterConfigurationViewport_AllNodesColorGrading, ColorGradingSettings)).ToSharedRef());
 		}
+	}
+
+private:
+	TWeakPtr<FDisplayClusterColorGradingDataModel> ColorGradingDataModel;
+};
+
+/** A property customizer that culls unneeded properties from the FDisplayClusterConfigurationViewport_PerNodeColorGrading struct to help speed up property node tree generation */
+class FFastPerNodeColorGradingCustomization : public IPropertyTypeCustomization
+{
+public:
+	FFastPerNodeColorGradingCustomization(const TSharedRef<FDisplayClusterColorGradingDataModel>& InColorGradingDataModel)
+		: ColorGradingDataModel(InColorGradingDataModel)
+	{ }
 
-		Category.AddProperty(PropertyHandle);
+	virtual void CustomizeHeader(TSharedRef<IPropertyHandle> StructPropertyHandle, FDetailWidgetRow& HeaderRow, IPropertyTypeCustomizationUtils& StructCustomizationUtils) override
+	{
+		HeaderRow.NameContent()
+		[
+			StructPropertyHandle->CreatePropertyNameWidget()
+		];
+	}
+
+	virtual void CustomizeChildren(TSharedRef<IPropertyHandle> StructPropertyHandle, IDetailChildrenBuilder& StructBuilder, IPropertyTypeCustomizationUtils& StructCustomizationUtils) override
+	{
+		const int32 ArrayIndex = StructPropertyHandle->GetIndexInArray();
+		const int32 GroupIndex = ColorGradingDataModel.IsValid() ? ColorGradingDataModel.Pin()->GetSelectedColorGradingGroupIndex() : INDEX_NONE;
+
+		StructBuilder.AddProperty(StructPropertyHandle->GetChildHandle(GET_MEMBER_NAME_CHECKED(FDisplayClusterConfigurationViewport_PerNodeColorGrading, bIsEnabled)).ToSharedRef());
+		StructBuilder.AddProperty(StructPropertyHandle->GetChildHandle(GET_MEMBER_NAME_CHECKED(FDisplayClusterConfigurationViewport_PerNodeColorGrading, Name)).ToSharedRef());
+
+		if (GroupIndex == ArrayIndex + 1)
+		{
+			StructBuilder.AddProperty(StructPropertyHandle->GetChildHandle(GET_MEMBER_NAME_CHECKED(FDisplayClusterConfigurationViewport_PerNodeColorGrading, ColorGradingSettings)).ToSharedRef());
+		}
 	}
 
 private:
-	EDetailsSectionType SectionType;
+	TWeakPtr<FDisplayClusterColorGradingDataModel> ColorGradingDataModel;
 };
 
+void FDisplayClusterColorGradingGenerator_ICVFXCamera::Initialize(const TSharedRef<class FDisplayClusterColorGradingDataModel>& ColorGradingDataModel, const TSharedRef<IPropertyRowGenerator>& PropertyRowGenerator)
+{
+	PropertyRowGenerator->RegisterInstancedCustomPropertyTypeLayout(FDisplayClusterConfigurationViewport_AllNodesColorGrading::StaticStruct()->GetFName(), FOnGetPropertyTypeCustomizationInstance::CreateLambda([ColorGradingDataModel]
+	{
+		return MakeShared<FFastAllNodesColorGradingCustomization>(ColorGradingDataModel);
+	}));
+
+	PropertyRowGenerator->RegisterInstancedCustomPropertyTypeLayout(FDisplayClusterConfigurationViewport_PerNodeColorGrading::StaticStruct()->GetFName(), FOnGetPropertyTypeCustomizationInstance::CreateLambda([ColorGradingDataModel]
+	{
+		return MakeShared<FFastPerNodeColorGradingCustomization>(ColorGradingDataModel);
+	}));
+
+	PropertyRowGenerator->RegisterInstancedCustomPropertyLayout(UDisplayClusterICVFXCameraComponent::StaticClass(), FOnGetDetailCustomizationInstance::CreateLambda([ColorGradingDataModel]
+	{
+		return MakeShared<FICVFXCameraColorGradingCustomization>(ColorGradingDataModel);
+	}));
+}
+
+void FDisplayClusterColorGradingGenerator_ICVFXCamera::Destroy(const TSharedRef<class FDisplayClusterColorGradingDataModel>& ColorGradingDataModel, const TSharedRef<IPropertyRowGenerator>& PropertyRowGenerator)
+{
+	PropertyRowGenerator->UnregisterInstancedCustomPropertyTypeLayout(FDisplayClusterConfigurationViewport_AllNodesColorGrading::StaticStruct()->GetFName());
+	PropertyRowGenerator->UnregisterInstancedCustomPropertyTypeLayout(FDisplayClusterConfigurationViewport_PerNodeColorGrading::StaticStruct()->GetFName());
+	PropertyRowGenerator->UnregisterInstancedCustomPropertyLayout(UDisplayClusterICVFXCameraComponent::StaticClass());
+}
+
 void FDisplayClusterColorGradingGenerator_ICVFXCamera::GenerateDataModel(IPropertyRowGenerator& PropertyRowGenerator, FDisplayClusterColorGradingDataModel& OutColorGradingDataModel)
 {
 	CameraComponents.Empty();
@@ -748,18 +970,7 @@ void FDisplayClusterColorGradingGenerator_ICVFXCamera::GenerateDataModel(IProper
 				FDisplayClusterColorGradingDataModel::FColorGradingGroup PerNodeGroup = CreateColorGradingGroup(PerNodeElementHandle);
 				PerNodeGroup.EditConditionPropertyHandle = PerNodeElementHandle->GetChildHandle(GET_MEMBER_NAME_CHECKED(FDisplayClusterConfigurationViewport_PerNodeColorGrading, bIsEnabled));
 
-				// Add per-node group specific properties and force the category of these properties to be at the top
-				AddDetailsViewPropertyToGroup(PerNodeElementHandle,
-					PerNodeGroup,
-					GET_MEMBER_NAME_CHECKED(FDisplayClusterConfigurationViewport_PerNodeColorGrading, bEntireClusterColorGrading),
-					TEXT("Per-Node Settings"));
-
-				AddDetailsViewPropertyToGroup(PerNodeElementHandle,
-					PerNodeGroup,
-					GET_MEMBER_NAME_CHECKED(FDisplayClusterConfigurationViewport_PerNodeColorGrading, bAllNodesColorGrading),
-					TEXT("Per-Node Settings"));
-
-				PerNodeGroup.DetailsViewCategorySortOrder.Add(CATEGORY_OVERRIDE_NAME("Per-Node Settings"), -1);
+				PerNodeGroup.DetailsViewCategories.Add(TEXT("DetailView_PerNode"));
 
 				TSharedPtr<IPropertyHandle> NamePropertyHandle = PerNodeElementHandle->GetChildHandle(GET_MEMBER_NAME_CHECKED(FDisplayClusterConfigurationViewport_PerNodeColorGrading, Name));
 				if (NamePropertyHandle.IsValid() && NamePropertyHandle->IsValidHandle())
@@ -827,19 +1038,13 @@ void FDisplayClusterColorGradingGenerator_ICVFXCamera::GenerateDataModel(IProper
 
 		FDisplayClusterColorGradingDataModel::FDetailsSubsection ICVFXDetailsSubsection;
 		ICVFXDetailsSubsection.DisplayName = LOCTEXT("ICVFXSubsectionLabel", "ICVFX");
-		ICVFXDetailsSubsection.DetailCustomizationDelegate = FOnGetDetailCustomizationInstance::CreateLambda([]
-		{
-			return MakeShared<FICVFXCameraDetailsSectionCustomization>(FICVFXCameraDetailsSectionCustomization::EDetailsSectionType::ICVFX);
-		});
+		ICVFXDetailsSubsection.Categories = { TEXT("CustomICVFXCategory"), TEXT("CustomSoftEdgeCategory"), TEXT("CustomBorderCategory") };
 
 		InnerFrustumDetailsSection.Subsections.Add(ICVFXDetailsSubsection);
 
 		FDisplayClusterColorGradingDataModel::FDetailsSubsection OverscanDetailsSubsection;
 		OverscanDetailsSubsection.DisplayName = LOCTEXT("OverscanDetailsSubsectionLabel", "Overscan");
-		OverscanDetailsSubsection.DetailCustomizationDelegate = FOnGetDetailCustomizationInstance::CreateLambda([]
-		{
-			return MakeShared<FICVFXCameraDetailsSectionCustomization>(FICVFXCameraDetailsSectionCustomization::EDetailsSectionType::Overscan);
-		});
+		OverscanDetailsSubsection.Categories = { TEXT("CustomOverscanCategory") };
 
 		InnerFrustumDetailsSection.Subsections.Add(OverscanDetailsSubsection);
 
@@ -853,19 +1058,13 @@ void FDisplayClusterColorGradingGenerator_ICVFXCamera::GenerateDataModel(IProper
 
 		FDisplayClusterColorGradingDataModel::FDetailsSubsection ChromakeyMarkersDetailsSubsection;
 		ChromakeyMarkersDetailsSubsection.DisplayName = LOCTEXT("ChromakeyMarkersDetailsSubsectionLabel", "Markers");
-		ChromakeyMarkersDetailsSubsection.DetailCustomizationDelegate = FOnGetDetailCustomizationInstance::CreateLambda([]
-		{
-			return MakeShared<FICVFXCameraDetailsSectionCustomization>(FICVFXCameraDetailsSectionCustomization::EDetailsSectionType::Chromakey_Markers);
-		});
+		ChromakeyMarkersDetailsSubsection.Categories = { TEXT("CustomChromakeyCategory"), TEXT("CustomChromakeyMarkersCategory") };
 
 		ChromakeyDetailsSection.Subsections.Add(ChromakeyMarkersDetailsSubsection);
 
 		FDisplayClusterColorGradingDataModel::FDetailsSubsection ChromakeyCustomDetailsSubsection;
 		ChromakeyCustomDetailsSubsection.DisplayName = LOCTEXT("ChromakeyCustomDetailsSubsectionLabel", "Custom");
-		ChromakeyCustomDetailsSubsection.DetailCustomizationDelegate = FOnGetDetailCustomizationInstance::CreateLambda([]
-		{
-			return MakeShared<FICVFXCameraDetailsSectionCustomization>(FICVFXCameraDetailsSectionCustomization::EDetailsSectionType::Chromakey_Custom);
-		});
+		ChromakeyCustomDetailsSubsection.Categories = { TEXT("CustomChromakeyCategory"), TEXT("CustomChromakeyCustomCategory") };
 
 		ChromakeyDetailsSection.Subsections.Add(ChromakeyCustomDetailsSubsection);
 
@@ -879,19 +1078,13 @@ void FDisplayClusterColorGradingGenerator_ICVFXCamera::GenerateDataModel(IProper
 
 		FDisplayClusterColorGradingDataModel::FDetailsSubsection AllNodesOCIODetailsSubsection;
 		AllNodesOCIODetailsSubsection.DisplayName = LOCTEXT("AllNodesOCIODetailsSubsectionLabel", "All Nodes");
-		AllNodesOCIODetailsSubsection.DetailCustomizationDelegate = FOnGetDetailCustomizationInstance::CreateLambda([]
-		{
-			return MakeShared<FICVFXCameraDetailsSectionCustomization>(FICVFXCameraDetailsSectionCustomization::EDetailsSectionType::OCIO_AllNodes);
-		});
+		AllNodesOCIODetailsSubsection.Categories = { TEXT("CustomAllNodesOCIOCategory") };
 
 		OCIODetailsSection.Subsections.Add(AllNodesOCIODetailsSubsection);
 
 		FDisplayClusterColorGradingDataModel::FDetailsSubsection PerNodeOCIODetailsSubsection;
 		PerNodeOCIODetailsSubsection.DisplayName = LOCTEXT("PerNodeOCIODetailsSubsectionLabel", "Per-Node");
-		PerNodeOCIODetailsSubsection.DetailCustomizationDelegate = FOnGetDetailCustomizationInstance::CreateLambda([]
-		{
-			return MakeShared<FICVFXCameraDetailsSectionCustomization>(FICVFXCameraDetailsSectionCustomization::EDetailsSectionType::OCIO_PerNode);
-		});
+		PerNodeOCIODetailsSubsection.Categories = { TEXT("CustomPerNodeOCIOCategory") };
 
 		OCIODetailsSection.Subsections.Add(PerNodeOCIODetailsSubsection);
 
@@ -985,7 +1178,6 @@ FText FDisplayClusterColorGradingGenerator_ICVFXCamera::GetNodeComboBoxText(int3
 	return FText::GetEmpty();
 }
 
-
 TSharedRef<SWidget> FDisplayClusterColorGradingGenerator_ICVFXCamera::GetNodeComboBoxMenu(int32 PerNodeColorGradingIndex) const
 {
 
diff --git a/Engine/Plugins/Runtime/nDisplay/Source/DisplayClusterColorGrading/Private/DataModelGenerators/DisplayClusterColorGradingGenerator_RootActor.h b/Engine/Plugins/Runtime/nDisplay/Source/DisplayClusterColorGrading/Private/DataModelGenerators/DisplayClusterColorGradingGenerator_RootActor.h
--- a/Engine/Plugins/Runtime/nDisplay/Source/DisplayClusterColorGrading/Private/DataModelGenerators/DisplayClusterColorGradingGenerator_RootActor.h
+++ b/Engine/Plugins/Runtime/nDisplay/Source/DisplayClusterColorGrading/Private/DataModelGenerators/DisplayClusterColorGradingGenerator_RootActor.h
@@ -25,9 +25,6 @@ protected:
 	/** Creates a color grading element structure from the specified property handle, whose child properties are expected to be colors with the ColorGradingMode metadata set */
 	FDisplayClusterColorGradingDataModel::FColorGradingElement CreateColorGradingElement(const TSharedPtr<IPropertyHandle>& GroupPropertyHandle, FName ElementPropertyName, FText ElementLabel);
 
-	/** Adds the specified property handle to the specified color grading group's list of properties to display in the details view */
-	void AddDetailsViewPropertyToGroup(const TSharedPtr<IPropertyHandle>& GroupPropertyHandle, FDisplayClusterColorGradingDataModel::FColorGradingGroup& Group, FName PropertyName, const FString& CategoryOverride);
-
 	/** Recursively searches the detail tree hierarchy for a property detail tree node whose name matches the specified name */
 	TSharedPtr<IDetailTreeNode> FindPropertyTreeNode(const TSharedRef<IDetailTreeNode>& Node, const FCachedPropertyPath& PropertyPath);
 
@@ -42,6 +39,8 @@ public:
 	static TSharedRef<IDisplayClusterColorGradingDataModelGenerator> MakeInstance();
 
 	//~ IDisplayClusterColorGradingDataModelGenerator interface
+	virtual void Initialize(const TSharedRef<class FDisplayClusterColorGradingDataModel>& ColorGradingDataModel, const TSharedRef<IPropertyRowGenerator>& PropertyRowGenerator) override;
+	virtual void Destroy(const TSharedRef<class FDisplayClusterColorGradingDataModel>& ColorGradingDataModel, const TSharedRef<IPropertyRowGenerator>& PropertyRowGenerator) override;
 	virtual void GenerateDataModel(IPropertyRowGenerator& PropertyRowGenerator, FDisplayClusterColorGradingDataModel& OutColorGradingDataModel) override;
 	//~ End IDisplayClusterColorGradingDataModelGenerator interface
 
@@ -70,6 +69,8 @@ public:
 	static TSharedRef<IDisplayClusterColorGradingDataModelGenerator> MakeInstance();
 
 	//~ IDisplayClusterColorGradingDataModelGenerator interface
+	virtual void Initialize(const TSharedRef<class FDisplayClusterColorGradingDataModel>& InColorGradingDataModel, const TSharedRef<IPropertyRowGenerator>& InPropertyRowGenerator) override;
+	virtual void Destroy(const TSharedRef<class FDisplayClusterColorGradingDataModel>& InColorGradingDataModel, const TSharedRef<IPropertyRowGenerator>& InPropertyRowGenerator) override;
 	virtual void GenerateDataModel(IPropertyRowGenerator& PropertyRowGenerator, FDisplayClusterColorGradingDataModel& OutColorGradingDataModel) override;
 	//~ End IDisplayClusterColorGradingDataModelGenerator interface
 
@@ -87,6 +88,6 @@ private:
 	TSharedRef<SWidget> GetNodeComboBoxMenu(int32 PerNodeColorGradingIndex) const;
 
 private:
-	/** A list of root actors that are being represented by the data model */
+	/** A list of camera components that are being represented by the data model */
 	TArray<TWeakObjectPtr<UDisplayClusterICVFXCameraComponent>> CameraComponents;
 };
\ No newline at end of file
diff --git a/Engine/Plugins/Runtime/nDisplay/Source/DisplayClusterColorGrading/Private/DetailView/SDisplayClusterColorGradingDetailTreeRow.cpp b/Engine/Plugins/Runtime/nDisplay/Source/DisplayClusterColorGrading/Private/DetailView/SDisplayClusterColorGradingDetailTreeRow.cpp
new file mode 100644
--- /dev/null
+++ b/Engine/Plugins/Runtime/nDisplay/Source/DisplayClusterColorGrading/Private/DetailView/SDisplayClusterColorGradingDetailTreeRow.cpp
@@ -0,0 +1,876 @@
+// Copyright Epic Games, Inc. All Rights Reserved.DisplayClusterColorGradingDetailView
+
+#include "SDisplayClusterColorGradingDetailTreeRow.h"
+
+#include "SDisplayClusterColorGradingDetailView.h"
+
+#include "DragAndDrop/DecoratedDragDropOp.h"
+#include "Framework/MultiBox/MultiBoxBuilder.h"
+#include "IPropertyRowGenerator.h"
+#include "IDetailDragDropHandler.h"
+#include "IDetailTreeNode.h"
+#include "Misc/ConfigCacheIni.h"
+#include "PropertyEditor/Private/DetailTreeNode.h"
+#include "ScopedTransaction.h"
+#include "Modules/ModuleManager.h"
+#include "Styling/StyleColors.h"
+#include "Widgets/Input/SCheckBox.h"
+#include "Widgets/Views/STableRow.h"
+#include "Widgets/Views/STreeView.h"
+
+namespace UE::DisplayClusterColorGradingDetailTreeRow
+{
+	DECLARE_DELEGATE_RetVal_OneParam(FSlateColor, FOnGetIndentBackgroundColor, int32);
+
+	// Copy of private SDetailRowIndent widget defined at PropertyEditor/Private/SDetailRowIndent.h
+	class SDetailRowIndent : public SCompoundWidget
+	{
+	public:
+		SLATE_BEGIN_ARGS(SDetailRowIndent) {}
+			SLATE_ATTRIBUTE(int32, Indent)
+			SLATE_EVENT(FOnGetIndentBackgroundColor, OnGetIndentBackgroundColor)
+		SLATE_END_ARGS()
+
+		void Construct(const FArguments& InArgs)
+		{
+			Indent = InArgs._Indent;
+			OnGetIndentBackgroundColor = InArgs._OnGetIndentBackgroundColor;
+
+			ChildSlot
+			[
+				SNew(SBox)
+				.WidthOverride(this, &SDetailRowIndent::GetIndentWidth)
+			];
+		}
+
+
+	private:
+		virtual int32 OnPaint(const FPaintArgs& Args,
+			const FGeometry& AllottedGeometry,
+			const FSlateRect& MyCullingRect,
+			FSlateWindowElementList& OutDrawElements,
+			int32 LayerId,
+			const FWidgetStyle& InWidgetStyle,
+			bool bParentEnabled) const override
+		{
+			const FSlateBrush* BackgroundBrush = FAppStyle::Get().GetBrush("DetailsView.CategoryMiddle");
+			const FSlateBrush* DropShadowBrush = FAppStyle::Get().GetBrush("DetailsView.ArrayDropShadow");
+
+			int32 IndentLevel = Indent.Get(0);
+			for (int32 IndentIndex = 0; IndentIndex < IndentLevel; ++IndentIndex)
+			{
+				FSlateColor BackgroundColor = OnGetIndentBackgroundColor.IsBound() ? OnGetIndentBackgroundColor.Execute(IndentIndex) : FSlateColor(FLinearColor::White);
+
+				FSlateDrawElement::MakeBox(
+					OutDrawElements,
+					LayerId,
+					AllottedGeometry.ToPaintGeometry(FVector2D(16 * IndentIndex, 0), FVector2D(16, AllottedGeometry.GetLocalSize().Y)),
+					BackgroundBrush,
+					ESlateDrawEffect::None,
+					BackgroundColor.GetColor(InWidgetStyle)
+				);
+
+				FSlateDrawElement::MakeBox(
+					OutDrawElements,
+					LayerId + 1,
+					AllottedGeometry.ToPaintGeometry(FVector2D(16 * IndentIndex, 0), FVector2D(16, AllottedGeometry.GetLocalSize().Y)),
+					DropShadowBrush
+				);
+			}
+
+			return LayerId + 1;
+		}
+
+		FOptionalSize GetIndentWidth() const
+		{
+			return Indent.Get(0) * 16.0f;
+		}
+
+	private:
+		TAttribute<int32> Indent;
+		FOnGetIndentBackgroundColor OnGetIndentBackgroundColor;
+	};
+
+	// Copy of private SEditConditionWidget defined in PropertyEditor/Private/PropertyEditorHelpers.h
+	class SEditConditionWidget : public SCompoundWidget
+	{
+		SLATE_BEGIN_ARGS(SEditConditionWidget) {}
+			SLATE_ATTRIBUTE(bool, EditConditionValue)
+			SLATE_EVENT(FOnBooleanValueChanged, OnEditConditionValueChanged)
+		SLATE_END_ARGS()
+
+		void Construct(const FArguments& Args)
+		{
+			EditConditionValue = Args._EditConditionValue;
+			OnEditConditionValueChanged = Args._OnEditConditionValueChanged;
+
+			ChildSlot
+			[
+				SNew(SCheckBox)
+				.OnCheckStateChanged(this, &SEditConditionWidget::OnEditConditionCheckChanged)
+				.IsChecked(this, &SEditConditionWidget::OnGetEditConditionCheckState)
+				.Visibility(this, &SEditConditionWidget::GetCheckBoxVisibility)
+			];
+		}
+
+	private:
+		bool HasEditConditionToggle() const
+		{
+			return OnEditConditionValueChanged.IsBound();
+		}
+
+		void OnEditConditionCheckChanged(ECheckBoxState CheckState)
+		{
+			checkSlow(HasEditConditionToggle());
+			FScopedTransaction EditConditionChangedTransaction(NSLOCTEXT("PropertyEditor", "UpdatedEditConditionFmt", "Edit Condition Changed"));
+			OnEditConditionValueChanged.ExecuteIfBound(CheckState == ECheckBoxState::Checked);
+		}
+
+		ECheckBoxState OnGetEditConditionCheckState() const
+		{
+			return EditConditionValue.Get() ? ECheckBoxState::Checked : ECheckBoxState::Unchecked;
+		}
+
+		EVisibility GetCheckBoxVisibility() const
+		{
+			return HasEditConditionToggle() ? EVisibility::Visible : EVisibility::Collapsed;
+		}
+
+	private:
+		TAttribute<bool> EditConditionValue;
+		FOnBooleanValueChanged OnEditConditionValueChanged;
+	};
+
+	// Copy of private SConstrainedBox defined in PropertyEditor/Private/SConstrainedBox.h
+	class SConstrainedBox : public SCompoundWidget
+	{
+	public:
+		SLATE_BEGIN_ARGS(SConstrainedBox) {}
+			SLATE_DEFAULT_SLOT(FArguments, Content)
+			SLATE_ATTRIBUTE(TOptional<float>, MinWidth)
+			SLATE_ATTRIBUTE(TOptional<float>, MaxWidth)
+		SLATE_END_ARGS()
+
+		void Construct(const FArguments& InArgs)
+		{
+			MinWidth = InArgs._MinWidth;
+			MaxWidth = InArgs._MaxWidth;
+
+			ChildSlot
+			[
+				InArgs._Content.Widget
+			];
+		}
+
+		virtual FVector2D ComputeDesiredSize(float LayoutScaleMultiplier) const override
+		{
+			const float MinWidthVal = MinWidth.Get().Get(0.0f);
+			const float MaxWidthVal = MaxWidth.Get().Get(0.0f);
+
+			if (MinWidthVal == 0.0f && MaxWidthVal == 0.0f)
+			{
+				return SCompoundWidget::ComputeDesiredSize(LayoutScaleMultiplier);
+			}
+			else
+			{
+				FVector2D ChildSize = ChildSlot.GetWidget()->GetDesiredSize();
+
+				float XVal = FMath::Max(MinWidthVal, ChildSize.X);
+				if (MaxWidthVal > MinWidthVal)
+				{
+					XVal = FMath::Min(MaxWidthVal, XVal);
+				}
+
+				return FVector2D(XVal, ChildSize.Y);
+			}
+		}
+
+	private:
+		TAttribute< TOptional<float> > MinWidth;
+		TAttribute< TOptional<float> > MaxWidth;
+	};
+
+	// Copy of private FArrayRowDragDropOp defined in PropertyEditor/Private/SDetailSingleItemRow.h
+	class FArrayRowDragDropOp : public FDecoratedDragDropOp
+	{
+	public:
+		DRAG_DROP_OPERATOR_TYPE(FArrayRowDragDropOp, FDecoratedDragDropOp)
+
+		FArrayRowDragDropOp(const TSharedPtr<SDisplayClusterColorGradingDetailTreeRow>& InRowBeingDragged)
+			: RowBeingDragged(InRowBeingDragged)
+		{
+			MouseCursor = EMouseCursor::GrabHandClosed;
+		}
+
+		void Init()
+		{
+			SetValidTarget(false);
+			SetupDefaults();
+			Construct();
+		}
+
+		void SetValidTarget(bool IsValidTarget)
+		{
+			if (IsValidTarget)
+			{
+				CurrentHoverText = NSLOCTEXT("ArrayDragDrop", "PlaceRowHere", "Place Row Here");
+				CurrentIconBrush = FAppStyle::GetBrush("Graph.ConnectorFeedback.OK");
+			}
+			else
+			{
+				CurrentHoverText = NSLOCTEXT("ArrayDragDrop", "CannotPlaceRowHere", "Cannot Place Row Here");
+				CurrentIconBrush = FAppStyle::GetBrush("Graph.ConnectorFeedback.Error");
+			}
+		}
+
+		TWeakPtr<SDisplayClusterColorGradingDetailTreeRow> RowBeingDragged;
+	};
+
+	// Copy of private SArrayRowHandle defined in PropertyEditor/Private/SDetailSingleItemRow.h
+	class SArrayRowHandle : public SCompoundWidget
+	{
+	public:
+		SLATE_BEGIN_ARGS(SArrayRowHandle) {}
+		SLATE_END_ARGS()
+
+		void Construct(const FArguments& InArgs, const TSharedPtr<SDisplayClusterColorGradingDetailTreeRow>& InParentRow)
+		{
+			ParentRow = InParentRow;
+			SetCursor(EMouseCursor::GrabHand);
+
+			ChildSlot
+			[
+				SNew(SBox)
+				.Padding(0.0f)
+				.HAlign(HAlign_Center)
+				.VAlign(VAlign_Center)
+				.WidthOverride(16.0f)
+				[
+					SNew(SImage)
+					.Image(FAppStyle::Get().GetBrush("VerticalBoxDragIndicatorShort"))
+				]
+			];
+		}
+
+		virtual FReply OnMouseButtonDown(const FGeometry& MyGeometry, const FPointerEvent& MouseEvent) override
+		{
+			return FReply::Handled().DetectDrag(SharedThis(this), EKeys::LeftMouseButton);
+		};
+
+		virtual FReply OnDragDetected(const FGeometry& MyGeometry, const FPointerEvent& MouseEvent) override
+		{
+			if (MouseEvent.IsMouseButtonDown(EKeys::LeftMouseButton))
+			{
+				TSharedPtr<FDragDropOperation> DragDropOp = nullptr;
+				if (ParentRow.IsValid())
+				{
+					DragDropOp = ParentRow.Pin()->CreateDragDropOperation();
+				}
+
+				if (DragDropOp.IsValid())
+				{
+					return FReply::Handled().BeginDragDrop(DragDropOp.ToSharedRef());
+				}
+			}
+
+			return FReply::Unhandled();
+		}
+
+	private:
+		TWeakPtr<IDetailDragDropHandler> CustomDragDropHandler;
+		TWeakPtr<SDisplayClusterColorGradingDetailTreeRow> ParentRow;
+	};
+}
+
+void SDisplayClusterColorGradingDetailTreeRow::Construct(const FArguments& InArgs,
+	const TSharedRef<FDisplayClusterColorGradingDetailTreeItem>& InDetailTreeItem,
+	const TSharedRef<STableViewBase>& InOwnerTableView,
+	const FDetailColumnSizeData& InColumnSizeData)
+{
+	DetailTreeItem = InDetailTreeItem;
+
+	TWeakPtr<STableViewBase> OwnerTableViewWeak = InOwnerTableView;
+
+	ExpanderArrowWidget = SNew(SExpanderArrow, SharedThis(this))
+		.BaseIndentLevel(0)
+		.IndentAmount(0.0f);
+
+	IndentWidget = SNew(UE::DisplayClusterColorGradingDetailTreeRow::SDetailRowIndent)
+		.Indent(this, &SDisplayClusterColorGradingDetailTreeRow::GetPropertyIndent)
+		.OnGetIndentBackgroundColor(this, &SDisplayClusterColorGradingDetailTreeRow::GetIndentBackgroundColor);
+
+	STableRow::FArguments TableRowArgs = STableRow::FArguments()
+		.Style(FAppStyle::Get(), "DetailsView.TreeView.TableRow")
+		.ShowSelection(false)
+		.OnDragLeave(this, &SDisplayClusterColorGradingDetailTreeRow::OnRowDragLeave)
+		.OnAcceptDrop(this, &SDisplayClusterColorGradingDetailTreeRow::OnRowAcceptDrop)
+		.OnCanAcceptDrop(this, &SDisplayClusterColorGradingDetailTreeRow::CanRowAcceptDrop)
+		[
+			SNew(SBorder)
+			.BorderImage(FAppStyle::Get().GetBrush("DetailsView.GridLine"))
+			.Padding(FMargin(0, 0, 0, 1))
+			.Clipping(EWidgetClipping::ClipToBounds)
+			[
+				SNew(SBox)
+				.MinDesiredHeight(26.0f)
+				[
+					SNew(SHorizontalBox)
+
+					+ SHorizontalBox::Slot()
+					[
+						SNew(SBorder)
+						.Padding(0)
+						.BorderImage(this, &SDisplayClusterColorGradingDetailTreeRow::GetRowBackgroundBrush)
+						.BorderBackgroundColor(this, &SDisplayClusterColorGradingDetailTreeRow::GetRowBackgroundColor)
+						[
+							CreatePropertyWidget(InColumnSizeData)
+						]
+					]
+
+					+ SHorizontalBox::Slot()
+					.HAlign(HAlign_Right)
+					.AutoWidth()
+					[
+						SNew(SBorder)
+						.Padding(FMargin(0, 0, 16.0f, 0))
+						.BorderImage(this, &SDisplayClusterColorGradingDetailTreeRow::GetScrollWellBackgroundBrush, OwnerTableViewWeak)
+						.BorderBackgroundColor(this, &SDisplayClusterColorGradingDetailTreeRow::GetScrollWellBackgroundColor, OwnerTableViewWeak)
+					]
+				]
+			]
+		];
+
+	STableRow<TSharedRef<FDisplayClusterColorGradingDetailTreeItem>>::Construct(TableRowArgs, InOwnerTableView);
+}
+
+void SDisplayClusterColorGradingDetailTreeRow::ConstructChildren(ETableViewMode::Type InOwnerTableMode, const TAttribute<FMargin>& InPadding, const TSharedRef<SWidget>& InContent)
+{
+	InnerContentSlot = nullptr;
+	SetContent(InContent);
+}
+
+TSharedPtr<IPropertyHandle> SDisplayClusterColorGradingDetailTreeRow::GetPropertyHandle() const
+{
+	TSharedPtr<IPropertyHandle> PropertyHandle = nullptr;
+
+	if (DetailTreeItem.IsValid())
+	{
+		PropertyHandle = DetailTreeItem.Pin()->GetPropertyHandle();
+	}
+
+	return PropertyHandle;
+}
+
+TSharedPtr<FDragDropOperation> SDisplayClusterColorGradingDetailTreeRow::CreateDragDropOperation()
+{
+	if (WidgetRow.CustomDragDropHandler)
+	{
+		TSharedPtr<FDragDropOperation> DragOp = WidgetRow.CustomDragDropHandler->CreateDragDropOperation();
+		return DragOp;
+	}
+	else
+	{
+		TSharedPtr<UE::DisplayClusterColorGradingDetailTreeRow::FArrayRowDragDropOp> ArrayDragOp = MakeShareable(new UE::DisplayClusterColorGradingDetailTreeRow::FArrayRowDragDropOp(SharedThis(this)));
+		ArrayDragOp->Init();
+		return ArrayDragOp;
+	}
+}
+
+TSharedRef<SWidget> SDisplayClusterColorGradingDetailTreeRow::CreatePropertyWidget(const FDetailColumnSizeData& InColumnSizeData)
+{
+	if (DetailTreeItem.IsValid())
+	{
+		TSharedPtr<FDisplayClusterColorGradingDetailTreeItem> PinnedDetailTreeItem = DetailTreeItem.Pin();
+		PinnedDetailTreeItem->GenerateDetailWidgetRow(WidgetRow);
+
+		if (!WidgetRow.HasAnyContent())
+		{
+			return SNullWidget::NullWidget;
+		}
+
+		TSharedRef<SSplitter> Splitter = SNew(SSplitter)
+			.Style(FAppStyle::Get(), "DetailsView.Splitter")
+			.PhysicalSplitterHandleSize(1.0f)
+			.HitDetectionSplitterHandleSize(5.0f)
+			.Orientation(Orient_Horizontal);
+
+		const bool bHasMultipleColumns = WidgetRow.HasColumns();
+
+		TSharedRef<SHorizontalBox> NameColumnBox = SNew(SHorizontalBox)
+			.Clipping(EWidgetClipping::OnDemand);
+
+		if (IndentWidget.IsValid())
+		{
+			NameColumnBox->AddSlot()
+				.AutoWidth()
+				.Padding(0)
+				.HAlign(HAlign_Left)
+				.VAlign(VAlign_Fill)
+				[
+					IndentWidget.ToSharedRef()
+				];
+		}
+
+		bool bNeedsReorderHandle = WidgetRow.CustomDragDropHandler.IsValid() || DetailTreeItem.Pin()->IsReorderable();
+		if (bNeedsReorderHandle)
+		{
+			NameColumnBox->AddSlot()
+				.HAlign(HAlign_Left)
+				.VAlign(VAlign_Center)
+				.Padding(-4, 0, -10, 0)
+				.AutoWidth()
+				[
+					SNew(UE::DisplayClusterColorGradingDetailTreeRow::SArrayRowHandle, SharedThis(this))
+					.IsEnabled(this, &SDisplayClusterColorGradingDetailTreeRow::IsRowEnabled)
+					.Visibility_Lambda([this]() { return IsHovered() ? EVisibility::Visible : EVisibility::Hidden; })
+				];
+		}
+
+		if (ExpanderArrowWidget.IsValid())
+		{
+			NameColumnBox->AddSlot()
+				.AutoWidth()
+				.Padding(2.f, 0.f, 0.f, 0.f)
+				.HAlign(HAlign_Left)
+				.VAlign(VAlign_Fill)
+				[
+					ExpanderArrowWidget.ToSharedRef()
+				];
+		}
+
+		NameColumnBox->AddSlot()
+			.AutoWidth()
+			.Padding(2.f, 0.f, 0.f, 0.f)
+			.HAlign(HAlign_Left)
+			.VAlign(VAlign_Center)
+			[
+				SNew(UE::DisplayClusterColorGradingDetailTreeRow::SEditConditionWidget)
+				.EditConditionValue(WidgetRow.EditConditionValue)
+				.OnEditConditionValueChanged(WidgetRow.OnEditConditionValueChanged)
+			];
+
+		if (bHasMultipleColumns)
+		{
+			WidgetRow.NameWidget.Widget->SetEnabled(TAttribute<bool>::CreateSP(this, &SDisplayClusterColorGradingDetailTreeRow::IsRowEnabled));
+
+			NameColumnBox->AddSlot()
+				.FillWidth(1.f)
+				.Padding(2.f, 0.f, 0.f, 0.f)
+				.HAlign(WidgetRow.NameWidget.HorizontalAlignment)
+				.VAlign(WidgetRow.NameWidget.VerticalAlignment)
+				[
+					WidgetRow.NameWidget.Widget
+				];
+		}
+		else
+		{
+			NameColumnBox->SetEnabled(TAttribute<bool>::CreateSP(this, &SDisplayClusterColorGradingDetailTreeRow::IsRowEnabled));
+
+			NameColumnBox->AddSlot()
+				.FillWidth(1.f)
+				.Padding(2.f, 0.f, 0.f, 0.f)
+				.HAlign(WidgetRow.WholeRowWidget.HorizontalAlignment)
+				.VAlign(WidgetRow.WholeRowWidget.VerticalAlignment)
+				[
+					WidgetRow.WholeRowWidget.Widget
+				];
+		}
+
+		Splitter->AddSlot()
+			.Value(InColumnSizeData.GetNameColumnWidth())
+			.OnSlotResized(InColumnSizeData.GetOnNameColumnResized())
+			[
+				NameColumnBox
+			];
+
+		if (bHasMultipleColumns && WidgetRow.HasValueContent())
+		{
+			WidgetRow.ValueWidget.Widget->SetEnabled(TAttribute<bool>::CreateSP(this, &SDisplayClusterColorGradingDetailTreeRow::IsRowValueEnabled));
+			WidgetRow.ExtensionWidget.Widget->SetEnabled(TAttribute<bool>::CreateSP(this, &SDisplayClusterColorGradingDetailTreeRow::IsRowEnabled));
+
+			Splitter->AddSlot()
+				.Value(InColumnSizeData.GetValueColumnWidth())
+				.OnSlotResized(InColumnSizeData.GetOnValueColumnResized())
+				[
+					SNew(SHorizontalBox)
+					.Clipping(EWidgetClipping::OnDemand)
+
+					+ SHorizontalBox::Slot()
+					.HAlign(WidgetRow.ValueWidget.HorizontalAlignment)
+					.VAlign(WidgetRow.ValueWidget.VerticalAlignment)
+					.Padding(12.0f, 0.0f, 2.0f, 0.0f)
+					[
+						SNew(UE::DisplayClusterColorGradingDetailTreeRow::SConstrainedBox)
+						.MinWidth(WidgetRow.ValueWidget.MinWidth)
+						.MaxWidth(WidgetRow.ValueWidget.MaxWidth)
+						[
+							WidgetRow.ValueWidget.Widget
+						]
+					]
+
+					+ SHorizontalBox::Slot()
+					.HAlign(WidgetRow.ExtensionWidget.HorizontalAlignment)
+					.VAlign(WidgetRow.ExtensionWidget.VerticalAlignment)
+					.Padding(5,0,0,0)
+					.AutoWidth()
+					[
+						WidgetRow.ExtensionWidget.Widget
+					]
+				];
+		}
+
+		if (PinnedDetailTreeItem->IsItem())
+		{
+			TArray<FPropertyRowExtensionButton> ExtensionButtons;
+
+			FPropertyRowExtensionButton& ResetToDefaultButton = ExtensionButtons.AddDefaulted_GetRef();
+			ResetToDefaultButton.Label = NSLOCTEXT("PropertyEditor", "ResetToDefault", "Reset to Default");
+			ResetToDefaultButton.UIAction = FUIAction(
+				FExecuteAction::CreateSP(this, &SDisplayClusterColorGradingDetailTreeRow::OnResetToDefaultClicked),
+				FCanExecuteAction::CreateSP(this, &SDisplayClusterColorGradingDetailTreeRow::CanResetToDefault)
+			);
+
+			// We could just collapse the Reset to Default button by setting the FIsActionButtonVisible delegate,
+			// but this would cause the reset to defaults not to reserve space in the toolbar and not be aligned across all rows.
+			// Instead, we show an empty icon and tooltip and disable the button.
+			static FSlateIcon EnabledResetToDefaultIcon(FAppStyle::Get().GetStyleSetName(), "PropertyWindow.DiffersFromDefault");
+			static FSlateIcon DisabledResetToDefaultIcon(FAppStyle::Get().GetStyleSetName(), "NoBrush");
+			ResetToDefaultButton.Icon = TAttribute<FSlateIcon>::Create([this]()
+			{
+				if (DetailTreeItem.IsValid() && DetailTreeItem.Pin()->IsResetToDefaultVisible())
+				{
+					return EnabledResetToDefaultIcon;
+				}
+
+				return DisabledResetToDefaultIcon;
+			});
+
+			ResetToDefaultButton.ToolTip = TAttribute<FText>::Create([this]()
+			{
+				if (DetailTreeItem.IsValid() && DetailTreeItem.Pin()->IsResetToDefaultVisible())
+				{
+					return NSLOCTEXT("PropertyEditor", "ResetToDefaultPropertyValueToolTip", "Reset this property to its default value.");
+				}
+
+				return FText::GetEmpty();
+			});
+
+			CreateGlobalExtensionWidgets(ExtensionButtons);
+
+			FSlimHorizontalToolBarBuilder ToolbarBuilder(TSharedPtr<FUICommandList>(), FMultiBoxCustomization::None);
+			ToolbarBuilder.SetLabelVisibility(EVisibility::Collapsed);
+			ToolbarBuilder.SetStyle(&FAppStyle::Get(), "DetailsView.ExtensionToolBar");
+			ToolbarBuilder.SetIsFocusable(false);
+
+			for (const FPropertyRowExtensionButton& Extension : ExtensionButtons)
+			{
+				ToolbarBuilder.AddToolBarButton(Extension.UIAction, NAME_None, Extension.Label, Extension.ToolTip, Extension.Icon);
+			}
+
+			Splitter->AddSlot()
+				.Value(InColumnSizeData.GetRightColumnWidth())
+				.OnSlotResized(InColumnSizeData.GetOnRightColumnResized())
+				.MinSize(InColumnSizeData.GetRightColumnMinWidth())
+				[
+					SNew(SBox)
+					.Padding(0)
+					.HAlign(HAlign_Fill)
+					.VAlign(VAlign_Center)
+					[
+						ToolbarBuilder.MakeWidget()
+					]
+				];
+		}
+
+		return Splitter;
+	}
+
+	return SNullWidget::NullWidget;
+}
+
+void SDisplayClusterColorGradingDetailTreeRow::CreateGlobalExtensionWidgets(TArray<FPropertyRowExtensionButton>& OutExtensions) const
+{
+	// fetch global extension widgets 
+	FPropertyEditorModule& PropertyEditorModule = FModuleManager::Get().GetModuleChecked<FPropertyEditorModule>("PropertyEditor");
+
+	FOnGenerateGlobalRowExtensionArgs Args;
+	Args.OwnerTreeNode = DetailTreeItem.Pin()->GetDetailTreeNode();
+	Args.PropertyHandle = DetailTreeItem.Pin()->GetPropertyHandle();
+
+	PropertyEditorModule.GetGlobalRowExtensionDelegate().Broadcast(Args, OutExtensions);
+}
+
+bool SDisplayClusterColorGradingDetailTreeRow::IsRowEnabled() const
+{
+	if (DetailTreeItem.IsValid())
+	{
+		return DetailTreeItem.Pin()->IsPropertyEditingEnabled().Get(true) && 
+			WidgetRow.IsEnabledAttr.Get(true) && 
+			WidgetRow.EditConditionValue.Get(true);
+	}
+
+	return false;
+}
+
+bool SDisplayClusterColorGradingDetailTreeRow::IsRowValueEnabled() const
+{
+	return IsRowEnabled() && WidgetRow.IsValueEnabledAttr.Get(true);
+}
+
+int32 SDisplayClusterColorGradingDetailTreeRow::GetPropertyIndent() const
+{
+	const int32 PropertyIndentLevel = FMath::Max(GetIndentLevel() - 1, 0);
+	return PropertyIndentLevel;
+}
+
+FSlateColor SDisplayClusterColorGradingDetailTreeRow::GetIndentBackgroundColor(int32 Indent) const
+{
+	int32 ColorIndex = 0;
+	int32 Increment = 1;
+
+	for (int32 Index = 0; Index < Indent; ++Index)
+	{
+		ColorIndex += Increment;
+
+		if (ColorIndex == 0 || ColorIndex == 3)
+		{
+			Increment = -Increment;
+		}
+	}
+
+	static const uint8 ColorOffsets[] =
+	{
+		0, 4, (4 + 2), (6 + 4), (10 + 6)
+	};
+
+	const FSlateColor BaseSlateColor = IsHovered() ?
+		FAppStyle::Get().GetSlateColor("Colors.Header") :
+		FAppStyle::Get().GetSlateColor("Colors.Panel");
+
+	const FColor BaseColor = BaseSlateColor.GetSpecifiedColor().ToFColor(true);
+
+	const FColor ColorWithOffset(
+		BaseColor.R + ColorOffsets[ColorIndex],
+		BaseColor.G + ColorOffsets[ColorIndex],
+		BaseColor.B + ColorOffsets[ColorIndex]);
+
+	return FSlateColor(FLinearColor::FromSRGBColor(ColorWithOffset));
+}
+
+const FSlateBrush* SDisplayClusterColorGradingDetailTreeRow::GetRowBackgroundBrush() const
+{
+	const bool bIsCategory = DetailTreeItem.IsValid() ? DetailTreeItem.Pin()->IsCategory() : false;
+	return bIsCategory ? FAppStyle::Get().GetBrush("DetailsView.CategoryTop") : FAppStyle::Get().GetBrush("DetailsView.CategoryMiddle");
+}
+
+FSlateColor SDisplayClusterColorGradingDetailTreeRow::GetRowBackgroundColor() const
+{
+	bool bIsCategory = DetailTreeItem.IsValid() ? DetailTreeItem.Pin()->IsCategory() : false;
+	return bIsCategory ? FSlateColor(FLinearColor::White) : GetIndentBackgroundColor(GetPropertyIndent());
+}
+
+const FSlateBrush* SDisplayClusterColorGradingDetailTreeRow::GetScrollWellBackgroundBrush(TWeakPtr<STableViewBase> OwnerTableViewWeak) const
+{
+	const bool bIsScrollBarVisible = OwnerTableViewWeak.Pin()->GetScrollbarVisibility() == EVisibility::Visible;
+
+	if (bIsScrollBarVisible)
+	{
+		return FAppStyle::Get().GetBrush("DetailsView.GridLine");
+	}
+
+	return GetRowBackgroundBrush();
+}
+
+FSlateColor SDisplayClusterColorGradingDetailTreeRow::GetScrollWellBackgroundColor(TWeakPtr<STableViewBase> OwnerTableViewWeak) const
+{
+	const bool bIsScrollBarVisible = OwnerTableViewWeak.Pin()->GetScrollbarVisibility() == EVisibility::Visible;
+
+	if (bIsScrollBarVisible)
+	{
+		return FSlateColor(EStyleColor::White);
+	}
+
+	return GetRowBackgroundColor();
+}
+
+void SDisplayClusterColorGradingDetailTreeRow::OnResetToDefaultClicked()
+{
+	if (DetailTreeItem.IsValid())
+	{
+		if (WidgetRow.CustomResetToDefault.IsSet())
+		{
+			WidgetRow.CustomResetToDefault.GetValue().OnResetToDefaultClicked(DetailTreeItem.Pin()->GetPropertyHandle());
+		}
+		else
+		{
+			DetailTreeItem.Pin()->ResetToDefault();
+		}
+	}
+}
+
+bool SDisplayClusterColorGradingDetailTreeRow::CanResetToDefault() const
+{
+	if (DetailTreeItem.IsValid())
+	{
+		if (WidgetRow.CustomResetToDefault.IsSet())
+		{
+			return WidgetRow.CustomResetToDefault.GetValue().IsResetToDefaultVisible(DetailTreeItem.Pin()->GetPropertyHandle());
+		}
+
+		return DetailTreeItem.Pin()->IsResetToDefaultVisible() && IsRowValueEnabled();
+	}
+
+	return false;
+}
+
+void SDisplayClusterColorGradingDetailTreeRow::OnRowDragLeave(const FDragDropEvent& DragDropEvent)
+{
+	if (TSharedPtr<FDecoratedDragDropOp> DecoratedOp = DragDropEvent.GetOperationAs<FDecoratedDragDropOp>())
+	{
+		DecoratedOp->ResetToDefaultToolTip();
+	}
+}
+
+FReply SDisplayClusterColorGradingDetailTreeRow::OnRowAcceptDrop(const FDragDropEvent& DragDropEvent, EItemDropZone DropZone, TSharedRef<FDisplayClusterColorGradingDetailTreeItem> TargetItem)
+{
+	if (WidgetRow.CustomDragDropHandler)
+	{
+		if (WidgetRow.CustomDragDropHandler->AcceptDrop(DragDropEvent, DropZone))
+		{
+			return FReply::Handled();
+		}
+
+		return FReply::Unhandled();
+	}
+
+	if (TSharedPtr<UE::DisplayClusterColorGradingDetailTreeRow::FArrayRowDragDropOp> ArrayRowDropOp = DragDropEvent.GetOperationAs<UE::DisplayClusterColorGradingDetailTreeRow::FArrayRowDragDropOp>())
+	{
+		if (ArrayRowDropOp->RowBeingDragged.IsValid())
+		{
+			TSharedPtr<IPropertyHandle> SwappingPropertyHandle = ArrayRowDropOp->RowBeingDragged.Pin()->GetPropertyHandle();
+			TSharedPtr<IPropertyHandle> ThisPropertyHandle = GetPropertyHandle();
+
+			if (SwappingPropertyHandle.IsValid() && ThisPropertyHandle.IsValid())
+			{
+				const int32 OriginalIndex = SwappingPropertyHandle->GetIndexInArray();
+				const int32 NewIndex = GetDropNewIndex(OriginalIndex, ThisPropertyHandle->GetIndexInArray(), DropZone);
+
+				// Need to swap the moving and target expansion states before saving
+				const bool bSwappingRowExpansion = ArrayRowDropOp->RowBeingDragged.Pin()->IsItemExpanded();
+				const bool bThisRowExpansion = IsItemExpanded();
+
+				if (bSwappingRowExpansion != bThisRowExpansion)
+				{
+					ArrayRowDropOp->RowBeingDragged.Pin()->ToggleExpansion();
+					ToggleExpansion();
+				}
+
+				FScopedTransaction Transaction(NSLOCTEXT("UnrealEd", "MoveRow", "Move Row"));
+
+				SwappingPropertyHandle->GetParentHandle()->NotifyPreChange();
+				SwappingPropertyHandle->GetParentHandle()->AsArray()->MoveElementTo(OriginalIndex, NewIndex);
+				SwappingPropertyHandle->GetParentHandle()->NotifyPostChange(EPropertyChangeType::ArrayMove);
+
+				return FReply::Handled();
+			}
+		}
+	}
+
+	return FReply::Unhandled();
+}
+
+TOptional<EItemDropZone> SDisplayClusterColorGradingDetailTreeRow::CanRowAcceptDrop(const FDragDropEvent& DragDropEvent, EItemDropZone DropZone, TSharedRef<FDisplayClusterColorGradingDetailTreeItem> Type)
+{
+	if (WidgetRow.CustomDragDropHandler)
+	{
+		// Disallow drop between expanded parent item and its first child
+		if (DropZone == EItemDropZone::BelowItem && IsItemExpanded())
+		{
+			return TOptional<EItemDropZone>();
+		}
+
+		return WidgetRow.CustomDragDropHandler->CanAcceptDrop(DragDropEvent, DropZone);
+	}
+
+	if (TSharedPtr<UE::DisplayClusterColorGradingDetailTreeRow::FArrayRowDragDropOp> ArrayRowDropOp = DragDropEvent.GetOperationAs<UE::DisplayClusterColorGradingDetailTreeRow::FArrayRowDragDropOp>())
+	{
+		// Can't drop onto another array item, so recompute our own drop zone to ensure it's above or below
+		const FGeometry& Geometry = GetTickSpaceGeometry();
+		const float LocalPointerY = Geometry.AbsoluteToLocal(DragDropEvent.GetScreenSpacePosition()).Y;
+		const EItemDropZone OverrideDropZone = LocalPointerY < Geometry.GetLocalSize().Y * 0.5f ? EItemDropZone::AboveItem : EItemDropZone::BelowItem;
+
+		const bool bIsValidDrop = IsValidDrop(DragDropEvent, OverrideDropZone);
+
+		ArrayRowDropOp->SetValidTarget(bIsValidDrop);
+
+		if (!bIsValidDrop)
+		{
+			return TOptional<EItemDropZone>();
+		}
+
+		return OverrideDropZone;
+	}
+
+	return TOptional<EItemDropZone>();
+}
+
+bool SDisplayClusterColorGradingDetailTreeRow::IsValidDrop(const FDragDropEvent& DragDropEvent, EItemDropZone DropZone) const
+{
+	if (DropZone == EItemDropZone::OntoItem)
+	{
+		return false;
+	}
+
+	if (DropZone == EItemDropZone::BelowItem && IsItemExpanded())
+	{
+		return false;
+	}
+
+	if (TSharedPtr<UE::DisplayClusterColorGradingDetailTreeRow::FArrayRowDragDropOp> ArrayRowDropOp = DragDropEvent.GetOperationAs<UE::DisplayClusterColorGradingDetailTreeRow::FArrayRowDragDropOp>())
+	{
+		if (DetailTreeItem.IsValid() && !DetailTreeItem.Pin()->IsReorderable())
+		{
+			return false;
+		}
+
+		if (ArrayRowDropOp->RowBeingDragged.IsValid())
+		{
+			TSharedPtr<IPropertyHandle> SwappingPropertyHandle = ArrayRowDropOp->RowBeingDragged.Pin()->GetPropertyHandle();
+			TSharedPtr<IPropertyHandle> ThisPropertyHandle = GetPropertyHandle();
+			if (ThisPropertyHandle.IsValid() && SwappingPropertyHandle.IsValid() && SwappingPropertyHandle->GetPropertyNode() != ThisPropertyHandle->GetPropertyNode())
+			{
+				const int32 OriginalIndex = SwappingPropertyHandle->GetIndexInArray();
+				const int32 NewIndex = GetDropNewIndex(OriginalIndex, ThisPropertyHandle->GetIndexInArray(), DropZone);
+
+				if (OriginalIndex != NewIndex)
+				{
+					if (SwappingPropertyHandle->GetParentHandle()->AsArray() && SwappingPropertyHandle->GetParentHandle()->GetPropertyNode() == ThisPropertyHandle->GetParentHandle()->GetPropertyNode())
+					{
+						return true;
+					}
+				}
+			}
+		}
+	}
+
+	return false;
+}
+
+int32 SDisplayClusterColorGradingDetailTreeRow::GetDropNewIndex(int32 OriginalIndex, int32 DropOntoIndex, EItemDropZone DropZone) const
+{
+	check(DropZone != EItemDropZone::OntoItem);
+
+	int32 NewIndex = DropOntoIndex;
+	if (DropZone == EItemDropZone::BelowItem)
+	{
+		// If the drop zone is below, then we actually move it to the next item's index
+		NewIndex++;
+	}
+	if (OriginalIndex < NewIndex)
+	{
+		// If the item is moved down the list, then all the other elements below it are shifted up one
+		NewIndex--;
+	}
+
+	return ensure(NewIndex >= 0) ? NewIndex : 0;
+}
\ No newline at end of file
diff --git a/Engine/Plugins/Runtime/nDisplay/Source/DisplayClusterColorGrading/Private/DetailView/SDisplayClusterColorGradingDetailTreeRow.h b/Engine/Plugins/Runtime/nDisplay/Source/DisplayClusterColorGrading/Private/DetailView/SDisplayClusterColorGradingDetailTreeRow.h
new file mode 100644
--- /dev/null
+++ b/Engine/Plugins/Runtime/nDisplay/Source/DisplayClusterColorGrading/Private/DetailView/SDisplayClusterColorGradingDetailTreeRow.h
@@ -0,0 +1,93 @@
+// Copyright Epic Games, Inc. All Rights Reserved.
+
+#pragma once
+
+#include "CoreMinimal.h"
+#include "Widgets/DeclarativeSyntaxSupport.h"
+#include "Widgets/Views/STableRow.h"
+#include "DetailWidgetRow.h"
+
+class FDetailColumnSizeData;
+class FDisplayClusterColorGradingDetailTreeItem;
+class IPropertyHandle;
+
+/** A copy/combination of SDetailSingleItemRow and SDetailCategoryTableRow defined in PropertyEditor/Private, used to display properties in a custom detail view widget */
+class SDisplayClusterColorGradingDetailTreeRow : public STableRow<TSharedRef<FDisplayClusterColorGradingDetailTreeItem>>
+{
+public:
+	SLATE_BEGIN_ARGS(SDisplayClusterColorGradingDetailTreeRow) {}
+	SLATE_END_ARGS()
+
+	//~ STableRow Interface
+	void Construct(const FArguments& InArgs, const TSharedRef<FDisplayClusterColorGradingDetailTreeItem>& InDetailTreeItem, const TSharedRef<STableViewBase>& InOwnerTableView, const FDetailColumnSizeData& InColumnSizeData);
+	virtual void ConstructChildren(ETableViewMode::Type InOwnerTableMode, const TAttribute<FMargin>& InPadding, const TSharedRef<SWidget>& InContent) override;
+	//~ End STableRow Interface
+
+	/** Gets the property handle of the property this row represents, if applicable */
+	TSharedPtr<IPropertyHandle> GetPropertyHandle() const;
+
+	/** Creates a new drag drop operation for the property row */
+	TSharedPtr<FDragDropOperation> CreateDragDropOperation();
+
+private:
+	/** Creates the widgets necessary to display the property this row represents */
+	TSharedRef<SWidget> CreatePropertyWidget(const FDetailColumnSizeData& InColumnSizeData);
+
+	/** Broadcasts the property editor module's global row extension delegate */
+	void CreateGlobalExtensionWidgets(TArray<FPropertyRowExtensionButton>& OutExtensions) const;
+
+	/** Gets whether the row is enabled */
+	bool IsRowEnabled() const;
+
+	/** Gets whether the row's value is enabled */
+	bool IsRowValueEnabled() const;
+
+	/** Gets the level of indent this row is at */
+	int32 GetPropertyIndent() const;
+
+	/** Gets the background color for the row depending on its level of indent */
+	FSlateColor GetIndentBackgroundColor(int32 Indent) const;
+
+	/** Gets the brush for the row's background border */
+	const FSlateBrush* GetRowBackgroundBrush() const;
+
+	/** Gets the color for the row's background border */
+	FSlateColor GetRowBackgroundColor() const;
+
+	/** Gets the brush for the row's scroll well */
+	const FSlateBrush* GetScrollWellBackgroundBrush(TWeakPtr<STableViewBase> OwnerTableViewWeak) const;
+
+	/** Gets the color for the row's scroll well */
+	FSlateColor GetScrollWellBackgroundColor(TWeakPtr<STableViewBase> OwnerTableViewWeak) const;
+
+	/** Raised when the property's Reset to Default button is pressed */
+	void OnResetToDefaultClicked();
+
+	/** Gets whether the property's Reset to Default button can be pressed */
+	bool CanResetToDefault() const;
+
+	/**  Raised when a drag drop event leaves this row */
+	void OnRowDragLeave(const FDragDropEvent& DragDropEvent);
+
+	/** Raised when a drag drop event is released on this row */
+	FReply OnRowAcceptDrop(const FDragDropEvent& DragDropEvent, EItemDropZone DropZone, TSharedRef<FDisplayClusterColorGradingDetailTreeItem> TargetItem);
+
+	/** Indicates whether the drag drop event can be accepted by this row */
+	TOptional<EItemDropZone> CanRowAcceptDrop(const FDragDropEvent& DragDropEvent, EItemDropZone DropZone, TSharedRef<FDisplayClusterColorGradingDetailTreeItem> Type);
+
+	/** Gets whether the specified drag drop event is allowed to be dropped onto this row */
+	bool IsValidDrop(const FDragDropEvent& DragDropEvent, EItemDropZone DropZone) const;
+
+	/** Gets the index to use when another row has been dragged onto this row, and an array element swap is taking place */
+	int32 GetDropNewIndex(int32 OriginalIndex, int32 DropOntoIndex, EItemDropZone DropZone) const;
+
+private:
+	/** A weak pointer to the detail tree item this row represents */
+	TWeakPtr<FDisplayClusterColorGradingDetailTreeItem> DetailTreeItem;
+
+	/** The cached property row widgets that were created for this row */
+	FDetailWidgetRow WidgetRow;
+
+	/** The indent widget used to properly indent the row's widgets */
+	TSharedPtr<SWidget> IndentWidget;
+};
\ No newline at end of file
diff --git a/Engine/Plugins/Runtime/nDisplay/Source/DisplayClusterColorGrading/Private/DetailView/SDisplayClusterColorGradingDetailView.cpp b/Engine/Plugins/Runtime/nDisplay/Source/DisplayClusterColorGrading/Private/DetailView/SDisplayClusterColorGradingDetailView.cpp
new file mode 100644
--- /dev/null
+++ b/Engine/Plugins/Runtime/nDisplay/Source/DisplayClusterColorGrading/Private/DetailView/SDisplayClusterColorGradingDetailView.cpp
@@ -0,0 +1,476 @@
+// Copyright Epic Games, Inc. All Rights Reserved.
+
+#include "SDisplayClusterColorGradingDetailView.h"
+
+#include "SDisplayClusterColorGradingDetailTreeRow.h"
+
+#include "DragAndDrop/DecoratedDragDropOp.h"
+#include "Framework/Application/SlateApplication.h"
+#include "Framework/MultiBox/MultiBoxBuilder.h"
+#include "IPropertyRowGenerator.h"
+#include "IDetailTreeNode.h"
+#include "Misc/App.h"
+#include "Misc/ConfigCacheIni.h"
+#include "PropertyEditor/Private/DetailTreeNode.h"
+#include "Styling/StyleColors.h"
+#include "Widgets/Input/SCheckBox.h"
+#include "Widgets/Views/STableRow.h"
+#include "Widgets/Views/STreeView.h"
+
+#include "PropertyEditor/Private/PropertyNode.h"
+
+void FDisplayClusterColorGradingDetailTreeItem::Initialize(const FOnFilterDetailTreeNode& NodeFilter)
+{
+	if (DetailTreeNode.IsValid())
+	{
+		PropertyHandle = DetailTreeNode.Pin()->CreatePropertyHandle();
+
+		TArray<TSharedRef<IDetailTreeNode>> ChildNodes;
+		DetailTreeNode.Pin()->GetChildren(ChildNodes);
+
+		for (const TSharedRef<IDetailTreeNode>& ChildNode : ChildNodes)
+		{
+			bool bShouldDisplayNode = true;
+			if (NodeFilter.IsBound())
+			{
+				bShouldDisplayNode = NodeFilter.Execute(ChildNode);
+			}
+
+			if (bShouldDisplayNode)
+			{
+				TSharedPtr<FDetailTreeNode> CastChildNode = StaticCastSharedRef<FDetailTreeNode>(ChildNode);
+				TSharedRef<FDisplayClusterColorGradingDetailTreeItem> ChildItem = MakeShared<FDisplayClusterColorGradingDetailTreeItem>(CastChildNode);
+				ChildItem->Parent = SharedThis(this);
+				ChildItem->Initialize(NodeFilter);
+
+				Children.Add(ChildItem);
+			}
+		}
+	}    
+}
+
+void FDisplayClusterColorGradingDetailTreeItem::GetChildren(TArray<TSharedRef<FDisplayClusterColorGradingDetailTreeItem>>& OutChildren) const
+{
+	OutChildren.Reset(Children.Num());
+	OutChildren.Append(Children);
+}
+
+TWeakPtr<IDetailTreeNode> FDisplayClusterColorGradingDetailTreeItem::GetDetailTreeNode() const
+{
+	return DetailTreeNode;
+}
+
+FName FDisplayClusterColorGradingDetailTreeItem::GetNodeName() const
+{
+	if (DetailTreeNode.IsValid())
+	{
+		return DetailTreeNode.Pin()->GetNodeName();
+	}
+
+	return NAME_None;
+}
+
+bool FDisplayClusterColorGradingDetailTreeItem::ShouldBeExpanded() const
+{
+	if (DetailTreeNode.IsValid())
+	{
+		return DetailTreeNode.Pin()->ShouldBeExpanded();
+	}
+
+	return false;
+}
+
+void FDisplayClusterColorGradingDetailTreeItem::OnItemExpansionChanged(bool bIsExpanded, bool bShouldSaveState)
+{
+	if (DetailTreeNode.IsValid())
+	{
+		DetailTreeNode.Pin()->OnItemExpansionChanged(bIsExpanded, bShouldSaveState);
+	}
+}
+
+bool FDisplayClusterColorGradingDetailTreeItem::IsResetToDefaultVisible() const
+{
+	if (PropertyHandle.IsValid())
+	{
+		if (PropertyHandle->HasMetaData("NoResetToDefault") || PropertyHandle->GetInstanceMetaData("NoResetToDefault"))
+		{
+			return false;
+		}
+
+		return PropertyHandle->CanResetToDefault();
+	}
+
+	return false;
+}
+
+void FDisplayClusterColorGradingDetailTreeItem::ResetToDefault()
+{
+	if (PropertyHandle.IsValid())
+	{
+		PropertyHandle->ResetToDefault();
+	}
+}
+
+TAttribute<bool> FDisplayClusterColorGradingDetailTreeItem::IsPropertyEditingEnabled() const
+{
+	if (DetailTreeNode.IsValid())
+	{
+		return DetailTreeNode.Pin()->IsPropertyEditingEnabled();
+	}
+
+	return false;
+}
+
+bool FDisplayClusterColorGradingDetailTreeItem::IsCategory() const
+{
+	if (DetailTreeNode.IsValid())
+	{
+		if (DetailTreeNode.Pin()->GetNodeType() == EDetailNodeType::Category)
+		{
+			return true;
+		}
+	}
+	
+	return false;
+}
+
+bool FDisplayClusterColorGradingDetailTreeItem::IsItem() const
+{
+	if (DetailTreeNode.IsValid())
+	{
+		return DetailTreeNode.Pin()->GetNodeType() == EDetailNodeType::Item;
+	}
+
+	return false;
+}
+
+bool FDisplayClusterColorGradingDetailTreeItem::IsReorderable() const
+{
+	if (PropertyHandle.IsValid())
+	{
+		if (TSharedPtr<IPropertyHandle> ParentHandle = PropertyHandle->GetParentHandle())
+		{
+			const bool bIsParentAnArray = ParentHandle->AsArray().IsValid();
+			const bool bIsParentArrayReorderable = !ParentHandle->HasMetaData(TEXT("EditFixedOrder")) && !ParentHandle->HasMetaData(TEXT("ArraySizeEnum"));
+			
+			return bIsParentAnArray && bIsParentArrayReorderable && !PropertyHandle->IsEditConst() && !FApp::IsGame();
+		}
+	}
+
+	return false;
+}
+
+void FDisplayClusterColorGradingDetailTreeItem::GenerateDetailWidgetRow(FDetailWidgetRow& OutDetailWidgetRow) const
+{
+	if (DetailTreeNode.IsValid())
+	{
+		DetailTreeNode.Pin()->GenerateStandaloneWidget(OutDetailWidgetRow);
+	}
+}
+
+void SDisplayClusterColorGradingDetailView::Construct(const FArguments& InArgs)
+{
+	PropertyRowGeneratorSource = InArgs._PropertyRowGeneratorSource;
+	OnFilterDetailTreeNode = InArgs._OnFilterDetailTreeNode;
+
+	ColumnSizeData.SetValueColumnWidth(0.5f);
+	ColumnSizeData.SetRightColumnMinWidth(22);
+
+	UpdateTreeNodes();
+
+	TSharedRef<SScrollBar> ExternalScrollbar = SNew(SScrollBar);
+	ExternalScrollbar->SetVisibility(TAttribute<EVisibility>(this, &SDisplayClusterColorGradingDetailView::GetScrollBarVisibility));
+
+	TSharedRef<SVerticalBox> VerticalBox = SNew(SVerticalBox);
+
+	VerticalBox->AddSlot()
+	.FillHeight(1)
+	.Padding(0)
+	[
+		SNew(SOverlay)
+		+ SOverlay::Slot()
+		[
+			SAssignNew(DetailTree, SDetailTree)
+			.TreeItemsSource(&RootTreeNodes)
+			.OnGenerateRow(this, &SDisplayClusterColorGradingDetailView::GenerateNodeRow)
+			.OnGetChildren(this, &SDisplayClusterColorGradingDetailView::GetChildrenForNode)
+			.OnSetExpansionRecursive(this, &SDisplayClusterColorGradingDetailView::OnSetExpansionRecursive)
+			.OnRowReleased(this, &SDisplayClusterColorGradingDetailView::OnRowReleased)
+			.OnExpansionChanged(this, &SDisplayClusterColorGradingDetailView::OnExpansionChanged)
+			.SelectionMode(ESelectionMode::None)
+			.HandleDirectionalNavigation(false)
+			.AllowOverscroll(EAllowOverscroll::Yes)
+			.ExternalScrollbar(ExternalScrollbar)
+		]
+
+		+ SOverlay::Slot()
+		.HAlign(HAlign_Right)
+		[
+			SNew(SBox)
+			.WidthOverride(16.0f)
+			[
+				ExternalScrollbar
+			]
+		]
+	];
+
+	ChildSlot
+	[
+		VerticalBox
+	];
+}
+
+void SDisplayClusterColorGradingDetailView::Tick(const FGeometry& AllottedGeometry, const double InCurrentTime, const float InDeltaTime)
+{
+	if (TreeItemsToSetExpansionState.Num() > 0)
+	{
+		for (const TPair<TWeakPtr<FDisplayClusterColorGradingDetailTreeItem>, bool>& Pair : TreeItemsToSetExpansionState)
+		{
+			if (TSharedPtr<FDisplayClusterColorGradingDetailTreeItem> DetailTreeItem = Pair.Key.Pin())
+			{
+				DetailTree->SetItemExpansion(DetailTreeItem.ToSharedRef(), Pair.Value);
+			}
+		}
+
+		TreeItemsToSetExpansionState.Empty();
+	}
+}
+
+void SDisplayClusterColorGradingDetailView::Refresh()
+{
+	UpdateTreeNodes();
+	DetailTree->RebuildList();
+}
+
+void SDisplayClusterColorGradingDetailView::SaveExpandedItems()
+{
+	TSet<FString> ObjectTypes;
+	if (PropertyRowGeneratorSource.IsValid())
+	{
+		const TArray<TWeakObjectPtr<UObject>>& SelectedObjects = PropertyRowGeneratorSource->GetSelectedObjects();
+
+		for (const TWeakObjectPtr<UObject>& Object : SelectedObjects)
+		{
+			if (Object.IsValid())
+			{
+				ObjectTypes.Add(Object->GetClass()->GetName());
+			}
+		}
+	}
+
+	FString ExpandedDetailNodesString;
+	for (const FString& DetailNode : ExpandedDetailNodes)
+	{
+		ExpandedDetailNodesString += DetailNode;
+		ExpandedDetailNodesString += TEXT(",");
+	}
+
+	for (const FString& ObjectType : ObjectTypes)
+	{
+		if (!ExpandedDetailNodesString.IsEmpty())
+		{
+			GConfig->SetString(TEXT("DisplayClusterColorGradingDetailsExpansion"), *ObjectType, *ExpandedDetailNodesString, GEditorPerProjectIni);
+		}
+		else
+		{
+			// If the expanded nodes string is empty but the saved expanded state is not, we want to save the empty string
+			FString SavedExpandedDetailNodesString;
+			GConfig->GetString(TEXT("DisplayClusterColorGradingDetailsExpansion"), *ObjectType, SavedExpandedDetailNodesString, GEditorPerProjectIni);
+
+			if (!SavedExpandedDetailNodesString.IsEmpty())
+			{
+				GConfig->SetString(TEXT("DisplayClusterColorGradingDetailsExpansion"), *ObjectType, *ExpandedDetailNodesString, GEditorPerProjectIni);
+			}
+		}
+	}
+}
+
+void SDisplayClusterColorGradingDetailView::RestoreExpandedItems()
+{
+	TSet<FString> ObjectTypes;
+	if (PropertyRowGeneratorSource.IsValid())
+	{
+		const TArray<TWeakObjectPtr<UObject>>& SelectedObjects = PropertyRowGeneratorSource->GetSelectedObjects();
+
+		for (const TWeakObjectPtr<UObject>& Object : SelectedObjects)
+		{
+			if (Object.IsValid())
+			{
+				ObjectTypes.Add(Object->GetClass()->GetName());
+			}
+		}
+	}
+
+	for (const FString& ObjectType : ObjectTypes)
+	{
+		FString SavedExpandedDetailNodesString;
+		GConfig->GetString(TEXT("DisplayClusterColorGradingDetailsExpansion"), *ObjectType, SavedExpandedDetailNodesString, GEditorPerProjectIni);
+		TArray<FString> SavedExpandedDetailNodes;
+		SavedExpandedDetailNodesString.ParseIntoArray(SavedExpandedDetailNodes, TEXT(","), true);
+
+		ExpandedDetailNodes.Append(SavedExpandedDetailNodes);
+	}
+}
+
+void SDisplayClusterColorGradingDetailView::UpdateTreeNodes()
+{
+	RootTreeNodes.Empty();
+
+	RestoreExpandedItems();
+
+	if (PropertyRowGeneratorSource.IsValid())
+	{
+		TArray<TSharedRef<IDetailTreeNode>> RawRootTreeNodes = PropertyRowGeneratorSource->GetRootTreeNodes();
+
+		for (const TSharedRef<IDetailTreeNode>& RootTreeNode : RawRootTreeNodes)
+		{
+			bool bShouldDisplayNode = true;
+			if (OnFilterDetailTreeNode.IsBound())
+			{
+				bShouldDisplayNode = OnFilterDetailTreeNode.Execute(RootTreeNode);
+			}
+
+			if (bShouldDisplayNode)
+			{
+				TSharedRef<FDetailTreeNode> CastRootTreeNode = StaticCastSharedRef<FDetailTreeNode>(RootTreeNode);
+				TSharedRef<FDisplayClusterColorGradingDetailTreeItem> RootTreeItem = MakeShared<FDisplayClusterColorGradingDetailTreeItem>(CastRootTreeNode);
+				RootTreeItem->Initialize(OnFilterDetailTreeNode);
+
+				RootTreeNodes.Add(RootTreeItem);
+
+				UpdateExpansionState(RootTreeItem);
+			}
+		}
+	}
+}
+
+void SDisplayClusterColorGradingDetailView::UpdateExpansionState(const TSharedRef<FDisplayClusterColorGradingDetailTreeItem> InTreeItem)
+{
+	if (InTreeItem->IsCategory())
+	{
+		TreeItemsToSetExpansionState.Add(InTreeItem, InTreeItem->ShouldBeExpanded());
+	}
+	else if (InTreeItem->IsItem())
+	{
+		TSharedPtr<FDisplayClusterColorGradingDetailTreeItem> ParentCategory = InTreeItem->GetParent().Pin();
+		while (ParentCategory.IsValid() && !ParentCategory->IsCategory())
+		{
+			ParentCategory = ParentCategory->GetParent().Pin();
+		}
+
+		FString Key;
+		if (ParentCategory.IsValid())
+		{
+			Key = ParentCategory->GetNodeName().ToString() + TEXT(".") + InTreeItem->GetNodeName().ToString();
+		}
+		else
+		{
+			Key = InTreeItem->GetNodeName().ToString();
+		}
+
+		const bool bShouldItemBeExpanded = ExpandedDetailNodes.Contains(Key) && InTreeItem->HasChildren();
+		TreeItemsToSetExpansionState.Add(InTreeItem, bShouldItemBeExpanded);
+	}
+
+	TArray<TSharedRef<FDisplayClusterColorGradingDetailTreeItem>> Children;
+	InTreeItem->GetChildren(Children);
+
+	for (const TSharedRef<FDisplayClusterColorGradingDetailTreeItem>& Child : Children)
+	{
+		UpdateExpansionState(Child);
+	}
+}
+
+TSharedRef<ITableRow> SDisplayClusterColorGradingDetailView::GenerateNodeRow(TSharedRef<FDisplayClusterColorGradingDetailTreeItem> InTreeItem, const TSharedRef<STableViewBase>& OwnerTable)
+{
+	return SNew(SDisplayClusterColorGradingDetailTreeRow, InTreeItem, OwnerTable, ColumnSizeData);
+}
+
+void SDisplayClusterColorGradingDetailView::GetChildrenForNode(TSharedRef<FDisplayClusterColorGradingDetailTreeItem> InTreeItem, TArray<TSharedRef<FDisplayClusterColorGradingDetailTreeItem>>& OutChildren)
+{
+	InTreeItem->GetChildren(OutChildren);
+}
+
+void SDisplayClusterColorGradingDetailView::SetNodeExpansionState(TSharedRef<FDisplayClusterColorGradingDetailTreeItem> InTreeItem, bool bIsItemExpanded, bool bRecursive)
+{
+	TArray<TSharedRef<FDisplayClusterColorGradingDetailTreeItem>> Children;
+	InTreeItem->GetChildren(Children);
+
+	if (Children.Num())
+	{
+		const bool bShouldSaveState = true;
+		InTreeItem->OnItemExpansionChanged(bIsItemExpanded, bShouldSaveState);
+
+		// Category nodes will save themselves to the editor config, but the item nodes can't, so manually save their expansion state here
+		if (InTreeItem->IsItem())
+		{
+			TSharedPtr<FDisplayClusterColorGradingDetailTreeItem> ParentCategory = InTreeItem->GetParent().Pin();
+			while (ParentCategory.IsValid() && !ParentCategory->IsCategory())
+			{
+				ParentCategory = ParentCategory->GetParent().Pin();
+			}
+
+			FString Key;
+			if (ParentCategory.IsValid())
+			{
+				Key = ParentCategory->GetNodeName().ToString() + TEXT(".") + InTreeItem->GetNodeName().ToString();
+			}
+			else
+			{
+				Key = InTreeItem->GetNodeName().ToString();
+			}
+			
+			if (bIsItemExpanded)
+			{
+				ExpandedDetailNodes.Add(Key);
+			}
+			else
+			{
+				ExpandedDetailNodes.Remove(Key);
+			}
+		}
+
+		if (bRecursive)
+		{
+			for (int32 ChildIndex = 0; ChildIndex < Children.Num(); ++ChildIndex)
+			{
+				SetNodeExpansionState(Children[ChildIndex], bIsItemExpanded, bRecursive);
+			}
+		}
+	}
+}
+
+void SDisplayClusterColorGradingDetailView::OnSetExpansionRecursive(TSharedRef<FDisplayClusterColorGradingDetailTreeItem> InTreeNode, bool bIsItemExpanded)
+{
+	SetNodeExpansionState(InTreeNode, bIsItemExpanded, true);
+	SaveExpandedItems();
+}
+
+void SDisplayClusterColorGradingDetailView::OnExpansionChanged(TSharedRef<FDisplayClusterColorGradingDetailTreeItem> InTreeNode, bool bIsItemExpanded)
+{
+	SetNodeExpansionState(InTreeNode, bIsItemExpanded, false);
+	SaveExpandedItems();
+}
+
+void SDisplayClusterColorGradingDetailView::OnRowReleased(const TSharedRef<ITableRow>& TableRow)
+{
+	// search upwards from the current keyboard-focused widget to see if it's contained in our row
+	TSharedPtr<SWidget> CurrentWidget = FSlateApplication::Get().GetKeyboardFocusedWidget();
+	while (CurrentWidget.IsValid())
+	{
+		if (CurrentWidget == TableRow->AsWidget())
+		{
+			// if so, clear focus so that any pending value changes are committed
+			FSlateApplication::Get().ClearKeyboardFocus();
+			return;
+		}
+
+		CurrentWidget = CurrentWidget->GetParentWidget();
+	}
+}
+
+EVisibility SDisplayClusterColorGradingDetailView::GetScrollBarVisibility() const
+{
+	const bool bShowScrollBar = RootTreeNodes.Num() > 0;
+	return bShowScrollBar ? EVisibility::Visible : EVisibility::Collapsed;
+}
\ No newline at end of file
diff --git a/Engine/Plugins/Runtime/nDisplay/Source/DisplayClusterColorGrading/Private/DetailView/SDisplayClusterColorGradingDetailView.h b/Engine/Plugins/Runtime/nDisplay/Source/DisplayClusterColorGrading/Private/DetailView/SDisplayClusterColorGradingDetailView.h
new file mode 100644
--- /dev/null
+++ b/Engine/Plugins/Runtime/nDisplay/Source/DisplayClusterColorGrading/Private/DetailView/SDisplayClusterColorGradingDetailView.h
@@ -0,0 +1,174 @@
+// Copyright Epic Games, Inc. All Rights Reserved.
+
+#pragma once
+
+#include "CoreMinimal.h"
+#include "DetailColumnSizeData.h"
+#include "DetailWidgetRow.h"
+#include "Widgets/DeclarativeSyntaxSupport.h"
+#include "Widgets/SCompoundWidget.h"
+#include "Widgets/Views/STableRow.h"
+
+class IDetailTreeNode;
+class IPropertyRowGenerator;
+class IPropertyHandle;
+class ITableRow;
+class FDetailTreeNode;
+class FDetailWidgetRow;
+class STableViewBase;
+
+template<typename T>
+class STreeView;
+
+DECLARE_DELEGATE_RetVal_OneParam(bool, FOnFilterDetailTreeNode, const TSharedRef<IDetailTreeNode>&);
+
+/** A wrapper used to abstract FDetailTreeNode, which is a private abstract class defined in PropertyEditor/Private */
+class FDisplayClusterColorGradingDetailTreeItem : public TSharedFromThis<FDisplayClusterColorGradingDetailTreeItem>
+{
+public:
+	FDisplayClusterColorGradingDetailTreeItem(const TSharedPtr<FDetailTreeNode>& InDetailTreeNode)
+		: DetailTreeNode(InDetailTreeNode)
+	{ };
+
+	/** Initializes the detail tree item, creating any child tree items needed */
+	void Initialize(const FOnFilterDetailTreeNode& NodeFilter);
+
+	/** Gets the parent detail tree item of this item */
+	TWeakPtr<FDisplayClusterColorGradingDetailTreeItem> GetParent() const { return Parent; }
+
+	/** Gets whether this tree item has any children */
+	bool HasChildren() const { return Children.Num() > 0; }
+
+	/** Gets the list of child tree items of this item */
+	void GetChildren(TArray<TSharedRef<FDisplayClusterColorGradingDetailTreeItem>>& OutChildren) const;
+
+	/** Gets the underlying IDetailTreeNode this detail tree item wraps */
+	TWeakPtr<IDetailTreeNode> GetDetailTreeNode() const;
+
+	/** Gets the property handle of the property this detail tree item represents */
+	TSharedPtr<IPropertyHandle> GetPropertyHandle() const { return PropertyHandle; }
+
+	/** Gets the name of this detail tree item */
+	FName GetNodeName() const;
+
+	/** Gets whether this detail tree item should be expanded */
+	bool ShouldBeExpanded() const;
+
+	/** Raised when this detail tree item's expansion state has been changed */
+	void OnItemExpansionChanged(bool bIsExpanded, bool bShouldSaveState);
+
+	/** Gets whether the "reset to default" button should be visible for this detail tree item */
+	bool IsResetToDefaultVisible() const;
+
+	/** Resets the property this detail tree item represents to its default value */
+	void ResetToDefault();
+
+	/** Gets an attribute that can be used to determine if property editing is enabled for this detail tree item */
+	TAttribute<bool> IsPropertyEditingEnabled() const;
+
+	/** Gets whether this detail tree item is a category */
+	bool IsCategory() const;
+
+	/** Gets whether this detail tree item is an item */
+	bool IsItem() const;
+
+	/** Gets whether this detail tree item can be reordered through a drag drop action */
+	bool IsReorderable() const;
+
+	/** Generates the row widgets for this detail tree item */
+	void GenerateDetailWidgetRow(FDetailWidgetRow& OutDetailWidgetRow) const;
+
+private:
+	/** A weak pointer to the detail tree node this item wraps */
+	TWeakPtr<FDetailTreeNode> DetailTreeNode;
+
+	/** The property handle for the property this item represents */
+	TSharedPtr<IPropertyHandle> PropertyHandle;
+
+	/** A weak pointer to this item's parent */
+	TWeakPtr<FDisplayClusterColorGradingDetailTreeItem> Parent;
+
+	/** A list of children of this item */
+	TArray<TSharedRef<FDisplayClusterColorGradingDetailTreeItem>> Children;
+};
+
+/**
+ * A custom detail view based on SDetailView that uses a property row generator as a source for the property nodes instead of generating them manually.
+ * Using an existing property row generator allows the detail view to display an object's properties much faster than the ordinary SDetailView, which
+ * has to regenerate a new property node tree every time the object being displayed is changed
+ */
+class SDisplayClusterColorGradingDetailView : public SCompoundWidget
+{
+public:
+	SLATE_BEGIN_ARGS(SDisplayClusterColorGradingDetailView) { }
+		SLATE_ARGUMENT(TSharedPtr<IPropertyRowGenerator>, PropertyRowGeneratorSource)
+		SLATE_EVENT(FOnFilterDetailTreeNode, OnFilterDetailTreeNode)
+	SLATE_END_ARGS()
+
+	void Construct(const FArguments& InArgs);
+
+	//~ SWidget interface
+	virtual void Tick(const FGeometry& AllottedGeometry, const double InCurrentTime, const float InDeltaTime) override;
+	//~ End SWidget interface
+
+	/** Regenerates this widget based on the current state of its property row generator source */
+	void Refresh();
+
+	/** Saves the expansion state of all properties being displayed in this detail view to the user's config file */
+	void SaveExpandedItems();
+
+	/** Restores the expansion state of all properties being displayed in this detail view from the user's config file */
+	void RestoreExpandedItems();
+
+private:
+	/** Updates the detail tree using the current state of the property row generator source */
+	void UpdateTreeNodes();
+
+	/** Updates the expansion state of the specified tree item using the stored expansion state configuration */
+	void UpdateExpansionState(const TSharedRef<FDisplayClusterColorGradingDetailTreeItem> InTreeItem);
+
+	/** Sets the expansion state of the specified tree item, and optionally recursively sets the expansion state of its children */
+	void SetNodeExpansionState(TSharedRef<FDisplayClusterColorGradingDetailTreeItem> InTreeNode, bool bIsItemExpanded, bool bRecursive);
+
+	/** Generates a table row widget for the specified tree item */
+	TSharedRef<ITableRow> GenerateNodeRow(TSharedRef<FDisplayClusterColorGradingDetailTreeItem> InTreeNode, const TSharedRef<STableViewBase>& OwnerTable);
+
+	/** Gets a list of child tree items for the specified tree item */
+	void GetChildrenForNode(TSharedRef<FDisplayClusterColorGradingDetailTreeItem> InTreeNode, TArray<TSharedRef<FDisplayClusterColorGradingDetailTreeItem>>& OutChildren);
+
+	/** Raised when the underlying tree widget is setting the expansion state of the specified tree item recursively */
+	void OnSetExpansionRecursive(TSharedRef<FDisplayClusterColorGradingDetailTreeItem> InTreeNode, bool bIsItemExpanded);
+
+	/** Raised when the underlying tree widget is setting the expansion state of the specified tree item */
+	void OnExpansionChanged(TSharedRef<FDisplayClusterColorGradingDetailTreeItem> InTreeNode, bool bIsItemExpanded);
+
+	/** Raised when the underlying tree widget is releasing the specified table row */
+	void OnRowReleased(const TSharedRef<ITableRow>& TableRow);
+
+	/** Gets the visibility of the scrollbar for the detail view */
+	EVisibility GetScrollBarVisibility() const;
+
+private:
+	typedef STreeView<TSharedRef<FDisplayClusterColorGradingDetailTreeItem>> SDetailTree;
+
+	/** The underlying tree view used to display the property widgets */
+	TSharedPtr<SDetailTree> DetailTree;
+
+	/** The source list of the root detail tree nodes being displayed by the tree widget */
+	TArray<TSharedRef<FDisplayClusterColorGradingDetailTreeItem>> RootTreeNodes;
+
+	/** The property row generator to generate the property widgets from */
+	TSharedPtr<IPropertyRowGenerator> PropertyRowGeneratorSource;
+
+	/** Column sizing data for the properties */
+	FDetailColumnSizeData ColumnSizeData;
+
+	/** A list of tree items whose expansion state needs to be set on the next tick */
+	TMap<TWeakPtr<FDisplayClusterColorGradingDetailTreeItem>, bool> TreeItemsToSetExpansionState;
+
+	/** A list of currently expanded detail nodes */
+	TSet<FString> ExpandedDetailNodes;
+
+	/** Delegate used to filter or process the detail tree nodes that are displayed in the detail view */
+	FOnFilterDetailTreeNode OnFilterDetailTreeNode;
+};
\ No newline at end of file
diff --git a/Engine/Plugins/Runtime/nDisplay/Source/DisplayClusterColorGrading/Private/DisplayClusterColorGradingDataModel.cpp b/Engine/Plugins/Runtime/nDisplay/Source/DisplayClusterColorGrading/Private/DisplayClusterColorGradingDataModel.cpp
--- a/Engine/Plugins/Runtime/nDisplay/Source/DisplayClusterColorGrading/Private/DisplayClusterColorGradingDataModel.cpp
+++ b/Engine/Plugins/Runtime/nDisplay/Source/DisplayClusterColorGrading/Private/DisplayClusterColorGradingDataModel.cpp
@@ -2,18 +2,50 @@
 
 #include "DisplayClusterColorGradingDataModel.h"
 
+#include "Drawer/DisplayClusterColorGradingDrawerState.h"
+
 #include "DisplayClusterRootActor.h"
 
-#include "IDetailTreeNode.h"
+#include "DetailWidgetRow.h"
 #include "IPropertyRowGenerator.h"
 #include "Modules/ModuleManager.h"
 #include "PropertyEditorModule.h"
 #include "PropertyHandle.h"
-#include "Widgets/Input/SCheckBox.h"
+
 
 #define LOCTEXT_NAMESPACE "DisplayClusterColorGrading"
 
-TMap<FName, FGetColorGradingDataModelGenerator> FDisplayClusterColorGradingDataModel::RegisteredDataModelGenerators;
+TMap<TWeakObjectPtr<UClass>, FGetColorGradingDataModelGenerator> FDisplayClusterColorGradingDataModel::RegisteredDataModelGenerators;
+
+/** Detail customizer intended for color FVector4 properties that don't generate property nodes for the child components of the vector, to speed up property node tree generation */
+class FFastColorStructCustomization : public IPropertyTypeCustomization
+{
+public:
+	static TSharedRef<IPropertyTypeCustomization> MakeInstance()
+	{
+		return MakeShareable(new FFastColorStructCustomization);
+	}
+
+	virtual void CustomizeHeader(TSharedRef<IPropertyHandle> StructPropertyHandle, FDetailWidgetRow& HeaderRow, IPropertyTypeCustomizationUtils& StructCustomizationUtils) override
+	{
+		HeaderRow.NameContent()
+		[
+			StructPropertyHandle->CreatePropertyNameWidget()
+		];
+	}
+
+	virtual void CustomizeChildren(TSharedRef<IPropertyHandle> StructPropertyHandle, IDetailChildrenBuilder& StructBuilder, IPropertyTypeCustomizationUtils& StructCustomizationUtils) override
+	{
+	}
+};
+
+class FColorPropertyTypeIdentifier : public IPropertyTypeIdentifier
+{
+	virtual bool IsPropertyTypeCustomized(const IPropertyHandle& PropertyHandle) const override
+	{
+		return PropertyHandle.HasMetaData(TEXT("ColorGradingMode"));
+	}
+};
 
 FDisplayClusterColorGradingDataModel::FDisplayClusterColorGradingDataModel()
 {
@@ -22,6 +54,14 @@ FDisplayClusterColorGradingDataModel::FDisplayClusterColorGradingDataModel()
 	FPropertyRowGeneratorArgs Args;
 	PropertyRowGenerator = PropertyEditorModule.CreatePropertyRowGenerator(Args);
 	PropertyRowGenerator->OnRowsRefreshed().AddRaw(this, &FDisplayClusterColorGradingDataModel::OnPropertyRowGeneratorRefreshed);
+
+	TSharedRef<FColorPropertyTypeIdentifier> ColorPropertyTypeIdentifier = MakeShared<FColorPropertyTypeIdentifier>();
+
+	// Since there is an entirely custom set of widgets for displaying and editing the color grading settings, set a customizer for any color vectors to prevent the 
+	// property row generator from generating child properties or extraneous widgets, which drastically helps improve performance when loading object properties
+	PropertyRowGenerator->RegisterInstancedCustomPropertyTypeLayout(NAME_Vector4, FOnGetPropertyTypeCustomizationInstance::CreateStatic(&FFastColorStructCustomization::MakeInstance), ColorPropertyTypeIdentifier);
+	PropertyRowGenerator->RegisterInstancedCustomPropertyTypeLayout(NAME_Vector4f, FOnGetPropertyTypeCustomizationInstance::CreateStatic(&FFastColorStructCustomization::MakeInstance), ColorPropertyTypeIdentifier);
+	PropertyRowGenerator->RegisterInstancedCustomPropertyTypeLayout(NAME_Vector4d, FOnGetPropertyTypeCustomizationInstance::CreateStatic(&FFastColorStructCustomization::MakeInstance), ColorPropertyTypeIdentifier);
 }
 
 TArray<TWeakObjectPtr<UObject>> FDisplayClusterColorGradingDataModel::GetObjects() const
@@ -36,6 +76,8 @@ TArray<TWeakObjectPtr<UObject>> FDisplayClusterColorGradingDataModel::GetObjects
 
 void FDisplayClusterColorGradingDataModel::SetObjects(const TArray<UObject*>& InObjects)
 {
+	TRACE_CPUPROFILER_EVENT_SCOPE(FDisplayClusterColorGradingDataModel::SetObjects);
+
 	// Only update the data model if the objects being set are new
 	bool bUpdateDataModel = false;
 	if (PropertyRowGenerator.IsValid())
@@ -63,6 +105,14 @@ void FDisplayClusterColorGradingDataModel::SetObjects(const TArray<UObject*>& In
 	{
 		Reset();
 
+		for (const UObject* Object : InObjects)
+		{
+			if (Object)
+			{
+				InitializeDataModelGenerator(Object->GetClass());
+			}
+		}
+
 		if (PropertyRowGenerator.IsValid())
 		{
 			PropertyRowGenerator->SetObjects(InObjects);
@@ -93,6 +143,12 @@ bool FDisplayClusterColorGradingDataModel::HasObjectOfType(const UClass* InClass
 
 void FDisplayClusterColorGradingDataModel::Reset()
 {
+	for (const TPair<TWeakObjectPtr<UClass>, TSharedPtr<IDisplayClusterColorGradingDataModelGenerator>>& GeneratorInstance : DataModelGeneratorInstances)
+	{
+		GeneratorInstance.Value->Destroy(SharedThis(this), PropertyRowGenerator.ToSharedRef());
+		PropertyRowGenerator->UnregisterInstancedCustomPropertyLayout(GeneratorInstance.Key.Get());
+	}
+
 	DataModelGeneratorInstances.Empty();
 	ColorGradingGroups.Empty();
 	DetailsSections.Empty();
@@ -100,10 +156,41 @@ void FDisplayClusterColorGradingDataModel::Reset()
 	SelectedColorGradingElementIndex = INDEX_NONE;
 	ColorGradingGroupToolBarWidget = nullptr;
 	bShowColorGradingGroupToolBar = false;
+}
+
+void FDisplayClusterColorGradingDataModel::GetDrawerState(FDisplayClusterColorGradingDrawerState& OutDrawerState)
+{
+	OutDrawerState.SelectedColorGradingGroup = SelectedColorGradingGroupIndex;
+	OutDrawerState.SelectedColorGradingElement = SelectedColorGradingElementIndex;
+}
+
+void FDisplayClusterColorGradingDataModel::SetDrawerState(const FDisplayClusterColorGradingDrawerState& InDrawerState)
+{
+	TRACE_CPUPROFILER_EVENT_SCOPE(FDisplayClusterColorGradingDataModel::SetDrawerState);
+
+	SelectedColorGradingGroupIndex = InDrawerState.SelectedColorGradingGroup;
+	SelectedColorGradingElementIndex = InDrawerState.SelectedColorGradingElement;
+
+	TArray<UObject*> ObjectsToSelect;
+	for (const TWeakObjectPtr<UObject>& Object : InDrawerState.SelectedObjects)
+	{
+		if (Object.IsValid())
+		{
+			ObjectsToSelect.Add(Object.Get());
+		}
+	}
+
+	for (const UObject* Object : ObjectsToSelect)
+	{
+		if (Object)
+		{
+			InitializeDataModelGenerator(Object->GetClass());
+		}
+	}
 
 	if (PropertyRowGenerator.IsValid())
 	{
-		PropertyRowGenerator->SetObjects(TArray<UObject*>());
+		PropertyRowGenerator->SetObjects(ObjectsToSelect);
 	}
 }
 
@@ -119,8 +206,28 @@ FDisplayClusterColorGradingDataModel::FColorGradingGroup* FDisplayClusterColorGr
 
 void FDisplayClusterColorGradingDataModel::SetSelectedColorGradingGroup(int32 InColorGradingGroupIndex)
 {
+	TRACE_CPUPROFILER_EVENT_SCOPE(FDisplayClusterColorGradingDataModel::SetSelectedColorGradingGroup);
+
 	SelectedColorGradingGroupIndex = InColorGradingGroupIndex <= ColorGradingGroups.Num() ? InColorGradingGroupIndex : INDEX_NONE;
+
+	// When the color grading group has changed, reset the selected color grading element as well
+	const bool bHasColorGradingElements = SelectedColorGradingGroupIndex > INDEX_NONE && ColorGradingGroups[SelectedColorGradingElementIndex].ColorGradingElements.Num();
+	SelectedColorGradingElementIndex = bHasColorGradingElements ? 0 : INDEX_NONE;
+
 	OnColorGradingGroupSelectionChangedDelegate.Broadcast();
+
+	// Force the property row generator to rebuild the property node tree, since the data model generators may have made some optimizations
+	// based on which color grading group is currently selected
+	TArray<UObject*> Objects;
+	for (const TWeakObjectPtr<UObject>& WeakObj : PropertyRowGenerator->GetSelectedObjects())
+	{
+		if (WeakObj.IsValid())
+		{
+			Objects.Add(WeakObj.Get());
+		}
+	}
+
+	PropertyRowGenerator->SetObjects(Objects);
 }
 
 FDisplayClusterColorGradingDataModel::FColorGradingElement* FDisplayClusterColorGradingDataModel::GetSelectedColorGradingElement()
@@ -143,37 +250,47 @@ void FDisplayClusterColorGradingDataModel::SetSelectedColorGradingElement(int32
 	OnColorGradingElementSelectionChangedDelegate.Broadcast();
 }
 
-TSharedPtr<IDisplayClusterColorGradingDataModelGenerator> FDisplayClusterColorGradingDataModel::GetDataModelGenerator(const UClass* InClass)
+void FDisplayClusterColorGradingDataModel::InitializeDataModelGenerator(UClass* InClass)
 {
-	const UClass* CurrentClass = InClass;
+	UClass* CurrentClass = InClass;
 	while (CurrentClass)
 	{
-		const FName ClassName = CurrentClass->GetFName();
-
-		if (DataModelGeneratorInstances.Contains(ClassName))
-		{
-			return DataModelGeneratorInstances[ClassName];
-		}
-		else if (RegisteredDataModelGenerators.Contains(ClassName))
+		if (RegisteredDataModelGenerators.Contains(CurrentClass) && !DataModelGeneratorInstances.Contains(CurrentClass))
 		{
-			if (RegisteredDataModelGenerators[ClassName].IsBound())
+			if (RegisteredDataModelGenerators[CurrentClass].IsBound())
 			{
-				TSharedRef<IDisplayClusterColorGradingDataModelGenerator> Generator = RegisteredDataModelGenerators[ClassName].Execute();
-				DataModelGeneratorInstances.Add(ClassName, Generator);
+				TSharedRef<IDisplayClusterColorGradingDataModelGenerator> Generator = RegisteredDataModelGenerators[CurrentClass].Execute();
+				Generator->Initialize(SharedThis(this), PropertyRowGenerator.ToSharedRef());
 
-				return Generator;
+				DataModelGeneratorInstances.Add(CurrentClass, Generator);
 			}
 		}
 
 		CurrentClass = CurrentClass->GetSuperClass();
 	}
+}
+
+TSharedPtr<IDisplayClusterColorGradingDataModelGenerator> FDisplayClusterColorGradingDataModel::GetDataModelGenerator(UClass* InClass) const
+{
+	UClass* CurrentClass = InClass;
+	while (CurrentClass)
+	{
+		if (DataModelGeneratorInstances.Contains(CurrentClass))
+		{
+			return DataModelGeneratorInstances[CurrentClass];
+		}
+
+		CurrentClass = CurrentClass->GetSuperClass();
+	}
 
 	return nullptr;
 }
 
 void FDisplayClusterColorGradingDataModel::OnPropertyRowGeneratorRefreshed()
 {
-	ColorGradingGroups.Empty();
+	TRACE_CPUPROFILER_EVENT_SCOPE(FDisplayClusterColorGradingDataModel::OnPropertyRowGeneratorRefreshed);
+
+ 	ColorGradingGroups.Empty();
 	DetailsSections.Empty();
 
 	const TArray<TWeakObjectPtr<UObject>>& SelectedObjects = PropertyRowGenerator->GetSelectedObjects();
diff --git a/Engine/Plugins/Runtime/nDisplay/Source/DisplayClusterColorGrading/Private/DisplayClusterColorGradingDataModel.h b/Engine/Plugins/Runtime/nDisplay/Source/DisplayClusterColorGrading/Private/DisplayClusterColorGradingDataModel.h
--- a/Engine/Plugins/Runtime/nDisplay/Source/DisplayClusterColorGrading/Private/DisplayClusterColorGradingDataModel.h
+++ b/Engine/Plugins/Runtime/nDisplay/Source/DisplayClusterColorGrading/Private/DisplayClusterColorGradingDataModel.h
@@ -8,6 +8,7 @@
 class IPropertyHandle;
 class IPropertyRowGenerator;
 class SWidget;
+struct FDisplayClusterColorGradingDrawerState;
 
 /** Interface that allows color grading data models to be generated for specific UObjects */
 class IDisplayClusterColorGradingDataModelGenerator : public TSharedFromThis<IDisplayClusterColorGradingDataModelGenerator>
@@ -15,6 +16,9 @@ class IDisplayClusterColorGradingDataModelGenerator : public TSharedFromThis<IDi
 public:
 	virtual ~IDisplayClusterColorGradingDataModelGenerator() {}
 
+	virtual void Initialize(const TSharedRef<class FDisplayClusterColorGradingDataModel>& ColorGradingDataModel, const TSharedRef<IPropertyRowGenerator>& PropertyRowGenerator) = 0;
+	virtual void Destroy(const TSharedRef<class FDisplayClusterColorGradingDataModel>& ColorGradingDataModel, const TSharedRef<IPropertyRowGenerator>& PropertyRowGenerator) = 0;
+
 	/** Called when the color grading data model should be generated by this generator */
 	virtual void GenerateDataModel(IPropertyRowGenerator& PropertyRowGenerator, class FDisplayClusterColorGradingDataModel& OutColorGradingDataModel) = 0;
 };
@@ -25,20 +29,20 @@ DECLARE_DELEGATE_RetVal(TSharedRef<IDisplayClusterColorGradingDataModelGenerator
  * A data model that stores the color grading properties from a UObject, which allows UObjects with widely varying color grading
  * interfaces to be used with the color grading panel.
  */
-class FDisplayClusterColorGradingDataModel
+class FDisplayClusterColorGradingDataModel : public TSharedFromThis<FDisplayClusterColorGradingDataModel>
 {
 public:
 	/** Registers a new data model generator used to populate a color grading data model for the specified class */
 	template<class T>
 	static void RegisterColorGradingDataModelGenerator(FGetColorGradingDataModelGenerator CreateGeneratorDelegate)
 	{
-		const FName ClassName = T::StaticClass()->GetFName();
-		RegisteredDataModelGenerators.Add(ClassName, CreateGeneratorDelegate);
+		UClass* Class = T::StaticClass();
+		RegisteredDataModelGenerators.Add(Class, CreateGeneratorDelegate);
 	}
 
 private:
 	/** A list of data model generators, indexed by class name, that have been registered */
-	static TMap<FName, FGetColorGradingDataModelGenerator> RegisteredDataModelGenerators;
+	static TMap<TWeakObjectPtr<UClass>, FGetColorGradingDataModelGenerator> RegisteredDataModelGenerators;
 
 
 public:
@@ -82,11 +86,8 @@ public:
 		/** The list of color grading elements (Shadows, Midtones, etc.) that belong in the group */
 		TArray<FColorGradingElement> ColorGradingElements;
 
-		/** A list of properties to display in the details view of the color grading panel */
-		TArray<TSharedPtr<IPropertyHandle>> DetailsViewPropertyHandles;
-
-		/** Optional sort order for the categories displayed in the details view */
-		TMap<FName, int32> DetailsViewCategorySortOrder;
+		/** A list of categories of properties to display in the detail view of the color grading panel */
+		TArray<FName> DetailsViewCategories;
 	};
 
 	/** Stores data used to generate a details subsection, which is a selectable element in a details section that modifies which properties are displayed in the section */
@@ -97,6 +98,9 @@ public:
 
 		/** A delegate that creates the details customization for the details subsection */
 		FOnGetDetailCustomizationInstance DetailCustomizationDelegate;
+
+		/** A list of categories to display in the detail subsection */
+		TArray<FName> Categories;
 	};
 
 	/** Stores data used to generate a details section in the drawer, which displays a specific subset of properties from the selected objects */
@@ -113,16 +117,19 @@ public:
 
 		/** A delegate that creates the details customization for the details section */
 		FOnGetDetailCustomizationInstance DetailCustomizationDelegate;
+
+		/** A list of categories to display in the detail subsection */
+		TArray<FName> Categories;
 	};
 
-	DECLARE_EVENT(FDisplayClusterColorGradingDataModel, FOnDataModelGenerated);
-	DECLARE_EVENT(FDisplayClusterColorGradingDataModel, FOnColorGradingSelectionChanged);
+	DECLARE_MULTICAST_DELEGATE(FOnDataModelGenerated);
+	DECLARE_MULTICAST_DELEGATE(FOnColorGradingSelectionChanged);
 
 public:
 	FDisplayClusterColorGradingDataModel();
 
 	/** Gets the property row generator that was used to generate this data model */
-	const TSharedRef<IPropertyRowGenerator> GetPropertyRowGenerator() const { return PropertyRowGenerator.ToSharedRef(); }
+	TSharedRef<IPropertyRowGenerator> GetPropertyRowGenerator() const { return PropertyRowGenerator.ToSharedRef(); }
 
 	/** Gets the list of objects from which the color grading properties were generated from */
 	TArray<TWeakObjectPtr<UObject>> GetObjects() const;
@@ -136,6 +143,12 @@ public:
 	/** Resets the data model to be empty */
 	void Reset();
 
+	/** Adds the state of the data model to the specified drawer state */
+	void GetDrawerState(FDisplayClusterColorGradingDrawerState& OutDrawerState);
+
+	/** Sets the state of the data model from the specified drawer state */
+	void SetDrawerState(const FDisplayClusterColorGradingDrawerState& InDrawerState);
+
 	/** Gets the index of the selected color grading group */
 	int32 GetSelectedColorGradingGroupIndex() const { return SelectedColorGradingGroupIndex; }
 
@@ -164,8 +177,11 @@ public:
 	FOnColorGradingSelectionChanged& OnColorGradingElementSelectionChanged() { return OnColorGradingElementSelectionChangedDelegate; }
 
 private:
+	/** Initializes any needed data model generators for the specified class */
+	void InitializeDataModelGenerator(UClass* InClass);
+
 	/** Attempts to find an instance of a data model generator that can be used for the specified class */
-	TSharedPtr<IDisplayClusterColorGradingDataModelGenerator> GetDataModelGenerator(const UClass* InClass);
+	TSharedPtr<IDisplayClusterColorGradingDataModelGenerator> GetDataModelGenerator(UClass* InClass) const;
 
 	/** Callback that is raised when the property row generator the data model is generated from has been refreshed */
 	void OnPropertyRowGeneratorRefreshed();
@@ -185,7 +201,7 @@ public:
 
 private:
 	/** A list of data model generators that have been created to generate the data model for the current objects */
-	TMap<FName, TSharedPtr<IDisplayClusterColorGradingDataModelGenerator>> DataModelGeneratorInstances;
+	TMap<TWeakObjectPtr<UClass>, TSharedPtr<IDisplayClusterColorGradingDataModelGenerator>> DataModelGeneratorInstances;
 
 	/** The property row generator which manages the property handles of the color grading objects */
 	TSharedPtr<IPropertyRowGenerator> PropertyRowGenerator;
diff --git a/Engine/Plugins/Runtime/nDisplay/Source/DisplayClusterColorGrading/Private/Drawer/DisplayClusterColorGradingDrawerState.h b/Engine/Plugins/Runtime/nDisplay/Source/DisplayClusterColorGrading/Private/Drawer/DisplayClusterColorGradingDrawerState.h
new file mode 100644
--- /dev/null
+++ b/Engine/Plugins/Runtime/nDisplay/Source/DisplayClusterColorGrading/Private/Drawer/DisplayClusterColorGradingDrawerState.h
@@ -0,0 +1,42 @@
+// Copyright Epic Games, Inc. All Rights Reserved.
+
+#pragma once
+
+#include "CoreMinimal.h"
+#include "Styling/SlateTypes.h"
+#include "SDisplayClusterColorGradingColorWheel.h"
+
+/** Drawer modes that the color grading drawer can be in */
+enum EDisplayClusterColorGradingDrawerMode
+{
+	ColorGrading,
+	DetailsView
+};
+
+/** Stores the state of the drawer UI that can be reloaded in cases where the drawer or any of its elements are reloaded (such as when the drawer is reopened or docked) */
+struct FDisplayClusterColorGradingDrawerState
+{
+	/** The objects that are selected in the list */
+	TArray<TWeakObjectPtr<UObject>> SelectedObjects;
+
+	/** The color grading group that is selected */
+	int32 SelectedColorGradingGroup = INDEX_NONE;
+
+	/** The color grading element that is selected */
+	int32 SelectedColorGradingElement = INDEX_NONE;
+
+	/** Indicates which color wheels are hidden */
+	TArray<bool> HiddenColorWheels;
+
+	/** The selected orientation of the color wheels */
+	EOrientation ColorWheelOrientation = EOrientation::Orient_Vertical;
+
+	/** The color display mode of the color wheels */
+	SDisplayClusterColorGradingColorWheel::EColorDisplayMode ColorDisplayMode;
+
+	/** The mode the drawer is in */
+	EDisplayClusterColorGradingDrawerMode DrawerMode;
+
+	/** Indicates which subsections were selected for each section in the details panel */
+	TArray<int32> SelectedDetailsSubsections;
+};
\ No newline at end of file
diff --git a/Engine/Plugins/Runtime/nDisplay/Source/DisplayClusterColorGrading/Private/Drawer/SDisplayClusterColorGradingDrawer.cpp b/Engine/Plugins/Runtime/nDisplay/Source/DisplayClusterColorGrading/Private/Drawer/SDisplayClusterColorGradingDrawer.cpp
--- a/Engine/Plugins/Runtime/nDisplay/Source/DisplayClusterColorGrading/Private/Drawer/SDisplayClusterColorGradingDrawer.cpp
+++ b/Engine/Plugins/Runtime/nDisplay/Source/DisplayClusterColorGrading/Private/Drawer/SDisplayClusterColorGradingDrawer.cpp
@@ -11,7 +11,6 @@
 
 #include "IDisplayClusterOperator.h"
 #include "IDisplayClusterOperatorViewModel.h"
-#include "DisplayClusterOperatorStatusBarExtender.h"
 
 #include "DisplayClusterRootActor.h"
 #include "Components/DisplayClusterICVFXCameraComponent.h"
@@ -23,7 +22,6 @@
 #include "Modules/ModuleManager.h"
 #include "ScopedTransaction.h"
 #include "Styling/AppStyle.h"
-#include "WidgetDrawerConfig.h"
 #include "Widgets/Input/SButton.h"
 #include "Widgets/Layout/SBorder.h"
 #include "Widgets/Layout/SBox.h"
@@ -366,11 +364,10 @@ void SDisplayClusterColorGradingDrawer::PostRedo(bool bSuccess)
 FDisplayClusterColorGradingDrawerState SDisplayClusterColorGradingDrawer::GetDrawerState() const
 {
 	FDisplayClusterColorGradingDrawerState DrawerState;
-
-	DrawerState.SelectedColorGradingGroup = ColorGradingDataModel->GetSelectedColorGradingGroupIndex();
-	DrawerState.SelectedColorGradingElement = ColorGradingDataModel->GetSelectedColorGradingElementIndex();
 	DrawerState.DrawerMode = CurrentDrawerMode;
 
+	ColorGradingDataModel->GetDrawerState(DrawerState);
+
 	if (ColorWheelPanel.IsValid())
 	{
 		ColorWheelPanel->GetDrawerState(DrawerState);
@@ -416,14 +413,14 @@ void SDisplayClusterColorGradingDrawer::SetDrawerState(const FDisplayClusterColo
 	{
 		if (SelectedObject.IsValid())
 		{
-			auto FindColorGraidngItem = [&SelectedObject](const FDisplayClusterColorGradingListItemRef& ColorGradingItem)
+			auto FindColorGradingItem = [&SelectedObject](const FDisplayClusterColorGradingListItemRef& ColorGradingItem)
 			{
 				return ColorGradingItem->Actor == SelectedObject || ColorGradingItem->Component == SelectedObject;
 			};
 
 			for (int32 Index = 0; Index < ColorGradingItemLists.Num(); ++Index)
 			{
-				if (FDisplayClusterColorGradingListItemRef* FoundItem = ColorGradingItemLists[Index].FindByPredicate(FindColorGraidngItem))
+				if (FDisplayClusterColorGradingListItemRef* FoundItem = ColorGradingItemLists[Index].FindByPredicate(FindColorGradingItem))
 				{
 					ItemsToSelect[Index].Add(*FoundItem);
 					break;
@@ -432,7 +429,7 @@ void SDisplayClusterColorGradingDrawer::SetDrawerState(const FDisplayClusterColo
 		}
 	}
 
-	// TODO: For now, since we don't support multiple color grading items selected at once, ensure either a level item or a root actor item is seleccted, not both
+	// TODO: For now, since we don't support multiple color grading items selected at once, ensure either a level item or a root actor item is selected, not both
 	for (int32 Index = 0; Index < ColorGradingItemLists.Num(); ++Index)
 	{
 		if (ItemsToSelect[Index].Num())
@@ -444,8 +441,7 @@ void SDisplayClusterColorGradingDrawer::SetDrawerState(const FDisplayClusterColo
 		}
 	}
 
-	ColorGradingDataModel->SetSelectedColorGradingGroup(InDrawerState.SelectedColorGradingGroup);
-	ColorGradingDataModel->SetSelectedColorGradingElement(InDrawerState.SelectedColorGradingElement);
+	ColorGradingDataModel->SetDrawerState(InDrawerState);
 
 	CurrentDrawerMode = InDrawerState.DrawerMode;
 
@@ -492,6 +488,7 @@ void SDisplayClusterColorGradingDrawer::SetDrawerStateToDefault()
 					if (ListItemPtr)
 					{
 						RootActorList->SetSelectedItems({ *ListItemPtr });
+						SetColorGradingDataModelObjects({ CCW });
 					}
 				}
 			}
@@ -499,6 +496,7 @@ void SDisplayClusterColorGradingDrawer::SetDrawerStateToDefault()
 			{
 				// The nDisplay stage actor is always the first item in the root actor color grading items list, so set that as the currently selected item
 				RootActorList->SetSelectedItems({ RootActorColorGradingItems[0] });
+				SetColorGradingDataModelObjects({ RootActorColorGradingItems[0]->Actor.Get() });
 			}
 		}
 	}
@@ -820,6 +818,18 @@ void SDisplayClusterColorGradingDrawer::FillColorCorrectionRegionColorGradingLis
 	}
 }
 
+void SDisplayClusterColorGradingDrawer::SetColorGradingDataModelObjects(const TArray<UObject*>& Objects)
+{
+	ColorGradingDataModel->SetObjects(Objects);
+
+	// If the current drawer mode is the details view but the newly selected objects don't have any details sections to display,
+	// switch to the color grading drawer mode
+	if (CurrentDrawerMode == EDisplayClusterColorGradingDrawerMode::DetailsView && ColorGradingDataModel->DetailsSections.Num() == 0)
+	{
+		CurrentDrawerMode = EDisplayClusterColorGradingDrawerMode::ColorGrading;
+	}
+}
+
 void SDisplayClusterColorGradingDrawer::FillColorGradingGroupToolBar()
 {
 	if (ColorGradingGroupToolBarBox.IsValid())
@@ -1067,14 +1077,12 @@ void SDisplayClusterColorGradingDrawer::OnListSelectionChanged(TSharedRef<SDispl
 
 			if (List->GetSelectedItems().Num())
 			{
-				bUpdateDataModelOnSelectionChanged = false;
 				List->SetSelectedItems(TArray<FDisplayClusterColorGradingListItemRef>());
-				bUpdateDataModelOnSelectionChanged = true;
 			}
 		}
 	}
 
-	if (bUpdateDataModelOnSelectionChanged)
+	if (SelectInfo != ESelectInfo::Direct)
 	{
 		TArray<UObject*> ObjectsToColorGrade;
 		for (const FDisplayClusterColorGradingListItemRef& SelectedObject : SelectedObjects)
@@ -1089,14 +1097,7 @@ void SDisplayClusterColorGradingDrawer::OnListSelectionChanged(TSharedRef<SDispl
 			}
 		}
 
-		ColorGradingDataModel->SetObjects(ObjectsToColorGrade);
-
-		// If the current drawer mode is the details view but the newly selected objects don't have any details sections to display,
-		// switch to the color grading drawer mode
-		if (CurrentDrawerMode == EDisplayClusterColorGradingDrawerMode::DetailsView && ColorGradingDataModel->DetailsSections.Num() == 0)
-		{
-			CurrentDrawerMode = EDisplayClusterColorGradingDrawerMode::ColorGrading;
-		}
+		SetColorGradingDataModelObjects(ObjectsToColorGrade);
 	}
 }
 
diff --git a/Engine/Plugins/Runtime/nDisplay/Source/DisplayClusterColorGrading/Private/Drawer/SDisplayClusterColorGradingDrawer.h b/Engine/Plugins/Runtime/nDisplay/Source/DisplayClusterColorGrading/Private/Drawer/SDisplayClusterColorGradingDrawer.h
--- a/Engine/Plugins/Runtime/nDisplay/Source/DisplayClusterColorGrading/Private/Drawer/SDisplayClusterColorGradingDrawer.h
+++ b/Engine/Plugins/Runtime/nDisplay/Source/DisplayClusterColorGrading/Private/Drawer/SDisplayClusterColorGradingDrawer.h
@@ -8,6 +8,7 @@
 #include "Widgets/DeclarativeSyntaxSupport.h"
 #include "Widgets/SCompoundWidget.h"
 
+#include "DisplayClusterColorGradingDrawerState.h"
 #include "SDisplayClusterColorGradingObjectList.h"
 #include "SDisplayClusterColorGradingColorWheel.h"
 
@@ -20,41 +21,6 @@ class SHorizontalBox;
 class SDisplayClusterColorGradingColorWheelPanel;
 class SDisplayClusterColorGradingDetailsPanel;
 
-/** Drawer modes that the color grading drawer can be in */
-enum EDisplayClusterColorGradingDrawerMode
-{
-	ColorGrading,
-	DetailsView
-};
-
-/** Stores the state of the drawer UI that can be reloaded in cases where the drawer or any of its elements are reloaded (such as when the drawer is reopened or docked) */
-struct FDisplayClusterColorGradingDrawerState
-{
-	/** The objects that are selected in the list */
-	TArray<TWeakObjectPtr<UObject>> SelectedObjects;
-
-	/** The color grading group that is selected */
-	int32 SelectedColorGradingGroup = INDEX_NONE;
-
-	/** The color grading element that is selected */
-	int32 SelectedColorGradingElement = INDEX_NONE;
-
-	/** Indicates which color wheels are hidden */
-	TArray<bool> HiddenColorWheels;
-
-	/** The selected orientation of the color wheels */
-	EOrientation ColorWheelOrientation = EOrientation::Orient_Vertical;
-
-	/** The color display mode of the color wheels */
-	SDisplayClusterColorGradingColorWheel::EColorDisplayMode ColorDisplayMode;
-
-	/** The mode the drawer is in */
-	EDisplayClusterColorGradingDrawerMode DrawerMode;
-
-	/** Indicates which subsections were selected for each section in the details panel */
-	TArray<int32> SelectedDetailsSubsections;
-};
-
 /** Color grading drawer widget, which displays a list of color gradable items, and the color wheel panel */
 class SDisplayClusterColorGradingDrawer : public SCompoundWidget, public FEditorUndoClient
 {
@@ -122,6 +88,9 @@ private:
 	/** Fills the specified list with any color gradable color correction regions in the currently loaded level */
 	void FillColorCorrectionRegionColorGradingList(TArray<FDisplayClusterColorGradingListItemRef>& InOutList);
 
+	/** Updates the color grading data model with the specified list of objects */
+	void SetColorGradingDataModelObjects(const TArray<UObject*>& Objects);
+
 	/** Fills the color grading group toolbar using the color grading data model */
 	void FillColorGradingGroupToolBar();
 
diff --git a/Engine/Plugins/Runtime/nDisplay/Source/DisplayClusterColorGrading/Private/SDisplayClusterColorGradingColorWheelPanel.cpp b/Engine/Plugins/Runtime/nDisplay/Source/DisplayClusterColorGrading/Private/SDisplayClusterColorGradingColorWheelPanel.cpp
--- a/Engine/Plugins/Runtime/nDisplay/Source/DisplayClusterColorGrading/Private/SDisplayClusterColorGradingColorWheelPanel.cpp
+++ b/Engine/Plugins/Runtime/nDisplay/Source/DisplayClusterColorGrading/Private/SDisplayClusterColorGradingColorWheelPanel.cpp
@@ -4,16 +4,12 @@
 
 #include "DisplayClusterColorGradingCommands.h"
 #include "SDisplayClusterColorGradingColorWheel.h"
-#include "Drawer/SDisplayClusterColorGradingDrawer.h"
+#include "DetailView/SDisplayClusterColorGradingDetailView.h"
+#include "Drawer/DisplayClusterColorGradingDrawerState.h"
 
-#include "DetailCategoryBuilder.h"
-#include "DetailLayoutBuilder.h"
 #include "Framework/MultiBox/MultiBoxBuilder.h"
-#include "IDetailCustomization.h"
 #include "IDetailTreeNode.h"
 #include "IPropertyRowGenerator.h"
-#include "Modules/ModuleManager.h"
-#include "PropertyEditorModule.h"
 #include "PropertyHandle.h"
 #include "Widgets/SBoxPanel.h"
 #include "Widgets/Input/SComboButton.h"
@@ -42,13 +38,6 @@ void SDisplayClusterColorGradingColorWheelPanel::Construct(const FArguments& InA
 		ColorGradingDataModel->OnColorGradingElementSelectionChanged().AddSP(this, &SDisplayClusterColorGradingColorWheelPanel::OnColorGradingElementSelectionChanged);
 	}
 
-	FPropertyEditorModule& EditModule = FModuleManager::Get().GetModuleChecked<FPropertyEditorModule>("PropertyEditor");
-	FDetailsViewArgs DetailsViewArgs;
-	DetailsViewArgs.bAllowSearch = false;
-	DetailsViewArgs.NameAreaSettings = FDetailsViewArgs::HideNameArea;
-	DetailsViewArgs.bHideSelectionTip = true;
-	DetailsView = EditModule.CreateDetailView(DetailsViewArgs);
-
 	CommandList = MakeShared<FUICommandList>();
 	BindCommands();
 
@@ -195,7 +184,9 @@ void SDisplayClusterColorGradingColorWheelPanel::Construct(const FArguments& InA
 		+ SSplitter::Slot()
 		.Value(0.2f)
 		[
-			DetailsView.ToSharedRef()
+			SAssignNew(DetailView, SDisplayClusterColorGradingDetailView)
+			.PropertyRowGeneratorSource(ColorGradingDataModel->GetPropertyRowGenerator())
+			.OnFilterDetailTreeNode(this, &SDisplayClusterColorGradingColorWheelPanel::FilterDetailTreeNode)
 		]
 	];
 }
@@ -208,17 +199,24 @@ void SDisplayClusterColorGradingColorWheelPanel::Refresh()
 		{
 			FillColorGradingGroupProperty(*ColorGradingGroup);
 			FillColorGradingElementsToolBar(ColorGradingGroup->ColorGradingElements);
-			ColorGradingDataModel->SetSelectedColorGradingElement(ColorGradingGroup->ColorGradingElements.Num() ? 0 : INDEX_NONE);
 
-			FillDetailsView(*ColorGradingGroup);
+			if (const FDisplayClusterColorGradingDataModel::FColorGradingElement* ColorGradingElement = ColorGradingDataModel->GetSelectedColorGradingElement())
+			{
+				FillColorWheels(*ColorGradingElement);
+			}
+			else
+			{
+				ClearColorWheels();
+			}
 		}
 		else
 		{
 			ClearColorGradingGroupProperty();
 			ClearColorGradingElementsToolBar();
-			ClearDetailsView();
 			ClearColorWheels();
 		}
+
+		DetailView->Refresh();
 	}
 }
 
@@ -245,14 +243,14 @@ void SDisplayClusterColorGradingColorWheelPanel::BindCommands()
 		Commands.SaturationColorWheelVisibility,
 		FExecuteAction::CreateSP(this, &SDisplayClusterColorGradingColorWheelPanel::ToggleColorWheelVisible, 3),
 		FCanExecuteAction(),
-		FIsActionChecked::CreateSP(this, &SDisplayClusterColorGradingColorWheelPanel::IsColorWheelVisibile, 3)
+		FIsActionChecked::CreateSP(this, &SDisplayClusterColorGradingColorWheelPanel::IsColorWheelVisible, 3)
 	);
 
 	CommandList->MapAction(
 		Commands.ContrastColorWheelVisibility,
 		FExecuteAction::CreateSP(this, &SDisplayClusterColorGradingColorWheelPanel::ToggleColorWheelVisible, 4),
 		FCanExecuteAction(),
-		FIsActionChecked::CreateSP(this, &SDisplayClusterColorGradingColorWheelPanel::IsColorWheelVisibile, 4)
+		FIsActionChecked::CreateSP(this, &SDisplayClusterColorGradingColorWheelPanel::IsColorWheelVisible, 4)
 	);
 
 	CommandList->MapAction(
@@ -415,95 +413,6 @@ void SDisplayClusterColorGradingColorWheelPanel::ClearColorGradingElementsToolBa
 	ColorGradingElementsToolBarBox->ClearChildren();
 }
 
-class FColorGradingDetailsCustomization : public IDetailCustomization
-{
-public:
-	FColorGradingDetailsCustomization(const FDisplayClusterColorGradingDataModel::FColorGradingGroup& ColorGradingGroup)
-		: PropertyHandles(ColorGradingGroup.DetailsViewPropertyHandles)
-		, CategorySortOrder(ColorGradingGroup.DetailsViewCategorySortOrder)
-	{ }
-
-	virtual void CustomizeDetails(IDetailLayoutBuilder& DetailBuilder) override
-	{
-		TArray<FName> Categories;
-		DetailBuilder.GetCategoryNames(Categories);
-
-		for (const FName& Category : Categories)
-		{
-			DetailBuilder.HideCategory(Category);
-		}
-
-		// TransformCommon is a custom category that doesn't get returned by GetCategoryNames that also needs to be hidden
-		DetailBuilder.HideCategory(TEXT("TransformCommon"));
-
-		TSet<FName> FinalCategories;
-		for (const TSharedPtr<IPropertyHandle>& PropertyHandle : PropertyHandles)
-		{
-			FName CategoryOverride = PropertyHandle->GetDefaultCategoryName();
-			FText CategoryOverrideText = PropertyHandle->GetDefaultCategoryText();
-			if (const FString* CategoryOverridePtr = PropertyHandle->GetInstanceMetaData(TEXT("CategoryOverride")))
-			{
-				CategoryOverride = FName(*CategoryOverridePtr + "_Override");
-				CategoryOverrideText = FText::FromString(*CategoryOverridePtr);
-			}
-
-			IDetailCategoryBuilder& CategoryBuilder = DetailBuilder.EditCategory(CategoryOverride, CategoryOverrideText);
-			CategoryBuilder.AddProperty(PropertyHandle);
-
-			FinalCategories.Add(CategoryOverride);
-		}
-
-		DetailBuilder.SortCategories([this, FinalCategories = MoveTemp(FinalCategories)](const TMap<FName, IDetailCategoryBuilder*>& CategoryMap)
-		{
-			for (const FName& Category : FinalCategories)
-			{
-				if (CategorySortOrder.Contains(Category))
-				{
-					CategoryMap[Category]->SetSortOrder(CategorySortOrder[Category]);
-				}
-			}
-		});
-	}
-
-private:
-	TArray<TSharedPtr<IPropertyHandle>> PropertyHandles;
-	TMap<FName, int32> CategorySortOrder;
-};
-
-void SDisplayClusterColorGradingColorWheelPanel::FillDetailsView(const FDisplayClusterColorGradingDataModel::FColorGradingGroup& ColorGradingGroup)
-{
-	if (DetailsView.IsValid() && ColorGradingGroup.DetailsViewPropertyHandles.Num())
-	{
-		TArray<UObject*> OuterObjects;
-
-		if (ColorGradingGroup.GroupPropertyHandle.IsValid())
-		{
-			ColorGradingGroup.GroupPropertyHandle->GetOuterObjects(OuterObjects);
-		}
-		else if (ColorGradingGroup.DetailsViewPropertyHandles[0].IsValid())
-		{
-			ColorGradingGroup.DetailsViewPropertyHandles[0]->GetOuterObjects(OuterObjects);
-		}
-
-		if (OuterObjects.Num())
-		{
-			DetailsView->RegisterInstancedCustomPropertyLayout(OuterObjects[0]->GetClass(), FOnGetDetailCustomizationInstance::CreateLambda([&ColorGradingGroup]()
-			{
-				return MakeShareable(new FColorGradingDetailsCustomization(ColorGradingGroup));
-			}));
-
-			const bool bForceRefresh = true;
-			DetailsView->SetObjects(OuterObjects, bForceRefresh);
-		}
-	}
-}
-
-void SDisplayClusterColorGradingColorWheelPanel::ClearDetailsView()
-{
-	TArray<UObject*> EmptyList;
-	DetailsView->SetObjects(EmptyList);
-}
-
 void SDisplayClusterColorGradingColorWheelPanel::FillColorWheels(const FDisplayClusterColorGradingDataModel::FColorGradingElement& ColorGradingElement)
 {
 	auto FillColorWheel = [this](const TSharedPtr<SDisplayClusterColorGradingColorWheel>& ColorWheel, const TSharedPtr<IPropertyHandle>& PropertyHandle)
@@ -565,6 +474,28 @@ void SDisplayClusterColorGradingColorWheelPanel::ClearColorWheels()
 	};
 }
 
+bool SDisplayClusterColorGradingColorWheelPanel::FilterDetailTreeNode(const TSharedRef<IDetailTreeNode>& InDetailTreeNode)
+{
+	if (ColorGradingDataModel)
+	{
+		if (FDisplayClusterColorGradingDataModel::FColorGradingGroup* ColorGradingGroup = ColorGradingDataModel->GetSelectedColorGradingGroup())
+		{
+			// Filter out any categories that are not configured by the data model to be displayed in the details section or subsection.
+			// All other nodes (which will be any child of the category), should be displayed.
+			if (InDetailTreeNode->GetNodeType() == EDetailNodeType::Category)
+			{
+				return ColorGradingGroup->DetailsViewCategories.Contains(InDetailTreeNode->GetNodeName());
+			}
+			else
+			{
+				return true;
+			}
+		}
+	}
+
+	return false;
+}
+
 void SDisplayClusterColorGradingColorWheelPanel::SetColorWheelOrientation(EOrientation NewOrientation)
 {
 	if (ColorWheelOrientation != NewOrientation)
@@ -594,7 +525,7 @@ void SDisplayClusterColorGradingColorWheelPanel::ToggleColorWheelVisible(int32 C
 	}
 }
 
-bool SDisplayClusterColorGradingColorWheelPanel::IsColorWheelVisibile(int32 ColorWheelIndex)
+bool SDisplayClusterColorGradingColorWheelPanel::IsColorWheelVisible(int32 ColorWheelIndex)
 {
 	if (ColorWheelIndex >= 0 && ColorWheelIndex < HiddenColorWheels.Num())
 	{
diff --git a/Engine/Plugins/Runtime/nDisplay/Source/DisplayClusterColorGrading/Private/SDisplayClusterColorGradingColorWheelPanel.h b/Engine/Plugins/Runtime/nDisplay/Source/DisplayClusterColorGrading/Private/SDisplayClusterColorGradingColorWheelPanel.h
--- a/Engine/Plugins/Runtime/nDisplay/Source/DisplayClusterColorGrading/Private/SDisplayClusterColorGradingColorWheelPanel.h
+++ b/Engine/Plugins/Runtime/nDisplay/Source/DisplayClusterColorGrading/Private/SDisplayClusterColorGradingColorWheelPanel.h
@@ -10,10 +10,9 @@
 #include "Widgets/SCompoundWidget.h"
 
 class IDetailTreeNode;
-class IDetaisView;
-class IPropertyHandle;
 class IPropertyRowGenerator;
 class SBox;
+class SDisplayClusterColorGradingDetailView;
 class SHorizontalBox;
 struct FDisplayClusterColorGradingDrawerState;
 
@@ -27,7 +26,7 @@ private:
 	static const uint32 NumColorWheels = 5;
 
 public:
-	~SDisplayClusterColorGradingColorWheelPanel();
+	virtual ~SDisplayClusterColorGradingColorWheelPanel() override;
 
 	SLATE_BEGIN_ARGS(SDisplayClusterColorGradingColorWheelPanel) {}
 		SLATE_ARGUMENT(TSharedPtr<FDisplayClusterColorGradingDataModel>, ColorGradingDataModelSource)
@@ -56,17 +55,16 @@ private:
 	void FillColorGradingElementsToolBar(const TArray<FDisplayClusterColorGradingDataModel::FColorGradingElement>& ColorGradingElements);
 	void ClearColorGradingElementsToolBar();
 
-	void FillDetailsView(const FDisplayClusterColorGradingDataModel::FColorGradingGroup& ColorGradingGroup);
-	void ClearDetailsView();
-
 	void FillColorWheels(const FDisplayClusterColorGradingDataModel::FColorGradingElement& ColorGradingElement);
 	void ClearColorWheels();
 
+	bool FilterDetailTreeNode(const TSharedRef<IDetailTreeNode>& InDetailTreeNode);
+
 	void SetColorWheelOrientation(EOrientation NewOrientation);
 	bool IsColorWheelOrientationSelected(EOrientation Orientation) const;
 
 	void ToggleColorWheelVisible(int32 ColorWheelIndex);
-	bool IsColorWheelVisibile(int32 ColorWheelIndex);
+	bool IsColorWheelVisible(int32 ColorWheelIndex);
 
 	void OnColorGradingGroupSelectionChanged();
 	void OnColorGradingElementSelectionChanged();
@@ -93,7 +91,7 @@ private:
 	TArray<TSharedPtr<SDisplayClusterColorGradingColorWheel>> ColorWheels;
 	TArray<bool> HiddenColorWheels;
 
-	TSharedPtr<IDetailsView> DetailsView;
+	TSharedPtr<SDisplayClusterColorGradingDetailView> DetailView;
 
 	/** Commands used by the color wheel panel */
 	TSharedPtr<FUICommandList> CommandList;
@@ -102,8 +100,8 @@ private:
 	int32 SelectedColorGradingGroup = INDEX_NONE;
 
 	/** The current color display mode for the color wheels */
-	EColorDisplayMode ColorDisplayMode;
+	EColorDisplayMode ColorDisplayMode = EColorDisplayMode::RGB;
 
-	/** The current orientation for the color wheeels */
-	EOrientation ColorWheelOrientation;
+	/** The current orientation for the color wheels */
+	EOrientation ColorWheelOrientation = EOrientation::Orient_Vertical;
 };
\ No newline at end of file
diff --git a/Engine/Plugins/Runtime/nDisplay/Source/DisplayClusterColorGrading/Private/SDisplayClusterColorGradingDetailsPanel.cpp b/Engine/Plugins/Runtime/nDisplay/Source/DisplayClusterColorGrading/Private/SDisplayClusterColorGradingDetailsPanel.cpp
--- a/Engine/Plugins/Runtime/nDisplay/Source/DisplayClusterColorGrading/Private/SDisplayClusterColorGradingDetailsPanel.cpp
+++ b/Engine/Plugins/Runtime/nDisplay/Source/DisplayClusterColorGrading/Private/SDisplayClusterColorGradingDetailsPanel.cpp
@@ -3,10 +3,10 @@
 #include "SDisplayClusterColorGradingDetailsPanel.h"
 
 #include "DisplayClusterColorGradingDataModel.h"
-#include "Drawer/SDisplayClusterColorGradingDrawer.h"
+#include "Drawer/DisplayClusterColorGradingDrawerState.h"
+#include "DetailView/SDisplayClusterColorGradingDetailView.h"
 
-#include "DetailCategoryBuilder.h"
-#include "IDetailsView.h"
+#include "IDetailTreeNode.h"
 #include "Modules/ModuleManager.h"
 #include "PropertyEditorModule.h"
 #include "Widgets/SBoxPanel.h"
@@ -19,19 +19,11 @@
 class SDetailsSectionView : public SCompoundWidget
 {
 public:
-	SLATE_BEGIN_ARGS(SDetailsSectionView)
-	{}
+	SLATE_BEGIN_ARGS(SDetailsSectionView) { }
 	SLATE_END_ARGS()
 
-	void Construct(const FArguments& InArgs)
+	void Construct(const FArguments& InArgs, TSharedPtr<IPropertyRowGenerator> PropertyRowGenerator)
 	{
-		FPropertyEditorModule& EditModule = FModuleManager::Get().GetModuleChecked<FPropertyEditorModule>("PropertyEditor");
-		FDetailsViewArgs DetailsViewArgs;
-		DetailsViewArgs.bAllowSearch = false;
-		DetailsViewArgs.NameAreaSettings = FDetailsViewArgs::HideNameArea;
-		DetailsViewArgs.bHideSelectionTip = true;
-		DetailsView = EditModule.CreateDetailView(DetailsViewArgs);
-
 		ChildSlot
 		[
 			SNew(SVerticalBox)
@@ -47,7 +39,9 @@ public:
 			.FillHeight(1.0f)
 			.Padding(2, 0)
 			[
-				DetailsView.ToSharedRef()
+				SAssignNew(DetailView, SDisplayClusterColorGradingDetailView)
+				.PropertyRowGeneratorSource(PropertyRowGenerator)
+				.OnFilterDetailTreeNode(this, &SDetailsSectionView::FilterDetailTreeNode)
 			]
 		];
 	}
@@ -56,18 +50,6 @@ public:
 	void FillDetailsSection(const TArray<TWeakObjectPtr<UObject>>& Objects, const FDisplayClusterColorGradingDataModel::FDetailsSection& InDetailsSection)
 	{
 		DetailsSection = &InDetailsSection;
-		ObjectClass = nullptr;
-
-		if (Objects.Num())
-		{
-			ObjectClass = UObject::StaticClass();
-
-			const TWeakObjectPtr<UObject>& Object = Objects[0];
-			if (Object.IsValid())
-			{
-				ObjectClass = Object->GetClass();
-			}
-		}
 
 		// If a details subsection isn't already selected and the details section has subsections, set the selected subsection to the first subsection
 		if (SelectedSubsectionIndex == INDEX_NONE && DetailsSection->Subsections.Num())
@@ -75,28 +57,18 @@ public:
 			SelectedSubsectionIndex = 0;
 		}
 
-		if (ObjectClass)
+		HeaderBox->SetContent(CreateHeader());
+
+		if (DetailView.IsValid())
 		{
-			// If a subsection is selected, use the details customization for that subsection. Otherwise, use the details customization for the entire section
-			if (SelectedSubsectionIndex > INDEX_NONE && SelectedSubsectionIndex < DetailsSection->Subsections.Num())
-			{
-				DetailsView->RegisterInstancedCustomPropertyLayout(ObjectClass, DetailsSection->Subsections[SelectedSubsectionIndex].DetailCustomizationDelegate);
-			}
-			else
-			{
-				DetailsView->RegisterInstancedCustomPropertyLayout(ObjectClass, DetailsSection->DetailCustomizationDelegate);
-			}
+			DetailView->Refresh();
 		}
-
-		DetailsView->SetObjects(Objects);
-		HeaderBox->SetContent(CreateHeader());
 	}
 
 	/** Clears the details view and header */
 	void EmptyDetailsSection()
 	{
 		DetailsSection = nullptr;
-		ObjectClass = nullptr;
 		SelectedSubsectionIndex = INDEX_NONE;
 
 		if (HeaderBox.IsValid())
@@ -104,10 +76,9 @@ public:
 			HeaderBox->SetContent(SNullWidget::NullWidget);
 		}
 
-		if (DetailsView.IsValid())
+		if (DetailView.IsValid())
 		{
-			TArray<UObject*> EmptyList;
-			DetailsView->SetObjects(EmptyList);
+			DetailView->Refresh();
 		}
 	}
 
@@ -123,23 +94,7 @@ public:
 		if (DetailsSection)
 		{
 			SelectedSubsectionIndex = InSubsectionIndex;
-
-			if (ObjectClass)
-			{
-				DetailsView->UnregisterInstancedCustomPropertyLayout(ObjectClass);
-
-				// If a subsection is selected, use the details customization for that subsection. Otherwise, use the details customization for the entire section
-				if (SelectedSubsectionIndex > INDEX_NONE && SelectedSubsectionIndex < DetailsSection->Subsections.Num())
-				{
-					DetailsView->RegisterInstancedCustomPropertyLayout(ObjectClass, DetailsSection->Subsections[SelectedSubsectionIndex].DetailCustomizationDelegate);
-				}
-				else
-				{
-					DetailsView->RegisterInstancedCustomPropertyLayout(ObjectClass, DetailsSection->DetailCustomizationDelegate);
-				}
-			}
-
-			DetailsView->ForceRefresh();
+			DetailView->Refresh();
 		}
 	}
 
@@ -230,16 +185,40 @@ private:
 		}
 	}
 
+	/** Filters the detail tree nodes to display for the details section */
+	bool FilterDetailTreeNode(const TSharedRef<IDetailTreeNode>& InDetailTreeNode)
+	{
+		if (DetailsSection)
+		{
+			// Filter out any categories that are not configured by the data model to be displayed in the details section or subsection.
+			// All other nodes (which will be any child of the category), should be displayed.
+			if (InDetailTreeNode->GetNodeType() == EDetailNodeType::Category)
+			{
+				if (DetailsSection->Subsections.Num() && SelectedSubsectionIndex > INDEX_NONE)
+				{
+					return DetailsSection->Subsections[SelectedSubsectionIndex].Categories.Contains(InDetailTreeNode->GetNodeName());
+				}
+				else
+				{
+					return DetailsSection->Categories.Contains(InDetailTreeNode->GetNodeName());
+				}
+			}
+			else
+			{
+				return true;
+			}
+		}
+
+		return false;
+	}
+
 private:
-	TSharedPtr<IDetailsView> DetailsView;
 	TSharedPtr<SBox> HeaderBox;
+	TSharedPtr<SDisplayClusterColorGradingDetailView> DetailView;
 
 	/** Pointer to the details section this view represents */
 	const FDisplayClusterColorGradingDataModel::FDetailsSection* DetailsSection = nullptr;
 
-	/** The object class of the objects that the details view is currently displaying */
-	UClass* ObjectClass = nullptr;
-
 	/** The index of the currently selected details subsection */
 	int32 SelectedSubsectionIndex = INDEX_NONE;
 };
@@ -259,7 +238,7 @@ void SDisplayClusterColorGradingDetailsPanel::Construct(const FArguments& InArgs
 	{
 		Splitter->AddSlot()
 		[
-			SAssignNew(DetailsSectionViews[Index], SDetailsSectionView)
+			SAssignNew(DetailsSectionViews[Index], SDetailsSectionView, ColorGradingDataModel->GetPropertyRowGenerator())
 			.Visibility(this, &SDisplayClusterColorGradingDetailsPanel::GetDetailsSectionVisibility, Index)
 		];
 	}
@@ -295,11 +274,11 @@ void SDisplayClusterColorGradingDetailsPanel::SetDrawerState(const FDisplayClust
 
 void SDisplayClusterColorGradingDetailsPanel::FillDetailsSections()
 {
+	TRACE_CPUPROFILER_EVENT_SCOPE(SDisplayClusterColorGradingDetailsPanel::FillDetailsSections);
 	if (ColorGradingDataModel.IsValid())
 	{
 		TArray<TWeakObjectPtr<UObject>> Objects = ColorGradingDataModel->GetObjects();
 
-		const int32 NumDetailsSections = FMath::Min(ColorGradingDataModel->DetailsSections.Num(), MaxNumDetailsSections);
 		for (int32 Index = 0; Index < MaxNumDetailsSections; ++Index)
 		{
 			if (Index < ColorGradingDataModel->DetailsSections.Num())
diff --git a/Engine/Plugins/Runtime/nDisplay/Source/DisplayClusterColorGrading/Private/SDisplayClusterColorGradingDetailsPanel.h b/Engine/Plugins/Runtime/nDisplay/Source/DisplayClusterColorGrading/Private/SDisplayClusterColorGradingDetailsPanel.h
--- a/Engine/Plugins/Runtime/nDisplay/Source/DisplayClusterColorGrading/Private/SDisplayClusterColorGradingDetailsPanel.h
+++ b/Engine/Plugins/Runtime/nDisplay/Source/DisplayClusterColorGrading/Private/SDisplayClusterColorGradingDetailsPanel.h
@@ -35,7 +35,7 @@ public:
 	void SetDrawerState(const FDisplayClusterColorGradingDrawerState& InDrawerState);
 
 private:
-	/** Fills the details sections based on the current state of the color gradin data model */
+	/** Fills the details sections based on the current state of the color grading data model */
 	void FillDetailsSections();
 
 	/** Gets the visibility state of the specified details section */
