Prevent PostEditChange message when interacting with a subobject of the CurrentConfigData.  This helps improve performance of the ICFVX panel when manipulating color settings.

diff --git a/Engine/Plugins/Runtime/nDisplay/Source/DisplayCluster/Private/Game/EngineClasses/Scene/DisplayClusterRootActorEditor.cpp b/Engine/Plugins/Runtime/nDisplay/Source/DisplayCluster/Private/Game/EngineClasses/Scene/DisplayClusterRootActorEditor.cpp
--- a/Engine/Plugins/Runtime/nDisplay/Source/DisplayCluster/Private/Game/EngineClasses/Scene/DisplayClusterRootActorEditor.cpp
+++ b/Engine/Plugins/Runtime/nDisplay/Source/DisplayCluster/Private/Game/EngineClasses/Scene/DisplayClusterRootActorEditor.cpp
@@ -696,10 +696,36 @@ static FName Name_RelativeLocation = USceneComponent::GetRelativeLocationPropert
 static FName Name_RelativeRotation = USceneComponent::GetRelativeRotationPropertyName();
 static FName Name_RelativeScale3D = USceneComponent::GetRelativeScale3DPropertyName();
 
+void ADisplayClusterRootActor::PostEditChangeChainProperty(FPropertyChangedChainEvent& PropertyChainEvent)
+{
+	const FProperty* RootProperty = PropertyChainEvent.PropertyChain.GetActiveNode()->GetValue();
+	const FProperty* TailProperty = PropertyChainEvent.PropertyChain.GetTail()->GetValue();
+	if (PropertyChainEvent.ChangeType  == EPropertyChangeType::Interactive && RootProperty && TailProperty)
+	{
+		const FName RootPropertyName = RootProperty->GetFName();
+		const FName TailPropertyName = TailProperty->GetFName();
+		if (RootPropertyName == GET_MEMBER_NAME_CHECKED(ADisplayClusterRootActor, CurrentConfigData)
+			&& RootPropertyName != TailPropertyName)
+		{
+			// Do not propagate the PostEditChangeProperty because we are in an interactive edit of a suboject
+			// and we do not need to rerun any construction scripts.
+			bIsInteractiveEditingSubobject = true;
+		}
+	}
+
+	Super::PostEditChangeChainProperty(PropertyChainEvent);
+	bIsInteractiveEditingSubobject = false;
+}
+
 void ADisplayClusterRootActor::PostEditChangeProperty(FPropertyChangedEvent& PropertyChangedEvent)
 {
 	const FName PropertyName = (PropertyChangedEvent.Property != nullptr) ? PropertyChangedEvent.Property->GetFName() : NAME_None;
-	
+
+	if (bIsInteractiveEditingSubobject)
+	{
+		return;
+	}
+
 	// The AActor method, simplified and modified to skip construction scripts.
 	// Component registration still needs to occur or the actor will look like it disappeared.
 	auto SuperCallWithoutConstructionScripts = [&]
@@ -727,7 +753,7 @@ void ADisplayClusterRootActor::PostEditChangeProperty(FPropertyChangedEvent& Pro
 		}
 
 		FEditorSupportDelegates::UpdateUI.Broadcast();
-		UObject::PostEditChangeProperty(PropertyChangedEvent);	
+		UObject::PostEditChangeProperty(PropertyChangedEvent);
 	};
 
 	bool bReinitializeActor = true;
diff --git a/Engine/Plugins/Runtime/nDisplay/Source/DisplayCluster/Public/DisplayClusterRootActor.h b/Engine/Plugins/Runtime/nDisplay/Source/DisplayCluster/Public/DisplayClusterRootActor.h
--- a/Engine/Plugins/Runtime/nDisplay/Source/DisplayCluster/Public/DisplayClusterRootActor.h
+++ b/Engine/Plugins/Runtime/nDisplay/Source/DisplayCluster/Public/DisplayClusterRootActor.h
@@ -471,6 +471,7 @@ protected:
 	void ImplRenderPreviewFrustums_Editor();
 	void ImplRenderPreviewViewportFrustum_Editor(const FMatrix ProjectionMatrix, const FMatrix ViewMatrix, const FVector ViewOrigin);
 
+	virtual void PostEditChangeChainProperty(FPropertyChangedChainEvent& PropertyChainEvent) override;
 	virtual void PostEditChangeProperty(FPropertyChangedEvent& PropertyChangedEvent) override;
 	virtual void PostEditMove(bool bFinished) override;
 
@@ -478,8 +479,12 @@ protected:
 	void HandleAssetReload(const EPackageReloadPhase InPackageReloadPhase, FPackageReloadedEvent* InPackageReloadedEvent);
 	
 private:
+
+	/** Flag to indicate if the user is currently interacting with a subobject of CurrentConfigData. */
+	bool bIsInteractiveEditingSubobject = false;
+
 	bool bIsSelectedInEditor = false;
-	
+
 	/** When the preview render should be directed to the post process render target. */
 	bool bOutputFrameToPostProcessRenderTarget;
 
