nDisplay: Added calls that snapshot a transaction on an object when changing slider or color wheel values in the ICVFX color grading drawer, which propagates the changes to MU while the user is using the slider, allowing updates to be seen realtime on nDisplay nodes

diff --git a/Engine/Plugins/Runtime/nDisplay/Source/DisplayClusterColorGrading/Private/DataModelGenerators/DisplayClusterColorGradingGenerator_RootActor.cpp b/Engine/Plugins/Runtime/nDisplay/Source/DisplayClusterColorGrading/Private/DataModelGenerators/DisplayClusterColorGradingGenerator_RootActor.cpp
--- a/Engine/Plugins/Runtime/nDisplay/Source/DisplayClusterColorGrading/Private/DataModelGenerators/DisplayClusterColorGradingGenerator_RootActor.cpp
+++ b/Engine/Plugins/Runtime/nDisplay/Source/DisplayClusterColorGrading/Private/DataModelGenerators/DisplayClusterColorGradingGenerator_RootActor.cpp
@@ -126,6 +126,30 @@ TSharedPtr<IPropertyHandle> FDisplayClusterColorGradingGenerator_ColorGradingRen
 
 #define CREATE_PROPERTY_PATH(RootObjectClass, PropertyPath) FCachedPropertyPath(GET_MEMBER_NAME_STRING_CHECKED(RootObjectClass, PropertyPath))
 
+TSharedPtr<IPropertyHandle> MakePropertyTransactional(const TSharedPtr<IPropertyHandle>& PropertyHandle)
+{
+	if (PropertyHandle.IsValid())
+	{
+		PropertyHandle->SetOnPropertyValueChanged(FSimpleDelegate::CreateLambda([PropertyHandle]
+		{
+			TArray<UObject*> OuterObjects;
+			PropertyHandle->GetOuterObjects(OuterObjects);
+			for (UObject* Object : OuterObjects)
+			{
+				if (!Object->HasAnyFlags(RF_Transactional))
+				{
+					Object->SetFlags(RF_Transactional);
+				}
+
+				SaveToTransactionBuffer(Object, false);
+				SnapshotTransactionBuffer(Object);
+			}
+		}));
+	}
+
+	return PropertyHandle;
+}
+
 TSharedRef<IDisplayClusterColorGradingDataModelGenerator> FDisplayClusterColorGradingGenerator_RootActor::MakeInstance()
 {
 	return MakeShareable(new FDisplayClusterColorGradingGenerator_RootActor());
@@ -193,28 +217,33 @@ private:
 		auto AddColorGradingSettings = [&DetailBuilder](const TSharedRef<IPropertyHandle>& ColorGradingSettingsHandle)
 		{
 			IDetailCategoryBuilder& DetailExposureCategoryBuilder = DetailBuilder.EditCategory(TEXT("DetailView_Exposure"), LOCTEXT("DetailView_ExposureDisplayName", "Exposure"));
-			DetailExposureCategoryBuilder.AddProperty(ColorGradingSettingsHandle->GetChildHandle(GET_MEMBER_NAME_CHECKED(FDisplayClusterConfigurationViewport_ColorGradingRenderingSettings, AutoExposureBias)));
+			DetailExposureCategoryBuilder.AddProperty(MakePropertyTransactional(ColorGradingSettingsHandle->GetChildHandle(GET_MEMBER_NAME_CHECKED(FDisplayClusterConfigurationViewport_ColorGradingRenderingSettings, AutoExposureBias))));
 
 			IDetailCategoryBuilder& DetailColorGradingCategoryBuilder = DetailBuilder.EditCategory(TEXT("DetailView_ColorGrading"), LOCTEXT("DetailView_ColorGradingDisplayName", "Color Grading"));
-			DetailColorGradingCategoryBuilder.AddProperty(ColorGradingSettingsHandle->GetChildHandle(GET_MEMBER_NAME_CHECKED(FDisplayClusterConfigurationViewport_ColorGradingRenderingSettings, ColorCorrectionShadowsMax)));
-			DetailColorGradingCategoryBuilder.AddProperty(ColorGradingSettingsHandle->GetChildHandle(GET_MEMBER_NAME_CHECKED(FDisplayClusterConfigurationViewport_ColorGradingRenderingSettings, ColorCorrectionHighlightsMin)));
-			DetailColorGradingCategoryBuilder.AddProperty(ColorGradingSettingsHandle->GetChildHandle(GET_MEMBER_NAME_CHECKED(FDisplayClusterConfigurationViewport_ColorGradingRenderingSettings, ColorCorrectionHighlightsMax)));
+			DetailColorGradingCategoryBuilder.AddProperty(MakePropertyTransactional(ColorGradingSettingsHandle->GetChildHandle(GET_MEMBER_NAME_CHECKED(FDisplayClusterConfigurationViewport_ColorGradingRenderingSettings, ColorCorrectionShadowsMax))));
+			DetailColorGradingCategoryBuilder.AddProperty(MakePropertyTransactional(ColorGradingSettingsHandle->GetChildHandle(GET_MEMBER_NAME_CHECKED(FDisplayClusterConfigurationViewport_ColorGradingRenderingSettings, ColorCorrectionHighlightsMin))));
+			DetailColorGradingCategoryBuilder.AddProperty(MakePropertyTransactional(ColorGradingSettingsHandle->GetChildHandle(GET_MEMBER_NAME_CHECKED(FDisplayClusterConfigurationViewport_ColorGradingRenderingSettings, ColorCorrectionHighlightsMax))));
 
 			IDetailCategoryBuilder& DetailWhiteBalanceCategoryBuilder = DetailBuilder.EditCategory(TEXT("DetailView_WhiteBalance"), LOCTEXT("DetailView_WhiteBalanceDisplayName", "White Balance"));
 			DetailWhiteBalanceCategoryBuilder.AddProperty(ColorGradingSettingsHandle->GetChildHandle(GET_MEMBER_NAME_CHECKED(FDisplayClusterConfigurationViewport_ColorGradingWhiteBalanceSettings, TemperatureType)));
-			DetailWhiteBalanceCategoryBuilder.AddProperty(ColorGradingSettingsHandle->GetChildHandle(GET_MEMBER_NAME_CHECKED(FDisplayClusterConfigurationViewport_ColorGradingWhiteBalanceSettings, WhiteTemp)));
-			DetailWhiteBalanceCategoryBuilder.AddProperty(ColorGradingSettingsHandle->GetChildHandle(GET_MEMBER_NAME_CHECKED(FDisplayClusterConfigurationViewport_ColorGradingWhiteBalanceSettings, WhiteTint)));
+			DetailWhiteBalanceCategoryBuilder.AddProperty(MakePropertyTransactional(ColorGradingSettingsHandle->GetChildHandle(GET_MEMBER_NAME_CHECKED(FDisplayClusterConfigurationViewport_ColorGradingWhiteBalanceSettings, WhiteTemp))));
+			DetailWhiteBalanceCategoryBuilder.AddProperty(MakePropertyTransactional(ColorGradingSettingsHandle->GetChildHandle(GET_MEMBER_NAME_CHECKED(FDisplayClusterConfigurationViewport_ColorGradingWhiteBalanceSettings, WhiteTint))));
 
 			IDetailCategoryBuilder& DetailMiscCategoryBuilder = DetailBuilder.EditCategory(TEXT("DetailView_Misc"), LOCTEXT("DetailView_MiscDisplayName", "Misc"));
-			DetailMiscCategoryBuilder.AddProperty(ColorGradingSettingsHandle->GetChildHandle(GET_MEMBER_NAME_CHECKED(FDisplayClusterConfigurationViewport_ColorGradingMiscSettings, BlueCorrection)));
-			DetailMiscCategoryBuilder.AddProperty(ColorGradingSettingsHandle->GetChildHandle(GET_MEMBER_NAME_CHECKED(FDisplayClusterConfigurationViewport_ColorGradingMiscSettings, ExpandGamut)));
+			DetailMiscCategoryBuilder.AddProperty(MakePropertyTransactional(ColorGradingSettingsHandle->GetChildHandle(GET_MEMBER_NAME_CHECKED(FDisplayClusterConfigurationViewport_ColorGradingMiscSettings, BlueCorrection))));
+			DetailMiscCategoryBuilder.AddProperty(MakePropertyTransactional(ColorGradingSettingsHandle->GetChildHandle(GET_MEMBER_NAME_CHECKED(FDisplayClusterConfigurationViewport_ColorGradingMiscSettings, ExpandGamut))));
 			DetailMiscCategoryBuilder.AddProperty(ColorGradingSettingsHandle->GetChildHandle(GET_MEMBER_NAME_CHECKED(FDisplayClusterConfigurationViewport_ColorGradingMiscSettings, SceneColorTint)));
 		};
 
 		UClass* ConfigDataClass = UDisplayClusterConfigurationData::StaticClass();
 		const int32 GroupIndex = ColorGradingDataModel.IsValid() ? ColorGradingDataModel.Pin()->GetSelectedColorGradingGroupIndex() : INDEX_NONE;
+		uint32 ArraySize = 0;
+
+		TSharedRef<IPropertyHandle> ArrayPropertyHandle = DetailBuilder.GetProperty(GET_MEMBER_NAME_CHECKED(UDisplayClusterConfigurationData, StageSettings.PerViewportColorGrading), ConfigDataClass);
+		check(ArrayPropertyHandle->AsArray());
+		ArrayPropertyHandle->AsArray()->GetNumElements(ArraySize);
 
-		if (GroupIndex > 0)
+		if (GroupIndex > 0 && GroupIndex < (int32)ArraySize)
 		{
 			const int32 Index = GroupIndex - 1;
 			IDetailCategoryBuilder& PerNodeSettingsCategoryBuilder = DetailBuilder.EditCategory(TEXT("DetailView_PerViewport"), LOCTEXT("DetailView_PerViewportDisplayName", "Per-Viewport Settings"));
@@ -740,27 +769,32 @@ private:
 		auto AddColorGradingSettings = [&DetailBuilder](const TSharedRef<IPropertyHandle>& ColorGradingSettingsHandle)
 		{
 			IDetailCategoryBuilder& DetailExposureCategoryBuilder = DetailBuilder.EditCategory(TEXT("DetailView_Exposure"), LOCTEXT("DetailView_ExposureDisplayName", "Exposure"));
-			DetailExposureCategoryBuilder.AddProperty(ColorGradingSettingsHandle->GetChildHandle(GET_MEMBER_NAME_CHECKED(FDisplayClusterConfigurationViewport_ColorGradingRenderingSettings, AutoExposureBias)));
+			DetailExposureCategoryBuilder.AddProperty(MakePropertyTransactional(ColorGradingSettingsHandle->GetChildHandle(GET_MEMBER_NAME_CHECKED(FDisplayClusterConfigurationViewport_ColorGradingRenderingSettings, AutoExposureBias))));
 
 			IDetailCategoryBuilder& DetailColorGradingCategoryBuilder = DetailBuilder.EditCategory(TEXT("DetailView_ColorGrading"), LOCTEXT("DetailView_ColorGradingDisplayName", "Color Grading"));
-			DetailColorGradingCategoryBuilder.AddProperty(ColorGradingSettingsHandle->GetChildHandle(GET_MEMBER_NAME_CHECKED(FDisplayClusterConfigurationViewport_ColorGradingRenderingSettings, ColorCorrectionShadowsMax)));
-			DetailColorGradingCategoryBuilder.AddProperty(ColorGradingSettingsHandle->GetChildHandle(GET_MEMBER_NAME_CHECKED(FDisplayClusterConfigurationViewport_ColorGradingRenderingSettings, ColorCorrectionHighlightsMin)));
-			DetailColorGradingCategoryBuilder.AddProperty(ColorGradingSettingsHandle->GetChildHandle(GET_MEMBER_NAME_CHECKED(FDisplayClusterConfigurationViewport_ColorGradingRenderingSettings, ColorCorrectionHighlightsMax)));
+			DetailColorGradingCategoryBuilder.AddProperty(MakePropertyTransactional(ColorGradingSettingsHandle->GetChildHandle(GET_MEMBER_NAME_CHECKED(FDisplayClusterConfigurationViewport_ColorGradingRenderingSettings, ColorCorrectionShadowsMax))));
+			DetailColorGradingCategoryBuilder.AddProperty(MakePropertyTransactional(ColorGradingSettingsHandle->GetChildHandle(GET_MEMBER_NAME_CHECKED(FDisplayClusterConfigurationViewport_ColorGradingRenderingSettings, ColorCorrectionHighlightsMin))));
+			DetailColorGradingCategoryBuilder.AddProperty(MakePropertyTransactional(ColorGradingSettingsHandle->GetChildHandle(GET_MEMBER_NAME_CHECKED(FDisplayClusterConfigurationViewport_ColorGradingRenderingSettings, ColorCorrectionHighlightsMax))));
 
 			IDetailCategoryBuilder& DetailWhiteBalanceCategoryBuilder = DetailBuilder.EditCategory(TEXT("DetailView_WhiteBalance"), LOCTEXT("DetailView_WhiteBalanceDisplayName", "White Balance"));
 			DetailWhiteBalanceCategoryBuilder.AddProperty(ColorGradingSettingsHandle->GetChildHandle(GET_MEMBER_NAME_CHECKED(FDisplayClusterConfigurationViewport_ColorGradingWhiteBalanceSettings, TemperatureType)));
-			DetailWhiteBalanceCategoryBuilder.AddProperty(ColorGradingSettingsHandle->GetChildHandle(GET_MEMBER_NAME_CHECKED(FDisplayClusterConfigurationViewport_ColorGradingWhiteBalanceSettings, WhiteTemp)));
-			DetailWhiteBalanceCategoryBuilder.AddProperty(ColorGradingSettingsHandle->GetChildHandle(GET_MEMBER_NAME_CHECKED(FDisplayClusterConfigurationViewport_ColorGradingWhiteBalanceSettings, WhiteTint)));
+			DetailWhiteBalanceCategoryBuilder.AddProperty(MakePropertyTransactional(ColorGradingSettingsHandle->GetChildHandle(GET_MEMBER_NAME_CHECKED(FDisplayClusterConfigurationViewport_ColorGradingWhiteBalanceSettings, WhiteTemp))));
+			DetailWhiteBalanceCategoryBuilder.AddProperty(MakePropertyTransactional(ColorGradingSettingsHandle->GetChildHandle(GET_MEMBER_NAME_CHECKED(FDisplayClusterConfigurationViewport_ColorGradingWhiteBalanceSettings, WhiteTint))));
 
 			IDetailCategoryBuilder& DetailMiscCategoryBuilder = DetailBuilder.EditCategory(TEXT("DetailView_Misc"), LOCTEXT("DetailView_MiscDisplayName", "Misc"));
-			DetailMiscCategoryBuilder.AddProperty(ColorGradingSettingsHandle->GetChildHandle(GET_MEMBER_NAME_CHECKED(FDisplayClusterConfigurationViewport_ColorGradingMiscSettings, BlueCorrection)));
-			DetailMiscCategoryBuilder.AddProperty(ColorGradingSettingsHandle->GetChildHandle(GET_MEMBER_NAME_CHECKED(FDisplayClusterConfigurationViewport_ColorGradingMiscSettings, ExpandGamut)));
+			DetailMiscCategoryBuilder.AddProperty(MakePropertyTransactional(ColorGradingSettingsHandle->GetChildHandle(GET_MEMBER_NAME_CHECKED(FDisplayClusterConfigurationViewport_ColorGradingMiscSettings, BlueCorrection))));
+			DetailMiscCategoryBuilder.AddProperty(MakePropertyTransactional(ColorGradingSettingsHandle->GetChildHandle(GET_MEMBER_NAME_CHECKED(FDisplayClusterConfigurationViewport_ColorGradingMiscSettings, ExpandGamut))));
 			DetailMiscCategoryBuilder.AddProperty(ColorGradingSettingsHandle->GetChildHandle(GET_MEMBER_NAME_CHECKED(FDisplayClusterConfigurationViewport_ColorGradingMiscSettings, SceneColorTint)));
 		};
 
 		const int32 GroupIndex = ColorGradingDataModel.IsValid() ? ColorGradingDataModel.Pin()->GetSelectedColorGradingGroupIndex() : INDEX_NONE;
+		uint32 ArraySize = 0;
+
+		TSharedPtr<IPropertyHandle> ArrayPropertyHandle = DetailBuilder.GetProperty(GET_MEMBER_NAME_CHECKED(UDisplayClusterICVFXCameraComponent, CameraSettings.PerNodeColorGrading));
+		check(ArrayPropertyHandle->AsArray());
+		ArrayPropertyHandle->AsArray()->GetNumElements(ArraySize);
 
-		if (GroupIndex > 0)
+		if (GroupIndex > 0 && GroupIndex < (int32)ArraySize)
 		{
 			const int32 Index = GroupIndex - 1;
 			IDetailCategoryBuilder& PerNodeSettingsCategoryBuilder = DetailBuilder.EditCategory(TEXT("DetailView_PerNode"), LOCTEXT("DetailView_PerNodeDisplayName", "Per-Node Settings"));
diff --git a/Engine/Plugins/Runtime/nDisplay/Source/DisplayClusterColorGrading/Private/SDisplayClusterColorGradingColorWheel.cpp b/Engine/Plugins/Runtime/nDisplay/Source/DisplayClusterColorGrading/Private/SDisplayClusterColorGradingColorWheel.cpp
--- a/Engine/Plugins/Runtime/nDisplay/Source/DisplayClusterColorGrading/Private/SDisplayClusterColorGradingColorWheel.cpp
+++ b/Engine/Plugins/Runtime/nDisplay/Source/DisplayClusterColorGrading/Private/SDisplayClusterColorGradingColorWheel.cpp
@@ -368,6 +368,27 @@ void SDisplayClusterColorGradingColorWheel::CommitColor(FVector4& NewValue, bool
 		{
 			ColorPropertyHandle.Pin()->SetValue(NewValue, EPropertyValueSetFlags::DefaultFlags);
 		}
+
+		TransactColorValue();
+	}
+}
+
+void SDisplayClusterColorGradingColorWheel::TransactColorValue()
+{
+	if (ColorPropertyHandle.IsValid())
+	{
+		TArray<UObject*> OuterObjects;
+		ColorPropertyHandle.Pin()->GetOuterObjects(OuterObjects);
+		for (UObject* Object : OuterObjects)
+		{
+			if (!Object->HasAnyFlags(RF_Transactional))
+			{
+				Object->SetFlags(RF_Transactional);
+			}
+
+			SaveToTransactionBuffer(Object, false);
+			SnapshotTransactionBuffer(Object);
+		}
 	}
 }
 
@@ -476,6 +497,7 @@ void SDisplayClusterColorGradingColorWheel::SetComponentValue(float NewValue, ui
 			}
 
 			ColorPropertyHandle.Pin()->SetValue(NewColorValue, bIsUsingComponentSlider ? EPropertyValueSetFlags::InteractiveChange : EPropertyValueSetFlags::DefaultFlags);
+			TransactColorValue();
 		}
 	}
 }
diff --git a/Engine/Plugins/Runtime/nDisplay/Source/DisplayClusterColorGrading/Private/SDisplayClusterColorGradingColorWheel.h b/Engine/Plugins/Runtime/nDisplay/Source/DisplayClusterColorGrading/Private/SDisplayClusterColorGradingColorWheel.h
--- a/Engine/Plugins/Runtime/nDisplay/Source/DisplayClusterColorGrading/Private/SDisplayClusterColorGradingColorWheel.h
+++ b/Engine/Plugins/Runtime/nDisplay/Source/DisplayClusterColorGrading/Private/SDisplayClusterColorGradingColorWheel.h
@@ -66,6 +66,7 @@ private:
 
 	bool GetColor(FVector4& OutCurrentColor);
 	void CommitColor(FVector4& NewValue, bool bShouldCommitValueChanges);
+	void TransactColorValue();
 
 	void BeginUsingColorPickerSlider();
 	void EndUsingColorPickerSlider();
@@ -119,7 +120,7 @@ private:
 	TAttribute<EColorDisplayMode> ColorDisplayMode;
 
 	/** The orientation of the color wheel, which determines whether the color component sliders are next to or above the color picker */
-	EOrientation Orientation;
+	EOrientation Orientation = Orient_Vertical;
 
 	/** Stored current min value of the color component numeric sliders */
 	TOptional<float> ComponentSliderDynamicMinValue;
