Fixed non-unity/pch compile errors caused by header removals in engine

diff --git a/Engine/Plugins/Animation/ControlRig/Source/ControlRig/Private/ControlRig.cpp b/Engine/Plugins/Animation/ControlRig/Source/ControlRig/Private/ControlRig.cpp
--- a/Engine/Plugins/Animation/ControlRig/Source/ControlRig/Private/ControlRig.cpp
+++ b/Engine/Plugins/Animation/ControlRig/Source/ControlRig/Private/ControlRig.cpp
@@ -18,6 +18,7 @@
 #include "Animation/AnimInstance.h"
 #include "Animation/AnimationPoseData.h"
 #include "Internationalization/TextKey.h"
+#include "UObject/Package.h"
 #if WITH_EDITOR
 #include "ControlRigModule.h"
 #include "Modules/ModuleManager.h"
diff --git a/Engine/Plugins/Developer/RenderDocPlugin/Source/RenderDocPlugin/Private/RenderDocPluginModule.cpp b/Engine/Plugins/Developer/RenderDocPlugin/Source/RenderDocPlugin/Private/RenderDocPluginModule.cpp
--- a/Engine/Plugins/Developer/RenderDocPlugin/Source/RenderDocPlugin/Private/RenderDocPluginModule.cpp
+++ b/Engine/Plugins/Developer/RenderDocPlugin/Source/RenderDocPlugin/Private/RenderDocPluginModule.cpp
@@ -2,6 +2,7 @@
 
 #include "RenderDocPluginModule.h"
 #include "Framework/Application/SlateApplication.h"
+#include "UnrealClient.h"
 
 #if PLATFORM_WINDOWS
 #include "Windows/AllowWindowsPlatformTypes.h"
diff --git a/Engine/Plugins/Enterprise/GLTFExporter/Source/GLTFExporter/Private/Exporters/GLTFExporterUtility.cpp b/Engine/Plugins/Enterprise/GLTFExporter/Source/GLTFExporter/Private/Exporters/GLTFExporterUtility.cpp
--- a/Engine/Plugins/Enterprise/GLTFExporter/Source/GLTFExporter/Private/Exporters/GLTFExporterUtility.cpp
+++ b/Engine/Plugins/Enterprise/GLTFExporter/Source/GLTFExporter/Private/Exporters/GLTFExporterUtility.cpp
@@ -1,6 +1,7 @@
 // Copyright Epic Games, Inc. All Rights Reserved.
 
 #include "Exporters/GLTFExporterUtility.h"
+#include "AssetRegistry/AssetData.h"
 #include "Materials/MaterialInstance.h"
 #include "AssetRegistry/AssetRegistryModule.h"
 #include "Animation/AnimSequence.h"
diff --git a/Engine/Plugins/Enterprise/VariantManagerContent/Source/VariantManagerContent/Private/LevelVariantSets.cpp b/Engine/Plugins/Enterprise/VariantManagerContent/Source/VariantManagerContent/Private/LevelVariantSets.cpp
--- a/Engine/Plugins/Enterprise/VariantManagerContent/Source/VariantManagerContent/Private/LevelVariantSets.cpp
+++ b/Engine/Plugins/Enterprise/VariantManagerContent/Source/VariantManagerContent/Private/LevelVariantSets.cpp
@@ -4,6 +4,7 @@
 
 #include "Engine/BlueprintGeneratedClass.h"
 #include "Engine/World.h"
+#include "UObject/Package.h"
 #include "VariantSet.h"
 #include "Variant.h"
 #include "VariantObjectBinding.h"
diff --git a/Engine/Plugins/Experimental/VirtualCamera/Source/VCamCore/Private/VCamComponent.cpp b/Engine/Plugins/Experimental/VirtualCamera/Source/VCamCore/Private/VCamComponent.cpp
--- a/Engine/Plugins/Experimental/VirtualCamera/Source/VCamCore/Private/VCamComponent.cpp
+++ b/Engine/Plugins/Experimental/VirtualCamera/Source/VCamCore/Private/VCamComponent.cpp
@@ -3,6 +3,7 @@
 #include "VCamComponent.h"
 
 #include "CineCameraComponent.h"
+#include "Engine/GameInstance.h"
 #include "Features/IModularFeatures.h"
 #include "Framework/Application/SlateApplication.h"
 #include "ILiveLinkClient.h"
diff --git a/Engine/Plugins/Experimental/VirtualCamera/Source/VCamCore/Private/VCamOutputRemoteSession.cpp b/Engine/Plugins/Experimental/VirtualCamera/Source/VCamCore/Private/VCamOutputRemoteSession.cpp
--- a/Engine/Plugins/Experimental/VirtualCamera/Source/VCamCore/Private/VCamOutputRemoteSession.cpp
+++ b/Engine/Plugins/Experimental/VirtualCamera/Source/VCamCore/Private/VCamOutputRemoteSession.cpp
@@ -1,6 +1,7 @@
 // Copyright Epic Games, Inc. All Rights Reserved.
 
 #include "VCamOutputRemoteSession.h"
+#include "UObject/Package.h"
 #include "VCamOutputComposure.h"
 #include "Channels/RemoteSessionImageChannel.h"
 #include "Channels/RemoteSessionInputChannel.h"
diff --git a/Engine/Plugins/Experimental/VirtualCamera/Source/VirtualCamera/Private/LevelSequencePlaybackController.cpp b/Engine/Plugins/Experimental/VirtualCamera/Source/VirtualCamera/Private/LevelSequencePlaybackController.cpp
--- a/Engine/Plugins/Experimental/VirtualCamera/Source/VirtualCamera/Private/LevelSequencePlaybackController.cpp
+++ b/Engine/Plugins/Experimental/VirtualCamera/Source/VirtualCamera/Private/LevelSequencePlaybackController.cpp
@@ -2,6 +2,7 @@
 
 #include "LevelSequencePlaybackController.h"
 
+#include "AssetRegistry/AssetData.h"
 #include "AssetRegistry/AssetRegistryModule.h"
 #include "AssetRegistry/IAssetRegistry.h"
 #include "HAL/FileManager.h"
diff --git a/Engine/Plugins/Experimental/VirtualCamera/Source/VirtualCamera/Private/VCamBlueprintFunctionLibrary.cpp b/Engine/Plugins/Experimental/VirtualCamera/Source/VirtualCamera/Private/VCamBlueprintFunctionLibrary.cpp
--- a/Engine/Plugins/Experimental/VirtualCamera/Source/VirtualCamera/Private/VCamBlueprintFunctionLibrary.cpp
+++ b/Engine/Plugins/Experimental/VirtualCamera/Source/VirtualCamera/Private/VCamBlueprintFunctionLibrary.cpp
@@ -4,6 +4,7 @@
 
 #include "CineCameraActor.h"
 #include "CineCameraComponent.h"
+#include "Engine/GameInstance.h"
 #include "LevelSequence.h"
 #include "VirtualCameraClipsMetaData.h"
 #include "Components/SceneCaptureComponent2D.h"
diff --git a/Engine/Plugins/Experimental/VirtualProductionUtilities/Source/VPUtilities/Private/VPFullScreenUserWidget.cpp b/Engine/Plugins/Experimental/VirtualProductionUtilities/Source/VPUtilities/Private/VPFullScreenUserWidget.cpp
--- a/Engine/Plugins/Experimental/VirtualProductionUtilities/Source/VPUtilities/Private/VPFullScreenUserWidget.cpp
+++ b/Engine/Plugins/Experimental/VirtualProductionUtilities/Source/VPUtilities/Private/VPFullScreenUserWidget.cpp
@@ -4,6 +4,7 @@
 
 #include "Components/PostProcessComponent.h"
 #include "Engine/Engine.h"
+#include "Engine/GameInstance.h"
 #include "Engine/GameViewportClient.h"
 #include "Engine/TextureRenderTarget2D.h"
 #include "Engine/UserInterfaceSettings.h"
diff --git a/Engine/Plugins/FX/Niagara/Source/Niagara/Private/NiagaraComponentRendererProperties.cpp b/Engine/Plugins/FX/Niagara/Source/Niagara/Private/NiagaraComponentRendererProperties.cpp
--- a/Engine/Plugins/FX/Niagara/Source/Niagara/Private/NiagaraComponentRendererProperties.cpp
+++ b/Engine/Plugins/FX/Niagara/Source/Niagara/Private/NiagaraComponentRendererProperties.cpp
@@ -6,6 +6,7 @@
 #include "NiagaraRendererComponents.h"
 #include "NiagaraComponent.h"
 #include "Modules/ModuleManager.h"
+#include "UObject/Package.h"
 #if WITH_EDITOR
 #include "Widgets/Images/SImage.h"
 #include "Styling/SlateIconFinder.h"
diff --git a/Engine/Plugins/Importers/USDImporter/Source/USDStage/Private/USDStageActor.cpp b/Engine/Plugins/Importers/USDImporter/Source/USDStage/Private/USDStageActor.cpp
--- a/Engine/Plugins/Importers/USDImporter/Source/USDStage/Private/USDStageActor.cpp
+++ b/Engine/Plugins/Importers/USDImporter/Source/USDStage/Private/USDStageActor.cpp
@@ -5,6 +5,7 @@
 #include "Engine/Blueprint.h"
 #include "Engine/BlueprintGeneratedClass.h"
 #include "UnrealUSDWrapper.h"
+#include "UObject/Package.h"
 #include "USDAssetCache.h"
 #include "USDAssetImportData.h"
 #include "USDClassesModule.h"
diff --git a/Engine/Plugins/Runtime/nDisplay/Source/DisplayCluster/Private/Render/Viewport/LightCard/DisplayClusterViewportLightCardManager.cpp b/Engine/Plugins/Runtime/nDisplay/Source/DisplayCluster/Private/Render/Viewport/LightCard/DisplayClusterViewportLightCardManager.cpp
--- a/Engine/Plugins/Runtime/nDisplay/Source/DisplayCluster/Private/Render/Viewport/LightCard/DisplayClusterViewportLightCardManager.cpp
+++ b/Engine/Plugins/Runtime/nDisplay/Source/DisplayCluster/Private/Render/Viewport/LightCard/DisplayClusterViewportLightCardManager.cpp
@@ -12,6 +12,7 @@
 
 #include "PreviewScene.h"
 #include "SceneInterface.h"
+#include "UObject/Package.h"
 
 /** Console variable used to control the size of the UV light card map texture */
 static TAutoConsoleVariable<int32> CVarUVLightCardTextureSize(
diff --git a/Engine/Plugins/Runtime/nDisplay/Source/DisplayClusterScenePreview/Private/DisplayClusterLightCardEditorHelper.cpp b/Engine/Plugins/Runtime/nDisplay/Source/DisplayClusterScenePreview/Private/DisplayClusterLightCardEditorHelper.cpp
--- a/Engine/Plugins/Runtime/nDisplay/Source/DisplayClusterScenePreview/Private/DisplayClusterLightCardEditorHelper.cpp
+++ b/Engine/Plugins/Runtime/nDisplay/Source/DisplayClusterScenePreview/Private/DisplayClusterLightCardEditorHelper.cpp
@@ -21,6 +21,7 @@
 #include "Containers/ArrayView.h"
 #include "Engine/StaticMesh.h"
 #include "Engine/Texture2D.h"
+#include "UObject/Package.h"
 #include "UObject/UObjectGlobals.h"
 
 #if WITH_EDITOR
diff --git a/Engine/Source/Runtime/Engine/Private/AssetManager.cpp b/Engine/Source/Runtime/Engine/Private/AssetManager.cpp
--- a/Engine/Source/Runtime/Engine/Private/AssetManager.cpp
+++ b/Engine/Source/Runtime/Engine/Private/AssetManager.cpp
@@ -3,6 +3,7 @@
 #include "Engine/AssetManager.h"
 #include "Engine/AssetManagerSettings.h"
 #include "Engine/BlueprintCore.h"
+#include "Engine/GameInstance.h"
 #include "Engine/PrimaryAssetLabel.h"
 #include "AssetRegistry/AssetData.h"
 #include "AssetRegistry/ARFilter.h"
diff --git a/Engine/Source/Runtime/Engine/Private/NetDriver.cpp b/Engine/Source/Runtime/Engine/Private/NetDriver.cpp
--- a/Engine/Source/Runtime/Engine/Private/NetDriver.cpp
+++ b/Engine/Source/Runtime/Engine/Private/NetDriver.cpp
@@ -5,6 +5,7 @@
 =============================================================================*/
 
 #include "Engine/NetDriver.h"
+#include "Engine/GameInstance.h"
 #include "Misc/CoreMisc.h"
 #include "Misc/CommandLine.h"
 #include "Misc/NetworkGuid.h"
diff --git a/Engine/Source/Runtime/Experimental/Iris/Core/Private/Iris/ReplicationSystem/ReplicationWriter.cpp b/Engine/Source/Runtime/Experimental/Iris/Core/Private/Iris/ReplicationSystem/ReplicationWriter.cpp
--- a/Engine/Source/Runtime/Experimental/Iris/Core/Private/Iris/ReplicationSystem/ReplicationWriter.cpp
+++ b/Engine/Source/Runtime/Experimental/Iris/Core/Private/Iris/ReplicationSystem/ReplicationWriter.cpp
@@ -5,6 +5,7 @@
 #include "Iris/IrisConfigInternal.h"
 #include "Iris/Core/IrisLog.h"
 #include "Iris/Core/IrisProfiler.h"
+#include "Misc/ScopeExit.h"
 #include "Net/Core/Trace/NetTrace.h"
 #include "Net/Core/Trace/NetDebugName.h"
 #include "Iris/PacketControl/PacketNotification.h"
diff --git a/Engine/Source/Runtime/LevelSequence/Private/LevelSequenceActorSpawner.cpp b/Engine/Source/Runtime/LevelSequence/Private/LevelSequenceActorSpawner.cpp
--- a/Engine/Source/Runtime/LevelSequence/Private/LevelSequenceActorSpawner.cpp
+++ b/Engine/Source/Runtime/LevelSequence/Private/LevelSequenceActorSpawner.cpp
@@ -1,6 +1,7 @@
 // Copyright Epic Games, Inc. All Rights Reserved.
 
 #include "LevelSequenceActorSpawner.h"
+#include "Misc/PackageName.h"
 #include "MovieSceneSpawnable.h"
 #include "IMovieScenePlayer.h"
 #include "GameFramework/Actor.h"
