[nDisplay] Fix non-lightcard stage actors failing to be automatically added to preview renderers

diff --git a/Engine/Plugins/Runtime/nDisplay/Source/DisplayClusterScenePreview/Private/DisplayClusterScenePreviewModule.cpp b/Engine/Plugins/Runtime/nDisplay/Source/DisplayClusterScenePreview/Private/DisplayClusterScenePreviewModule.cpp
--- a/Engine/Plugins/Runtime/nDisplay/Source/DisplayClusterScenePreview/Private/DisplayClusterScenePreviewModule.cpp
+++ b/Engine/Plugins/Runtime/nDisplay/Source/DisplayClusterScenePreview/Private/DisplayClusterScenePreviewModule.cpp
@@ -169,6 +169,7 @@ bool FDisplayClusterScenePreviewModule::ClearRendererScene(int32 RendererId)
 	{
 		Config->Renderer->ClearScene();
 		Config->AddedActors.Empty();
+		Config->AutoActors.Empty();
 		return true;
 	}
 
@@ -437,6 +438,8 @@ void FDisplayClusterScenePreviewModule::RegisterRootActorEvents(ADisplayClusterR
 void FDisplayClusterScenePreviewModule::AutoPopulateScene(FRendererConfig& RendererConfig)
 {
 	RendererConfig.Renderer->ClearScene();
+	RendererConfig.AddedActors.Empty();
+	RendererConfig.AutoActors.Empty();
 
 	if (ADisplayClusterRootActor* RootActor = InternalGetRendererRootActor(RendererConfig))
 	{
@@ -482,6 +485,11 @@ void FDisplayClusterScenePreviewModule::AutoPopulateScene(FRendererConfig& Rende
 
 			for (AActor* Actor : Actors)
 			{
+				if (RendererConfig.AddedActors.Contains(Actor))
+				{
+					continue;
+				}
+
 				RendererConfig.Renderer->AddActor(Actor);
 				RendererConfig.AddedActors.Add(Actor);
 				RendererConfig.AutoActors.Add(Actor);
@@ -626,7 +634,7 @@ void FDisplayClusterScenePreviewModule::OnLevelActorDeleted(AActor* Actor)
 
 void FDisplayClusterScenePreviewModule::OnLevelActorAdded(AActor* Actor)
 {
-	if (Cast<ADisplayClusterLightCardActor>(Actor) == nullptr)
+	if (!Actor || !Actor->Implements<UDisplayClusterStageActor>())
 	{
 		return;
 	}
diff --git a/Engine/Plugins/Runtime/nDisplay/Source/DisplayClusterScenePreview/Private/DisplayClusterScenePreviewModule.h b/Engine/Plugins/Runtime/nDisplay/Source/DisplayClusterScenePreview/Private/DisplayClusterScenePreviewModule.h
--- a/Engine/Plugins/Runtime/nDisplay/Source/DisplayClusterScenePreview/Private/DisplayClusterScenePreviewModule.h
+++ b/Engine/Plugins/Runtime/nDisplay/Source/DisplayClusterScenePreview/Private/DisplayClusterScenePreviewModule.h
@@ -62,7 +62,7 @@ private:
 		/** Actors that have been automatically added to the scene. */
 		TArray<TWeakObjectPtr<AActor>> AutoActors;
 
-		/** If true, automatically update the renderer with lightcards belonging to the root actor. */
+		/** If true, automatically update the renderer with stage actors belonging to the root actor. */
 		bool bAutoUpdateLightcards = false;
 
 		/** If true, the scene needs to be updated before the next render. This is only relevant if bAutoUpdateLightcards is true. */
