[nDisplay] Fix drag controls in ortho mode on rotated DCRAs

diff --git a/Engine/Plugins/Runtime/nDisplay/Source/DisplayClusterScenePreview/Private/DisplayClusterLightCardEditorHelper.cpp b/Engine/Plugins/Runtime/nDisplay/Source/DisplayClusterScenePreview/Private/DisplayClusterLightCardEditorHelper.cpp
--- a/Engine/Plugins/Runtime/nDisplay/Source/DisplayClusterScenePreview/Private/DisplayClusterLightCardEditorHelper.cpp
+++ b/Engine/Plugins/Runtime/nDisplay/Source/DisplayClusterScenePreview/Private/DisplayClusterLightCardEditorHelper.cpp
@@ -1657,6 +1657,14 @@ FVector FDisplayClusterLightCardEditorHelper::TraceScreenRay(const FVector& Orth
 				Direction = (ClosestPoint - ViewOrigin).GetSafeNormal();
 			}
 		}
+
+		ADisplayClusterRootActor* RootActor = UpdateRootActor();
+		if (RootActor)
+		{
+			// Rotate direction back into local space
+			const FQuat RootRotation = RootActor->GetTransform().GetRotation().Inverse();
+			Direction = RootRotation.RotateVector(Direction);
+		}
 	}
 
 	return Direction;
