nDisplay: fixed Node's TextureShare flag serialization

diff --git a/Engine/Plugins/Runtime/nDisplay/Source/DisplayClusterConfiguration/Private/Formats/JSON500/DisplayClusterConfigurationJsonParser_500.cpp b/Engine/Plugins/Runtime/nDisplay/Source/DisplayClusterConfiguration/Private/Formats/JSON500/DisplayClusterConfigurationJsonParser_500.cpp
--- a/Engine/Plugins/Runtime/nDisplay/Source/DisplayClusterConfiguration/Private/Formats/JSON500/DisplayClusterConfigurationJsonParser_500.cpp
+++ b/Engine/Plugins/Runtime/nDisplay/Source/DisplayClusterConfiguration/Private/Formats/JSON500/DisplayClusterConfigurationJsonParser_500.cpp
@@ -221,6 +221,7 @@ namespace JSON500
 				Node->Host = CfgNode.Value.Host;
 				Node->bIsSoundEnabled = CfgNode.Value.Sound;
 				Node->bIsFullscreen = CfgNode.Value.FullScreen;
+				Node->bEnableTextureShare = CfgNode.Value.TextureShare;
 				Node->WindowRect = FDisplayClusterConfigurationRectangle(CfgNode.Value.Window.X, CfgNode.Value.Window.Y, CfgNode.Value.Window.W, CfgNode.Value.Window.H);
 
 				// Viewports
@@ -412,6 +413,7 @@ namespace JSON500
 				Node.Host = CfgNode.Value->Host;
 				Node.Sound = CfgNode.Value->bIsSoundEnabled;
 				Node.FullScreen = CfgNode.Value->bIsFullscreen;
+				Node.TextureShare = CfgNode.Value->bEnableTextureShare;
 				Node.Window = FDisplayClusterConfigurationJsonRectangle_500(CfgNode.Value->WindowRect.X, CfgNode.Value->WindowRect.Y, CfgNode.Value->WindowRect.W, CfgNode.Value->WindowRect.H);
 
 				// Viewports
diff --git a/Engine/Plugins/Runtime/nDisplay/Source/DisplayClusterConfiguration/Private/Formats/JSON500/DisplayClusterConfigurationJsonTypes_500.h b/Engine/Plugins/Runtime/nDisplay/Source/DisplayClusterConfiguration/Private/Formats/JSON500/DisplayClusterConfigurationJsonTypes_500.h
--- a/Engine/Plugins/Runtime/nDisplay/Source/DisplayClusterConfiguration/Private/Formats/JSON500/DisplayClusterConfigurationJsonTypes_500.h
+++ b/Engine/Plugins/Runtime/nDisplay/Source/DisplayClusterConfiguration/Private/Formats/JSON500/DisplayClusterConfigurationJsonTypes_500.h
@@ -498,6 +498,9 @@ public:
 	UPROPERTY()
 	bool FullScreen = false;
 
+	UPROPERTY()
+	bool TextureShare = false;
+
 	UPROPERTY()
 	FDisplayClusterConfigurationJsonRectangle_500 Window;
 
