nDisplay: fixed mpcdi rendering preview and cleanup

diff --git a/Engine/Plugins/Runtime/nDisplay/Source/DisplayClusterProjection/Private/Policy/MPCDI/DisplayClusterProjectionMPCDIPolicy.cpp b/Engine/Plugins/Runtime/nDisplay/Source/DisplayClusterProjection/Private/Policy/MPCDI/DisplayClusterProjectionMPCDIPolicy.cpp
--- a/Engine/Plugins/Runtime/nDisplay/Source/DisplayClusterProjection/Private/Policy/MPCDI/DisplayClusterProjectionMPCDIPolicy.cpp
+++ b/Engine/Plugins/Runtime/nDisplay/Source/DisplayClusterProjection/Private/Policy/MPCDI/DisplayClusterProjectionMPCDIPolicy.cpp
@@ -446,6 +446,9 @@ UMeshComponent* FDisplayClusterProjectionMPCDIPolicy::GetOrCreatePreviewMeshComp
 			MeshComp->CreateMeshSection(0, MeshData.Vertices, MeshData.Triangles, MeshData.Normal, MeshData.UV, TArray<FColor>(), TArray<FProcMeshTangent>(), false);
 			MeshComp->SetIsVisualizationComponent(true);
 
+			// Because of "nDisplay.render.show.visualizationcomponents" we need extra flag to exclude this geometry from render
+			MeshComp->SetHiddenInGame(true);
+
 			// Store reference to mesh component
 			PreviewMeshComponentRef.SetSceneComponent(MeshComp);
 			return MeshComp;
