nDisplay: Fixed issue with translucency sort order in the ICVFX panel viewport, and added translucency sort priority to stage actor proxies (light cards and CCWs) to ensure they are always rendered on top of screens when flush with them

diff --git a/Engine/Plugins/Runtime/nDisplay/Source/DisplayClusterLightCardEditor/Private/Viewport/DisplayClusterLightCardEditorViewportClient.cpp b/Engine/Plugins/Runtime/nDisplay/Source/DisplayClusterLightCardEditor/Private/Viewport/DisplayClusterLightCardEditorViewportClient.cpp
--- a/Engine/Plugins/Runtime/nDisplay/Source/DisplayClusterLightCardEditor/Private/Viewport/DisplayClusterLightCardEditorViewportClient.cpp
+++ b/Engine/Plugins/Runtime/nDisplay/Source/DisplayClusterLightCardEditor/Private/Viewport/DisplayClusterLightCardEditorViewportClient.cpp
@@ -1601,6 +1601,12 @@ void FDisplayClusterLightCardEditorViewportClient::UpdatePreviewActor(ADisplayCl
 
 				FBox BoundingBox = RootActorProxy->GetComponentsBoundingBox();
 				RootActorBoundingRadius = FMath::Max(BoundingBox.Min.Length(), BoundingBox.Max.Length());
+
+				// Set translucency sort priority of root actor proxy primitive components so that actors that are flush with screens are rendered on top of them
+				RootActorProxy->ForEachComponent<UPrimitiveComponent>(false, [](UPrimitiveComponent* InPrimitiveComponent)
+				{
+					InPrimitiveComponent->SetTranslucentSortPriority(-10);
+				});
 			}
 
 			// Filter out any primitives hidden in game except screen components
diff --git a/Engine/Plugins/Runtime/nDisplay/Source/DisplayClusterLightCardEditorShaders/Private/DisplayClusterMeshProjectionRenderer.cpp b/Engine/Plugins/Runtime/nDisplay/Source/DisplayClusterLightCardEditorShaders/Private/DisplayClusterMeshProjectionRenderer.cpp
--- a/Engine/Plugins/Runtime/nDisplay/Source/DisplayClusterLightCardEditorShaders/Private/DisplayClusterMeshProjectionRenderer.cpp
+++ b/Engine/Plugins/Runtime/nDisplay/Source/DisplayClusterLightCardEditorShaders/Private/DisplayClusterMeshProjectionRenderer.cpp
@@ -349,18 +349,23 @@ protected:
 		if (bTranslucencyPass)
 		{
 			uint16 SortKeyPriority = 0;
-			float DistanceOffset = 0.0f;
+			float Distance = 0.0f;
 
 			if (PrimitiveSceneProxy)
 			{
 				const FPrimitiveSceneInfo* PrimitiveSceneInfo = PrimitiveSceneProxy->GetPrimitiveSceneInfo();
 				SortKeyPriority = (uint16)((int32)PrimitiveSceneInfo->Proxy->GetTranslucencySortPriority() - (int32)SHRT_MIN);
-				DistanceOffset = PrimitiveSceneInfo->Proxy->GetTranslucencySortDistanceOffset();
+
+				// Use the standard sort by distance method for translucent objects
+				const float DistanceOffset = PrimitiveSceneInfo->Proxy->GetTranslucencySortDistanceOffset();
+				const FVector BoundsOrigin = PrimitiveSceneProxy->GetBounds().Origin;
+				const FVector ViewOrigin = this->ViewIfDynamicMeshCommand->ViewMatrices.GetViewOrigin();
+				Distance = (BoundsOrigin - ViewOrigin).Size() + DistanceOffset;
 			}
 
 			SortKey.Translucent.MeshIdInPrimitive = MeshBatch.MeshIdInPrimitive;
 			SortKey.Translucent.Priority = SortKeyPriority;
-			SortKey.Translucent.Distance = *(uint32*)(&DistanceOffset);
+			SortKey.Translucent.Distance = (uint32)~BitInvertIfNegativeFloat(*(uint32*)&Distance);
 		}
 		else
 		{
@@ -372,6 +377,14 @@ protected:
 		return SortKey;
 	}
 
+private:
+	/** Inverts the bits of the floating point number if that number is negative */
+	uint32 BitInvertIfNegativeFloat(uint32 FloatBit)
+	{
+		unsigned Mask = -int32(FloatBit >> 31) | 0x80000000;
+		return FloatBit ^ Mask;
+	}
+
 private:
 	bool bTranslucencyPass;
 	bool bIgnoreTranslucency;
@@ -478,7 +491,7 @@ protected:
 		}
 	}
 
-	virtual FMeshProjectionHitProxyShaderElementData CreateShaderElementData(const FMeshBatch& RESTRICT MeshBatch, const FPrimitiveSceneProxy* RESTRICT PrimitiveSceneProxy, int32 StaticMeshId)
+	virtual FMeshProjectionHitProxyShaderElementData CreateShaderElementData(const FMeshBatch& RESTRICT MeshBatch, const FPrimitiveSceneProxy* RESTRICT PrimitiveSceneProxy, int32 StaticMeshId) override
 	{
 		FMeshProjectionHitProxyShaderElementData ShaderElementData;
 		ShaderElementData.InitializeMeshMaterialData(this->ViewIfDynamicMeshCommand, PrimitiveSceneProxy, MeshBatch, StaticMeshId, false);
@@ -814,7 +827,7 @@ FVector FDisplayClusterMeshProjectionTransform::UnprojectPosition(const FVector&
 		UnprojectedPosition = InvViewMatrix.TransformPosition(ProjectedViewPos);
 	}
 
-	return ProjectedPosition;
+	return UnprojectedPosition;
 }
 
 FDisplayClusterMeshProjectionRenderer::~FDisplayClusterMeshProjectionRenderer()
