nDisplay: Fix UV selection through CCWs in UV mode.
Optimize proxy creation for single proxies.
Speculative workaround for hard-to-reproduce invalid hitproxy causing a crash when clicking on the viewport.

diff --git a/Engine/Plugins/Runtime/nDisplay/Source/DisplayClusterLightCardEditor/Private/Viewport/DisplayClusterLightCardEditorViewportClient.cpp b/Engine/Plugins/Runtime/nDisplay/Source/DisplayClusterLightCardEditor/Private/Viewport/DisplayClusterLightCardEditorViewportClient.cpp
--- a/Engine/Plugins/Runtime/nDisplay/Source/DisplayClusterLightCardEditor/Private/Viewport/DisplayClusterLightCardEditorViewportClient.cpp
+++ b/Engine/Plugins/Runtime/nDisplay/Source/DisplayClusterLightCardEditor/Private/Viewport/DisplayClusterLightCardEditorViewportClient.cpp
@@ -481,6 +481,15 @@ void FDisplayClusterLightCardEditorViewportClient::Draw(const FSceneView* View,
 			{
 				continue;
 			}
+
+			if (const IDisplayClusterStageActor* StageActor = Cast<IDisplayClusterStageActor>(BillboardComponent->GetOwner()))
+			{
+				if ((bIsUVProjection && !StageActor->IsUVActor()) || (!bIsUVProjection && StageActor->IsUVActor()))
+				{
+					continue;
+				}
+			}
+			
 			FSpriteProxy SpriteProxy = FSpriteProxy::FromBillboard(BillboardComponent.Get());
 		
 			FVector ProjectedLocation = ProjectWorldPosition(SpriteProxy.WorldPosition, BillboardViewMatrices);
@@ -1158,7 +1167,7 @@ void FDisplayClusterLightCardEditorViewportClient::UpdateProxyTransforms(const F
 {
 	DECLARE_SCOPE_CYCLE_COUNTER(TEXT("FDisplayClusterLightCardEditorViewportClient::UpdateProxyTransforms"), STAT_UpdateProxyTransforms, STATGROUP_NDisplayLightCardEditor);
 	
-	if (InActorProxy.LevelInstance.IsValid() && InActorProxy.Proxy.IsValid())
+	if (RootActorLevelInstance.IsValid() && InActorProxy.LevelInstance.IsValid() && InActorProxy.Proxy.IsValid())
 	{
 		const FTransform RALevelTransformNoScale(RootActorLevelInstance->GetActorRotation(), RootActorLevelInstance->GetActorLocation(), FVector::OneVector);
 		const FTransform LCLevelRelativeToRALevel = InActorProxy.LevelInstance.AsActorChecked()->GetTransform().GetRelativeTransform(RALevelTransformNoScale);
@@ -1211,6 +1220,14 @@ void FDisplayClusterLightCardEditorViewportClient::ProcessClick(FSceneView& View
 
 	if (HitProxy)
 	{
+		const HHitProxy* HitProxyTest = Viewport->GetHitProxy(HitX, HitY);
+		if (!ensure(HitProxyTest == HitProxy))
+		{
+			// Speculative work around for invalid hit proxy. The HitProxy passed to ProcessClick may have been deleted,
+			// but since we're dealing with a raw ptr we can't check that. This is so we can possibly hit it and debug.
+			return;
+		}
+		
 		if (HitProxy->IsA(HActor::StaticGetType()))
 		{
 			HActor* ActorHitProxy = static_cast<HActor*>(HitProxy);
@@ -1218,17 +1235,16 @@ void FDisplayClusterLightCardEditorViewportClient::ProcessClick(FSceneView& View
 			{
 				if (ActorHitProxy->PrimComponent && ActorHitProxy->PrimComponent->IsA<UStaticMeshComponent>())
 				{
-					ADisplayClusterLightCardActor* TracedLightCard = TraceScreenForLightCard(View, HitX, HitY);
-					SelectActor(TracedLightCard, bMultiSelect);
+					AActor* TracedActor = TraceScreenForActor(View, HitX, HitY);
+					SelectActor(TracedActor, bMultiSelect);
 				}
 			}
-			else if (UE::DisplayClusterLightCardEditorUtils::IsProxySelectable(ActorHitProxy->Actor) && ActorProxies.Contains(ActorHitProxy->Actor))
+			else if (ActorProxies.Contains(ActorHitProxy->Actor) && UE::DisplayClusterLightCardEditorUtils::IsProxySelectable(ActorHitProxy->Actor))
 			{
 				SelectActor(ActorHitProxy->Actor, bMultiSelect);
 			}
 			else if (!bMultiSelect)
 			{
-
 				// Unless a right click is being performed, clear the selection if no geometry was clicked
 				if (!bIsRightClickSelection)
 				{
@@ -1342,8 +1358,7 @@ EMouseCursor::Type FDisplayClusterLightCardEditorViewportClient::GetCursor(FView
 							.SetRealtimeUpdate(IsRealtime()));
 						FSceneView* View = CalcSceneView(&ViewFamily);
 
-						ADisplayClusterLightCardActor* TracedLightCard = TraceScreenForLightCard(*View, X, Y);
-						if (TracedLightCard)
+						if (const AActor* TracedActor = TraceScreenForActor(*View, X, Y))
 						{
 							MouseCursor = EMouseCursor::Crosshairs;
 						}
@@ -1498,7 +1513,16 @@ void FDisplayClusterLightCardEditorViewportClient::UpdatePreviewActor(ADisplayCl
 		return;
 	}
 
-	if (!StageActor)
+	if (StageActor && ActorsRefreshing.Contains(StageActor))
+	{
+		return;
+	}
+
+	if (StageActor)
+	{
+		ActorsRefreshing.Add(StageActor);
+	}
+	else
 	{
 		ProxyTypesRefreshing.Add(ProxyType);
 	}
@@ -1546,6 +1570,7 @@ void FDisplayClusterLightCardEditorViewportClient::UpdatePreviewActor(ADisplayCl
 
 			if (StageActor)
 			{
+				ActorsRefreshing.Remove(StageActor);
 				DestroyProxy(StageActor);
 			}
 			else
@@ -2570,7 +2595,7 @@ double FDisplayClusterLightCardEditorViewportClient::GetActorSpinDelta(FViewport
 	return FMath::RadiansToDegrees(Theta - LastTheta);
 }
 
-ADisplayClusterLightCardActor* FDisplayClusterLightCardEditorViewportClient::TraceScreenForLightCard(const FSceneView& View, int32 HitX, int32 HitY)
+AActor* FDisplayClusterLightCardEditorViewportClient::TraceScreenForActor(const FSceneView& View, int32 HitX, int32 HitY)
 {
 	UWorld* PreviewWorld = PreviewScene->GetWorld();
 	check(PreviewWorld);
@@ -2584,7 +2609,18 @@ ADisplayClusterLightCardActor* FDisplayClusterLightCardEditorViewportClient::Tra
 
 	FCollisionQueryParams TraceParams(SCENE_QUERY_STAT(LightCardTrace), true);
 
-	bool bHitLightCard = false;
+	const bool bInUVMode = ProjectionMode == EDisplayClusterMeshProjectionType::UV;
+	auto IsValidActor = [&](AActor* InActor)
+	{
+		bool bValidForUV = true;
+		if (const IDisplayClusterStageActor* StageActor = Cast<IDisplayClusterStageActor>(InActor))
+		{
+			// Ignore actors not supporting the current UV mode
+			bValidForUV = (bInUVMode && StageActor->IsUVActor()) || (!bInUVMode && !StageActor->IsUVActor());
+		}
+		return bValidForUV && UE::DisplayClusterLightCardEditorUtils::IsProxySelectable(InActor) && ActorProxies.Contains(InActor);
+	};
+	
 	FHitResult ScreenHitResult;
 	if (PreviewWorld->LineTraceSingleByObjectType(ScreenHitResult, CursorRayStart, CursorRayEnd, FCollisionObjectQueryParams(FCollisionObjectQueryParams::InitType::AllObjects), TraceParams))
 	{
@@ -2617,11 +2653,11 @@ ADisplayClusterLightCardActor* FDisplayClusterLightCardEditorViewportClient::Tra
 						{
 							for (FHitResult& HitResult : HitResults)
 							{
-								if (ADisplayClusterLightCardActor* LightCardActor = Cast<ADisplayClusterLightCardActor>(HitResult.GetActor()))
+								if (AActor* HitStageActor = HitResult.GetActor())
 								{
-									if (ActorProxies.Contains(LightCardActor))
+									if (IsValidActor(HitStageActor))
 									{
-										return LightCardActor;
+										return HitStageActor;
 									}
 								}
 							}
@@ -2629,9 +2665,9 @@ ADisplayClusterLightCardActor* FDisplayClusterLightCardEditorViewportClient::Tra
 					}
 				}
 			}
-			else if (UE::DisplayClusterLightCardEditorUtils::IsProxySelectable(HitActor) && ActorProxies.Contains(HitActor))
+			else if (IsValidActor(HitActor))
 			{
-				return Cast<ADisplayClusterLightCardActor>(HitActor);
+				return HitActor;
 			}
 		}
 	}
diff --git a/Engine/Plugins/Runtime/nDisplay/Source/DisplayClusterLightCardEditor/Private/Viewport/DisplayClusterLightCardEditorViewportClient.h b/Engine/Plugins/Runtime/nDisplay/Source/DisplayClusterLightCardEditor/Private/Viewport/DisplayClusterLightCardEditorViewportClient.h
--- a/Engine/Plugins/Runtime/nDisplay/Source/DisplayClusterLightCardEditor/Private/Viewport/DisplayClusterLightCardEditorViewportClient.h
+++ b/Engine/Plugins/Runtime/nDisplay/Source/DisplayClusterLightCardEditor/Private/Viewport/DisplayClusterLightCardEditorViewportClient.h
@@ -216,8 +216,8 @@ private:
 	/** Determines the appropriate spin delta needed to rotate the light card */
 	double GetActorSpinDelta(FViewport* InViewport);
 
-	/** Traces to find the light card corresponding to a click on a stage screen */
-	ADisplayClusterLightCardActor* TraceScreenForLightCard(const FSceneView& View, int32 HitX, int32 HitY);
+	/** Traces to find the actor corresponding to a click on a stage screen */
+	AActor* TraceScreenForActor(const FSceneView& View, int32 HitX, int32 HitY);
 
 	/** Projects the specified world position to the viewport's current projection space */
 	FVector ProjectWorldPosition(const FVector& UnprojectedWorldPosition, const FViewMatrices& ViewMatrices) const;
@@ -305,6 +305,7 @@ private:
 	TWeakPtr<FDisplayClusterLightCardEditor> LightCardEditorPtr;
 	TWeakObjectPtr<ADisplayClusterRootActor> RootActorProxy;
 	TWeakObjectPtr<ADisplayClusterRootActor> RootActorLevelInstance;
+	TSet<AActor*> ActorsRefreshing;
 
 	/** Proxy types that are currently refreshing over 1-2 frames */
 	TSet<EDisplayClusterLightCardEditorProxyType> ProxyTypesRefreshing;
diff --git a/Engine/Plugins/Runtime/nDisplay/Source/DisplayClusterScenePreview/Private/DisplayClusterLightCardEditorHelper.cpp b/Engine/Plugins/Runtime/nDisplay/Source/DisplayClusterScenePreview/Private/DisplayClusterLightCardEditorHelper.cpp
--- a/Engine/Plugins/Runtime/nDisplay/Source/DisplayClusterScenePreview/Private/DisplayClusterLightCardEditorHelper.cpp
+++ b/Engine/Plugins/Runtime/nDisplay/Source/DisplayClusterScenePreview/Private/DisplayClusterLightCardEditorHelper.cpp
@@ -923,10 +923,10 @@ FDisplayClusterMeshProjectionPrimitiveFilter::FPrimitiveFilter FDisplayClusterLi
 	// Create a lambda function so that it's safe to access even if this helper is destroyed
 	return FDisplayClusterMeshProjectionPrimitiveFilter::FPrimitiveFilter::CreateLambda([bIsUVProjection](const UPrimitiveComponent* PrimitiveComponent)
 		{
-			if (ADisplayClusterLightCardActor* LightCard = Cast<ADisplayClusterLightCardActor>(PrimitiveComponent->GetOwner()))
+			if (IDisplayClusterStageActor* StageActor = Cast<IDisplayClusterStageActor>(PrimitiveComponent->GetOwner()))
 			{
-				// Only render the UV light cards when in UV projection mode, and only render non-UV light cards in any other projection mode
-				return bIsUVProjection ? LightCard->bIsUVLightCard : !LightCard->bIsUVLightCard;
+				// Only render the UV actors when in UV projection mode, and only render non-UV actors in any other projection mode
+				return bIsUVProjection ? StageActor->IsUVActor() : !StageActor->IsUVActor();
 			}
 
 			return true;
@@ -940,10 +940,10 @@ FDisplayClusterMeshProjectionPrimitiveFilter::FPrimitiveFilter FDisplayClusterLi
 	// Create a lambda function so that it's safe to access even if this helper is destroyed
 	return FDisplayClusterMeshProjectionPrimitiveFilter::FPrimitiveFilter::CreateLambda([bIsUVProjection](const UPrimitiveComponent* PrimitiveComponent)
 		{
-			if (ADisplayClusterLightCardActor* LightCard = Cast<ADisplayClusterLightCardActor>(PrimitiveComponent->GetOwner()))
+			if (IDisplayClusterStageActor* StageActor = Cast<IDisplayClusterStageActor>(PrimitiveComponent->GetOwner()))
 			{
-				// When in UV projection mode, don't render the UV light cards using the UV projection, render them linearly
-				if (bIsUVProjection && LightCard->bIsUVLightCard)
+				// When in UV projection mode, don't render the UV actors using the UV projection, render them linearly
+				if (bIsUVProjection && StageActor->IsUVActor())
 				{
 					return false;
 				}
