nDisplay: Fix warning when node id not found in non-operational DCRA.

This may happen when there are additional root actors in the scene with an active blueprint event graph.

diff --git a/Engine/Plugins/Runtime/nDisplay/Source/DisplayCluster/Private/Game/EngineClasses/Scene/DisplayClusterRootActor.cpp b/Engine/Plugins/Runtime/nDisplay/Source/DisplayCluster/Private/Game/EngineClasses/Scene/DisplayClusterRootActor.cpp
--- a/Engine/Plugins/Runtime/nDisplay/Source/DisplayCluster/Private/Game/EngineClasses/Scene/DisplayClusterRootActor.cpp
+++ b/Engine/Plugins/Runtime/nDisplay/Source/DisplayCluster/Private/Game/EngineClasses/Scene/DisplayClusterRootActor.cpp
@@ -1137,6 +1137,13 @@ bool ADisplayClusterRootActor::SetReplaceTextureFlagForAllViewports(bool bReplac
 	}
 	else
 	{
+		// No need to set this on a non operational nDisplay root actor.
+		if ((GDisplayCluster->GetOperationMode() == EDisplayClusterOperationMode::Cluster) 
+			&& (this != Display.GetGameMgr()->GetRootActor()))
+		{
+			return false;
+		}
+
 		UDisplayClusterConfigurationClusterNode* Node = ConfigData->Cluster->GetNode(NodeId);
 
 		if (!Node)
