nDisplay: Added Reset to Default button to the color grading drawer's color wheels, and fixed the color wheel slider tooltips. Also added missing context menu to the custom details view used by the color grading drawer.

diff --git a/Engine/Plugins/Runtime/nDisplay/Source/DisplayClusterColorGrading/DisplayClusterColorGrading.Build.cs b/Engine/Plugins/Runtime/nDisplay/Source/DisplayClusterColorGrading/DisplayClusterColorGrading.Build.cs
--- a/Engine/Plugins/Runtime/nDisplay/Source/DisplayClusterColorGrading/DisplayClusterColorGrading.Build.cs
+++ b/Engine/Plugins/Runtime/nDisplay/Source/DisplayClusterColorGrading/DisplayClusterColorGrading.Build.cs
@@ -13,6 +13,7 @@ public class DisplayClusterColorGrading : ModuleRules
 				"DisplayClusterConfiguration",
 				"DisplayClusterOperator",
 
+				"ApplicationCore",
 				"AppFramework",
 				"ColorCorrectRegions",
 				"Core",
diff --git a/Engine/Plugins/Runtime/nDisplay/Source/DisplayClusterColorGrading/Private/DetailView/SDisplayClusterColorGradingDetailTreeRow.cpp b/Engine/Plugins/Runtime/nDisplay/Source/DisplayClusterColorGrading/Private/DetailView/SDisplayClusterColorGradingDetailTreeRow.cpp
--- a/Engine/Plugins/Runtime/nDisplay/Source/DisplayClusterColorGrading/Private/DetailView/SDisplayClusterColorGradingDetailTreeRow.cpp
+++ b/Engine/Plugins/Runtime/nDisplay/Source/DisplayClusterColorGrading/Private/DetailView/SDisplayClusterColorGradingDetailTreeRow.cpp
@@ -5,18 +5,21 @@
 #include "SDisplayClusterColorGradingDetailView.h"
 
 #include "DragAndDrop/DecoratedDragDropOp.h"
+#include "Framework/Application/SlateApplication.h"
+#include "Framework/Application/MenuStack.h"
 #include "Framework/MultiBox/MultiBoxBuilder.h"
+#include "HAL/PlatformApplicationMisc.h"
 #include "IPropertyRowGenerator.h"
 #include "IDetailDragDropHandler.h"
 #include "IDetailTreeNode.h"
-#include "Misc/ConfigCacheIni.h"
+#include "Modules/ModuleManager.h"
+#include "Layout/WidgetPath.h"
 #include "PropertyEditor/Private/DetailTreeNode.h"
 #include "ScopedTransaction.h"
-#include "Modules/ModuleManager.h"
 #include "Styling/StyleColors.h"
 #include "Widgets/Input/SCheckBox.h"
+#include "Widgets/Views/SListView.h"
 #include "Widgets/Views/STableRow.h"
-#include "Widgets/Views/STreeView.h"
 
 namespace UE::DisplayClusterColorGradingDetailTreeRow
 {
@@ -349,6 +352,100 @@ void SDisplayClusterColorGradingDetailTreeRow::ConstructChildren(ETableViewMode:
 	SetContent(InContent);
 }
 
+FReply SDisplayClusterColorGradingDetailTreeRow::OnMouseButtonUp(const FGeometry& MyGeometry, const FPointerEvent& MouseEvent)
+{
+	if (DetailTreeItem.IsValid() && MouseEvent.GetEffectingButton() == EKeys::RightMouseButton && !StaticCastSharedRef<STableViewBase>(OwnerTablePtr.Pin()->AsWidget())->IsRightClickScrolling())
+	{
+		FMenuBuilder MenuBuilder(true, nullptr, nullptr, true);
+
+		bool bShouldOpenMenu = false;
+
+		if (DetailTreeItem.Pin()->HasChildren())
+		{
+			bShouldOpenMenu = true;
+
+			FUIAction ExpandAllAction(FExecuteAction::CreateSP(this, &SDisplayClusterColorGradingDetailTreeRow::SetExpansionStateForAll, true));
+			FUIAction CollapseAllAction(FExecuteAction::CreateSP(this, &SDisplayClusterColorGradingDetailTreeRow::SetExpansionStateForAll, false));
+
+			MenuBuilder.BeginSection(NAME_None, NSLOCTEXT("PropertyView", "ExpansionHeading", "Expansion"));
+			MenuBuilder.AddMenuEntry(NSLOCTEXT("PropertyView", "CollapseAll", "Collapse All"), NSLOCTEXT("PropertyView", "CollapseAll_ToolTip", "Collapses this item and all children"), FSlateIcon(), CollapseAllAction);
+			MenuBuilder.AddMenuEntry(NSLOCTEXT("PropertyView", "ExpandAll", "Expand All"), NSLOCTEXT("PropertyView", "ExpandAll_ToolTip", "Expands this item and all children"), FSlateIcon(), ExpandAllAction);
+			MenuBuilder.EndSection();
+		}
+
+		if (DetailTreeItem.Pin()->IsCopyable())
+		{
+			bShouldOpenMenu = true;
+
+			// Hide separator line if it only contains the SearchWidget, making the next 2 elements the top of the list
+			if (MenuBuilder.GetMultiBox()->GetBlocks().Num() > 1)
+			{
+				MenuBuilder.AddMenuSeparator();
+			}
+
+			FUIAction CopyAction = WidgetRow.IsCopyPasteBound() ?
+				WidgetRow.CopyMenuAction :
+				FExecuteAction::CreateSP(this, &SDisplayClusterColorGradingDetailTreeRow::CopyPropertyValue);
+
+			MenuBuilder.AddMenuEntry(
+				NSLOCTEXT("PropertyView", "CopyProperty", "Copy"),
+				NSLOCTEXT("PropertyView", "CopyProperty_ToolTip", "Copy this property value"),
+				FSlateIcon(FCoreStyle::Get().GetStyleSetName(), "GenericCommands.Copy"),
+				CopyAction);
+
+			FUIAction PasteAction = WidgetRow.IsCopyPasteBound() ?
+				WidgetRow.PasteMenuAction :
+				FUIAction(
+					FExecuteAction::CreateSP(this, &SDisplayClusterColorGradingDetailTreeRow::PastePropertyValue),
+					FCanExecuteAction::CreateSP(this, &SDisplayClusterColorGradingDetailTreeRow::CanPastePropertyValue));
+
+			MenuBuilder.AddMenuEntry(
+				NSLOCTEXT("PropertyView", "PasteProperty", "Paste"),
+				NSLOCTEXT("PropertyView", "PasteProperty_ToolTip", "Paste the copied value here"),
+				FSlateIcon(FCoreStyle::Get().GetStyleSetName(), "GenericCommands.Paste"),
+				PasteAction);
+
+			MenuBuilder.AddMenuEntry(
+				NSLOCTEXT("PropertyView", "CopyPropertyDisplayName", "Copy Display Name"),
+				NSLOCTEXT("PropertyView", "CopyPropertyDisplayName_ToolTip", "Copy this property display name"),
+				FSlateIcon(FCoreStyle::Get().GetStyleSetName(), "GenericCommands.Copy"),
+				FExecuteAction::CreateSP(this, &SDisplayClusterColorGradingDetailTreeRow::CopyPropertyName));
+		}
+
+		if (WidgetRow.CustomMenuItems.Num() > 0)
+		{
+			bShouldOpenMenu = true;
+
+			// Hide separator line if it only contains the SearchWidget, making the next 2 elements the top of the list
+			if (MenuBuilder.GetMultiBox()->GetBlocks().Num() > 1)
+			{
+				MenuBuilder.AddMenuSeparator();
+			}
+
+			for (const FDetailWidgetRow::FCustomMenuData& CustomMenuData : WidgetRow.CustomMenuItems)
+			{
+				// Add the menu entry
+				MenuBuilder.AddMenuEntry(
+					CustomMenuData.Name,
+					CustomMenuData.Tooltip,
+					CustomMenuData.SlateIcon,
+					CustomMenuData.Action);
+			}
+		}
+
+		if (bShouldOpenMenu)
+		{
+			FWidgetPath WidgetPath = MouseEvent.GetEventPath() != nullptr ? *MouseEvent.GetEventPath() : FWidgetPath();
+
+			FSlateApplication::Get().PushMenu(AsShared(), WidgetPath, MenuBuilder.MakeWidget(), MouseEvent.GetScreenSpacePosition(), FPopupTransitionEffect::ContextMenu);
+
+			return FReply::Handled();
+		}
+	}
+
+	return STableRow<TSharedRef<FDisplayClusterColorGradingDetailTreeItem, ESPMode::ThreadSafe>>::OnMouseButtonUp(MyGeometry, MouseEvent);
+}
+
 TSharedPtr<IPropertyHandle> SDisplayClusterColorGradingDetailTreeRow::GetPropertyHandle() const
 {
 	TSharedPtr<IPropertyHandle> PropertyHandle = nullptr;
@@ -873,4 +970,87 @@ int32 SDisplayClusterColorGradingDetailTreeRow::GetDropNewIndex(int32 OriginalIn
 	}
 
 	return ensure(NewIndex >= 0) ? NewIndex : 0;
-}
\ No newline at end of file
+}
+
+void SDisplayClusterColorGradingDetailTreeRow::CopyPropertyName()
+{
+	if (DetailTreeItem.IsValid())
+	{
+		TSharedPtr<IPropertyHandle> PropertyHandle = DetailTreeItem.Pin()->GetPropertyHandle();
+		if (PropertyHandle.IsValid())
+		{
+			FPlatformApplicationMisc::ClipboardCopy(*PropertyHandle->GetPropertyDisplayName().ToString());
+		}
+	}
+}
+
+void SDisplayClusterColorGradingDetailTreeRow::CopyPropertyValue()
+{
+	if (DetailTreeItem.IsValid())
+	{
+		TSharedPtr<IPropertyHandle> PropertyHandle = DetailTreeItem.Pin()->GetPropertyHandle();
+		if (PropertyHandle.IsValid())
+		{
+			FString Value;
+			if (PropertyHandle->GetValueAsFormattedString(Value, PPF_Copy) == FPropertyAccess::Success)
+			{
+				FPlatformApplicationMisc::ClipboardCopy(*Value);
+			}
+		}
+	}
+}
+
+bool SDisplayClusterColorGradingDetailTreeRow::CanPastePropertyValue()
+{
+	FString ClipboardContent;
+	if (DetailTreeItem.IsValid())
+	{
+		if (DetailTreeItem.Pin()->GetPropertyHandle()->IsEditConst())
+		{
+			return false;
+		}
+
+		FPlatformApplicationMisc::ClipboardPaste(ClipboardContent);
+	}
+
+	return !ClipboardContent.IsEmpty();
+}
+
+void SDisplayClusterColorGradingDetailTreeRow::PastePropertyValue()
+{
+	FString ClipboardContent;
+	FPlatformApplicationMisc::ClipboardPaste(ClipboardContent);
+
+	if (!ClipboardContent.IsEmpty() && DetailTreeItem.IsValid())
+	{
+		TSharedPtr<IPropertyHandle> PropertyHandle = DetailTreeItem.Pin()->GetPropertyHandle();
+		if (PropertyHandle.IsValid())
+		{
+			PropertyHandle->SetValueFromFormattedString(ClipboardContent, EPropertyValueSetFlags::InstanceObjects);
+		}
+	}
+}
+
+void SDisplayClusterColorGradingDetailTreeRow::SetExpansionStateForAll(bool bShouldBeExpanded)
+{
+	if (DetailTreeItem.IsValid())
+	{
+		SetExpansionStateRecursive(DetailTreeItem.Pin().ToSharedRef(), bShouldBeExpanded);
+	}
+}
+
+void SDisplayClusterColorGradingDetailTreeRow::SetExpansionStateRecursive(const TSharedRef<FDisplayClusterColorGradingDetailTreeItem>& TreeItem, bool bShouldBeExpanded)
+{
+	if (OwnerTablePtr.IsValid())
+	{
+		OwnerTablePtr.Pin()->Private_SetItemExpansion(TreeItem, bShouldBeExpanded);
+
+		TArray<TSharedRef<FDisplayClusterColorGradingDetailTreeItem>> Children;
+		TreeItem->GetChildren(Children);
+
+		for (const TSharedRef<FDisplayClusterColorGradingDetailTreeItem>& Child : Children)
+		{
+			SetExpansionStateRecursive(Child, bShouldBeExpanded);
+		}
+	}
+}
diff --git a/Engine/Plugins/Runtime/nDisplay/Source/DisplayClusterColorGrading/Private/DetailView/SDisplayClusterColorGradingDetailTreeRow.h b/Engine/Plugins/Runtime/nDisplay/Source/DisplayClusterColorGrading/Private/DetailView/SDisplayClusterColorGradingDetailTreeRow.h
--- a/Engine/Plugins/Runtime/nDisplay/Source/DisplayClusterColorGrading/Private/DetailView/SDisplayClusterColorGradingDetailTreeRow.h
+++ b/Engine/Plugins/Runtime/nDisplay/Source/DisplayClusterColorGrading/Private/DetailView/SDisplayClusterColorGradingDetailTreeRow.h
@@ -21,6 +21,7 @@ public:
 	//~ STableRow Interface
 	void Construct(const FArguments& InArgs, const TSharedRef<FDisplayClusterColorGradingDetailTreeItem>& InDetailTreeItem, const TSharedRef<STableViewBase>& InOwnerTableView, const FDetailColumnSizeData& InColumnSizeData);
 	virtual void ConstructChildren(ETableViewMode::Type InOwnerTableMode, const TAttribute<FMargin>& InPadding, const TSharedRef<SWidget>& InContent) override;
+	virtual FReply OnMouseButtonUp(const FGeometry& MyGeometry, const FPointerEvent& MouseEvent) override;
 	//~ End STableRow Interface
 
 	/** Gets the property handle of the property this row represents, if applicable */
@@ -81,6 +82,24 @@ private:
 	/** Gets the index to use when another row has been dragged onto this row, and an array element swap is taking place */
 	int32 GetDropNewIndex(int32 OriginalIndex, int32 DropOntoIndex, EItemDropZone DropZone) const;
 
+	/** Copies the property name to the clipboard */
+	void CopyPropertyName();
+
+	/** Copies the property value to the clipboard */
+	void CopyPropertyValue();
+
+	/** Checks to see if this property can receive a pasted value from the clipboard */
+	bool CanPastePropertyValue();
+
+	/** Pastes the value in the clipboard into this property */
+	void PastePropertyValue();
+
+	/** Sets the expansion state for this row and all of its children */
+	void SetExpansionStateForAll(bool bShouldBeExpanded);
+
+	/** Recursively sets the expansion state for the specified tree item and its children */
+	void SetExpansionStateRecursive(const TSharedRef<FDisplayClusterColorGradingDetailTreeItem>& TreeItem, bool bShouldBeExpanded);
+
 private:
 	/** A weak pointer to the detail tree item this row represents */
 	TWeakPtr<FDisplayClusterColorGradingDetailTreeItem> DetailTreeItem;
diff --git a/Engine/Plugins/Runtime/nDisplay/Source/DisplayClusterColorGrading/Private/DetailView/SDisplayClusterColorGradingDetailView.cpp b/Engine/Plugins/Runtime/nDisplay/Source/DisplayClusterColorGrading/Private/DetailView/SDisplayClusterColorGradingDetailView.cpp
--- a/Engine/Plugins/Runtime/nDisplay/Source/DisplayClusterColorGrading/Private/DetailView/SDisplayClusterColorGradingDetailView.cpp
+++ b/Engine/Plugins/Runtime/nDisplay/Source/DisplayClusterColorGrading/Private/DetailView/SDisplayClusterColorGradingDetailView.cpp
@@ -160,6 +160,30 @@ bool FDisplayClusterColorGradingDetailTreeItem::IsReorderable() const
 	return false;
 }
 
+bool FDisplayClusterColorGradingDetailTreeItem::IsCopyable() const
+{
+	if (PropertyHandle.IsValid() && PropertyHandle->IsValidHandle())
+	{
+		static const FName DisableCopyPasteMetaDataName("DisableCopyPaste");
+
+		// Check to see if this property or any of its parents have the DisableCopyPaste metadata
+		TSharedPtr<IPropertyHandle> CurrentPropertyHandle = PropertyHandle;
+		while (CurrentPropertyHandle.IsValid())
+		{
+			if (CurrentPropertyHandle->HasMetaData(DisableCopyPasteMetaDataName))
+			{
+				return false;
+			}
+
+			CurrentPropertyHandle = CurrentPropertyHandle->GetParentHandle();
+		}
+
+		return true;
+	}
+
+	return false;
+}
+
 void FDisplayClusterColorGradingDetailTreeItem::GenerateDetailWidgetRow(FDetailWidgetRow& OutDetailWidgetRow) const
 {
 	if (DetailTreeNode.IsValid())
diff --git a/Engine/Plugins/Runtime/nDisplay/Source/DisplayClusterColorGrading/Private/DetailView/SDisplayClusterColorGradingDetailView.h b/Engine/Plugins/Runtime/nDisplay/Source/DisplayClusterColorGrading/Private/DetailView/SDisplayClusterColorGradingDetailView.h
--- a/Engine/Plugins/Runtime/nDisplay/Source/DisplayClusterColorGrading/Private/DetailView/SDisplayClusterColorGradingDetailView.h
+++ b/Engine/Plugins/Runtime/nDisplay/Source/DisplayClusterColorGrading/Private/DetailView/SDisplayClusterColorGradingDetailView.h
@@ -75,6 +75,9 @@ public:
 	/** Gets whether this detail tree item can be reordered through a drag drop action */
 	bool IsReorderable() const;
 
+	/** Gets whether this detail tree item can be copied */
+	bool IsCopyable() const;
+
 	/** Generates the row widgets for this detail tree item */
 	void GenerateDetailWidgetRow(FDetailWidgetRow& OutDetailWidgetRow) const;
 
diff --git a/Engine/Plugins/Runtime/nDisplay/Source/DisplayClusterColorGrading/Private/SDisplayClusterColorGradingColorWheel.cpp b/Engine/Plugins/Runtime/nDisplay/Source/DisplayClusterColorGrading/Private/SDisplayClusterColorGradingColorWheel.cpp
--- a/Engine/Plugins/Runtime/nDisplay/Source/DisplayClusterColorGrading/Private/SDisplayClusterColorGradingColorWheel.cpp
+++ b/Engine/Plugins/Runtime/nDisplay/Source/DisplayClusterColorGrading/Private/SDisplayClusterColorGradingColorWheel.cpp
@@ -24,7 +24,7 @@ void SDisplayClusterColorGradingColorWheel::Construct(const FArguments& InArgs)
 
 		+ SVerticalBox::Slot()
 		.AutoHeight()
-		.HAlign(HAlign_Center)
+		.HAlign(HAlign_Fill)
 		[
 			SAssignNew(HeaderBox, SBox)
 		]
@@ -605,11 +605,11 @@ FText SDisplayClusterColorGradingColorWheel::GetComponentToolTipText(uint32 Comp
 	{
 		if (ColorDisplayMode.Get(EColorDisplayMode::RGB) != EColorDisplayMode::RGB)
 		{
-			return RGBToolTipText[ComponentIndex];
+			return HSVToolTipText[ComponentIndex];
 		}
 		else
 		{
-			return HSVToolTipText[ComponentIndex];
+			return RGBToolTipText[ComponentIndex];
 		}
 	}
 	else if (ComponentIndex == 3)
diff --git a/Engine/Plugins/Runtime/nDisplay/Source/DisplayClusterColorGrading/Private/SDisplayClusterColorGradingColorWheelPanel.cpp b/Engine/Plugins/Runtime/nDisplay/Source/DisplayClusterColorGrading/Private/SDisplayClusterColorGradingColorWheelPanel.cpp
--- a/Engine/Plugins/Runtime/nDisplay/Source/DisplayClusterColorGrading/Private/SDisplayClusterColorGradingColorWheelPanel.cpp
+++ b/Engine/Plugins/Runtime/nDisplay/Source/DisplayClusterColorGrading/Private/SDisplayClusterColorGradingColorWheelPanel.cpp
@@ -10,6 +10,7 @@
 #include "Framework/MultiBox/MultiBoxBuilder.h"
 #include "IDetailTreeNode.h"
 #include "IPropertyRowGenerator.h"
+#include "Modules/ModuleManager.h"
 #include "PropertyHandle.h"
 #include "Widgets/SBoxPanel.h"
 #include "Widgets/Input/SComboButton.h"
@@ -420,38 +421,7 @@ void SDisplayClusterColorGradingColorWheelPanel::FillColorWheels(const FDisplayC
 		if (ColorWheel.IsValid())
 		{
 			ColorWheel->SetColorPropertyHandle(PropertyHandle);
-
-			if (TSharedPtr<IDetailTreeNode> TreeNode = ColorGradingDataModel->GetPropertyRowGenerator()->FindTreeNode(PropertyHandle))
-			{
-				FNodeWidgets NodeWidgets = TreeNode->CreateNodeWidgets();
-
-				TSharedRef<SHorizontalBox> PropertyNameBox = SNew(SHorizontalBox);
-
-				if (NodeWidgets.EditConditionWidget.IsValid())
-				{
-					PropertyNameBox->AddSlot()
-						.VAlign(VAlign_Center)
-						.HAlign(HAlign_Left)
-						.Padding(2, 0, 0, 0)
-						.AutoWidth()
-						[
-							NodeWidgets.EditConditionWidget.ToSharedRef()
-						];
-				}
-
-				if (NodeWidgets.NameWidget.IsValid())
-				{
-					PropertyNameBox->AddSlot()
-						.HAlign(NodeWidgets.NameWidgetLayoutData.HorizontalAlignment)
-						.VAlign(NodeWidgets.NameWidgetLayoutData.VerticalAlignment)
-						.Padding(2, 0, 0, 0)
-						[
-							NodeWidgets.NameWidget.ToSharedRef()
-						];
-				}
-
-				ColorWheel->SetHeaderContent(PropertyNameBox);
-			}
+			ColorWheel->SetHeaderContent(CreateColorWheelHeaderWidget(PropertyHandle));
 		}
 	};
 
@@ -474,6 +444,124 @@ void SDisplayClusterColorGradingColorWheelPanel::ClearColorWheels()
 	};
 }
 
+TSharedRef<SWidget> SDisplayClusterColorGradingColorWheelPanel::CreateColorWheelHeaderWidget(const TSharedPtr<IPropertyHandle>& ColorPropertyHandle)
+{
+	if (TSharedPtr<IDetailTreeNode> TreeNode = ColorGradingDataModel->GetPropertyRowGenerator()->FindTreeNode(ColorPropertyHandle))
+	{
+		FNodeWidgets NodeWidgets = TreeNode->CreateNodeWidgets();
+
+		TSharedRef<SHorizontalBox> PropertyNameBox = SNew(SHorizontalBox);
+
+		PropertyNameBox->AddSlot()
+			.FillWidth(1.0f)
+			[
+				SNew(SSpacer)
+			];
+
+		if (NodeWidgets.EditConditionWidget.IsValid())
+		{
+			PropertyNameBox->AddSlot()
+				.VAlign(VAlign_Center)
+				.HAlign(HAlign_Left)
+				.Padding(2, 0, 0, 0)
+				.AutoWidth()
+				[
+					NodeWidgets.EditConditionWidget.ToSharedRef()
+				];
+		}
+
+		if (NodeWidgets.NameWidget.IsValid())
+		{
+			PropertyNameBox->AddSlot()
+				.HAlign(NodeWidgets.NameWidgetLayoutData.HorizontalAlignment)
+				.VAlign(NodeWidgets.NameWidgetLayoutData.VerticalAlignment)
+				.Padding(2, 0, 0, 0)
+				[
+					NodeWidgets.NameWidget.ToSharedRef()
+				];
+
+			PropertyNameBox->AddSlot()
+				.HAlign(HAlign_Right)
+				.VAlign(VAlign_Center)
+				.FillWidth(1.0f)
+				[
+					SNew(SBox)
+					.WidthOverride(22.0f)
+					[
+						CreateColorPropertyExtensions(ColorPropertyHandle, TreeNode)
+					]
+				];
+		}
+
+		return PropertyNameBox;
+	}
+
+	return SNullWidget::NullWidget;
+}
+
+TSharedRef<SWidget> SDisplayClusterColorGradingColorWheelPanel::CreateColorPropertyExtensions(const TSharedPtr<IPropertyHandle>& ColorPropertyHandle, const TSharedPtr<IDetailTreeNode>& DetailTreeNode)
+{
+	TArray<FPropertyRowExtensionButton> ExtensionButtons;
+
+	// Use a weak pointer to pass into delegates
+	TWeakPtr<IPropertyHandle> WeakColorPropertyHandle = ColorPropertyHandle;
+
+	FPropertyRowExtensionButton& ResetToDefaultButton = ExtensionButtons.AddDefaulted_GetRef();
+	ResetToDefaultButton.Label = NSLOCTEXT("PropertyEditor", "ResetToDefault", "Reset to Default");
+	ResetToDefaultButton.UIAction = FUIAction(
+		FExecuteAction::CreateLambda([WeakColorPropertyHandle]
+		{
+			if (WeakColorPropertyHandle.IsValid())
+			{
+				WeakColorPropertyHandle.Pin()->ResetToDefault();
+			}
+		}),
+		FCanExecuteAction::CreateLambda([WeakColorPropertyHandle]
+		{
+			const bool bIsEditable = WeakColorPropertyHandle.Pin()->IsEditable();
+			return bIsEditable;
+		}),
+		FIsActionChecked(),
+		FIsActionButtonVisible::CreateLambda([WeakColorPropertyHandle]
+		{
+			bool bShowResetToDefaultButton = false;
+			if (WeakColorPropertyHandle.IsValid())
+			{
+				if (!WeakColorPropertyHandle.Pin()->HasMetaData("NoResetToDefault") && !WeakColorPropertyHandle.Pin()->GetInstanceMetaData("NoResetToDefault"))
+				{
+					bShowResetToDefaultButton = WeakColorPropertyHandle.Pin()->CanResetToDefault();
+				}
+			}
+
+			return bShowResetToDefaultButton;
+		})
+	);
+
+	ResetToDefaultButton.Icon = FSlateIcon(FAppStyle::Get().GetStyleSetName(), "PropertyWindow.DiffersFromDefault");
+	ResetToDefaultButton.ToolTip = NSLOCTEXT("PropertyEditor", "ResetToDefaultPropertyValueToolTip", "Reset this property to its default value.");
+
+	// Add any global row extensions that are registered for the color property
+	FPropertyEditorModule& PropertyEditorModule = FModuleManager::Get().GetModuleChecked<FPropertyEditorModule>("PropertyEditor");
+
+	FOnGenerateGlobalRowExtensionArgs Args;
+	Args.OwnerTreeNode = DetailTreeNode;
+	Args.PropertyHandle = ColorPropertyHandle;
+
+	PropertyEditorModule.GetGlobalRowExtensionDelegate().Broadcast(Args, ExtensionButtons);
+
+	FSlimHorizontalToolBarBuilder ToolbarBuilder(TSharedPtr<FUICommandList>(), FMultiBoxCustomization::None);
+	ToolbarBuilder.SetLabelVisibility(EVisibility::Collapsed);
+	ToolbarBuilder.SetStyle(&FAppStyle::Get(), "DetailsView.ExtensionToolBar");
+	ToolbarBuilder.SetIsFocusable(false);
+
+	for (const FPropertyRowExtensionButton& Extension : ExtensionButtons)
+	{
+		ToolbarBuilder.AddToolBarButton(Extension.UIAction, NAME_None, Extension.Label, Extension.ToolTip, Extension.Icon);
+	}
+
+	return ToolbarBuilder.MakeWidget();
+}
+
 bool SDisplayClusterColorGradingColorWheelPanel::FilterDetailTreeNode(const TSharedRef<IDetailTreeNode>& InDetailTreeNode)
 {
 	if (ColorGradingDataModel)
diff --git a/Engine/Plugins/Runtime/nDisplay/Source/DisplayClusterColorGrading/Private/SDisplayClusterColorGradingColorWheelPanel.h b/Engine/Plugins/Runtime/nDisplay/Source/DisplayClusterColorGrading/Private/SDisplayClusterColorGradingColorWheelPanel.h
--- a/Engine/Plugins/Runtime/nDisplay/Source/DisplayClusterColorGrading/Private/SDisplayClusterColorGradingColorWheelPanel.h
+++ b/Engine/Plugins/Runtime/nDisplay/Source/DisplayClusterColorGrading/Private/SDisplayClusterColorGradingColorWheelPanel.h
@@ -58,6 +58,9 @@ private:
 	void FillColorWheels(const FDisplayClusterColorGradingDataModel::FColorGradingElement& ColorGradingElement);
 	void ClearColorWheels();
 
+	TSharedRef<SWidget> CreateColorWheelHeaderWidget(const TSharedPtr<IPropertyHandle>& ColorPropertyHandle);
+	TSharedRef<SWidget> CreateColorPropertyExtensions(const TSharedPtr<IPropertyHandle>& ColorPropertyHandle, const TSharedPtr<IDetailTreeNode>& DetailTreeNode);
+
 	bool FilterDetailTreeNode(const TSharedRef<IDetailTreeNode>& InDetailTreeNode);
 
 	void SetColorWheelOrientation(EOrientation NewOrientation);
