nDisplay: Disable placing new actors in ICVFX panel when there is no root actor, fixing a crash.

diff --git a/Engine/Plugins/Runtime/nDisplay/Source/DisplayClusterLightCardEditor/Private/DisplayClusterLightCardEditor.cpp b/Engine/Plugins/Runtime/nDisplay/Source/DisplayClusterLightCardEditor/Private/DisplayClusterLightCardEditor.cpp
--- a/Engine/Plugins/Runtime/nDisplay/Source/DisplayClusterLightCardEditor/Private/DisplayClusterLightCardEditor.cpp
+++ b/Engine/Plugins/Runtime/nDisplay/Source/DisplayClusterLightCardEditor/Private/DisplayClusterLightCardEditor.cpp
@@ -470,7 +470,7 @@ void FDisplayClusterLightCardEditor::AddLightCardsToActor(const TArray<ADisplayC
 	}
 }
 
-bool FDisplayClusterLightCardEditor::CanAddLightCard() const
+bool FDisplayClusterLightCardEditor::CanAddNewActor() const
 {
 	return ActiveRootActor.IsValid() && ActiveRootActor->GetWorld() != nullptr;
 }
@@ -1092,15 +1092,19 @@ TSharedRef<SWidget> FDisplayClusterLightCardEditor::GeneratePlaceActorsMenu()
 			FText Label = Class->GetDisplayNameText();
 			FSlateIcon StageActorIcon = FSlateIconFinder::FindIconForClass(Class);
 			MenuBuilder.AddMenuEntry(Label, LOCTEXT("AddStageActorHeader", "Add a stage actor to the scene"), StageActorIcon,
-				FUIAction(FExecuteAction::CreateRaw(this, &FDisplayClusterLightCardEditor::AddNewDynamic, Class)));
+				FUIAction(FExecuteAction::CreateSP(this, &FDisplayClusterLightCardEditor::AddNewDynamic, Class),
+					FCanExecuteAction::CreateSP(this, &FDisplayClusterLightCardEditor::CanAddNewActor)));
+		}
+
+		if (CanAddNewActor())
+		{
+			MenuBuilder.AddSubMenu(LOCTEXT("AllTemplatesLabel", "All Templates"),
+			LOCTEXT("AllTemplatesTooltip", "Select a template"),
+			FNewMenuDelegate::CreateSP(this, &FDisplayClusterLightCardEditor::GenerateTemplateSubMenu),
+			false,
+			FSlateIcon(FDisplayClusterLightCardEditorStyle::Get().GetStyleSetName(), TEXT("DisplayClusterLightCardEditor.Template")),
+			true);
 		}
-		
-		MenuBuilder.AddSubMenu(LOCTEXT("AllTemplatesLabel", "All Templates"),
-		LOCTEXT("AllTemplatesTooltip", "Select a template"),
-		FNewMenuDelegate::CreateSP(this, &FDisplayClusterLightCardEditor::GenerateTemplateSubMenu),
-		false,
-		FSlateIcon(FDisplayClusterLightCardEditorStyle::Get().GetStyleSetName(), TEXT("DisplayClusterLightCardEditor.Template")),
-		true);
 	}
 	MenuBuilder.EndSection();
 	MenuBuilder.BeginSection("Favorites", LOCTEXT("FavoritesMenuHeader", "Favorites"));
@@ -1158,6 +1162,7 @@ TSharedRef<SWidget> FDisplayClusterLightCardEditor::GeneratePlaceActorsMenu()
 					SpawnActor(TemplateWeakPtr.Get());
 				}
 			});
+			EntryParams.DirectActions.CanExecuteAction = FCanExecuteAction::CreateSP(this, &FDisplayClusterLightCardEditor::CanAddNewActor);
 			
 			MenuBuilder.AddMenuEntry(EntryParams);
 		}
@@ -1226,6 +1231,7 @@ TSharedRef<SWidget> FDisplayClusterLightCardEditor::GeneratePlaceActorsMenu()
 					}
 				}
 			});
+			EntryParams.DirectActions.CanExecuteAction = FCanExecuteAction::CreateSP(this, &FDisplayClusterLightCardEditor::CanAddNewActor);
 			
 			MenuBuilder.AddMenuEntry(EntryParams);
 		}
@@ -1479,17 +1485,17 @@ void FDisplayClusterLightCardEditor::BindCommands()
 	CommandList->MapAction(
 		FDisplayClusterLightCardEditorCommands::Get().AddNewLightCard,
 		FExecuteAction::CreateSP(this, &FDisplayClusterLightCardEditor::AddNewLightCard),
-		FCanExecuteAction::CreateSP(this, &FDisplayClusterLightCardEditor::CanAddLightCard));
+		FCanExecuteAction::CreateSP(this, &FDisplayClusterLightCardEditor::CanAddNewActor));
 
 	CommandList->MapAction(
 		FDisplayClusterLightCardEditorCommands::Get().AddNewFlag,
 		FExecuteAction::CreateSP(this, &FDisplayClusterLightCardEditor::AddNewFlag),
-		FCanExecuteAction::CreateSP(this, &FDisplayClusterLightCardEditor::CanAddLightCard));
+		FCanExecuteAction::CreateSP(this, &FDisplayClusterLightCardEditor::CanAddNewActor));
 	
 	CommandList->MapAction(
 		FDisplayClusterLightCardEditorCommands::Get().AddExistingLightCard,
 		FExecuteAction::CreateSP(this, &FDisplayClusterLightCardEditor::AddExistingLightCard),
-		FCanExecuteAction::CreateSP(this, &FDisplayClusterLightCardEditor::CanAddLightCard));
+		FCanExecuteAction::CreateSP(this, &FDisplayClusterLightCardEditor::CanAddNewActor));
 
 	CommandList->MapAction(
 		FDisplayClusterLightCardEditorCommands::Get().RemoveLightCard,
diff --git a/Engine/Plugins/Runtime/nDisplay/Source/DisplayClusterLightCardEditor/Private/DisplayClusterLightCardEditor.h b/Engine/Plugins/Runtime/nDisplay/Source/DisplayClusterLightCardEditor/Private/DisplayClusterLightCardEditor.h
--- a/Engine/Plugins/Runtime/nDisplay/Source/DisplayClusterLightCardEditor/Private/DisplayClusterLightCardEditor.h
+++ b/Engine/Plugins/Runtime/nDisplay/Source/DisplayClusterLightCardEditor/Private/DisplayClusterLightCardEditor.h
@@ -117,7 +117,7 @@ public:
 	void AddLightCardsToActor(const TArray<ADisplayClusterLightCardActor*>& LightCards);
 
 	/** If a Light Card can currently be added */
-	bool CanAddLightCard() const;
+	bool CanAddNewActor() const;
 
 	/** Copies any selected actors to the clipboard, and then deletes them */
 	void CutSelectedActors();
diff --git a/Engine/Plugins/Runtime/nDisplay/Source/DisplayClusterLightCardEditor/Private/Viewport/DisplayClusterLightcardEditorViewport.cpp b/Engine/Plugins/Runtime/nDisplay/Source/DisplayClusterLightCardEditor/Private/Viewport/DisplayClusterLightcardEditorViewport.cpp
--- a/Engine/Plugins/Runtime/nDisplay/Source/DisplayClusterLightCardEditor/Private/Viewport/DisplayClusterLightcardEditorViewport.cpp
+++ b/Engine/Plugins/Runtime/nDisplay/Source/DisplayClusterLightCardEditor/Private/Viewport/DisplayClusterLightcardEditorViewport.cpp
@@ -734,13 +734,15 @@ void SDisplayClusterLightCardEditorViewport::MakePlaceActorsSubMenu(FMenuBuilder
 		FDisplayClusterLightCardEditorCommands::Get().AddNewFlag->GetLabel(),
 		FDisplayClusterLightCardEditorCommands::Get().AddNewFlag->GetDescription(),
 		FlagIcon,
-		FUIAction(FExecuteAction::CreateSP(this, &SDisplayClusterLightCardEditorViewport::AddFlagHere)));
+		FUIAction(FExecuteAction::CreateSP(this, &SDisplayClusterLightCardEditorViewport::AddFlagHere),
+			FCanExecuteAction::CreateSP(this, &SDisplayClusterLightCardEditorViewport::CanPlaceActorHere)));
 
 	MenuBuilder.AddMenuEntry(
 		FDisplayClusterLightCardEditorCommands::Get().AddNewLightCard->GetLabel(),
 		FDisplayClusterLightCardEditorCommands::Get().AddNewLightCard->GetDescription(),
 		bIsUVMode ? UVLightCardIcon : LightCardIcon,
-		FUIAction(FExecuteAction::CreateSP(this, &SDisplayClusterLightCardEditorViewport::AddLightCardHere)));
+		FUIAction(FExecuteAction::CreateSP(this, &SDisplayClusterLightCardEditorViewport::AddLightCardHere),
+			FCanExecuteAction::CreateSP(this, &SDisplayClusterLightCardEditorViewport::CanPlaceActorHere)));
 
 	TSet<UClass*> StageActorClasses = UE::DisplayClusterLightCardEditorUtils::GetAllStageActorClasses();
 	for (UClass* Class : StageActorClasses)
@@ -757,7 +759,8 @@ void SDisplayClusterLightCardEditorViewport::MakePlaceActorsSubMenu(FMenuBuilder
 			Label,
 			LOCTEXT("AddStageActorHeader", "Add a stage actor to the scene"), 
 			StageActorIcon,
-			FUIAction(FExecuteAction::CreateRaw(this, &SDisplayClusterLightCardEditorViewport::AddStageActorHere, Class)));
+			FUIAction(FExecuteAction::CreateSP(this, &SDisplayClusterLightCardEditorViewport::AddStageActorHere, Class),
+				FCanExecuteAction::CreateSP(this, &SDisplayClusterLightCardEditorViewport::CanPlaceActorHere)));
 	}
 }
 
@@ -924,6 +927,11 @@ void SDisplayClusterLightCardEditorViewport::AddStageActorHere(UClass* InClass)
 	}
 }
 
+bool SDisplayClusterLightCardEditorViewport::CanPlaceActorHere() const
+{
+	return LightCardEditorPtr.IsValid() && LightCardEditorPtr.Pin()->CanAddNewActor();
+}
+
 void SDisplayClusterLightCardEditorViewport::ToggleLabels()
 {
 	// TODO: Handle CCR
diff --git a/Engine/Plugins/Runtime/nDisplay/Source/DisplayClusterLightCardEditor/Private/Viewport/DisplayClusterLightcardEditorViewport.h b/Engine/Plugins/Runtime/nDisplay/Source/DisplayClusterLightCardEditor/Private/Viewport/DisplayClusterLightcardEditorViewport.h
--- a/Engine/Plugins/Runtime/nDisplay/Source/DisplayClusterLightCardEditor/Private/Viewport/DisplayClusterLightcardEditorViewport.h
+++ b/Engine/Plugins/Runtime/nDisplay/Source/DisplayClusterLightCardEditor/Private/Viewport/DisplayClusterLightcardEditorViewport.h
@@ -110,6 +110,9 @@ private:
 	/** Adds a new stage actor of the specified class and moves it to the cached mouse position */
 	void AddStageActorHere(UClass* InClass);
 
+	/** If placing actors is allowed */
+	bool CanPlaceActorHere() const;
+
 	/** Toggles light card labels on or off */
 	void ToggleLabels();
 
