nDisplay: fixed threading crash at import object's sync data

diff --git a/Engine/Plugins/Runtime/nDisplay/Source/DisplayCluster/Private/Cluster/DisplayClusterClusterManager.cpp b/Engine/Plugins/Runtime/nDisplay/Source/DisplayCluster/Private/Cluster/DisplayClusterClusterManager.cpp
--- a/Engine/Plugins/Runtime/nDisplay/Source/DisplayCluster/Private/Cluster/DisplayClusterClusterManager.cpp
+++ b/Engine/Plugins/Runtime/nDisplay/Source/DisplayCluster/Private/Cluster/DisplayClusterClusterManager.cpp
@@ -608,6 +608,8 @@ void FDisplayClusterClusterManager::ImportObjectsData(const EDisplayClusterSyncG
 			UE_LOG(LogDisplayClusterCluster, VeryVerbose, TEXT("sync-data: %s=%s"), *It->Key, *It->Value);
 		}
 
+		FScopeLock Lock(&ObjectsToSyncCS);
+
 		for (IDisplayClusterClusterSyncObject* SyncObj : ObjectsToSync[InSyncGroup])
 		{
 			if (SyncObj && SyncObj->IsActive())
