nDisplay: Fix proxies being added to ICVFX panel renderer multiple times.

diff --git a/Engine/Plugins/Runtime/nDisplay/Source/DisplayClusterLightCardEditor/Private/Viewport/DisplayClusterLightCardEditorViewportClient.cpp b/Engine/Plugins/Runtime/nDisplay/Source/DisplayClusterLightCardEditor/Private/Viewport/DisplayClusterLightCardEditorViewportClient.cpp
--- a/Engine/Plugins/Runtime/nDisplay/Source/DisplayClusterLightCardEditor/Private/Viewport/DisplayClusterLightCardEditorViewportClient.cpp
+++ b/Engine/Plugins/Runtime/nDisplay/Source/DisplayClusterLightCardEditor/Private/Viewport/DisplayClusterLightCardEditorViewportClient.cpp
@@ -1097,12 +1097,12 @@ const FDisplayClusterLightCardEditorViewportClient::FActorProxy* FDisplayCluster
 	});
 }
 
-void FDisplayClusterLightCardEditorViewportClient::CreateStageActorProxy(AActor* InLevelInstance)
+AActor* FDisplayClusterLightCardEditorViewportClient::CreateStageActorProxy(AActor* InLevelInstance)
 {
 	if (!IsValid(InLevelInstance))
 	{
 		// Can happen if the level actor was destroyed the tick prior to proxy creation.
-		return;
+		return nullptr;
 	}
 	
 	const FTransform RALevelTransformNoScale(RootActorLevelInstance->GetActorRotation(), RootActorLevelInstance->GetActorLocation(), FVector::OneVector);
@@ -1161,6 +1161,8 @@ void FDisplayClusterLightCardEditorViewportClient::CreateStageActorProxy(AActor*
 	ProjectionHelper->VerifyAndFixActorOrigin(ActorProxy);
 
 	UpdateProxyTransforms(ActorProxyStruct);
+
+	return ActorProxy;
 }
 
 void FDisplayClusterLightCardEditorViewportClient::UpdateProxyTransforms(const FActorProxy& InActorProxy)
@@ -1592,6 +1594,8 @@ void FDisplayClusterLightCardEditorViewportClient::UpdatePreviewActor(ADisplayCl
 			RootActorLevelInstance = RootActorPtr;
 
 			SubscribeToRootActor();
+
+			TArray<TObjectPtr<AActor>> ActorProxiesCreated;
 			
 			if (ProxyType == EDisplayClusterLightCardEditorProxyType::All ||
 				ProxyType == EDisplayClusterLightCardEditorProxyType::RootActor)
@@ -1646,14 +1650,20 @@ void FDisplayClusterLightCardEditorViewportClient::UpdatePreviewActor(ADisplayCl
 				{
 					if (StageActor)
 					{
-						CreateStageActorProxy(StageActor);
+						if (AActor* ActorProxy = CreateStageActorProxy(StageActor))
+						{
+							ActorProxiesCreated.Add(ActorProxy);
+						}
 					}
 					else
 					{
 						TArray<AActor*> ManagedActors = LightCardEditorPtr.Pin()->FindAllManagedActors();
 						for (AActor* Actor : ManagedActors)
 						{
-							CreateStageActorProxy(Actor);
+							if (AActor* ActorProxy = CreateStageActorProxy(Actor))
+							{
+								ActorProxiesCreated.Add(ActorProxy);
+							}
 						}
 					}
 				}
@@ -1669,9 +1679,9 @@ void FDisplayClusterLightCardEditorViewportClient::UpdatePreviewActor(ADisplayCl
 			}
 
 			// Make sure proxies are added to the renderer. Necessary for selections to render even if stage actors were not modified but root actor was updated
-			for (const FActorProxy& ActorProxy : ActorProxies)
+			for (const TObjectPtr<AActor>& ActorProxy : ActorProxiesCreated)
 			{
-				IDisplayClusterScenePreview::Get().AddActorToRenderer(PreviewRendererId, ActorProxy.Proxy.AsActor(), [this, ActorProxy](const UPrimitiveComponent* PrimitiveComponent)
+				IDisplayClusterScenePreview::Get().AddActorToRenderer(PreviewRendererId, ActorProxy, [this, ActorProxy](const UPrimitiveComponent* PrimitiveComponent)
 				{
 					// Always add the light card mesh component to the renderer's scene even if it is marked hidden in game, since UV light cards will purposefully
 					// hide the light card mesh since it isn't supposed to exist in 3D space. The light card mesh will be appropriately filtered when the scene is
diff --git a/Engine/Plugins/Runtime/nDisplay/Source/DisplayClusterLightCardEditor/Private/Viewport/DisplayClusterLightCardEditorViewportClient.h b/Engine/Plugins/Runtime/nDisplay/Source/DisplayClusterLightCardEditor/Private/Viewport/DisplayClusterLightCardEditorViewportClient.h
--- a/Engine/Plugins/Runtime/nDisplay/Source/DisplayClusterLightCardEditor/Private/Viewport/DisplayClusterLightCardEditorViewportClient.h
+++ b/Engine/Plugins/Runtime/nDisplay/Source/DisplayClusterLightCardEditor/Private/Viewport/DisplayClusterLightCardEditorViewportClient.h
@@ -288,7 +288,7 @@ private:
 	const FActorProxy* FindActorProxyFromActor(AActor* InActor) const;
 
 	/** Create a stage actor proxy */
-	void CreateStageActorProxy(AActor* InLevelInstance);
+	AActor* CreateStageActorProxy(AActor* InLevelInstance);
 	
 	/** Update the transforms for the specific proxy */
 	void UpdateProxyTransforms(const FActorProxy& InActorProxy);
