UnrealEditor: Handle private includes from so the path isn't rooted in the module parent. This should be cleaned up to not require including private headers in public files.

#rnx

diff --git a/Engine/Plugins/Compositing/Composure/Source/Composure/Composure.Build.cs b/Engine/Plugins/Compositing/Composure/Source/Composure/Composure.Build.cs
--- a/Engine/Plugins/Compositing/Composure/Source/Composure/Composure.Build.cs
+++ b/Engine/Plugins/Compositing/Composure/Source/Composure/Composure.Build.cs
@@ -1,5 +1,7 @@
 // Copyright Epic Games, Inc. All Rights Reserved.
 
+using System.IO;
+
 namespace UnrealBuildTool.Rules
 {
 	public class Composure : ModuleRules
@@ -54,6 +56,9 @@ namespace UnrealBuildTool.Rules
 					}
 					);
             }
-        }
-    }
+
+			// TODO: Use proper module reference (currently cannot reference editor module for non-editor target)
+			PrivateIncludePaths.Add(Path.Combine(GetModuleDirectory("ComposureLayersEditor"), "Public")); // ICompElementManager.h & CompElementEditorModule.h
+		}
+	}
 }
diff --git a/Engine/Plugins/Compositing/Composure/Source/Composure/Private/ComposureBlueprintLibrary.cpp b/Engine/Plugins/Compositing/Composure/Source/Composure/Private/ComposureBlueprintLibrary.cpp
--- a/Engine/Plugins/Compositing/Composure/Source/Composure/Private/ComposureBlueprintLibrary.cpp
+++ b/Engine/Plugins/Compositing/Composure/Source/Composure/Private/ComposureBlueprintLibrary.cpp
@@ -6,13 +6,13 @@
 #include "Camera/PlayerCameraManager.h"
 #include "CameraCalibrationSubsystem.h"
 #include "CineCameraComponent.h"
+#include "CompElementEditorModule.h"
 #include "Components/SceneCaptureComponent2D.h"
-#include "ComposureLayersEditor/Public/ICompElementManager.h"
-#include "ComposureLayersEditor/Public/CompElementEditorModule.h"
 #include "ComposurePlayerCompositingTarget.h"
 #include "ComposureUtils.h"
 #include "Engine/LocalPlayer.h"
 #include "GameFramework/PlayerController.h"
+#include "ICompElementManager.h"
 #include "LensDistortionModelHandlerBase.h"
 #include "Modules/ModuleManager.h"
 #include "Slate/SceneViewport.h"
diff --git a/Engine/Plugins/Developer/Concert/ConcertSync/ConcertSyncTest/Source/ConcertSyncTest/ConcertSyncTest.Build.cs b/Engine/Plugins/Developer/Concert/ConcertSync/ConcertSyncTest/Source/ConcertSyncTest/ConcertSyncTest.Build.cs
--- a/Engine/Plugins/Developer/Concert/ConcertSync/ConcertSyncTest/Source/ConcertSyncTest/ConcertSyncTest.Build.cs
+++ b/Engine/Plugins/Developer/Concert/ConcertSync/ConcertSyncTest/Source/ConcertSyncTest/ConcertSyncTest.Build.cs
@@ -1,5 +1,7 @@
 // Copyright Epic Games, Inc. All Rights Reserved.
 
+using System.IO;
+
 namespace UnrealBuildTool.Rules
 {
 	public class ConcertSyncTest : ModuleRules
@@ -24,6 +26,13 @@ namespace UnrealBuildTool.Rules
 					"ConcertSyncServer",
 				}
 			);
+
+			PrivateIncludePaths.AddRange(
+				new string[]
+				{
+					Path.Combine(GetModuleDirectory("Core"), "Tests"),
+				}
+			);
 		}
 	}
 }
diff --git a/Engine/Plugins/Developer/Concert/ConcertSync/ConcertSyncTest/Source/ConcertSyncTest/Private/HistoryEdition/HistoryBottomUpAnalysisTests.cpp b/Engine/Plugins/Developer/Concert/ConcertSync/ConcertSyncTest/Source/ConcertSyncTest/Private/HistoryEdition/HistoryBottomUpAnalysisTests.cpp
--- a/Engine/Plugins/Developer/Concert/ConcertSync/ConcertSyncTest/Source/ConcertSyncTest/Private/HistoryEdition/HistoryBottomUpAnalysisTests.cpp
+++ b/Engine/Plugins/Developer/Concert/ConcertSync/ConcertSyncTest/Source/ConcertSyncTest/Private/HistoryEdition/HistoryBottomUpAnalysisTests.cpp
@@ -10,7 +10,7 @@
 #include "Util/ScopedSessionDatabase.h"
 
 #include "Algo/AllOf.h"
-#include "Core/Tests/Containers/TestUtils.h"
+#include "Containers/TestUtils.h"
 #include "Misc/AutomationTest.h"
 
 namespace UE::ConcertSyncTests::AnalysisTests
diff --git a/Engine/Plugins/Developer/Concert/ConcertSync/ConcertSyncTest/Source/ConcertSyncTest/Private/HistoryEdition/HistoryTopDownAnalysisTests.cpp b/Engine/Plugins/Developer/Concert/ConcertSync/ConcertSyncTest/Source/ConcertSyncTest/Private/HistoryEdition/HistoryTopDownAnalysisTests.cpp
--- a/Engine/Plugins/Developer/Concert/ConcertSync/ConcertSyncTest/Source/ConcertSyncTest/Private/HistoryEdition/HistoryTopDownAnalysisTests.cpp
+++ b/Engine/Plugins/Developer/Concert/ConcertSync/ConcertSyncTest/Source/ConcertSyncTest/Private/HistoryEdition/HistoryTopDownAnalysisTests.cpp
@@ -10,7 +10,7 @@
 #include "Util/ScopedSessionDatabase.h"
 
 #include "Algo/AllOf.h"
-#include "Core/Tests/Containers/TestUtils.h"
+#include "Containers/TestUtils.h"
 #include "Misc/AutomationTest.h"
 
 namespace UE::ConcertSyncTests
diff --git a/Engine/Plugins/Editor/SpeedTreeImporter/Source/SpeedTreeImporter/Private/SpeedTreeImportFactory.cpp b/Engine/Plugins/Editor/SpeedTreeImporter/Source/SpeedTreeImporter/Private/SpeedTreeImportFactory.cpp
--- a/Engine/Plugins/Editor/SpeedTreeImporter/Source/SpeedTreeImporter/Private/SpeedTreeImportFactory.cpp
+++ b/Engine/Plugins/Editor/SpeedTreeImporter/Source/SpeedTreeImporter/Private/SpeedTreeImportFactory.cpp
@@ -7,8 +7,6 @@
 #include "Editor.h"
 #include "EditorFramework/AssetImportData.h"
 #include "EditorReimportHandler.h"
-#include "RenderUtils.h"
-#include "Styling/AppStyle.h"
 #include "Engine/Engine.h"
 #include "Engine/StaticMesh.h"
 #include "Engine/Texture.h"
@@ -27,8 +25,10 @@
 #include "Misc/PackageName.h"
 #include "Misc/Paths.h"
 #include "Modules/ModuleManager.h"
+#include "RenderUtils.h"
 #include "SlateOptMacros.h"
 #include "StaticParameterSet.h"
+#include "Styling/AppStyle.h"
 #include "Styling/SlateTypes.h"
 #include "UObject/ConstructorHelpers.h"
 #include "UObject/GCObject.h"
@@ -50,25 +50,25 @@
 #include "Materials/MaterialExpressionConstant.h"
 #include "Materials/MaterialExpressionConstant3Vector.h"
 #include "Materials/MaterialExpressionConstant4Vector.h"
+#include "Materials/MaterialExpressionConstantBiasScale.h"
 #include "Materials/MaterialExpressionLinearInterpolate.h"
-#include "Materials/MaterialExpressionPower.h"
 #include "Materials/MaterialExpressionMultiply.h"
+#include "Materials/MaterialExpressionPower.h"
 #include "Materials/MaterialExpressionSpeedTree.h"
 #include "Materials/MaterialExpressionTextureCoordinate.h"
 #include "Materials/MaterialExpressionTextureSample.h"
 #include "Materials/MaterialExpressionTwoSidedSign.h"
 #include "Materials/MaterialExpressionVertexColor.h"
-#include "Materials/MaterialExpressionConstantBiasScale.h"
 
+#include "AssetRegistry/AssetRegistryModule.h"
+#include "AssetToolsModule.h"
+#include "ComponentReregisterContext.h"
+#include "Editor/UnrealEd/Private/GeomFitUtils.h"
 #include "Interfaces/IMainFrameModule.h"
 #include "ObjectTools.h"
 #include "PackageTools.h"
-#include "AssetRegistry/AssetRegistryModule.h"
-#include "AssetToolsModule.h"
-#include "UnrealEd/Private/GeomFitUtils.h"
 #include "SpeedTreeWind.h"
 #include "StaticMeshAttributes.h"
-#include "ComponentReregisterContext.h"
 
 #if WITH_SPEEDTREE
 
diff --git a/Engine/Plugins/Experimental/Animation/PoseSearch/Source/Editor/PoseSearchEditor.Build.cs b/Engine/Plugins/Experimental/Animation/PoseSearch/Source/Editor/PoseSearchEditor.Build.cs
--- a/Engine/Plugins/Experimental/Animation/PoseSearch/Source/Editor/PoseSearchEditor.Build.cs
+++ b/Engine/Plugins/Experimental/Animation/PoseSearch/Source/Editor/PoseSearchEditor.Build.cs
@@ -1,5 +1,6 @@
 // Copyright Epic Games, Inc. All Rights Reserved.
 
+using System.IO;
 using UnrealBuildTool;
 
 public class PoseSearchEditor : ModuleRules
@@ -51,5 +52,11 @@ public class PoseSearchEditor : ModuleRules
 				"ToolWidgets"
 			}
 		);
+
+		// TODO: Should not be including private headers from a different module
+		PrivateIncludePaths.AddRange(
+			new string[] {
+				Path.Combine(GetModuleDirectory("PoseSearch"), "Private"), // For PoseSearchTraceLogger.h
+			});
 	}
 }
\ No newline at end of file
diff --git a/Engine/Plugins/Experimental/Animation/PoseSearch/Source/Editor/Private/Trace/PoseSearchTraceAnalyzer.cpp b/Engine/Plugins/Experimental/Animation/PoseSearch/Source/Editor/Private/Trace/PoseSearchTraceAnalyzer.cpp
--- a/Engine/Plugins/Experimental/Animation/PoseSearch/Source/Editor/Private/Trace/PoseSearchTraceAnalyzer.cpp
+++ b/Engine/Plugins/Experimental/Animation/PoseSearch/Source/Editor/Private/Trace/PoseSearchTraceAnalyzer.cpp
@@ -4,7 +4,7 @@
 
 #include "HAL/LowLevelMemTracker.h"
 #include "PoseSearchTraceProvider.h"
-#include "Runtime/Private/Trace/PoseSearchTraceLogger.h"
+#include "Trace/PoseSearchTraceLogger.h"
 #include "TraceServices/Model/AnalysisSession.h"
 
 namespace UE::PoseSearch
diff --git a/Engine/Plugins/Experimental/Animation/PoseSearch/Source/Editor/Private/Trace/PoseSearchTraceProvider.h b/Engine/Plugins/Experimental/Animation/PoseSearch/Source/Editor/Private/Trace/PoseSearchTraceProvider.h
--- a/Engine/Plugins/Experimental/Animation/PoseSearch/Source/Editor/Private/Trace/PoseSearchTraceProvider.h
+++ b/Engine/Plugins/Experimental/Animation/PoseSearch/Source/Editor/Private/Trace/PoseSearchTraceProvider.h
@@ -4,7 +4,7 @@
 
 #include "CoreMinimal.h"
 #include "Model/PointTimeline.h"
-#include "Runtime/Private/Trace/PoseSearchTraceLogger.h"
+#include "Trace/PoseSearchTraceLogger.h"
 #include "TraceServices/Model/AnalysisSession.h"
 
 
diff --git a/Engine/Plugins/Experimental/DataInterfaceGraph/Source/DataInterfaceGraphEditor/Private/DataInterfacePropertyTypeCustomization.cpp b/Engine/Plugins/Experimental/DataInterfaceGraph/Source/DataInterfaceGraphEditor/Private/DataInterfacePropertyTypeCustomization.cpp
--- a/Engine/Plugins/Experimental/DataInterfaceGraph/Source/DataInterfaceGraphEditor/Private/DataInterfacePropertyTypeCustomization.cpp
+++ b/Engine/Plugins/Experimental/DataInterfaceGraph/Source/DataInterfaceGraphEditor/Private/DataInterfacePropertyTypeCustomization.cpp
@@ -3,16 +3,16 @@
 #include "DataInterfacePropertyTypeCustomization.h"
 
 #include "DataInterfaceGraph.h"
-#include "PropertyHandle.h"
-#include "IDataInterface.h"
+#include "DataInterfaceWidgetFactories.h"
 #include "DetailWidgetRow.h"
+#include "Editor/PropertyEditor/Private/ObjectPropertyNode.h"
+#include "Editor/PropertyEditor/Private/PropertyNode.h"
+#include "IDataInterface.h"
+#include "IDetailChildrenBuilder.h"
 #include "PropertyCustomizationHelpers.h"
-#include "Widgets/Text/STextBlock.h"
+#include "PropertyHandle.h"
 #include "SClassViewer.h"
-#include "DataInterfaceWidgetFactories.h"
-#include "IDetailChildrenBuilder.h"
-#include "PropertyEditor/Private/ObjectPropertyNode.h"
-#include "PropertyEditor/Private/PropertyNode.h"
+#include "Widgets/Text/STextBlock.h"
 
 #define LOCTEXT_NAMESPACE "DataInterfacePropertyTypeCustomization"
 
diff --git a/Engine/Plugins/Experimental/GPULightmass/Source/GPULightmass/Private/LightmapDenoising.h b/Engine/Plugins/Experimental/GPULightmass/Source/GPULightmass/Private/LightmapDenoising.h
--- a/Engine/Plugins/Experimental/GPULightmass/Source/GPULightmass/Private/LightmapDenoising.h
+++ b/Engine/Plugins/Experimental/GPULightmass/Source/GPULightmass/Private/LightmapDenoising.h
@@ -13,7 +13,7 @@
 #endif
 
 // TBB suffers from extreme fragmentation problem in editor
-#include "Core/Private/HAL/Allocators/AnsiAllocator.h"
+#include "Runtime/Core/Private/HAL/Allocators/AnsiAllocator.h"
 
 struct FDenoiserFilterSet
 {
diff --git a/Engine/Plugins/Experimental/GPULightmass/Source/GPULightmass/Private/LightmapRenderer.h b/Engine/Plugins/Experimental/GPULightmass/Source/GPULightmass/Private/LightmapRenderer.h
--- a/Engine/Plugins/Experimental/GPULightmass/Source/GPULightmass/Private/LightmapRenderer.h
+++ b/Engine/Plugins/Experimental/GPULightmass/Source/GPULightmass/Private/LightmapRenderer.h
@@ -2,12 +2,12 @@
 
 #pragma once
 
+#include "GPULightmassSettings.h"
 #include "LightmapTilePool.h"
 #include "Misc/QueuedThreadPool.h"
 #include "RHIGPUReadback.h"
-#include "GPULightmassSettings.h"
 // TBB suffers from extreme fragmentation problem in editor
-#include "Core/Private/HAL/Allocators/AnsiAllocator.h"
+#include "Runtime/Core/Private/HAL/Allocators/AnsiAllocator.h"
 
 namespace GPULightmass
 {
diff --git a/Engine/Plugins/Runtime/nDisplay/Source/DisplayClusterConfigurator/Private/Views/Viewport/DisplayClusterConfiguratorSCSEditorViewportClient.h b/Engine/Plugins/Runtime/nDisplay/Source/DisplayClusterConfigurator/Private/Views/Viewport/DisplayClusterConfiguratorSCSEditorViewportClient.h
--- a/Engine/Plugins/Runtime/nDisplay/Source/DisplayClusterConfigurator/Private/Views/Viewport/DisplayClusterConfiguratorSCSEditorViewportClient.h
+++ b/Engine/Plugins/Runtime/nDisplay/Source/DisplayClusterConfigurator/Private/Views/Viewport/DisplayClusterConfiguratorSCSEditorViewportClient.h
@@ -3,7 +3,7 @@
 #pragma once
 
 #include "CoreMinimal.h"
-#include "Kismet/Private/SCSEditorViewportClient.h"
+#include "Editor/Kismet/Private/SCSEditorViewportClient.h"
 #include "UnrealWidgetFwd.h"
 
 class UMaterialInstanceConstant;
diff --git a/Engine/Plugins/Runtime/nDisplay/Source/DisplayClusterEditor/DisplayClusterEditor.Build.cs b/Engine/Plugins/Runtime/nDisplay/Source/DisplayClusterEditor/DisplayClusterEditor.Build.cs
--- a/Engine/Plugins/Runtime/nDisplay/Source/DisplayClusterEditor/DisplayClusterEditor.Build.cs
+++ b/Engine/Plugins/Runtime/nDisplay/Source/DisplayClusterEditor/DisplayClusterEditor.Build.cs
@@ -16,5 +16,8 @@ public class DisplayClusterEditor : ModuleRules
 				"Engine",
 				"UnrealEd"
 			});
+
+		// TODO: Should not be including private headers from other modules
+		PrivateIncludePaths.Add(Path.Combine(GetModuleDirectory("DisplayCluster"), "Private")); // For IPDisplayCluster.h
 	}
 }
diff --git a/Engine/Plugins/Runtime/nDisplay/Source/DisplayClusterEditor/Private/DisplayClusterEditorEngine.cpp b/Engine/Plugins/Runtime/nDisplay/Source/DisplayClusterEditor/Private/DisplayClusterEditorEngine.cpp
--- a/Engine/Plugins/Runtime/nDisplay/Source/DisplayClusterEditor/Private/DisplayClusterEditorEngine.cpp
+++ b/Engine/Plugins/Runtime/nDisplay/Source/DisplayClusterEditor/Private/DisplayClusterEditorEngine.cpp
@@ -5,10 +5,10 @@
 
 #include "DisplayClusterRootActor.h"
 
-#include "IDisplayClusterConfiguration.h"
 #include "DisplayClusterConfigurationTypes.h"
+#include "IDisplayClusterConfiguration.h"
 
-#include "DisplayCluster/Private/IPDisplayCluster.h"
+#include "IPDisplayCluster.h"
 
 #include "Engine/LevelStreaming.h"
 
diff --git a/Engine/Plugins/Runtime/nDisplay/Source/DisplayClusterLightCardEditor/Private/Viewport/DisplayClusterLightCardEditorViewportClient.cpp b/Engine/Plugins/Runtime/nDisplay/Source/DisplayClusterLightCardEditor/Private/Viewport/DisplayClusterLightCardEditorViewportClient.cpp
--- a/Engine/Plugins/Runtime/nDisplay/Source/DisplayClusterLightCardEditor/Private/Viewport/DisplayClusterLightCardEditorViewportClient.cpp
+++ b/Engine/Plugins/Runtime/nDisplay/Source/DisplayClusterLightCardEditor/Private/Viewport/DisplayClusterLightCardEditorViewportClient.cpp
@@ -7,39 +7,40 @@
 #include "DisplayClusterLightCardEditorWidget.h"
 #include "LightCardTemplates/DisplayClusterLightCardTemplate.h"
 
-#include "Components/DisplayClusterPreviewComponent.h"
 #include "Components/DisplayClusterCameraComponent.h"
+#include "Components/DisplayClusterPreviewComponent.h"
 #include "Components/DisplayClusterScreenComponent.h"
 #include "DisplayClusterConfigurationTypes.h"
-#include "DisplayClusterLightCardEditorHelper.h"
-#include "DisplayClusterProjectionStrings.h"
-#include "DisplayClusterRootActor.h"
 #include "DisplayClusterLightCardActor.h"
+#include "DisplayClusterLightCardEditor.h"
+#include "DisplayClusterLightCardEditorHelper.h"
 #include "DisplayClusterLightCardEditorLog.h"
 #include "DisplayClusterLightCardEditorUtils.h"
+#include "DisplayClusterProjectionStrings.h"
+#include "DisplayClusterRootActor.h"
 #include "IDisplayClusterScenePreview.h"
-#include "DisplayClusterLightCardEditor.h"
 #include "SceneManagement.h"
 #include "Settings/DisplayClusterLightCardEditorSettings.h"
 
 #include "AudioDevice.h"
 #include "CameraController.h"
+#include "Components/BillboardComponent.h"
 #include "Components/DirectionalLightComponent.h"
 #include "Components/LineBatchComponent.h"
 #include "Components/PostProcessComponent.h"
-#include "Components/SkyLightComponent.h"
 #include "Components/PrimitiveComponent.h"
+#include "Components/SkyLightComponent.h"
 #include "Debug/DebugDrawService.h"
-#include "EngineUtils.h"
-#include "EditorModes.h"
 #include "Editor/UnrealEdEngine.h"
 #include "EditorModeManager.h"
-#include "EngineModule.h"
+#include "EditorModes.h"
 #include "Engine/Canvas.h"
 #include "Engine/TextureRenderTarget2D.h"
+#include "EngineModule.h"
+#include "EngineUtils.h"
 #include "ImageUtils.h"
-#include "Kismet/KismetMathLibrary.h"
 #include "Kismet/GameplayStatics.h"
+#include "Kismet/KismetMathLibrary.h"
 #include "KismetProceduralMeshLibrary.h"
 #include "LegacyScreenPercentageDriver.h"
 #include "Math/UnrealMathUtility.h"
@@ -47,14 +48,13 @@
 #include "PreviewScene.h"
 #include "ProceduralMeshComponent.h"
 #include "RayTracingDebugVisualizationMenuCommands.h"
-#include "Renderer/Private/SceneRendering.h"
+#include "Runtime/Renderer/Private/SceneRendering.h"
 #include "ScopedTransaction.h"
 #include "Settings/LevelEditorViewportSettings.h"
 #include "Slate/SceneViewport.h"
 #include "TextureResource.h"
 #include "UnrealEdGlobals.h"
 #include "UnrealWidget.h"
-#include "Components/BillboardComponent.h"
 #include "Widgets/Docking/SDockTab.h"
 
 
diff --git a/Engine/Plugins/Runtime/nDisplay/Source/DisplayClusterTests/DisplayClusterTests.Build.cs b/Engine/Plugins/Runtime/nDisplay/Source/DisplayClusterTests/DisplayClusterTests.Build.cs
--- a/Engine/Plugins/Runtime/nDisplay/Source/DisplayClusterTests/DisplayClusterTests.Build.cs
+++ b/Engine/Plugins/Runtime/nDisplay/Source/DisplayClusterTests/DisplayClusterTests.Build.cs
@@ -20,6 +20,12 @@ public class DisplayClusterTests : ModuleRules
 				"UnrealEd"
 			});
 
+		// TODO: Should not be including private headers from a different module
+		PrivateIncludePaths.AddRange(
+			new string[] {
+				Path.Combine(GetModuleDirectory("DisplayClusterConfigurator"), "Private"), // For DisplayClusterConfiguratorPropertyUtils.h
+			});
+
 		OptimizeCode = CodeOptimization.Never;
 		PCHUsage = PCHUsageMode.NoPCHs;
 	}
diff --git a/Engine/Plugins/Runtime/nDisplay/Source/DisplayClusterTests/Private/DisplayCluster.spec.cpp b/Engine/Plugins/Runtime/nDisplay/Source/DisplayClusterTests/Private/DisplayCluster.spec.cpp
--- a/Engine/Plugins/Runtime/nDisplay/Source/DisplayClusterTests/Private/DisplayCluster.spec.cpp
+++ b/Engine/Plugins/Runtime/nDisplay/Source/DisplayClusterTests/Private/DisplayCluster.spec.cpp
@@ -1,14 +1,14 @@
 // Copyright Epic Games, Inc. All Rights Reserved.
 
-#include "ComponentAssetBroker.h"
-#include "DisplayClusterRootActor.h"
-#include "DisplayClusterTestUtils.h"
 #include "AssetRegistry/AssetRegistryModule.h"
 #include "Blueprints/DisplayClusterBlueprint.h"
 #include "Blueprints/DisplayClusterBlueprintGeneratedClass.h"
+#include "ClusterConfiguration/DisplayClusterConfiguratorClusterUtils.h"
+#include "ComponentAssetBroker.h"
 #include "Components/DisplayClusterScreenComponent.h"
 #include "Components/DisplayClusterXformComponent.h"
-#include "DisplayClusterConfigurator/Private/ClusterConfiguration/DisplayClusterConfiguratorClusterUtils.h"
+#include "DisplayClusterRootActor.h"
+#include "DisplayClusterTestUtils.h"
 #include "Engine/SCS_Node.h"
 #include "Engine/SimpleConstructionScript.h"
 #include "Kismet2/ComponentEditorUtils.h"
diff --git a/Engine/Plugins/Runtime/nDisplay/Source/DisplayClusterTests/Private/DisplayClusterTestUtils.cpp b/Engine/Plugins/Runtime/nDisplay/Source/DisplayClusterTests/Private/DisplayClusterTestUtils.cpp
--- a/Engine/Plugins/Runtime/nDisplay/Source/DisplayClusterTests/Private/DisplayClusterTestUtils.cpp
+++ b/Engine/Plugins/Runtime/nDisplay/Source/DisplayClusterTests/Private/DisplayClusterTestUtils.cpp
@@ -2,15 +2,15 @@
 
 #include "DisplayClusterTestUtils.h"
 
+#include "AssetRegistry/AssetRegistryModule.h"
+#include "Blueprints/DisplayClusterBlueprint.h"
+#include "ClusterConfiguration/DisplayClusterConfiguratorClusterUtils.h"
+#include "DisplayClusterConfiguratorFactory.h"
 #include "DisplayClusterTestsModule.h"
 #include "Engine/Engine.h"
 #include "Engine/World.h"
 #include "ObjectTools.h"
 #include "PackageTools.h"
-#include "AssetRegistry/AssetRegistryModule.h"
-#include "Blueprints/DisplayClusterBlueprint.h"
-#include "DisplayClusterConfigurator/Private/ClusterConfiguration/DisplayClusterConfiguratorClusterUtils.h"
-#include "DisplayClusterConfigurator/Private/DisplayClusterConfiguratorFactory.h"
 
 #if WITH_DEV_AUTOMATION_TESTS
 
diff --git a/Engine/Plugins/Runtime/nDisplay/Source/DisplayClusterTests/Private/DisplayClusterTestUtils.h b/Engine/Plugins/Runtime/nDisplay/Source/DisplayClusterTests/Private/DisplayClusterTestUtils.h
--- a/Engine/Plugins/Runtime/nDisplay/Source/DisplayClusterTests/Private/DisplayClusterTestUtils.h
+++ b/Engine/Plugins/Runtime/nDisplay/Source/DisplayClusterTests/Private/DisplayClusterTestUtils.h
@@ -2,11 +2,11 @@
 
 #pragma once
 #include "DisplayClusterConfigurationTypes_Viewport.h"
-#include "PropertyEditor/Private/PropertyNode.h"
+#include "Editor/PropertyEditor/Private/PropertyNode.h"
 
 #if WITH_DEV_AUTOMATION_TESTS
 
-#include "DisplayClusterConfigurator/Private/DisplayClusterConfiguratorPropertyUtils.h"
+#include "DisplayClusterConfiguratorPropertyUtils.h"
 #include "Kismet2/BlueprintEditorUtils.h"
 
 class UDisplayClusterBlueprint;
diff --git a/Engine/Plugins/VirtualProduction/LevelSnapshots/Source/LevelSnapshotsEditor/LevelSnapshotsEditor.build.cs b/Engine/Plugins/VirtualProduction/LevelSnapshots/Source/LevelSnapshotsEditor/LevelSnapshotsEditor.build.cs
--- a/Engine/Plugins/VirtualProduction/LevelSnapshots/Source/LevelSnapshotsEditor/LevelSnapshotsEditor.build.cs
+++ b/Engine/Plugins/VirtualProduction/LevelSnapshots/Source/LevelSnapshotsEditor/LevelSnapshotsEditor.build.cs
@@ -44,5 +44,11 @@ public class LevelSnapshotsEditor : ModuleRules
 				"WorkspaceMenuStructure"
 			}
 			);
+
+		// TODO: Should not be including private headers from a different module
+		PrivateIncludePaths.AddRange(
+			new string[] {
+				Path.Combine(GetModuleDirectory("GameplayMediaEncoder"), "Private"), // For GameplayMediaEncoderCommon.h
+			});
 	}
 }
diff --git a/Engine/Plugins/VirtualProduction/LevelSnapshots/Source/LevelSnapshotsEditor/Private/LevelSnapshotsEditorFunctionLibrary.cpp b/Engine/Plugins/VirtualProduction/LevelSnapshots/Source/LevelSnapshotsEditor/Private/LevelSnapshotsEditorFunctionLibrary.cpp
--- a/Engine/Plugins/VirtualProduction/LevelSnapshots/Source/LevelSnapshotsEditor/Private/LevelSnapshotsEditorFunctionLibrary.cpp
+++ b/Engine/Plugins/VirtualProduction/LevelSnapshots/Source/LevelSnapshotsEditor/Private/LevelSnapshotsEditorFunctionLibrary.cpp
@@ -6,8 +6,8 @@
 
 #include "Editor.h"
 #include "Engine/World.h"
+#include "GameplayMediaEncoderCommon.h"
 #include "LevelSnapshot.h"
-#include "GameplayMediaEncoder/Private/GameplayMediaEncoderCommon.h"
 #include "Logging/MessageLog.h"
 #include "UObject/MetaData.h"
 #include "Util/TakeSnapshotUtil.h"
diff --git a/Engine/Plugins/VirtualProduction/LevelSnapshots/Source/LevelSnapshotsEditor/Private/Util/TakeSnapshotUtil.cpp b/Engine/Plugins/VirtualProduction/LevelSnapshots/Source/LevelSnapshotsEditor/Private/Util/TakeSnapshotUtil.cpp
--- a/Engine/Plugins/VirtualProduction/LevelSnapshots/Source/LevelSnapshotsEditor/Private/Util/TakeSnapshotUtil.cpp
+++ b/Engine/Plugins/VirtualProduction/LevelSnapshots/Source/LevelSnapshotsEditor/Private/Util/TakeSnapshotUtil.cpp
@@ -5,8 +5,8 @@
 #include "Components/Viewport.h"
 #include "Data/LevelSnapshotsEditorData.h"
 #include "LevelSnapshot.h"
-#include "LevelSnapshotsEditorSettings.h"
 #include "LevelSnapshotsEditorModule.h"
+#include "LevelSnapshotsEditorSettings.h"
 #include "LevelSnapshotsEditorStyle.h"
 #include "LevelSnapshotsLog.h"
 #include "Widgets/SLevelSnapshotsEditorCreationForm.h"
@@ -14,14 +14,14 @@
 #include "AssetRegistry/AssetRegistryModule.h"
 #include "AssetToolsModule.h"
 #include "AssetViewUtils.h"
-#include "IAssetTools.h"
 #include "Editor.h"
 #include "Engine/World.h"
-#include "ObjectTools.h"
 #include "Framework/Notifications/NotificationManager.h"
-#include "GameplayMediaEncoder/Private/GameplayMediaEncoderCommon.h"
+#include "GameplayMediaEncoderCommon.h"
+#include "IAssetTools.h"
 #include "Misc/FeedbackContext.h"
 #include "Misc/MessageDialog.h"
+#include "ObjectTools.h"
 #include "UObject/Package.h"
 #include "UObject/SavePackage.h"
 #include "Widgets/Notifications/SNotificationList.h"
diff --git a/Engine/Plugins/VirtualProduction/LevelSnapshots/Source/LevelSnapshotsEditor/Private/Widgets/SLevelSnapshotsEditorBrowser.cpp b/Engine/Plugins/VirtualProduction/LevelSnapshots/Source/LevelSnapshotsEditor/Private/Widgets/SLevelSnapshotsEditorBrowser.cpp
--- a/Engine/Plugins/VirtualProduction/LevelSnapshots/Source/LevelSnapshotsEditor/Private/Widgets/SLevelSnapshotsEditorBrowser.cpp
+++ b/Engine/Plugins/VirtualProduction/LevelSnapshots/Source/LevelSnapshotsEditor/Private/Widgets/SLevelSnapshotsEditorBrowser.cpp
@@ -2,10 +2,10 @@
 
 #include "Widgets/SLevelSnapshotsEditorBrowser.h"
 
-#include "LevelSnapshot.h"
 #include "Data/LevelSnapshotsEditorData.h"
-#include "LevelSnapshotsEditor/Private/LevelSnapshotsEditorModule.h"
+#include "LevelSnapshot.h"
 #include "LevelSnapshotsEditorFunctionLibrary.h"
+#include "LevelSnapshotsEditorModule.h"
 #include "LevelSnapshotsEditorStyle.h"
 #include "LevelSnapshotsLog.h"
 
diff --git a/Engine/Plugins/VirtualProduction/RemoteControl/Source/RemoteControlUI/Private/Commands/RemoteControlCommands.cpp b/Engine/Plugins/VirtualProduction/RemoteControl/Source/RemoteControlUI/Private/Commands/RemoteControlCommands.cpp
--- a/Engine/Plugins/VirtualProduction/RemoteControl/Source/RemoteControlUI/Private/Commands/RemoteControlCommands.cpp
+++ b/Engine/Plugins/VirtualProduction/RemoteControl/Source/RemoteControlUI/Private/Commands/RemoteControlCommands.cpp
@@ -2,7 +2,7 @@
 
 #include "Commands/RemoteControlCommands.h"
 
-#include "RemoteControlUI/Private/UI/RemoteControlPanelStyle.h"
+#include "UI/RemoteControlPanelStyle.h"
 
 #define LOCTEXT_NAMESPACE "RemoteControlCommands"
 
