nDisplay: Fixed issue where orthogonal projections in the light card editor were not depth sorting properly, caused by a redefinition of the WORLD_MAX constant to a much larger value.

diff --git a/Engine/Plugins/Runtime/nDisplay/Source/DisplayClusterLightCardEditor/Private/Viewport/DisplayClusterLightCardEditorViewportClient.cpp b/Engine/Plugins/Runtime/nDisplay/Source/DisplayClusterLightCardEditor/Private/Viewport/DisplayClusterLightCardEditorViewportClient.cpp
--- a/Engine/Plugins/Runtime/nDisplay/Source/DisplayClusterLightCardEditor/Private/Viewport/DisplayClusterLightCardEditorViewportClient.cpp
+++ b/Engine/Plugins/Runtime/nDisplay/Source/DisplayClusterLightCardEditor/Private/Viewport/DisplayClusterLightCardEditorViewportClient.cpp
@@ -1233,7 +1233,9 @@ void FDisplayClusterLightCardEditorViewportClient::ProcessClick(FSceneView& View
 		if (HitProxy->IsA(HActor::StaticGetType()))
 		{
 			HActor* ActorHitProxy = static_cast<HActor*>(HitProxy);
-			if (ActorHitProxy->Actor == RootActorProxy.Get())
+
+			// Only perform the ray trace when not in UV mode, as it doesn't make sense to ray trace against UV space
+			if (ActorHitProxy->Actor == RootActorProxy.Get() && ProjectionMode != EDisplayClusterMeshProjectionType::UV)
 			{
 				if (ActorHitProxy->PrimComponent && ActorHitProxy->PrimComponent->IsA<UStaticMeshComponent>())
 				{
diff --git a/Engine/Plugins/Runtime/nDisplay/Source/DisplayClusterScenePreview/Private/DisplayClusterLightCardEditorHelper.cpp b/Engine/Plugins/Runtime/nDisplay/Source/DisplayClusterScenePreview/Private/DisplayClusterLightCardEditorHelper.cpp
--- a/Engine/Plugins/Runtime/nDisplay/Source/DisplayClusterScenePreview/Private/DisplayClusterLightCardEditorHelper.cpp
+++ b/Engine/Plugins/Runtime/nDisplay/Source/DisplayClusterScenePreview/Private/DisplayClusterLightCardEditorHelper.cpp
@@ -821,8 +821,8 @@ void FDisplayClusterLightCardEditorHelper::GetSceneViewInitOptions(
 
 	if (bIsOrthographic)
 	{
-		const float ZScale = 0.5f / HALF_WORLD_MAX;
-		const float ZOffset = HALF_WORLD_MAX;
+		const float ZScale = 0.5f / UE_OLD_HALF_WORLD_MAX;
+		const float ZOffset = UE_OLD_HALF_WORLD_MAX;
 
 		const float FOVScale = FMath::Tan(HalfFOV) / InDPIScale;
 
