nDisplay: nDisplay PIE preview

diff --git a/Engine/Plugins/Runtime/nDisplay/Source/DisplayCluster/Private/Game/EngineClasses/Scene/DisplayClusterRootActorEditor.cpp b/Engine/Plugins/Runtime/nDisplay/Source/DisplayCluster/Private/Game/EngineClasses/Scene/DisplayClusterRootActorEditor.cpp
--- a/Engine/Plugins/Runtime/nDisplay/Source/DisplayCluster/Private/Game/EngineClasses/Scene/DisplayClusterRootActorEditor.cpp
+++ b/Engine/Plugins/Runtime/nDisplay/Source/DisplayCluster/Private/Game/EngineClasses/Scene/DisplayClusterRootActorEditor.cpp
@@ -116,7 +116,7 @@ void ADisplayClusterRootActor::Tick_Editor(float DeltaSeconds)
 	{
 		ResetPreviewInternals_Editor();
 	}
-	else if (!IsRunningGameOrPIE())
+	else if (bPIEPreviewEnable && (PreviewNodeId == DisplayClusterConfigurationStrings::gui::preview::PreviewNodeNone))
 	{
 		if (bDeferPreviewGeneration)
 		{
@@ -895,6 +895,10 @@ void ADisplayClusterRootActor::UpdatePreviewComponents()
 
 				// Always reinitialize so changes impact the preview component.
 				PreviewComp->InitializePreviewComponent(this, Node.Key, Viewport.Key, Viewport.Value);
+				if (UMeshComponent* PreviewMesh = PreviewComp->GetPreviewMesh())
+				{
+					PreviewMesh->SetHiddenInGame(false);
+				}
 
 				IteratedPreviewComponents.Add(PreviewComp);
 			}
diff --git a/Engine/Plugins/Runtime/nDisplay/Source/DisplayCluster/Public/DisplayClusterRootActor.h b/Engine/Plugins/Runtime/nDisplay/Source/DisplayCluster/Public/DisplayClusterRootActor.h
--- a/Engine/Plugins/Runtime/nDisplay/Source/DisplayCluster/Public/DisplayClusterRootActor.h
+++ b/Engine/Plugins/Runtime/nDisplay/Source/DisplayCluster/Public/DisplayClusterRootActor.h
@@ -311,6 +311,10 @@ public:
 	/** Render the scene and display it as a preview on the nDisplay root actor in the editor.  This will impact editor performance. */
 	UPROPERTY(EditAnywhere, BlueprintReadWrite, Category = "Editor Preview", meta = (DisplayName = "Enable Editor Preview"))
 	bool bPreviewEnable = true;
+
+	/** Render the scene and display it in the PIE as a preview of the nDisplay root actor. This will impact PIE performance. */
+	UPROPERTY(EditAnywhere, BlueprintReadWrite, Category = "Editor Preview", meta = (DisplayName = "Enabled in PIE"))
+	bool bPIEPreviewEnable = true;
 	
 	/** Adjust resolution scaling for the editor preview. */
 	UPROPERTY(EditAnywhere, BlueprintReadWrite, Category = "Editor Preview", meta = (DisplayName = "Preview Screen Percentage", ClampMin = "0.05", UIMin = "0.05", ClampMax = "1", UIMax = "1", EditCondition = "bPreviewEnable"))
