[EpicStageApp] Fix CCWs being created at the wrong position

diff --git a/Engine/Plugins/Runtime/nDisplay/Source/DisplayClusterScenePreview/Private/DisplayClusterLightCardEditorHelper.cpp b/Engine/Plugins/Runtime/nDisplay/Source/DisplayClusterScenePreview/Private/DisplayClusterLightCardEditorHelper.cpp
--- a/Engine/Plugins/Runtime/nDisplay/Source/DisplayClusterScenePreview/Private/DisplayClusterLightCardEditorHelper.cpp
+++ b/Engine/Plugins/Runtime/nDisplay/Source/DisplayClusterScenePreview/Private/DisplayClusterLightCardEditorHelper.cpp
@@ -464,6 +464,8 @@ void FDisplayClusterLightCardEditorHelper::MoveActorsToPixel(
 		{
 			if (Actor.IsValid())
 			{
+				VerifyAndFixActorOrigin(Actor);
+
 				const FSphericalCoordinates ActorCoords = GetActorCoordinates(Actor);
 
 				AverageCoords = AverageCoords + ActorCoords;
