nDisplay: Preview restricted to Editor or PIE (not enabled in -game).

diff --git a/Engine/Plugins/Runtime/nDisplay/Source/DisplayCluster/Private/Game/EngineClasses/Scene/DisplayClusterRootActorEditor.cpp b/Engine/Plugins/Runtime/nDisplay/Source/DisplayCluster/Private/Game/EngineClasses/Scene/DisplayClusterRootActorEditor.cpp
--- a/Engine/Plugins/Runtime/nDisplay/Source/DisplayCluster/Private/Game/EngineClasses/Scene/DisplayClusterRootActorEditor.cpp
+++ b/Engine/Plugins/Runtime/nDisplay/Source/DisplayCluster/Private/Game/EngineClasses/Scene/DisplayClusterRootActorEditor.cpp
@@ -112,11 +112,7 @@ void ADisplayClusterRootActor::Destructor_Editor()
 
 void ADisplayClusterRootActor::Tick_Editor(float DeltaSeconds)
 {
-	if (!IsPreviewEnabled())
-	{
-		ResetPreviewInternals_Editor();
-	}
-	else if (bPIEPreviewEnable && (PreviewNodeId == DisplayClusterConfigurationStrings::gui::preview::PreviewNodeNone))
+	if (IsPreviewEnabled())
 	{
 		if (bDeferPreviewGeneration)
 		{
@@ -125,7 +121,7 @@ void ADisplayClusterRootActor::Tick_Editor(float DeltaSeconds)
 			bDeferPreviewGeneration = false;
 			UpdatePreviewComponents();
 		}
-		
+
 		// Update preview RTTs correspond to 'TickPerFrame' value
 		if (++TickPerFrameCounter >= TickPerFrame)
 		{
@@ -141,6 +137,10 @@ void ADisplayClusterRootActor::Tick_Editor(float DeltaSeconds)
 		// preview frustums on each tick
 		ImplRenderPreviewFrustums_Editor();
 	}
+	else
+	{
+		ResetPreviewInternals_Editor();
+	}
 }
 
 void ADisplayClusterRootActor::PostActorCreated_Editor()
@@ -192,23 +192,15 @@ bool ADisplayClusterRootActor::IsPreviewEnabled() const
 		return false;
 	}
 
-	//@todo: (GUI) Scene preview can be disabled when the configuration window with internal preview is open.
-#if 0
-	bool bIsScenePreview = true; //@todo: handle GUI logic
-	bool bIsConfigurationPreviewUsed = false; //@todo: handle GUI logic
-
-	if (bIsScenePreview == bIsConfigurationPreviewUsed)
-	{
-		return false;
-	}
-#endif
-
-	if (!PreviewEnableOverriders.IsEmpty())
+	// -game or PIE case
+	if (IsRunningGameOrPIE())
 	{
-		return true;
+		// Only PIE is currently supported
+		return bPIEPreviewEnable && GIsPlayInEditorWorld && (PreviewNodeId == DisplayClusterConfigurationStrings::gui::preview::PreviewNodeNone);
 	}
 
-	return bPreviewEnable;
+	// Editor case
+	return bPreviewEnable || !PreviewEnableOverriders.IsEmpty();
 }
 
 // Return all RTT RHI resources for preview
