nDisplay: UV light cards now respect translucent sort priority. This allows consistent ordering between UV light cards and fixes UV lightcard indicators not displaying in non-UV modes.
When adding a new UV lightcard to the root actor it will adjust the priority to be the highest in the layer.

diff --git a/Engine/Plugins/Runtime/nDisplay/Source/DisplayCluster/Private/Game/EngineClasses/Scene/DisplayClusterLightCardActor.cpp b/Engine/Plugins/Runtime/nDisplay/Source/DisplayCluster/Private/Game/EngineClasses/Scene/DisplayClusterLightCardActor.cpp
--- a/Engine/Plugins/Runtime/nDisplay/Source/DisplayCluster/Private/Game/EngineClasses/Scene/DisplayClusterLightCardActor.cpp
+++ b/Engine/Plugins/Runtime/nDisplay/Source/DisplayCluster/Private/Game/EngineClasses/Scene/DisplayClusterLightCardActor.cpp
@@ -398,6 +398,10 @@ void ADisplayClusterLightCardActor::UpdateUVIndicator()
 
 	if (GIsEditor && bIsUVLightCard)
 	{
+		// Always keep indicator on top of the light card
+		UVIndicatorComponent->TranslucencySortPriority = LightCardComponent->TranslucencySortPriority;
+		UVIndicatorComponent->TranslucencySortDistanceOffset = LightCardComponent->TranslucencySortDistanceOffset - 0.1f;
+		
 		// Keep 1:1 aspect ratio
 		if (const USceneComponent* UVParentComponent = UVIndicatorComponent->GetAttachParent())
 		{
@@ -435,6 +439,12 @@ void ADisplayClusterLightCardActor::UpdateUVIndicator()
 
 void ADisplayClusterLightCardActor::ShowLightCardLabel(bool bValue, float ScaleValue, ADisplayClusterRootActor* InRootActor)
 {
+	if (IsUVActor())
+	{
+		// UV actors don't currently support labels
+		return;
+	}
+
 #if WITH_EDITOR
 	LabelComponent->Modify(false);
 	LightCardComponent->Modify(false);
@@ -512,6 +522,35 @@ void ADisplayClusterLightCardActor::AddToLightCardLayer(ADisplayClusterRootActor
 	ULayersSubsystem* LayersSubsystem = GEditor ? GEditor->GetEditorSubsystem<ULayersSubsystem>() : nullptr;
 	if (LayersSubsystem)
 	{
+		// Adjust transparent sort order for UV light cards
+		if (IsUVActor())
+		{
+			int32 HighestPriority = MIN_int32;
+			bool bExistingPrioritiesFound = false;
+			
+			const TArray<AActor*> ActorsInLayer = LayersSubsystem->GetActorsFromLayer(LightCardLayerName);
+			for (AActor* Actor : ActorsInLayer)
+			{
+				if (const ADisplayClusterLightCardActor* LightCard = Cast<ADisplayClusterLightCardActor>(Actor))
+				{
+					if (LightCard->IsUVActor())
+					{
+						if (LightCard->LightCardComponent->TranslucencySortPriority > HighestPriority)
+						{
+							HighestPriority = LightCard->LightCardComponent->TranslucencySortPriority;
+						}
+						bExistingPrioritiesFound = true;
+					}
+				}
+			}
+
+			if (bExistingPrioritiesFound)
+			{
+				// Newly added cards should be on the top-most visibility layer
+				LightCardComponent->TranslucencySortPriority = HighestPriority < MAX_int32 ? HighestPriority + 1 : HighestPriority;
+			}
+		}
+
 		LayersSubsystem->AddActorsToLayer(TArray<AActor*>{ this }, LightCardLayerName);
 		bLayerAdded = true;
 	}
diff --git a/Engine/Plugins/Runtime/nDisplay/Source/DisplayClusterShaders/Private/Shaders/DisplayClusterShadersPreprocess_UVLightCards.cpp b/Engine/Plugins/Runtime/nDisplay/Source/DisplayClusterShaders/Private/Shaders/DisplayClusterShadersPreprocess_UVLightCards.cpp
--- a/Engine/Plugins/Runtime/nDisplay/Source/DisplayClusterShaders/Private/Shaders/DisplayClusterShadersPreprocess_UVLightCards.cpp
+++ b/Engine/Plugins/Runtime/nDisplay/Source/DisplayClusterShaders/Private/Shaders/DisplayClusterShadersPreprocess_UVLightCards.cpp
@@ -101,7 +101,7 @@ public:
 		FMeshMaterialShaderElementData ShaderElementData;
 		ShaderElementData.InitializeMeshMaterialData(ViewIfDynamicMeshCommand, PrimitiveSceneProxy, MeshBatch, StaticMeshId, false);
 
-		FMeshDrawCommandSortKey SortKey = CalculateMeshStaticSortKey(PassShaders.VertexShader.GetShader(), PassShaders.PixelShader.GetShader());
+		FMeshDrawCommandSortKey SortKey = CreateMeshSortKey(MeshBatch, PrimitiveSceneProxy, Material, PassShaders.VertexShader.GetShader(), PassShaders.PixelShader.GetShader());
 
 		BuildMeshDrawCommands(
 			MeshBatch,
@@ -118,6 +118,45 @@ public:
 			ShaderElementData);
 	}
 
+protected:
+	virtual FMeshDrawCommandSortKey CreateMeshSortKey(const FMeshBatch& RESTRICT MeshBatch,
+	const FPrimitiveSceneProxy* RESTRICT PrimitiveSceneProxy,
+	const FMaterial& Material,
+	const FMeshMaterialShader* VertexShader,
+	const FMeshMaterialShader* PixelShader)
+	{
+		FMeshDrawCommandSortKey SortKey = FMeshDrawCommandSortKey::Default;
+		
+		uint16 SortKeyPriority = 0;
+		float Distance = 0.0f;
+
+		if (PrimitiveSceneProxy)
+		{
+			const FPrimitiveSceneInfo* PrimitiveSceneInfo = PrimitiveSceneProxy->GetPrimitiveSceneInfo();
+			SortKeyPriority = (uint16)((int32)PrimitiveSceneInfo->Proxy->GetTranslucencySortPriority() - (int32)SHRT_MIN);
+
+			// Use the standard sort by distance method for translucent objects
+			const float DistanceOffset = PrimitiveSceneInfo->Proxy->GetTranslucencySortDistanceOffset();
+			const FVector BoundsOrigin = PrimitiveSceneProxy->GetBounds().Origin;
+			const FVector ViewOrigin = this->ViewIfDynamicMeshCommand->ViewMatrices.GetViewOrigin();
+			Distance = (BoundsOrigin - ViewOrigin).Size() + DistanceOffset;
+		}
+
+		SortKey.Translucent.MeshIdInPrimitive = MeshBatch.MeshIdInPrimitive;
+		SortKey.Translucent.Priority = SortKeyPriority;
+		SortKey.Translucent.Distance = (uint32)~BitInvertIfNegativeFloat(*(uint32*)&Distance);
+
+		return SortKey;
+	}
+
+private:
+	/** Inverts the bits of the floating point number if that number is negative */
+	uint32 BitInvertIfNegativeFloat(uint32 FloatBit)
+	{
+		unsigned Mask = -int32(FloatBit >> 31) | 0x80000000;
+		return FloatBit ^ Mask;
+	}
+
 private:
 	FMeshPassProcessorRenderState DrawRenderState;
 };
@@ -187,8 +226,8 @@ bool FDisplayClusterShadersPreprocess_UVLightCards::RenderPreprocess_UVLightCard
 		FPlane(0, 1, 0, 0),
 		FPlane(0, 0, 0, 1));
 
-	const float ZScale = 0.5f / HALF_WORLD_MAX;
-	const float ZOffset = HALF_WORLD_MAX;
+	const float ZScale = 0.5f / UE_OLD_HALF_WORLD_MAX;
+	const float ZOffset = UE_OLD_HALF_WORLD_MAX;
 	const float OrthoSize = 0.5f * ProjectionPlaneSize;
 
 	if ((bool)ERHIZBuffer::IsInverted)
