nDisplay: headless rendering v2

diff --git a/Engine/Plugins/Runtime/nDisplay/Source/DisplayCluster/Private/Cluster/Controller/DisplayClusterClusterNodeCtrlBase.h b/Engine/Plugins/Runtime/nDisplay/Source/DisplayCluster/Private/Cluster/Controller/DisplayClusterClusterNodeCtrlBase.h
--- a/Engine/Plugins/Runtime/nDisplay/Source/DisplayCluster/Private/Cluster/Controller/DisplayClusterClusterNodeCtrlBase.h
+++ b/Engine/Plugins/Runtime/nDisplay/Source/DisplayCluster/Private/Cluster/Controller/DisplayClusterClusterNodeCtrlBase.h
@@ -102,7 +102,7 @@ public:
 	//////////////////////////////////////////////////////////////////////////////////////////////
 	// IDisplayClusterProtocolRenderSync - default overrides
 	//////////////////////////////////////////////////////////////////////////////////////////////
-	virtual EDisplayClusterCommResult WaitForSwapSync() override
+	virtual EDisplayClusterCommResult SyncOnBarrier() override
 	{
 		return EDisplayClusterCommResult::NotImplemented;
 	}
diff --git a/Engine/Plugins/Runtime/nDisplay/Source/DisplayCluster/Private/Cluster/Controller/DisplayClusterClusterNodeCtrlNull.h b/Engine/Plugins/Runtime/nDisplay/Source/DisplayCluster/Private/Cluster/Controller/DisplayClusterClusterNodeCtrlNull.h
--- a/Engine/Plugins/Runtime/nDisplay/Source/DisplayCluster/Private/Cluster/Controller/DisplayClusterClusterNodeCtrlNull.h
+++ b/Engine/Plugins/Runtime/nDisplay/Source/DisplayCluster/Private/Cluster/Controller/DisplayClusterClusterNodeCtrlNull.h
@@ -97,7 +97,7 @@ public:
 	//////////////////////////////////////////////////////////////////////////////////////////////
 	// IDisplayClusterProtocolRenderSync
 	//////////////////////////////////////////////////////////////////////////////////////////////
-	virtual EDisplayClusterCommResult WaitForSwapSync() override final
+	virtual EDisplayClusterCommResult SyncOnBarrier() override final
 	{
 		return EDisplayClusterCommResult::NotImplemented;
 	}
diff --git a/Engine/Plugins/Runtime/nDisplay/Source/DisplayCluster/Private/Cluster/Controller/DisplayClusterClusterNodeCtrlSecondary.cpp b/Engine/Plugins/Runtime/nDisplay/Source/DisplayCluster/Private/Cluster/Controller/DisplayClusterClusterNodeCtrlSecondary.cpp
--- a/Engine/Plugins/Runtime/nDisplay/Source/DisplayCluster/Private/Cluster/Controller/DisplayClusterClusterNodeCtrlSecondary.cpp
+++ b/Engine/Plugins/Runtime/nDisplay/Source/DisplayCluster/Private/Cluster/Controller/DisplayClusterClusterNodeCtrlSecondary.cpp
@@ -92,10 +92,10 @@ EDisplayClusterCommResult FDisplayClusterClusterNodeCtrlSecondary::GetNativeInpu
 //////////////////////////////////////////////////////////////////////////////////////////////
 // IDisplayClusterProtocolRenderSync
 //////////////////////////////////////////////////////////////////////////////////////////////
-EDisplayClusterCommResult FDisplayClusterClusterNodeCtrlSecondary::WaitForSwapSync()
+EDisplayClusterCommResult FDisplayClusterClusterNodeCtrlSecondary::SyncOnBarrier()
 {
 	checkSlow(RenderSyncClient.IsValid());
-	const EDisplayClusterCommResult CommResult = RenderSyncClient->WaitForSwapSync();
+	const EDisplayClusterCommResult CommResult = RenderSyncClient->SyncOnBarrier();
 	GDisplayCluster->GetPrivateClusterMgr()->GetFailoverNodeController()->HandleCommResult(CommResult);
 	return CommResult;
 }
diff --git a/Engine/Plugins/Runtime/nDisplay/Source/DisplayCluster/Private/Cluster/Controller/DisplayClusterClusterNodeCtrlSecondary.h b/Engine/Plugins/Runtime/nDisplay/Source/DisplayCluster/Private/Cluster/Controller/DisplayClusterClusterNodeCtrlSecondary.h
--- a/Engine/Plugins/Runtime/nDisplay/Source/DisplayCluster/Private/Cluster/Controller/DisplayClusterClusterNodeCtrlSecondary.h
+++ b/Engine/Plugins/Runtime/nDisplay/Source/DisplayCluster/Private/Cluster/Controller/DisplayClusterClusterNodeCtrlSecondary.h
@@ -45,7 +45,7 @@ public:
 	//////////////////////////////////////////////////////////////////////////////////////////////
 	// IDisplayClusterProtocolRenderSync
 	//////////////////////////////////////////////////////////////////////////////////////////////
-	virtual EDisplayClusterCommResult WaitForSwapSync() override final;
+	virtual EDisplayClusterCommResult SyncOnBarrier() override final;
 
 public:
 	//////////////////////////////////////////////////////////////////////////////////////////////
diff --git a/Engine/Plugins/Runtime/nDisplay/Source/DisplayCluster/Private/Network/Protocol/IDisplayClusterProtocolRenderSync.h b/Engine/Plugins/Runtime/nDisplay/Source/DisplayCluster/Private/Network/Protocol/IDisplayClusterProtocolRenderSync.h
--- a/Engine/Plugins/Runtime/nDisplay/Source/DisplayCluster/Private/Network/Protocol/IDisplayClusterProtocolRenderSync.h
+++ b/Engine/Plugins/Runtime/nDisplay/Source/DisplayCluster/Private/Network/Protocol/IDisplayClusterProtocolRenderSync.h
@@ -6,8 +6,8 @@
 
 
 /**
- * Rendering synchronization protocol. Used to synchronize everything we need
- * on the render thread.
+ * Rendering synchronization protocol.
+ * Used to synchronize frame presentation on RHI thread.
  */
 class IDisplayClusterProtocolRenderSync
 {
@@ -15,6 +15,6 @@ public:
 	virtual ~IDisplayClusterProtocolRenderSync() = default;
 
 public:
-	// Swap sync barrier
-	virtual EDisplayClusterCommResult WaitForSwapSync() = 0;
+	// Synchronize RHI threads on a network barrier
+	virtual EDisplayClusterCommResult SyncOnBarrier() = 0;
 };
diff --git a/Engine/Plugins/Runtime/nDisplay/Source/DisplayCluster/Private/Network/Service/RenderSync/DisplayClusterRenderSyncClient.cpp b/Engine/Plugins/Runtime/nDisplay/Source/DisplayCluster/Private/Network/Service/RenderSync/DisplayClusterRenderSyncClient.cpp
--- a/Engine/Plugins/Runtime/nDisplay/Source/DisplayCluster/Private/Network/Service/RenderSync/DisplayClusterRenderSyncClient.cpp
+++ b/Engine/Plugins/Runtime/nDisplay/Source/DisplayCluster/Private/Network/Service/RenderSync/DisplayClusterRenderSyncClient.cpp
@@ -52,11 +52,11 @@ bool FDisplayClusterRenderSyncClient::Connect(const FString& Address, const uint
 //////////////////////////////////////////////////////////////////////////////////////////////
 // IDisplayClusterProtocolRenderSync
 //////////////////////////////////////////////////////////////////////////////////////////////
-EDisplayClusterCommResult FDisplayClusterRenderSyncClient::WaitForSwapSync()
+EDisplayClusterCommResult FDisplayClusterRenderSyncClient::SyncOnBarrier()
 {
 	static const TSharedPtr<FDisplayClusterPacketInternal> Request(
 		new FDisplayClusterPacketInternal(
-			DisplayClusterRenderSyncStrings::WaitForSwapSync::Name,
+			DisplayClusterRenderSyncStrings::SyncOnBarrier::Name,
 			DisplayClusterRenderSyncStrings::TypeRequest,
 			DisplayClusterRenderSyncStrings::ProtocolName)
 	);
diff --git a/Engine/Plugins/Runtime/nDisplay/Source/DisplayCluster/Private/Network/Service/RenderSync/DisplayClusterRenderSyncClient.h b/Engine/Plugins/Runtime/nDisplay/Source/DisplayCluster/Private/Network/Service/RenderSync/DisplayClusterRenderSyncClient.h
--- a/Engine/Plugins/Runtime/nDisplay/Source/DisplayCluster/Private/Network/Service/RenderSync/DisplayClusterRenderSyncClient.h
+++ b/Engine/Plugins/Runtime/nDisplay/Source/DisplayCluster/Private/Network/Service/RenderSync/DisplayClusterRenderSyncClient.h
@@ -28,5 +28,5 @@ public:
 	//////////////////////////////////////////////////////////////////////////////////////////////
 	// IDisplayClusterProtocolRenderSync
 	//////////////////////////////////////////////////////////////////////////////////////////////
-	virtual EDisplayClusterCommResult WaitForSwapSync() override;
+	virtual EDisplayClusterCommResult SyncOnBarrier() override;
 };
diff --git a/Engine/Plugins/Runtime/nDisplay/Source/DisplayCluster/Private/Network/Service/RenderSync/DisplayClusterRenderSyncService.cpp b/Engine/Plugins/Runtime/nDisplay/Source/DisplayCluster/Private/Network/Service/RenderSync/DisplayClusterRenderSyncService.cpp
--- a/Engine/Plugins/Runtime/nDisplay/Source/DisplayCluster/Private/Network/Service/RenderSync/DisplayClusterRenderSyncService.cpp
+++ b/Engine/Plugins/Runtime/nDisplay/Source/DisplayCluster/Private/Network/Service/RenderSync/DisplayClusterRenderSyncService.cpp
@@ -12,36 +12,18 @@
 #include "Cluster/Failover/IDisplayClusterFailoverNodeController.h"
 #include "Config/IPDisplayClusterConfigManager.h"
 
+#include "DisplayClusterConfigurationStrings.h"
 #include "DisplayClusterConfigurationTypes.h"
 
 #include "Misc/DisplayClusterGlobals.h"
 #include "Misc/DisplayClusterLog.h"
 
-#include "Interfaces/IPv4/IPv4Endpoint.h"
-
 
 FDisplayClusterRenderSyncService::FDisplayClusterRenderSyncService()
 	: FDisplayClusterService(FString("SRV_RS"))
 {
-	// Get list of cluster node IDs
-	TArray<FString> NodeIds;
-	GDisplayCluster->GetPrivateClusterMgr()->GetNodeIds(NodeIds);
-
-	// Network settings from config data
-	const FDisplayClusterConfigurationNetworkSettings& NetworkSettings = GDisplayCluster->GetConfigMgr()->GetConfig()->Cluster->Network;
-
-	// Instantiate service barriers
-	BarrierSwap = FDisplayClusterBarrierFactory::CreateBarrier(NodeIds, NetworkSettings.RenderSyncBarrierTimeout, TEXT("RenderSync_barrier"));
-
-	// Put the barriers into an aux container
-	ServiceBarriers.Emplace(BarrierSwap->GetName(), &BarrierSwap);
-
-	// Subscribe for barrier timeout events
-	for (TPair<FString, TUniquePtr<IDisplayClusterBarrier>*>& BarrierIt : ServiceBarriers)
-	{
-		(*BarrierIt.Value)->OnBarrierTimeout().AddRaw(this, &FDisplayClusterRenderSyncService::ProcessBarrierTimeout);
-	}
-
+	// Perform barriers initialization depending on current circumstances
+	InitializeBarriers();
 	// Subscribe for SessionClosed events
 	OnSessionClosed().AddRaw(this, &FDisplayClusterRenderSyncService::ProcessSessionClosed);
 }
@@ -49,11 +31,7 @@ FDisplayClusterRenderSyncService::FDisplayClusterRenderSyncService()
 FDisplayClusterRenderSyncService::~FDisplayClusterRenderSyncService()
 {
 	// Unsubscribe from barrier timeout events
-	for (TPair<FString, TUniquePtr<IDisplayClusterBarrier>*>& BarrierIt : ServiceBarriers)
-	{
-		(*BarrierIt.Value)->OnBarrierTimeout().RemoveAll(this);
-	}
-
+	UnsubscribeFromAllBarrierEvents();
 	// Unsubscribe from SessionClosed notifications
 	OnSessionClosed().RemoveAll(this);
 
@@ -64,10 +42,7 @@ FDisplayClusterRenderSyncService::~FDisplayClusterRenderSyncService()
 bool FDisplayClusterRenderSyncService::Start(const FString& Address, const uint16 Port)
 {
 	// Activate all barriers
-	for (TPair<FString, TUniquePtr<IDisplayClusterBarrier>*>& BarrierIt : ServiceBarriers)
-	{
-		(*BarrierIt.Value)->Activate();
-	}
+	ActivateAllBarriers();
 
 	return FDisplayClusterServer::Start(Address, Port);
 }
@@ -75,10 +50,7 @@ bool FDisplayClusterRenderSyncService::Start(const FString& Address, const uint1
 bool FDisplayClusterRenderSyncService::Start(TSharedPtr<FDisplayClusterTcpListener>& ExternalListener)
 {
 	// Activate all barriers
-	for (TPair<FString, TUniquePtr<IDisplayClusterBarrier>*>& BarrierIt : ServiceBarriers)
-	{
-		(*BarrierIt.Value)->Activate();
-	}
+	ActivateAllBarriers();
 
 	return FDisplayClusterServer::Start(ExternalListener);
 }
@@ -86,10 +58,7 @@ bool FDisplayClusterRenderSyncService::Start(TSharedPtr<FDisplayClusterTcpListen
 void FDisplayClusterRenderSyncService::Shutdown()
 {
 	// Deactivate all barriers
-	for (TPair<FString, TUniquePtr<IDisplayClusterBarrier>*>& BarrierIt : ServiceBarriers)
-	{
-		(*BarrierIt.Value)->Deactivate();
-	}
+	DeactivateAllBarriers();
 
 	return FDisplayClusterServer::Shutdown();
 }
@@ -103,10 +72,7 @@ FString FDisplayClusterRenderSyncService::GetProtocolName() const
 void FDisplayClusterRenderSyncService::KillSession(const FString& NodeId)
 {
 	// Before killing the session on the parent level, we need to unregister this node from the barriers
-	for (TPair<FString, TUniquePtr<IDisplayClusterBarrier>*>& BarrierIt : ServiceBarriers)
-	{
-		(*BarrierIt.Value)->UnregisterSyncNode(NodeId);
-	}
+	UnregisterClusterNode(NodeId);
 
 	// Now do the session related job
 	FDisplayClusterServer::KillSession(NodeId);
@@ -128,10 +94,7 @@ void FDisplayClusterRenderSyncService::ProcessSessionClosed(const FDisplayCluste
 		if (!NodeId.IsEmpty())
 		{
 			// We have to unregister the node that just disconnected from all barriers
-			for (TPair<FString, TUniquePtr<IDisplayClusterBarrier>*>& BarrierIt : ServiceBarriers)
-			{
-				(*BarrierIt.Value)->UnregisterSyncNode(NodeId);
-			}
+			UnregisterClusterNode(NodeId);
 
 			// Notify others about node fail
 			OnNodeFailed().Broadcast(NodeId, ENodeFailType::ConnectionLost);
@@ -142,18 +105,12 @@ void FDisplayClusterRenderSyncService::ProcessSessionClosed(const FDisplayCluste
 // Callbck on barrier timeout
 void FDisplayClusterRenderSyncService::ProcessBarrierTimeout(const FString& BarrierName, const TArray<FString>& NodesTimedOut)
 {
-	// We have to unregister the node that just timed out from all barriers
-	for (TPair<FString, TUniquePtr<IDisplayClusterBarrier>*>& BarrierIt : ServiceBarriers)
-	{
-		for (const FString& NodeId : NodesTimedOut)
-		{
-			(*BarrierIt.Value)->UnregisterSyncNode(NodeId);
-		}
-	}
-
-	// Notify others about timeout
 	for (const FString& NodeId : NodesTimedOut)
 	{
+		// We have to unregister the node that just timed out from all barriers
+		UnregisterClusterNode(NodeId);
+
+		// Notify others about timeout
 		OnNodeFailed().Broadcast(NodeId, ENodeFailType::BarrierTimeOut);
 	}
 }
@@ -187,9 +144,10 @@ TSharedPtr<FDisplayClusterPacketInternal> FDisplayClusterRenderSyncService::Proc
 	TSharedPtr<FDisplayClusterPacketInternal> Response = MakeShared<FDisplayClusterPacketInternal>(Request->GetName(), DisplayClusterRenderSyncStrings::TypeResponse, Request->GetProtocol());
 
 	// Dispatch the packet
-	if (Request->GetName().Equals(DisplayClusterRenderSyncStrings::WaitForSwapSync::Name, ESearchCase::IgnoreCase))
+	if (Request->GetName().Equals(DisplayClusterRenderSyncStrings::SyncOnBarrier::Name, ESearchCase::IgnoreCase))
 	{
-		const EDisplayClusterCommResult CommResult = WaitForSwapSync();
+		// Process command
+		const EDisplayClusterCommResult CommResult = SyncOnBarrier();
 		Response->SetCommResult(CommResult);
 
 		return Response;
@@ -205,11 +163,138 @@ TSharedPtr<FDisplayClusterPacketInternal> FDisplayClusterRenderSyncService::Proc
 //////////////////////////////////////////////////////////////////////////////////////////////
 // IDisplayClusterProtocolRenderSync
 //////////////////////////////////////////////////////////////////////////////////////////////
-EDisplayClusterCommResult FDisplayClusterRenderSyncService::WaitForSwapSync()
+EDisplayClusterCommResult FDisplayClusterRenderSyncService::SyncOnBarrier()
 {
-	TRACE_CPUPROFILER_EVENT_SCOPE(nD SRV_RS::WaitForSwapSync);
+	TRACE_CPUPROFILER_EVENT_SCOPE(nD SRV_RS::SyncOnBarrier);
+	
+	const FString NodeId = GetSessionInfoCache().NodeId.Get("");
+
+	IDisplayClusterBarrier* Barrier = GetBarrierForNode(NodeId);
+	if (ensureMsgf(Barrier, TEXT("%s could not find a barrier for node '%s'"), *GetName(), *NodeId))
+	{
+		Barrier->Wait(NodeId);
+	}
 
-	const FString CachedNodeId = GetSessionInfoCache().NodeId.Get("");
-	BarrierSwap->Wait(CachedNodeId);
 	return EDisplayClusterCommResult::Ok;
 }
+
+void FDisplayClusterRenderSyncService::InitializeBarriers()
+{
+	// Get configuration data
+	const UDisplayClusterConfigurationData* const Config = GDisplayCluster->GetPrivateConfigMgr()->GetConfig();
+	if (!Config)
+	{
+		UE_LOG(LogDisplayClusterNetwork, Error, TEXT("%s - No configuartion data, can't initialize barriers"), *GetName());
+		return;
+	}
+
+	// Get list of cluster node IDs (runtime nodes)
+	TArray<FString> RuntimeNodeIds;
+	GDisplayCluster->GetPrivateClusterMgr()->GetNodeIds(RuntimeNodeIds);
+
+	// Filter nodes that are expected to present in runtime
+	const TMap<FString, TObjectPtr<UDisplayClusterConfigurationClusterNode>> RuntimeNodeCfgs = Config->Cluster->Nodes.FilterByPredicate(
+		[RuntimeNodeIds](const TPair<FString, TObjectPtr<UDisplayClusterConfigurationClusterNode>>& Pair)
+		{
+			return RuntimeNodeIds.Contains(Pair.Key);
+		});
+
+	// Build sync groups (this one is an aux mapping of policy IDs to their node IDs)
+	TMap<FString, TArray<FString>> SyncGroups;
+	for (const TPair<FString, TObjectPtr<UDisplayClusterConfigurationClusterNode>>& RuntimeNodeCfg : RuntimeNodeCfgs)
+	{
+		FString SyncPolicyId;
+
+		// Headless nodes always use the same policy
+		if (RuntimeNodeCfg.Value->bRenderHeadless)
+		{
+			SyncPolicyId = DisplayClusterConfigurationStrings::config::cluster::render_sync::HeadlessRenderingSyncPolicy;
+		}
+#if 0 // Not implemented yet
+		// Individual sync policy
+		else if(NodeCfgIt.Value->bOverrideSyncPolicy)
+		{
+			SyncPolicyId = NodeCfgIt.Value->Sync.RenderSyncPolicy.Type;
+		}
+#endif
+		// Then default sync policy
+		else
+		{
+			SyncPolicyId = Config->Cluster->Sync.RenderSyncPolicy.Type;
+		}
+
+		// To prevent any case related issues caused by manual .ndisplay file editing. We won't
+		// need it once we refactor all the IDs to be FName's.
+		SyncPolicyId.ToLowerInline();
+
+		// Create new sync group if not yet available
+		TArray<FString>* NodesInSyncGroup = SyncGroups.Find(SyncPolicyId);
+		if (!NodesInSyncGroup)
+		{
+			NodesInSyncGroup = &SyncGroups.Emplace(SyncPolicyId);
+		}
+
+		// Assign this cluster node to the specific sync group
+		NodesInSyncGroup->Add(RuntimeNodeCfg.Key);
+		// Update node-policy mapping
+		NodeToPolicyMap.Emplace(RuntimeNodeCfg.Key, SyncPolicyId);
+	}
+
+	// Finally, initialize barriers for every sync group
+	for (const TPair<FString, TArray<FString>>& SyncGroup : SyncGroups)
+	{
+		const FString BarrierName = SyncGroup.Key + TEXT("_barrier");
+
+		// Instantiate the barrier
+		TUniquePtr<IDisplayClusterBarrier> Barrier = FDisplayClusterBarrierFactory::CreateBarrier(SyncGroup.Value, Config->Cluster->Network.RenderSyncBarrierTimeout, BarrierName);
+		// Subscribe for barrier timeout events
+		Barrier->OnBarrierTimeout().AddRaw(this, &FDisplayClusterRenderSyncService::ProcessBarrierTimeout);
+		// Store barrier instance
+		PolicyToBarrierMap.Emplace(SyncGroup.Key, MoveTemp(Barrier));
+	}
+}
+
+void FDisplayClusterRenderSyncService::ActivateAllBarriers()
+{
+	for (TPair<FString, TUniquePtr<IDisplayClusterBarrier>>& PolicyToBarrierIt : PolicyToBarrierMap)
+	{
+		PolicyToBarrierIt.Value->Activate();
+	}
+}
+
+void FDisplayClusterRenderSyncService::DeactivateAllBarriers()
+{
+	for (TPair<FString, TUniquePtr<IDisplayClusterBarrier>>& PolicyToBarrierIt : PolicyToBarrierMap)
+	{
+		PolicyToBarrierIt.Value->Deactivate();
+	}
+}
+
+void FDisplayClusterRenderSyncService::UnsubscribeFromAllBarrierEvents()
+{
+	for (TPair<FString, TUniquePtr<IDisplayClusterBarrier>>& PolicyToBarrierIt : PolicyToBarrierMap)
+	{
+		PolicyToBarrierIt.Value->OnBarrierTimeout().RemoveAll(this);
+	}
+}
+
+IDisplayClusterBarrier* FDisplayClusterRenderSyncService::GetBarrierForNode(const FString& NodeId) const
+{
+	if (const FString* const Policy = NodeToPolicyMap.Find(NodeId))
+	{
+		if (const TUniquePtr<IDisplayClusterBarrier>* const Barrier = PolicyToBarrierMap.Find(*Policy))
+		{
+			return Barrier->Get();
+		}
+	}
+
+	return nullptr;
+}
+
+void FDisplayClusterRenderSyncService::UnregisterClusterNode(const FString& NodeId)
+{
+	if (IDisplayClusterBarrier* Barrier = GetBarrierForNode(NodeId))
+	{
+		Barrier->UnregisterSyncNode(NodeId);
+	}
+}
diff --git a/Engine/Plugins/Runtime/nDisplay/Source/DisplayCluster/Private/Network/Service/RenderSync/DisplayClusterRenderSyncService.h b/Engine/Plugins/Runtime/nDisplay/Source/DisplayCluster/Private/Network/Service/RenderSync/DisplayClusterRenderSyncService.h
--- a/Engine/Plugins/Runtime/nDisplay/Source/DisplayCluster/Private/Network/Service/RenderSync/DisplayClusterRenderSyncService.h
+++ b/Engine/Plugins/Runtime/nDisplay/Source/DisplayCluster/Private/Network/Service/RenderSync/DisplayClusterRenderSyncService.h
@@ -55,12 +55,33 @@ private:
 	//////////////////////////////////////////////////////////////////////////////////////////////
 	// IDisplayClusterProtocolRenderSync
 	//////////////////////////////////////////////////////////////////////////////////////////////
-	virtual EDisplayClusterCommResult WaitForSwapSync() override;
+	virtual EDisplayClusterCommResult SyncOnBarrier() override;
 
 private:
-	// Swap sync barrier
-	TUniquePtr<IDisplayClusterBarrier> BarrierSwap;
+	// Initializes barriers based on available nodes and their sync policies
+	void InitializeBarriers();
 
-	// Auxiliary container that keeps all the barriers
-	TMap<FString, TUniquePtr<IDisplayClusterBarrier>*> ServiceBarriers;
+	// Activates all service barriers
+	void ActivateAllBarriers();
+	// Deactivates all service barriers
+	void DeactivateAllBarriers();
+
+	// Unsubscribe from all events of all internal barriers
+	void UnsubscribeFromAllBarrierEvents();
+
+	// Returns barrier of a node's sync group
+	IDisplayClusterBarrier* GetBarrierForNode(const FString& NodeId) const;
+
+	// Unregister cluster node from a barrier
+	void UnregisterClusterNode(const FString& NodeId);
+
+private:
+	// We can now use different sync policies within the same cluster. Since each policy may
+	// have its own synchronization logic, mainly barriers utilization, we need to have
+	// individual sync barriers for every sync group. A sync group is a list of nodes that
+	// use the same sync policy.
+	TMap<FString, TUniquePtr<IDisplayClusterBarrier>> PolicyToBarrierMap;
+
+	// Node ID to sync policy ID (or sync group) mapping
+	TMap<FString, FString> NodeToPolicyMap;
 };
diff --git a/Engine/Plugins/Runtime/nDisplay/Source/DisplayCluster/Private/Network/Service/RenderSync/DisplayClusterRenderSyncStrings.h b/Engine/Plugins/Runtime/nDisplay/Source/DisplayCluster/Private/Network/Service/RenderSync/DisplayClusterRenderSyncStrings.h
--- a/Engine/Plugins/Runtime/nDisplay/Source/DisplayCluster/Private/Network/Service/RenderSync/DisplayClusterRenderSyncStrings.h
+++ b/Engine/Plugins/Runtime/nDisplay/Source/DisplayCluster/Private/Network/Service/RenderSync/DisplayClusterRenderSyncStrings.h
@@ -15,8 +15,8 @@ namespace DisplayClusterRenderSyncStrings
 
 	constexpr static auto ArgumentsDefaultCategory = "RS";
 
-	namespace WaitForSwapSync
+	namespace SyncOnBarrier
 	{
-		constexpr static auto Name = "WaitForSwapSync";
+		constexpr static auto Name = "SyncOnBarrier";
 	};
 };
diff --git a/Engine/Plugins/Runtime/nDisplay/Source/DisplayCluster/Private/Render/DisplayClusterRenderManager.cpp b/Engine/Plugins/Runtime/nDisplay/Source/DisplayCluster/Private/Render/DisplayClusterRenderManager.cpp
--- a/Engine/Plugins/Runtime/nDisplay/Source/DisplayCluster/Private/Render/DisplayClusterRenderManager.cpp
+++ b/Engine/Plugins/Runtime/nDisplay/Source/DisplayCluster/Private/Render/DisplayClusterRenderManager.cpp
@@ -582,15 +582,32 @@ TSharedPtr<IDisplayClusterRenderSyncPolicy> FDisplayClusterRenderManager::Create
 		return nullptr;
 	}
 
-	// Create sync policy specified in a config file
-	const FString SyncPolicyType = ConfigData->Cluster->Sync.RenderSyncPolicy.Type;
+	FString SyncPolicyType;
+	TMap<FString, FString>* SyncPolicyParams = nullptr;
 
+	// Always use 'none' sync policy while rendering headless
+	if (FParse::Param(FCommandLine::Get(), TEXT("RenderOffscreen")))
+	{
+		SyncPolicyType   = DisplayClusterConfigurationStrings::config::cluster::render_sync::HeadlessRenderingSyncPolicy;
+		SyncPolicyParams = nullptr;
+		UE_LOG(LogDisplayClusterRender, Log, TEXT("Headless rendering requested. Using '%s' sync policy."), *SyncPolicyType);
+	}
+	// Otherwise use sync policy from config
+	else
+	{
+		SyncPolicyType   = ConfigData->Cluster->Sync.RenderSyncPolicy.Type;
+		SyncPolicyParams = &ConfigData->Cluster->Sync.RenderSyncPolicy.Parameters;
+		UE_LOG(LogDisplayClusterRender, Log, TEXT("Requested sync policy is '%s'"), *SyncPolicyType);
+	}
+
+	// Instantiate policy instance
 	TSharedPtr<IDisplayClusterRenderSyncPolicy> NewSyncPolicy;
 	if (SyncPolicyFactories.Contains(SyncPolicyType))
 	{
 		UE_LOG(LogDisplayClusterRender, Log, TEXT("A factory for the requested synchronization policy <%s> was found"), *SyncPolicyType);
-		NewSyncPolicy = SyncPolicyFactories[SyncPolicyType]->Create(SyncPolicyType, ConfigData->Cluster->Sync.RenderSyncPolicy.Parameters);
+		NewSyncPolicy = SyncPolicyFactories[SyncPolicyType]->Create(SyncPolicyType, SyncPolicyParams ? *SyncPolicyParams : TMap<FString, FString>());
 	}
+	// Fallback if requested policy is not available
 	else
 	{
 		const FString DefaultPolicy = DisplayClusterConfigurationStrings::config::cluster::render_sync::EthernetBarrier;
diff --git a/Engine/Plugins/Runtime/nDisplay/Source/DisplayCluster/Private/Render/Synchronization/DisplayClusterRenderSyncPolicyBase.cpp b/Engine/Plugins/Runtime/nDisplay/Source/DisplayCluster/Private/Render/Synchronization/DisplayClusterRenderSyncPolicyBase.cpp
--- a/Engine/Plugins/Runtime/nDisplay/Source/DisplayCluster/Private/Render/Synchronization/DisplayClusterRenderSyncPolicyBase.cpp
+++ b/Engine/Plugins/Runtime/nDisplay/Source/DisplayCluster/Private/Render/Synchronization/DisplayClusterRenderSyncPolicyBase.cpp
@@ -40,11 +40,11 @@ bool FDisplayClusterRenderSyncPolicyBase::Initialize()
 	return true;
 }
 
-void FDisplayClusterRenderSyncPolicyBase::SyncBarrierRenderThread()
+void FDisplayClusterRenderSyncPolicyBase::SyncOnBarrier()
 {
 	{
 		TRACE_CPUPROFILER_EVENT_SCOPE(nDisplay SyncPolicyBase::SyncBarrier);
-		GDisplayCluster->GetPrivateClusterMgr()->GetClusterNodeController()->WaitForSwapSync();
+		GDisplayCluster->GetPrivateClusterMgr()->GetClusterNodeController()->SyncOnBarrier();
 	}
 }
 
diff --git a/Engine/Plugins/Runtime/nDisplay/Source/DisplayCluster/Private/Render/Synchronization/DisplayClusterRenderSyncPolicyBase.h b/Engine/Plugins/Runtime/nDisplay/Source/DisplayCluster/Private/Render/Synchronization/DisplayClusterRenderSyncPolicyBase.h
--- a/Engine/Plugins/Runtime/nDisplay/Source/DisplayCluster/Private/Render/Synchronization/DisplayClusterRenderSyncPolicyBase.h
+++ b/Engine/Plugins/Runtime/nDisplay/Source/DisplayCluster/Private/Render/Synchronization/DisplayClusterRenderSyncPolicyBase.h
@@ -29,7 +29,7 @@ public:
 
 public:
 	// Performs Ethernet barrier syncrhonization with other nodes in the cluster
-	void SyncBarrierRenderThread();
+	void SyncOnBarrier();
 
 	// Returns policy configuration parameters
 	const TMap<FString, FString>& GetParameters() const
diff --git a/Engine/Plugins/Runtime/nDisplay/Source/DisplayCluster/Private/Render/Synchronization/DisplayClusterRenderSyncPolicyEthernetBarrier.cpp b/Engine/Plugins/Runtime/nDisplay/Source/DisplayCluster/Private/Render/Synchronization/DisplayClusterRenderSyncPolicyEthernetBarrier.cpp
--- a/Engine/Plugins/Runtime/nDisplay/Source/DisplayCluster/Private/Render/Synchronization/DisplayClusterRenderSyncPolicyEthernetBarrier.cpp
+++ b/Engine/Plugins/Runtime/nDisplay/Source/DisplayCluster/Private/Render/Synchronization/DisplayClusterRenderSyncPolicyEthernetBarrier.cpp
@@ -15,7 +15,7 @@ bool FDisplayClusterRenderSyncPolicyEthernetBarrier::SynchronizeClusterRendering
 	// Wait unless the frame is rendered
 	WaitForFrameCompletion();
 	// Sync on the network barrier
-	SyncBarrierRenderThread();
+	SyncOnBarrier();
 	// Ask engine to present the frame
 	return true;
 }
diff --git a/Engine/Plugins/Runtime/nDisplay/Source/DisplayCluster/Private/Render/Synchronization/Windows/DisplayClusterRenderSyncPolicyEthernetWindows.cpp b/Engine/Plugins/Runtime/nDisplay/Source/DisplayCluster/Private/Render/Synchronization/Windows/DisplayClusterRenderSyncPolicyEthernetWindows.cpp
--- a/Engine/Plugins/Runtime/nDisplay/Source/DisplayCluster/Private/Render/Synchronization/Windows/DisplayClusterRenderSyncPolicyEthernetWindows.cpp
+++ b/Engine/Plugins/Runtime/nDisplay/Source/DisplayCluster/Private/Render/Synchronization/Windows/DisplayClusterRenderSyncPolicyEthernetWindows.cpp
@@ -72,7 +72,7 @@ bool FDisplayClusterRenderSyncPolicyEthernet::SynchronizeClusterRendering(int32&
 			// Wait unless the frame is rendered
 			WaitForFrameCompletion();
 			// Barrier sync only
-			SyncBarrierRenderThread();
+			SyncOnBarrier();
 			// Let the engine present this frame
 			bNeedEnginePresent = true;
 		}
@@ -234,7 +234,7 @@ void FDisplayClusterRenderSyncPolicyEthernet::Step_WaitForEthernetBarrierSignal_
 {
 	// Align render threads with a barrier
 	B1B = FPlatformTime::Seconds();
-	SyncBarrierRenderThread();
+	SyncOnBarrier();
 	B1A = FPlatformTime::Seconds();
 }
 
@@ -263,7 +263,7 @@ void FDisplayClusterRenderSyncPolicyEthernet::Step_SkipPresentationOnClosestVBla
 void FDisplayClusterRenderSyncPolicyEthernet::Step_WaitForEthernetBarrierSignal_2()
 {
 	B2B = FPlatformTime::Seconds();
-	SyncBarrierRenderThread();
+	SyncOnBarrier();
 	B2A = FPlatformTime::Seconds();
 }
 
diff --git a/Engine/Plugins/Runtime/nDisplay/Source/DisplayCluster/Private/Render/Synchronization/Windows/DisplayClusterRenderSyncPolicyNvidiaWindows.cpp b/Engine/Plugins/Runtime/nDisplay/Source/DisplayCluster/Private/Render/Synchronization/Windows/DisplayClusterRenderSyncPolicyNvidiaWindows.cpp
--- a/Engine/Plugins/Runtime/nDisplay/Source/DisplayCluster/Private/Render/Synchronization/Windows/DisplayClusterRenderSyncPolicyNvidiaWindows.cpp
+++ b/Engine/Plugins/Runtime/nDisplay/Source/DisplayCluster/Private/Render/Synchronization/Windows/DisplayClusterRenderSyncPolicyNvidiaWindows.cpp
@@ -208,7 +208,7 @@ bool FDisplayClusterRenderSyncPolicyNvidia::SynchronizeClusterRendering(int32& I
 	if (bNvDiagPresent)
 	{
 		UE_LOG(LogDisplayClusterRenderSync, VeryVerbose, TEXT("NVAPI DIAG: wait start"));
-		SyncBarrierRenderThread();
+		SyncOnBarrier();
 		UE_LOG(LogDisplayClusterRenderSync, VeryVerbose, TEXT("NVAPI DIAG: wait end"));
 
 		// Disable barrier when it reaches the set limit.
@@ -332,14 +332,14 @@ bool FDisplayClusterRenderSyncPolicyNvidia::InitializeNvidiaSwapLock()
 		// upcoming V-blank. Normally, the barrier synchronization takes up to 500 microseconds which is about 3% of frame time.
 		// So the probability of a situation where some nodes left the barrier before V-blank while others after is really low.
 		// However, it's still possible. And cluster restart should likely be successfull.
-		SyncBarrierRenderThread();
+		SyncOnBarrier();
 
 		// Assuming all the nodes in the right place and have some time before V-blank. Let's wait untill it happend.
 		WaitForVBlank();
 
 		// We don't really know either we're at the very beginning of V-blank interval or in the end. Let's spend some time
 		// on the barrier to let V-blank interval over.
-		SyncBarrierRenderThread();
+		SyncOnBarrier();
 
 		// Yes, the task scheduler can break the logic above by allocating CPU resources to this thread like 20ms later which
 		// makes us behind of the next V-blank for 60Hz output. And we probably should play around thread priorities to reduce
diff --git a/Engine/Plugins/Runtime/nDisplay/Source/DisplayClusterConfiguration/Private/Formats/JSON500/DisplayClusterConfigurationJsonParser_500.cpp b/Engine/Plugins/Runtime/nDisplay/Source/DisplayClusterConfiguration/Private/Formats/JSON500/DisplayClusterConfigurationJsonParser_500.cpp
--- a/Engine/Plugins/Runtime/nDisplay/Source/DisplayClusterConfiguration/Private/Formats/JSON500/DisplayClusterConfigurationJsonParser_500.cpp
+++ b/Engine/Plugins/Runtime/nDisplay/Source/DisplayClusterConfiguration/Private/Formats/JSON500/DisplayClusterConfigurationJsonParser_500.cpp
@@ -221,6 +221,7 @@ namespace JSON500
 				Node->Host = CfgNode.Value.Host;
 				Node->bIsSoundEnabled = CfgNode.Value.Sound;
 				Node->bIsFullscreen = CfgNode.Value.FullScreen;
+				Node->bRenderHeadless = CfgNode.Value.RenderHeadless;
 				Node->bEnableTextureShare = CfgNode.Value.TextureShare;
 				Node->WindowRect = FDisplayClusterConfigurationRectangle(CfgNode.Value.Window.X, CfgNode.Value.Window.Y, CfgNode.Value.Window.W, CfgNode.Value.Window.H);
 
@@ -413,6 +414,7 @@ namespace JSON500
 				Node.Host = CfgNode.Value->Host;
 				Node.Sound = CfgNode.Value->bIsSoundEnabled;
 				Node.FullScreen = CfgNode.Value->bIsFullscreen;
+				Node.RenderHeadless = CfgNode.Value->bRenderHeadless;
 				Node.TextureShare = CfgNode.Value->bEnableTextureShare;
 				Node.Window = FDisplayClusterConfigurationJsonRectangle_500(CfgNode.Value->WindowRect.X, CfgNode.Value->WindowRect.Y, CfgNode.Value->WindowRect.W, CfgNode.Value->WindowRect.H);
 
diff --git a/Engine/Plugins/Runtime/nDisplay/Source/DisplayClusterConfiguration/Private/Formats/JSON500/DisplayClusterConfigurationJsonTypes_500.h b/Engine/Plugins/Runtime/nDisplay/Source/DisplayClusterConfiguration/Private/Formats/JSON500/DisplayClusterConfigurationJsonTypes_500.h
--- a/Engine/Plugins/Runtime/nDisplay/Source/DisplayClusterConfiguration/Private/Formats/JSON500/DisplayClusterConfigurationJsonTypes_500.h
+++ b/Engine/Plugins/Runtime/nDisplay/Source/DisplayClusterConfiguration/Private/Formats/JSON500/DisplayClusterConfigurationJsonTypes_500.h
@@ -498,6 +498,9 @@ public:
 	UPROPERTY()
 	bool FullScreen = false;
 
+	UPROPERTY()
+	bool RenderHeadless = false;
+
 	UPROPERTY()
 	bool TextureShare = false;
 
diff --git a/Engine/Plugins/Runtime/nDisplay/Source/DisplayClusterConfiguration/Public/DisplayClusterConfigurationStrings.h b/Engine/Plugins/Runtime/nDisplay/Source/DisplayClusterConfiguration/Public/DisplayClusterConfigurationStrings.h
--- a/Engine/Plugins/Runtime/nDisplay/Source/DisplayClusterConfiguration/Public/DisplayClusterConfigurationStrings.h
+++ b/Engine/Plugins/Runtime/nDisplay/Source/DisplayClusterConfiguration/Public/DisplayClusterConfigurationStrings.h
@@ -84,6 +84,9 @@ namespace DisplayClusterConfigurationStrings
 				static constexpr const TCHAR* Nvidia            = TEXT("nvidia");
 				static constexpr const TCHAR* NvidiaSwapBarrier = TEXT("swap_barrier");
 				static constexpr const TCHAR* NvidiaSwapGroup   = TEXT("swap_group");
+
+				// Always use 'none' for headless rendering
+				static constexpr const TCHAR* HeadlessRenderingSyncPolicy = None;
 			}
 		}
 
diff --git a/Engine/Plugins/Runtime/nDisplay/Source/DisplayClusterConfiguration/Public/DisplayClusterConfigurationTypes.h b/Engine/Plugins/Runtime/nDisplay/Source/DisplayClusterConfiguration/Public/DisplayClusterConfigurationTypes.h
--- a/Engine/Plugins/Runtime/nDisplay/Source/DisplayClusterConfiguration/Public/DisplayClusterConfigurationTypes.h
+++ b/Engine/Plugins/Runtime/nDisplay/Source/DisplayClusterConfiguration/Public/DisplayClusterConfigurationTypes.h
@@ -334,6 +334,10 @@ public:
 	UPROPERTY(EditAnywhere, BlueprintReadWrite, Category = "Configuration", meta = (DisplayName = "Output Remapping"))
 	FDisplayClusterConfigurationFramePostProcess_OutputRemap OutputRemap;
 
+	/** Activates headless rendering for this cluster node */
+	UPROPERTY(EditAnywhere, BlueprintReadonly, Category = "Configuration", meta = (DisplayName = "Headless Rendering"))
+	bool bRenderHeadless = false;
+
 	/** Enables texture sharing for this cluster node */
 	UPROPERTY(EditAnywhere, BlueprintReadonly, Category = "Configuration", meta = (DisplayName = "Enable Texture Share"))
 	bool bEnableTextureShare = false;
