nDisplay: fixed Linux build errors introduced in 23370057

diff --git a/Engine/Plugins/Runtime/nDisplay/Source/DisplayCluster/Private/Render/Synchronization/Linux/DisplayClusterRenderSyncPolicyEthernetLinux.cpp b/Engine/Plugins/Runtime/nDisplay/Source/DisplayCluster/Private/Render/Synchronization/Linux/DisplayClusterRenderSyncPolicyEthernetLinux.cpp
--- a/Engine/Plugins/Runtime/nDisplay/Source/DisplayCluster/Private/Render/Synchronization/Linux/DisplayClusterRenderSyncPolicyEthernetLinux.cpp
+++ b/Engine/Plugins/Runtime/nDisplay/Source/DisplayCluster/Private/Render/Synchronization/Linux/DisplayClusterRenderSyncPolicyEthernetLinux.cpp
@@ -8,7 +8,7 @@ bool FDisplayClusterRenderSyncPolicyEthernet::SynchronizeClusterRendering(int32&
 	// Wait unless the frame is rendered
 	WaitForFrameCompletion();
 	// Sync on the network barrier
-	SyncBarrierRenderThread();
+	SyncOnBarrier();
 	// Ask engine to present the frame
 	return true;
 }
diff --git a/Engine/Plugins/Runtime/nDisplay/Source/DisplayCluster/Private/Render/Synchronization/Linux/DisplayClusterRenderSyncPolicyNvidiaLinux.cpp b/Engine/Plugins/Runtime/nDisplay/Source/DisplayCluster/Private/Render/Synchronization/Linux/DisplayClusterRenderSyncPolicyNvidiaLinux.cpp
--- a/Engine/Plugins/Runtime/nDisplay/Source/DisplayCluster/Private/Render/Synchronization/Linux/DisplayClusterRenderSyncPolicyNvidiaLinux.cpp
+++ b/Engine/Plugins/Runtime/nDisplay/Source/DisplayCluster/Private/Render/Synchronization/Linux/DisplayClusterRenderSyncPolicyNvidiaLinux.cpp
@@ -27,7 +27,7 @@ bool FDisplayClusterRenderSyncPolicyNvidia::SynchronizeClusterRendering(int32& I
 	// Wait unless the frame is rendered
 	WaitForFrameCompletion();
 	// Sync on the network barrier
-	SyncBarrierRenderThread();
+	SyncOnBarrier();
 	// Ask engine to present the frame
 	return true;
 }
