nDisplay: Fix CCW selection not visible in the ICVFX panel.
Regression from 23230783.

diff --git a/Engine/Plugins/Runtime/nDisplay/Source/DisplayClusterLightCardEditor/Private/Viewport/DisplayClusterLightCardEditorViewportClient.cpp b/Engine/Plugins/Runtime/nDisplay/Source/DisplayClusterLightCardEditor/Private/Viewport/DisplayClusterLightCardEditorViewportClient.cpp
--- a/Engine/Plugins/Runtime/nDisplay/Source/DisplayClusterLightCardEditor/Private/Viewport/DisplayClusterLightCardEditorViewportClient.cpp
+++ b/Engine/Plugins/Runtime/nDisplay/Source/DisplayClusterLightCardEditor/Private/Viewport/DisplayClusterLightCardEditorViewportClient.cpp
@@ -1683,12 +1683,23 @@ void FDisplayClusterLightCardEditorViewportClient::UpdatePreviewActor(ADisplayCl
 			// Make sure proxies are added to the renderer. Necessary for selections to render even if stage actors were not modified but root actor was updated
 			for (const TObjectPtr<AActor>& ActorProxy : ActorProxiesCreated)
 			{
-				IDisplayClusterScenePreview::Get().AddActorToRenderer(PreviewRendererId, ActorProxy, [this, ActorProxy](const UPrimitiveComponent* PrimitiveComponent)
+				// Hack - CL 23230783 sets CCW meshes to hidden which causes problems with selection, so always add CCWs to the renderer.
+				bool bIsCCW = false;
+				for (const UClass* Class = ActorProxy->GetClass(); Class && (UObject::StaticClass() != Class); Class = Class->GetSuperClass())
+				{
+					if (Class->GetName() == TEXT("ColorCorrectionWindow"))
+					{
+						bIsCCW = true;
+						break;
+					}
+				}
+				
+				IDisplayClusterScenePreview::Get().AddActorToRenderer(PreviewRendererId, ActorProxy, [this, ActorProxy, bIsCCW](const UPrimitiveComponent* PrimitiveComponent)
 				{
 					// Always add the light card mesh component to the renderer's scene even if it is marked hidden in game, since UV light cards will purposefully
 					// hide the light card mesh since it isn't supposed to exist in 3D space. The light card mesh will be appropriately filtered when the scene is
 					// rendered based on the projection mode
-					if (PrimitiveComponent->GetFName() == TEXT("LightCard"))
+					if (PrimitiveComponent->GetFName() == TEXT("LightCard") || bIsCCW)
 					{
 						return true;
 					}
