nDisplay: Fix RootActorOwner not always being set for LightCards.

diff --git a/Engine/Plugins/Runtime/nDisplay/Source/DisplayCluster/Private/Game/EngineClasses/Scene/DisplayClusterLightCardActor.cpp b/Engine/Plugins/Runtime/nDisplay/Source/DisplayCluster/Private/Game/EngineClasses/Scene/DisplayClusterLightCardActor.cpp
--- a/Engine/Plugins/Runtime/nDisplay/Source/DisplayCluster/Private/Game/EngineClasses/Scene/DisplayClusterLightCardActor.cpp
+++ b/Engine/Plugins/Runtime/nDisplay/Source/DisplayCluster/Private/Game/EngineClasses/Scene/DisplayClusterLightCardActor.cpp
@@ -560,6 +560,8 @@ void ADisplayClusterLightCardActor::AddToLightCardLayer(ADisplayClusterRootActor
 		// Likely -game and we need to add the layer directly
 		Layers.AddUnique(LightCardLayerName);
 	}
+
+	SetRootActorOwner(InRootActor);
 }
 
 void ADisplayClusterLightCardActor::UpdateStageActorTransform()
diff --git a/Engine/Plugins/Runtime/nDisplay/Source/DisplayCluster/Private/Game/EngineClasses/Scene/DisplayClusterRootActor.cpp b/Engine/Plugins/Runtime/nDisplay/Source/DisplayCluster/Private/Game/EngineClasses/Scene/DisplayClusterRootActor.cpp
--- a/Engine/Plugins/Runtime/nDisplay/Source/DisplayCluster/Private/Game/EngineClasses/Scene/DisplayClusterRootActor.cpp
+++ b/Engine/Plugins/Runtime/nDisplay/Source/DisplayCluster/Private/Game/EngineClasses/Scene/DisplayClusterRootActor.cpp
@@ -683,6 +683,7 @@ void ADisplayClusterRootActor::UpdateLightCardPositions()
 					{
 						if (LightCardActor->Layers.Contains(ActorLayer.Name))
 						{
+							LightCardActor->SetRootActorOwner(this);
 							if (LightCardActor->bLockToOwningRootActor)
 							{
 								LightCardActor->SetActorLocation(Location);
