nDisplay: Added ability for light cards to remain flush with their owning stage actor when the user modifies their positional properties directly through the details panel

diff --git a/Engine/Plugins/Runtime/nDisplay/Source/DisplayCluster/DisplayCluster.Build.cs b/Engine/Plugins/Runtime/nDisplay/Source/DisplayCluster/DisplayCluster.Build.cs
--- a/Engine/Plugins/Runtime/nDisplay/Source/DisplayCluster/DisplayCluster.Build.cs
+++ b/Engine/Plugins/Runtime/nDisplay/Source/DisplayCluster/DisplayCluster.Build.cs
@@ -20,6 +20,7 @@ public class DisplayCluster : ModuleRules
 				"Core",
 				"CoreUObject",
 				"DisplayClusterConfiguration",
+				"DisplayClusterLightCardEditorShaders",
 				"DisplayClusterLightCardExtender",
 				"Engine"
 			});
diff --git a/Engine/Plugins/Runtime/nDisplay/Source/DisplayCluster/Private/Game/EngineClasses/Scene/DisplayClusterLightCardActor.cpp b/Engine/Plugins/Runtime/nDisplay/Source/DisplayCluster/Private/Game/EngineClasses/Scene/DisplayClusterLightCardActor.cpp
--- a/Engine/Plugins/Runtime/nDisplay/Source/DisplayCluster/Private/Game/EngineClasses/Scene/DisplayClusterLightCardActor.cpp
+++ b/Engine/Plugins/Runtime/nDisplay/Source/DisplayCluster/Private/Game/EngineClasses/Scene/DisplayClusterLightCardActor.cpp
@@ -52,6 +52,7 @@ ADisplayClusterLightCardActor::ADisplayClusterLightCardActor(const FObjectInitia
 	, Yaw(0.f)
 	, Scale(FVector2D(1.f))
 	, RadialOffset(-0.5)
+	, bAlwaysFlushToWall(true)
 	, Mask(EDisplayClusterLightCardMask::Circle)
 	, Texture(nullptr)
 	, Color(FLinearColor(1.f, 1.f, 1.f, 1.f))
@@ -154,11 +155,18 @@ void ADisplayClusterLightCardActor::PostEditChangeProperty(FPropertyChangedEvent
 		PropertyChangedEvent.Property->GetFName() == GET_MEMBER_NAME_CHECKED(ADisplayClusterLightCardActor, Longitude) ||
 		PropertyChangedEvent.Property->GetFName() == GET_MEMBER_NAME_CHECKED(ADisplayClusterLightCardActor, Latitude) ||
 		PropertyChangedEvent.Property->GetFName() == GET_MEMBER_NAME_CHECKED(ADisplayClusterLightCardActor, Spin) ||
+		PropertyChangedEvent.Property->GetFName() == GET_MEMBER_NAME_CHECKED(ADisplayClusterLightCardActor, Pitch) ||
+		PropertyChangedEvent.Property->GetFName() == GET_MEMBER_NAME_CHECKED(ADisplayClusterLightCardActor, Yaw) ||
 		PropertyChangedEvent.Property->GetFName() == GET_MEMBER_NAME_CHECKED(ADisplayClusterLightCardActor, RadialOffset) ||
 		PropertyChangedEvent.Property->GetFName() == GET_MEMBER_NAME_CHECKED(ADisplayClusterLightCardActor, Scale) ||
 		PropertyChangedEvent.Property->GetFName() == GET_MEMBER_NAME_CHECKED(ADisplayClusterLightCardActor, bIsUVLightCard)))
 	{
 		UpdateStageActorTransform();
+
+		if (bAlwaysFlushToWall && RootActorOwner.IsValid())
+		{
+			RootActorOwner->MakeStageActorFlushToWall(this);
+		}
 	}
 
 	if (PropertyChangedEvent.Property && (
@@ -566,7 +574,7 @@ void ADisplayClusterLightCardActor::AddToLightCardLayer(ADisplayClusterRootActor
 
 void ADisplayClusterLightCardActor::UpdateStageActorTransform()
 {
-	// If the light card is in UV space, set the spring arm's trasform to be zero, effectively removing it from the transform hierarchy
+	// If the light card is in UV space, set the spring arm's transform to be zero, effectively removing it from the transform hierarchy
 	// This allows the light card to be positioned with the actor's cartesian coordinates instead of longitude and latitude
 	if (bIsUVLightCard)
 	{
@@ -600,7 +608,12 @@ FTransform ADisplayClusterLightCardActor::GetStageActorTransform(bool bRemoveOri
 {
 	FTransform Transform;
 
-	FVector Position = LightCardComponent->GetComponentLocation();
+	const FVector ForwardVector = (GetActorRotation().Quaternion() * FRotator(-Latitude, Longitude, 0.0).Quaternion()).RotateVector(-FVector::ForwardVector);
+	constexpr float DistanceFactor = 0.99f;
+	FVector Position = bIsUVLightCard ?
+		FVector(DistanceFactor * UVPlaneDefaultDistance, -UVPlaneDefaultSize * (0.5 - UVCoordinates.X), UVPlaneDefaultSize * (0.5 - UVCoordinates.Y)) :
+		GetActorLocation() + DistanceFromCenter * ForwardVector;
+
 	if (bRemoveOrigin)
 	{
 		Position -= GetActorLocation();
diff --git a/Engine/Plugins/Runtime/nDisplay/Source/DisplayCluster/Private/Game/EngineClasses/Scene/DisplayClusterRootActor.cpp b/Engine/Plugins/Runtime/nDisplay/Source/DisplayCluster/Private/Game/EngineClasses/Scene/DisplayClusterRootActor.cpp
--- a/Engine/Plugins/Runtime/nDisplay/Source/DisplayCluster/Private/Game/EngineClasses/Scene/DisplayClusterRootActor.cpp
+++ b/Engine/Plugins/Runtime/nDisplay/Source/DisplayCluster/Private/Game/EngineClasses/Scene/DisplayClusterRootActor.cpp
@@ -43,6 +43,7 @@
 #include "Game/EngineClasses/Scene/DisplayClusterRootActorInitializer.h"
 
 #include "Algo/MaxElement.h"
+#include "Components/DisplayClusterStageGeometryComponent.h"
 #include "UObject/Package.h"
 
 #if WITH_EDITOR
@@ -82,6 +83,8 @@ ADisplayClusterRootActor::ADisplayClusterRootActor(const FObjectInitializer& Obj
 	// A helper component to trigger nDisplay Tick() during Tick phase
 	SyncTickComponent = CreateDefaultSubobject<UDisplayClusterSyncTickComponent>(TEXT("DisplayClusterSyncTick"));
 
+	StageGeometryComponent = CreateDefaultSubobject<UDisplayClusterStageGeometryComponent>(TEXT("DisplayClusterStageGeometry"));
+
 	// Default nDisplay camera
 	DefaultViewPoint = CreateDefaultSubobject<UDisplayClusterCameraComponent>(TEXT("DefaultViewPoint"));
 	DefaultViewPoint->AttachToComponent(GetRootComponent(), FAttachmentTransformRules::KeepRelativeTransform);
@@ -577,6 +580,8 @@ void ADisplayClusterRootActor::InitializeRootActor()
 		ViewportManager = MakeUnique<FDisplayClusterViewportManager>();
 	}
 
+	StageGeometryComponent->Invalidate();
+
 	// Packaged, PIE and -game runtime
 	if (IsRunningGameOrPIE())
 	{
@@ -698,6 +703,55 @@ void ADisplayClusterRootActor::UpdateLightCardPositions()
 	}
 }
 
+bool ADisplayClusterRootActor::GetFlushPositionAndNormal(const FVector& WorldPosition, FVector& OutPosition, FVector& OutNormal)
+{
+	const FVector StagePosition = GetCommonViewPoint()->GetComponentLocation();
+	const FVector Direction = (WorldPosition - StagePosition).GetSafeNormal();
+
+	float Distance;
+	FVector Normal;
+
+	if (StageGeometryComponent->GetStageDistanceAndNormal(Direction, Distance, Normal))
+	{
+		OutPosition = Distance * Direction + GetCommonViewPoint()->GetComponentLocation();
+
+		// Normal is returned in the local "radial basis", meaning that the x axis of the basis points radially inwards.
+		// Convert this to world coordinates using the radial basis made from the world direction (no need to account
+		// for root actor rotation here since Direction is already in world coordinates)
+		const FMatrix RadialBasis = FRotationMatrix::MakeFromX(Direction);
+		OutNormal = RadialBasis.TransformVector(Normal);
+		return true;
+	}
+
+	return false;
+}
+
+bool ADisplayClusterRootActor::MakeStageActorFlushToWall(const TScriptInterface<IDisplayClusterStageActor>& StageActor, double DesiredOffsetFromFlush)
+{
+	if (StageActor.GetObject() && !StageActor->IsUVActor())
+	{
+		FVector Position;
+		FVector Normal;
+		if (GetFlushPositionAndNormal(StageActor->GetStageActorTransform().GetLocation(), Position, Normal))
+		{
+			const FTransform StageActorTransform = StageActor->GetOrigin();
+			const FVector LocalPosition = StageActorTransform.InverseTransformPosition(Position);
+			const FVector LocalNormal = FRotationMatrix::MakeFromX(LocalPosition.GetSafeNormal()).InverseTransformVector(StageActorTransform.InverseTransformVectorNoScale(Normal));
+			const FRotator Rotation = FRotationMatrix::MakeFromX(-LocalNormal).Rotator();
+			const float Distance = FMath::Max(FMath::Min(LocalPosition.Length(), StageGeometryComponent->GetStageBoundingRadius()) + DesiredOffsetFromFlush, 0);
+
+			StageActor->SetDistanceFromCenter(Distance);
+			StageActor->SetPitch(Rotation.Pitch);
+			StageActor->SetYaw(Rotation.Yaw);
+
+			StageActor->UpdateStageActorTransform();
+			return true;
+		}
+	}
+
+	return false;
+}
+
 bool ADisplayClusterRootActor::IsBlueprint() const
 {
 	for (UClass* Class = GetClass(); Class; Class = Class->GetSuperClass())
diff --git a/Engine/Plugins/Runtime/nDisplay/Source/DisplayCluster/Private/Game/EngineClasses/Scene/DisplayClusterRootActorEditor.cpp b/Engine/Plugins/Runtime/nDisplay/Source/DisplayCluster/Private/Game/EngineClasses/Scene/DisplayClusterRootActorEditor.cpp
--- a/Engine/Plugins/Runtime/nDisplay/Source/DisplayCluster/Private/Game/EngineClasses/Scene/DisplayClusterRootActorEditor.cpp
+++ b/Engine/Plugins/Runtime/nDisplay/Source/DisplayCluster/Private/Game/EngineClasses/Scene/DisplayClusterRootActorEditor.cpp
@@ -41,6 +41,7 @@
 #include "Render/Viewport/IDisplayClusterViewport.h"
 #include "Render/Viewport/DisplayClusterViewportManager.h"
 #include "TextureResource.h"
+#include "Components/DisplayClusterStageGeometryComponent.h"
 
 //////////////////////////////////////////////////////////////////////////////////////////////
 // IN-EDITOR STUFF
@@ -100,6 +101,11 @@ void ADisplayClusterRootActor::Constructor_Editor()
 	ResetPreviewInternals_Editor();
 
 	FCoreUObjectDelegates::OnPackageReloaded.AddUObject(this, &ADisplayClusterRootActor::HandleAssetReload);
+
+	if (GEditor)
+	{
+		GEditor->OnEndObjectMovement().AddUObject(this, &ADisplayClusterRootActor::OnEndObjectMovement);
+	}
 }
 
 void ADisplayClusterRootActor::Destructor_Editor()
@@ -108,6 +114,11 @@ void ADisplayClusterRootActor::Destructor_Editor()
 	OnPreviewDestroyed.Unbind();
 
 	FCoreUObjectDelegates::OnPackageReloaded.RemoveAll(this);
+
+	if (GEditor)
+	{
+		GEditor->OnEndObjectMovement().RemoveAll(this);
+	}
 }
 
 void ADisplayClusterRootActor::Tick_Editor(float DeltaSeconds)
@@ -177,6 +188,8 @@ void ADisplayClusterRootActor::RerunConstructionScripts_Editor()
 
 	// Reset preview components before DCRA rebuild
 	ResetPreviewComponents_Editor(false);
+
+	StageGeometryComponent->Invalidate();
 }
 
 void ADisplayClusterRootActor::EnableEditorRender(bool bValue)
@@ -832,6 +845,15 @@ void ADisplayClusterRootActor::HandleAssetReload(const EPackageReloadPhase InPac
 	}
 }
 
+void ADisplayClusterRootActor::OnEndObjectMovement(UObject& InObject)
+{
+	// If any of this stage actor's components have been moved, invalidate the stage geometry map
+	if (InObject.IsA<USceneComponent>() && Cast<USceneComponent>(&InObject)->GetOwner() == this)
+	{
+		StageGeometryComponent->Invalidate();
+	}
+}
+
 void ADisplayClusterRootActor::ResetPreviewComponents_Editor(bool bInRestoreSceneMaterial)
 {
 	TArray<UDisplayClusterPreviewComponent*> AllPreviewComponents;
diff --git a/Engine/Plugins/Runtime/nDisplay/Source/DisplayCluster/Private/Game/EngineClasses/Scene/DisplayClusterStageGeometryComponent.cpp b/Engine/Plugins/Runtime/nDisplay/Source/DisplayCluster/Private/Game/EngineClasses/Scene/DisplayClusterStageGeometryComponent.cpp
new file mode 100644
--- /dev/null
+++ b/Engine/Plugins/Runtime/nDisplay/Source/DisplayCluster/Private/Game/EngineClasses/Scene/DisplayClusterStageGeometryComponent.cpp
@@ -0,0 +1,382 @@
+// Copyright Epic Games, Inc. All Rights Reserved.
+
+#include "Components/DisplayClusterStageGeometryComponent.h"
+
+#include "CanvasTypes.h"
+#include "DisplayClusterMeshProjectionRenderer.h"
+#include "DisplayClusterRootActor.h"
+#include "DisplayClusterConfigurationTypes.h"
+#include "ProceduralMeshComponent.h"
+#include "TextureResource.h"
+#include "Components/DisplayClusterScreenComponent.h"
+#include "Engine/TextureRenderTarget2D.h"
+#include "Brushes/SlateImageBrush.h"
+#include "Widgets/Images/SImage.h"
+
+const uint32 UDisplayClusterStageGeometryComponent::GeometryMapSize = 512;
+const float UDisplayClusterStageGeometryComponent::GeometryMapFOV = 2.0f * FMath::RadiansToDegrees(FMath::Atan(0.55 * PI));
+
+UDisplayClusterStageGeometryComponent::UDisplayClusterStageGeometryComponent()
+{
+	SetActive(true);
+	PrimaryComponentTick.bCanEverTick = true;
+
+	Renderer = MakeShared<FDisplayClusterMeshProjectionRenderer>();
+}
+
+void UDisplayClusterStageGeometryComponent::TickComponent(float DeltaTime, ELevelTick TickType, FActorComponentTickFunction* ThisTickFunction)
+{
+	Super::TickComponent(DeltaTime, TickType, ThisTickFunction);
+
+	if (bUpdateGeometryMap)
+	{
+		RedrawGeometryMap();
+	}
+}
+
+#if WITH_EDITOR
+namespace UE::DisplayClusterStageGeometryComponent
+{
+	class STextureImage : public SCompoundWidget
+	{
+	public:
+		SLATE_BEGIN_ARGS(STextureImage) {}
+			SLATE_ARGUMENT(UTexture2D*, Texture)
+		SLATE_END_ARGS()
+
+		void Construct(const FArguments& InArgs)
+		{
+			if (InArgs._Texture)
+			{
+				FTextureResource* Resource = InArgs._Texture->GetResource();
+				if (Resource)
+				{
+					FVector2D Size(Resource->GetSizeX(), Resource->GetSizeY());
+					TextureBrush = MakeShared<FSlateImageBrush>(InArgs._Texture,
+						Size,
+						FSlateColor(FLinearColor(1, 1, 1, 1)),
+						ESlateBrushTileType::NoTile,
+						ESlateBrushImageType::Linear);
+				}
+			}
+
+			if (!TextureBrush.IsValid())
+			{
+				TextureBrush = MakeShared<FSlateBrush>();
+			}
+
+			ChildSlot
+			[
+				SNew(SImage)
+				.Image(TextureBrush.Get())
+			];
+		}
+
+	private:
+		TSharedPtr<FSlateBrush> TextureBrush;
+	};
+}
+
+TSharedPtr<SWidget> UDisplayClusterStageGeometryComponent::GetCustomEditorPreviewWidget()
+{
+	return SNew(SVerticalBox)
+		+ SVerticalBox::Slot()
+		[
+			SNew(UE::DisplayClusterStageGeometryComponent::STextureImage)
+			.Texture(NorthGeometryMap.Texture)
+		]
+		+ SVerticalBox::Slot()
+		[
+			SNew(UE::DisplayClusterStageGeometryComponent::STextureImage)
+			.Texture(SouthGeometryMap.Texture)
+		];
+}
+#endif
+
+void UDisplayClusterStageGeometryComponent::Invalidate(bool bForceImmediateRedraw)
+{
+	if (bForceImmediateRedraw)
+	{
+		RedrawGeometryMap();
+	}
+	else
+	{
+		bUpdateGeometryMap = true;
+	}
+}
+
+bool UDisplayClusterStageGeometryComponent::GetStageDistanceAndNormal(const FVector& InDirection, float& OutDistance, FVector& OutNormal)
+{
+	if (bGeometryMapLoaded)
+	{
+		const bool bIsQueryingNorthernHemisphere = InDirection.Z > 0 ? true : false;
+		FDisplayClusterStageGeometryMap& GeometryMap = bIsQueryingNorthernHemisphere ? NorthGeometryMap : SouthGeometryMap;
+
+		auto GetPixel = [&GeometryMap](uint32 InX, uint32 InY)
+		{
+			uint32 ClampedX = FMath::Clamp(InX, (uint32)0, GeometryMapSize - 1);
+			uint32 ClampedY = FMath::Clamp(InY, (uint32)0, GeometryMapSize - 1);
+
+			return GeometryMap.GeometryData[ClampedY * GeometryMapSize + ClampedX].GetFloats();
+		};
+
+		// Use the view matrix of the geometry map to convert from world coordinates to view coordinates, then convert those coordinates into dome projection space
+		const FVector ViewPos = GeometryMap.ViewMatrices.GetViewMatrix().TransformVector(InDirection);
+		const FVector ProjectedViewPos = FDisplayClusterMeshProjectionRenderer::ProjectViewPosition(ViewPos, EDisplayClusterMeshProjectionType::Azimuthal);
+
+		const FVector4 ScreenPos = GeometryMap.ViewMatrices.GetProjectionMatrix().TransformFVector4(FVector4(ProjectedViewPos, 1));
+
+		if (ScreenPos.W != 0.0)
+		{
+			const float InvW = (ScreenPos.W > 0.0f ? 1.0f : -1.0f) / ScreenPos.W;
+			const float Y = (GProjectionSignY > 0.0f) ? ScreenPos.Y : 1.0f - ScreenPos.Y;
+			const FVector2D PixelPos = FVector2D((0.5f + ScreenPos.X * 0.5f * InvW) * GeometryMapSize, (0.5f - Y * 0.5f * InvW) * GeometryMapSize);
+
+			// Perform a bilinear interpolation on the computed pixel position to ensure a continuous normal regardless of the resolution of the normal map
+			const uint32 PixelX = FMath::Floor(PixelPos.X - 0.5f);
+			const uint32 PixelY = FMath::Floor(PixelPos.Y - 0.5f);
+			const float PixelXFrac = FMath::Frac(PixelPos.X);
+			const float PixelYFrac = FMath::Frac(PixelPos.Y);
+
+			const FLinearColor NormalData = FMath::Lerp(
+				FMath::Lerp(GetPixel(PixelX, PixelY), GetPixel(PixelX + 1, PixelY), PixelXFrac),
+				FMath::Lerp(GetPixel(PixelX, PixelY + 1), GetPixel(PixelX + 1, PixelY + 1), PixelXFrac),
+				PixelYFrac);
+
+			const FVector NormalVector = 2.f * FVector(NormalData.R, NormalData.G, NormalData.B) - 1.f;
+			OutNormal = NormalVector.GetSafeNormal();
+
+			// Make sure the depth value is not 0, as that will cause a divide by zero when transformed, resulting in an NaN distance being returned
+			const float Depth = FMath::Max(0.001f, NormalData.A);
+
+			FVector4 DepthPos = GeometryMap.ViewMatrices.GetInvProjectionMatrix().TransformFVector4(FVector4(ScreenPos.X * InvW, ScreenPos.Y * InvW, Depth, 1.0f));
+			DepthPos /= DepthPos.W;
+
+			const FVector UnprojectedDepthPos = FDisplayClusterMeshProjectionRenderer::UnprojectViewPosition(DepthPos, EDisplayClusterMeshProjectionType::Azimuthal);
+			OutDistance = UnprojectedDepthPos.Length();
+			return true;
+		}
+	}
+
+	OutNormal = FVector::ZeroVector;
+	OutDistance = 0.0;
+	return false;
+}
+
+bool UDisplayClusterStageGeometryComponent::MorphProceduralMesh(UProceduralMeshComponent* InProceduralMeshComponent)
+{
+	if (bGeometryMapLoaded)
+	{
+		FProcMeshSection& Section = *InProceduralMeshComponent->GetProcMeshSection(0);
+
+		for (int32 Index = 0; Index < Section.ProcVertexBuffer.Num(); ++Index)
+		{
+			FProcMeshVertex& Vertex = Section.ProcVertexBuffer[Index];
+
+			const FVector VertexDirection = Vertex.Position.GetSafeNormal();
+			float Distance;
+			FVector Normal;
+			GetStageDistanceAndNormal(VertexDirection, Distance, Normal);
+
+			const FVector NewPosition = VertexDirection * Distance;
+			Vertex.Position = NewPosition;
+
+			const FMatrix RadialBasis = FRotationMatrix::MakeFromX(VertexDirection);
+
+			const FVector WorldNormal = RadialBasis.TransformVector(Normal);
+			Vertex.Normal = WorldNormal;
+		}
+
+		InProceduralMeshComponent->SetProcMeshSection(0, Section);
+	}
+
+	return false;
+}
+
+UTextureRenderTarget2D* UDisplayClusterStageGeometryComponent::CreateRenderTarget()
+{
+	UTextureRenderTarget2D* RenderTarget = NewObject<UTextureRenderTarget2D>(this);
+	RenderTarget->InitCustomFormat(GeometryMapSize, GeometryMapSize, EPixelFormat::PF_FloatRGBA, false);
+	RenderTarget->ClearColor = FLinearColor(0.0f, 0.0f, 0.0f, 1.0f);
+
+	return RenderTarget;
+}
+
+UTexture2D* UDisplayClusterStageGeometryComponent::CreateTexture2D()
+{
+#if WITH_EDITOR
+	UTexture2D* Texture2D = NewObject<UTexture2D>(this);
+	Texture2D->Source.Init(GeometryMapSize, GeometryMapSize, 1, 1, TSF_BGRA8);
+
+	return Texture2D;
+#else
+	return nullptr;
+#endif
+}
+
+void UDisplayClusterStageGeometryComponent::RedrawGeometryMap()
+{
+	if (!bGeometryMapLoaded)
+	{
+		NorthGeometryMap.RenderTarget = CreateRenderTarget();
+		NorthGeometryMap.Texture = CreateTexture2D();
+
+		SouthGeometryMap.RenderTarget = CreateRenderTarget();
+		SouthGeometryMap.Texture = CreateTexture2D();
+
+		bGeometryMapLoaded = true;
+	}
+
+	UpdateStageGeometry();
+	GenerateGeometryMap(true);
+	GenerateGeometryMap(false);
+
+	bUpdateGeometryMap = false;
+}
+
+void UDisplayClusterStageGeometryComponent::UpdateStageGeometry()
+{
+	Renderer->ClearScene();
+
+	if (ADisplayClusterRootActor* RootActor = Cast<ADisplayClusterRootActor>(GetOwner()))
+	{
+		TArray<FString> ProjectionMeshNames;
+
+		if (UDisplayClusterConfigurationData* Config = RootActor->GetConfigData())
+		{
+			Config->GetReferencedMeshNames(ProjectionMeshNames);
+		}
+
+		Renderer->AddActor(RootActor, [&ProjectionMeshNames](const UPrimitiveComponent* PrimitiveComponent)
+		{
+			// Filter out any primitive component that isn't a projection mesh (a static mesh that has a Mesh projection configured for it) or a screen component
+			const bool bIsProjectionMesh = PrimitiveComponent->IsA<UStaticMeshComponent>() && ProjectionMeshNames.Contains(PrimitiveComponent->GetName());
+			const bool bIsScreen = PrimitiveComponent->IsA<UDisplayClusterScreenComponent>();
+			return bIsProjectionMesh || bIsScreen;
+		});
+
+		const FBox BoundingBox = RootActor->GetComponentsBoundingBox().ShiftBy(-RootActor->GetActorLocation());
+		StageBoundingRadius = FMath::Max(BoundingBox.Min.Length(), BoundingBox.Max.Length());
+	}
+}
+
+void UDisplayClusterStageGeometryComponent::GenerateGeometryMap(bool bIsNorthMap)
+{
+	if (ADisplayClusterRootActor* RootActor = Cast<ADisplayClusterRootActor>(GetOwner()))
+	{
+		UWorld* World = RootActor->GetWorld();
+		if (!World)
+		{
+			return;
+		}
+
+		FDisplayClusterStageGeometryMap& GeometryMap = bIsNorthMap ? NorthGeometryMap : SouthGeometryMap;
+
+		// Rotate the normal map's direction to account for the actor's world rotation
+		FRotator InverseRootRotation = RootActor->GetTransform().Rotator().GetInverse();
+		FVector NormalMapDirection = InverseRootRotation.RotateVector(bIsNorthMap ? FVector::UpVector : FVector::DownVector);
+
+		FDisplayClusterMeshProjectionRenderSettings RenderSettings;
+		RenderSettings.RenderType = EDisplayClusterMeshProjectionOutput::Normals;
+		RenderSettings.ProjectionType = EDisplayClusterMeshProjectionType::Azimuthal;
+		RenderSettings.NormalCorrectionMatrix = FMatrix44f(FRotationMatrix::Make(InverseRootRotation));
+
+		GetSceneViewInitOptions(NormalMapDirection, RenderSettings.ViewInitOptions);
+		GeometryMap.ViewMatrices = FViewMatrices(RenderSettings.ViewInitOptions);
+
+		FTextureRenderTarget2DResource* TexResource = (FTextureRenderTarget2DResource*)GeometryMap.RenderTarget->GetResource();
+		if (TexResource)
+		{
+			FCanvas Canvas(TexResource, nullptr, World, World->Scene->GetFeatureLevel());
+			Canvas.Clear(FLinearColor::Black);
+
+			Renderer->Render(&Canvas, World->Scene, RenderSettings);
+			Canvas.Flush_GameThread();
+
+			TexResource->ReadFloat16Pixels(GeometryMap.GeometryData);
+
+#if WITH_EDITOR
+			if (GeometryMap.Texture)
+			{
+				uint8* MipData = GeometryMap.Texture->Source.LockMip(0);
+				for (int32 Index = 0; Index < GeometryMap.GeometryData.Num(); ++Index)
+				{
+					uint8* DestPtr = &MipData[Index * sizeof(FColor)];
+					const FColor SourceColor = GeometryMap.GeometryData[Index].GetFloats().ToFColor(false);
+
+					*DestPtr++ = SourceColor.B;
+					*DestPtr++ = SourceColor.G;
+					*DestPtr++ = SourceColor.R;
+					*DestPtr++ = 0xFF;
+				}
+				GeometryMap.Texture->Source.UnlockMip(0);
+
+				GeometryMap.Texture->PostEditChange();
+			}
+#endif
+		}
+	}
+}
+
+void UDisplayClusterStageGeometryComponent::GetSceneViewInitOptions(const FVector& ViewDirection, FSceneViewInitOptions& OutViewInitOptions)
+{
+	if (ADisplayClusterRootActor* RootActor = Cast<ADisplayClusterRootActor>(GetOwner()))
+	{
+		OutViewInitOptions.ViewLocation = RootActor->GetCommonViewPoint()->GetComponentLocation();
+		OutViewInitOptions.ViewRotation = FRotator(RootActor->GetActorRotation().Quaternion() * ViewDirection.Rotation().Quaternion());
+		OutViewInitOptions.ViewOrigin = OutViewInitOptions.ViewLocation;
+
+		OutViewInitOptions.SetViewRectangle(FIntRect(0, 0, GeometryMapSize, GeometryMapSize));
+
+		AWorldSettings* WorldSettings = nullptr;
+
+		if (RootActor->GetWorld())
+		{
+			WorldSettings = RootActor->GetWorld()->GetWorldSettings();
+		}
+
+		if (WorldSettings != nullptr)
+		{
+			OutViewInitOptions.WorldToMetersScale = WorldSettings->WorldToMeters;
+		}
+
+		// Rotate view 90 degrees
+		OutViewInitOptions.ViewRotationMatrix = FInverseRotationMatrix(OutViewInitOptions.ViewRotation) * FMatrix(
+			FPlane(0, 0, 1, 0),
+			FPlane(1, 0, 0, 0),
+			FPlane(0, 1, 0, 0),
+			FPlane(0, 0, 0, 1));
+
+		const float MinZ = GNearClippingPlane;
+		const float MaxZ = FMath::Max(StageBoundingRadius, MinZ);
+
+		// Avoid zero ViewFOVs which cause divide by zero's in projection matrix
+		const float MatrixFOV = FMath::Max(0.001f, GeometryMapFOV) * (float)PI / 360.0f;
+
+		const float XAxisMultiplier = 1.0f;
+		const float YAxisMultiplier = 1.0f;
+
+		if ((bool)ERHIZBuffer::IsInverted)
+		{
+			OutViewInitOptions.ProjectionMatrix = FReversedZPerspectiveMatrix(
+				MatrixFOV,
+				MatrixFOV,
+				XAxisMultiplier,
+				YAxisMultiplier,
+				MinZ,
+				MaxZ);
+		}
+		else
+		{
+			OutViewInitOptions.ProjectionMatrix = FPerspectiveMatrix(
+				MatrixFOV,
+				MatrixFOV,
+				XAxisMultiplier,
+				YAxisMultiplier,
+				MinZ,
+				MaxZ);
+		}
+
+		OutViewInitOptions.FOV = GeometryMapFOV;
+	}
+}
diff --git a/Engine/Plugins/Runtime/nDisplay/Source/DisplayCluster/Public/Components/DisplayClusterStageGeometryComponent.h b/Engine/Plugins/Runtime/nDisplay/Source/DisplayCluster/Public/Components/DisplayClusterStageGeometryComponent.h
new file mode 100644
--- /dev/null
+++ b/Engine/Plugins/Runtime/nDisplay/Source/DisplayCluster/Public/Components/DisplayClusterStageGeometryComponent.h
@@ -0,0 +1,136 @@
+// Copyright Epic Games, Inc. All Rights Reserved.
+
+#pragma once
+
+#include "CoreMinimal.h"
+#include "SceneView.h"
+#include "Components/ActorComponent.h"
+
+#include "DisplayClusterStageGeometryComponent.generated.h"
+
+class UProceduralMeshComponent;
+class UTexture2D;
+class UTextureRenderTarget2D;
+class FDisplayClusterMeshProjectionRenderer;
+
+/** A geometry map for the stage, which contains a generated texture of the stage geometry's depth and normal vectors */
+USTRUCT()
+struct FDisplayClusterStageGeometryMap
+{
+	GENERATED_BODY()
+
+	/** The render target of the geometry map */
+	UPROPERTY()
+	TObjectPtr<UTextureRenderTarget2D> RenderTarget;
+
+	/** A displayable texture of the geometry map */
+	UPROPERTY()
+	TObjectPtr<UTexture2D> Texture;
+
+	/** The raw texture data of the geometry map, which stores the normal vector in the RGB components and the depth in the A component */
+	TArray<FFloat16Color> GeometryData;
+
+	/** The view matrices used to render the geometry map */
+	FViewMatrices ViewMatrices;
+};
+
+/** A component that stores the generated geometry map of the stage actor, which is used for placing stage actors (light cards, CCWs, etc) flush to the stage's walls and ceiling */
+UCLASS()
+class DISPLAYCLUSTER_API UDisplayClusterStageGeometryComponent : public UActorComponent
+{
+	GENERATED_BODY()
+
+private:
+	/** The default size of the geometry map texture */
+	static const uint32 GeometryMapSize;
+
+	/** The field of view to render the geometry map with */
+	static const float GeometryMapFOV;
+
+public:
+	UDisplayClusterStageGeometryComponent();
+
+	//~ Begin UObject interface
+	virtual void TickComponent(float DeltaTime, ELevelTick TickType, FActorComponentTickFunction* ThisTickFunction) override;
+
+#if WITH_EDITOR
+	virtual bool GetEditorPreviewInfo(float DeltaTime, FMinimalViewInfo& ViewOut) override { return true; }
+	virtual TSharedPtr<SWidget> GetCustomEditorPreviewWidget() override;
+#endif
+	//~ End UObject interface
+
+	/**
+	 * Invalidates and regenerates the cached stage geometry map
+	 * @param bForceImmediateRedraw Indicates whether the geometry map is regenerated immediately or is queued to redraw on the next tick
+	 */
+	UFUNCTION(BlueprintCallable, Category = "NDisplay")
+	void Invalidate(bool bForceImmediateRedraw = false);
+
+	/**
+	 * Gets the distance and normal vector (in radial space) of the stage's geometry in the specified world direction
+	 * @param InDirection The direction in world coordinates to query the stage geometry map
+	 * @param OutDistance The distance from the stage's default view origin to the nearest stage geometry in the specified direction
+	 * @param OutNormal The normal vector in radial space (with x axis pointing in the direction of the stage's default view origin) of the nearest stage geometry in the specified direction
+	 * @return true if the stage geometry map was successfully queried, otherwise false
+	 */
+	UFUNCTION(BlueprintCallable, Category = "NDisplay")
+	bool GetStageDistanceAndNormal(const FVector& InDirection, float& OutDistance, FVector& OutNormal);
+
+	/**
+	 * Gets the bounding radius of the stage's geometry
+	 */
+	UFUNCTION(BlueprintGetter, Category = "NDisplay")
+	float GetStageBoundingRadius() const { return StageBoundingRadius; }
+
+	/**
+	 * Morphs the specified procedural mesh to match the stage's geometry map
+	 */
+	bool MorphProceduralMesh(UProceduralMeshComponent* InProceduralMeshComponent);
+
+	/** Gets the geometry map texture for the northern hemisphere of the stage */
+	UTexture2D* GetNorthGeometryMapTexture() const { return NorthGeometryMap.Texture; }
+
+	/** Gets the geometry map texture for the southern hemisphere of the stage */
+	UTexture2D* GetSouthGeometryMapTexture() const { return SouthGeometryMap.Texture; }
+
+private:
+	/** Creates a new render target that can be used to render the stage's geometry map */
+	UTextureRenderTarget2D* CreateRenderTarget();
+
+	/** Creates a new texture that can be used to display the stage's geometry map */
+	UTexture2D* CreateTexture2D();
+
+	/** Redraws the stage's geometry map */
+	void RedrawGeometryMap();
+
+	/** Updates the stage's current geometry with the renderer */
+	void UpdateStageGeometry();
+
+	/** Renders the stage's geometry to the specified geometry map */
+	void GenerateGeometryMap(bool bIsNorthMap);
+
+	/** Computes the scene view init options necessary to render the geometry map in the specified view direction */
+	void GetSceneViewInitOptions(const FVector& ViewDirection, FSceneViewInitOptions& OutViewInitOptions);
+
+private:
+	/** The cached geometry map for the stage's northern hemisphere */
+	UPROPERTY(Transient)
+	FDisplayClusterStageGeometryMap NorthGeometryMap;
+
+	/** The cached geometry map for the stage's southern hemisphere */
+	UPROPERTY(Transient)
+	FDisplayClusterStageGeometryMap SouthGeometryMap;
+
+	/** The cached bounding radius of the stage's geometry */
+	UPROPERTY(Transient)
+	float StageBoundingRadius = 0.0f;
+
+	/** The renderer used to generate the stage's geometry map */
+	TSharedPtr<FDisplayClusterMeshProjectionRenderer> Renderer;
+
+	/** Indicates that the geometry map should be updated on the next tick */
+	bool bUpdateGeometryMap = false;
+
+	/** Indicates that the resources for the geometry map have been loaded */
+	bool bGeometryMapLoaded = false;
+};
\ No newline at end of file
diff --git a/Engine/Plugins/Runtime/nDisplay/Source/DisplayCluster/Public/DisplayClusterLightCardActor.h b/Engine/Plugins/Runtime/nDisplay/Source/DisplayCluster/Public/DisplayClusterLightCardActor.h
--- a/Engine/Plugins/Runtime/nDisplay/Source/DisplayCluster/Public/DisplayClusterLightCardActor.h
+++ b/Engine/Plugins/Runtime/nDisplay/Source/DisplayCluster/Public/DisplayClusterLightCardActor.h
@@ -111,7 +111,7 @@ public:
 	// ~Begin IDisplayClusterStageActor interface
 	/** Updates the Light Card transform based on its positional properties (Lat, Long, etc.) */
 	virtual void UpdateStageActorTransform() override;
-	
+
 	/**
 	 * Gets the transform in world space of the light card component
 	 */
@@ -149,6 +149,8 @@ public:
 	virtual void SetScale(const FVector2D& InScale) override;
 	virtual FVector2D GetScale() const override;
 
+	virtual void SetAlwaysFlushToWall(bool bInAlwaysFlushToWall) override { bAlwaysFlushToWall = bInAlwaysFlushToWall; }
+	virtual bool IsAlwaysFlushToWall() const override { return bAlwaysFlushToWall; }
 	virtual void SetUVCoordinates(const FVector2D& InUVCoordinates) override;
 	virtual FVector2D GetUVCoordinates() const override;
 
@@ -190,6 +192,10 @@ public:
 	UPROPERTY(EditAnywhere, BlueprintReadWrite, Category = "Orientation", meta = (EditCondition = "!bIsUVLightCard", HideEditConditionToggle, EditConditionHides))
 	double RadialOffset;
 
+	/** Indicates whether the light card is always made to be flush to a stage wall or not */
+	UPROPERTY(EditAnywhere, BlueprintReadWrite, Category = "Orientation", meta = (EditCondition = "!bIsUVLightCard", HideEditConditionToggle, EditConditionHides))
+	bool bAlwaysFlushToWall;
+
 	UPROPERTY(EditAnywhere, BlueprintReadWrite, Category = "Appearance")
 	EDisplayClusterLightCardMask Mask;
 
diff --git a/Engine/Plugins/Runtime/nDisplay/Source/DisplayCluster/Public/DisplayClusterRootActor.h b/Engine/Plugins/Runtime/nDisplay/Source/DisplayCluster/Public/DisplayClusterRootActor.h
--- a/Engine/Plugins/Runtime/nDisplay/Source/DisplayCluster/Public/DisplayClusterRootActor.h
+++ b/Engine/Plugins/Runtime/nDisplay/Source/DisplayCluster/Public/DisplayClusterRootActor.h
@@ -29,11 +29,13 @@
 class IDisplayClusterConfiguratorBlueprintEditor;
 #endif
 
+class IDisplayClusterStageActor;
 class USceneComponent;
 class UDisplayClusterConfigurationData;
 class UDisplayClusterCameraComponent;
 class UDisplayClusterOriginComponent;
 class UDisplayClusterPreviewComponent;
+class UDisplayClusterStageGeometryComponent;
 class UDisplayClusterSyncTickComponent;
 class UProceduralMeshComponent;
 
@@ -119,6 +121,15 @@ protected:
 	void UpdateLightCardPositions();
 
 public:
+	UFUNCTION(BlueprintCallable, Category = "NDisplay")
+	bool GetFlushPositionAndNormal(const FVector& WorldPosition, FVector& OutPosition, FVector& OutNormal);
+
+	UFUNCTION(BlueprintCallable, Category = "NDisplay")
+	bool MakeStageActorFlushToWall(const TScriptInterface<IDisplayClusterStageActor>& StageActor, double DesiredOffsetFromFlush = 0.0f);
+
+	UFUNCTION(BlueprintGetter)
+	UDisplayClusterStageGeometryComponent* GetStageGeometryComponent() const { return StageGeometryComponent; }
+
 	UFUNCTION(BlueprintGetter)
 	UDisplayClusterCameraComponent* GetDefaultCamera() const;
 
@@ -262,6 +273,10 @@ private:
 	UPROPERTY()
 	TObjectPtr<UDisplayClusterSyncTickComponent> SyncTickComponent;
 
+	/** Component that stores the stage's geometry map, which is used to make objects flush with the stage's walls and ceilings */
+	UPROPERTY()
+	TObjectPtr<UDisplayClusterStageGeometryComponent> StageGeometryComponent;
+
 private:
 	// Current operation mode
 	EDisplayClusterOperationMode OperationMode;
@@ -481,7 +496,10 @@ protected:
 
 	/** Called when the asset has been reloaded in the editor. */
 	void HandleAssetReload(const EPackageReloadPhase InPackageReloadPhase, FPackageReloadedEvent* InPackageReloadedEvent);
-	
+
+	/** Raised when any object is moved within a level. Used to update the stage's geometry map if any of its components have been moved */
+	void OnEndObjectMovement(UObject& InObject);
+
 private:
 
 	/** Flag to indicate if the user is currently interacting with a subobject of CurrentConfigData. */
diff --git a/Engine/Plugins/Runtime/nDisplay/Source/DisplayClusterLightCardEditor/Private/Viewport/DisplayClusterLightCardEditorViewportClient.cpp b/Engine/Plugins/Runtime/nDisplay/Source/DisplayClusterLightCardEditor/Private/Viewport/DisplayClusterLightCardEditorViewportClient.cpp
--- a/Engine/Plugins/Runtime/nDisplay/Source/DisplayClusterLightCardEditor/Private/Viewport/DisplayClusterLightCardEditorViewportClient.cpp
+++ b/Engine/Plugins/Runtime/nDisplay/Source/DisplayClusterLightCardEditor/Private/Viewport/DisplayClusterLightCardEditorViewportClient.cpp
@@ -55,6 +55,7 @@
 #include "TextureResource.h"
 #include "UnrealEdGlobals.h"
 #include "UnrealWidget.h"
+#include "Components/DisplayClusterStageGeometryComponent.h"
 #include "Widgets/Docking/SDockTab.h"
 
 
@@ -1612,6 +1613,10 @@ void FDisplayClusterLightCardEditorViewportClient::UpdatePreviewActor(ADisplayCl
 			
 				PreviewWorld->GetCurrentLevel()->AddLoadedActor(RootActorProxy.Get());
 
+				// Draw the geometry map for the proxy stage actor immediately to avoid a race condition where the geometry map could render
+				// before the actor location changes propagate to its component proxies, resulting in an inaccurate proxy geometry map
+				RootActorProxy->GetStageGeometryComponent()->Invalidate(true);
+
 				// Spawned actor will take the transform values from the template, so manually reset them to zero here
 				RootActorProxy->SetActorLocation(FVector::ZeroVector);
 				RootActorProxy->SetActorRotation(FRotator::ZeroRotator);
diff --git a/Engine/Plugins/Runtime/nDisplay/Source/DisplayClusterScenePreview/Private/DisplayClusterLightCardEditorHelper.cpp b/Engine/Plugins/Runtime/nDisplay/Source/DisplayClusterScenePreview/Private/DisplayClusterLightCardEditorHelper.cpp
--- a/Engine/Plugins/Runtime/nDisplay/Source/DisplayClusterScenePreview/Private/DisplayClusterLightCardEditorHelper.cpp
+++ b/Engine/Plugins/Runtime/nDisplay/Source/DisplayClusterScenePreview/Private/DisplayClusterLightCardEditorHelper.cpp
@@ -14,10 +14,10 @@
 #include "StageActor/DisplayClusterStageActorTemplate.h"
 
 #include "CanvasTypes.h"
-#include "ImageUtils.h"
 #include "KismetProceduralMeshLibrary.h"
 #include "PreviewScene.h"
 #include "ProceduralMeshComponent.h"
+#include "Components/DisplayClusterStageGeometryComponent.h"
 #include "Containers/ArrayView.h"
 #include "Engine/StaticMesh.h"
 #include "Engine/Texture2D.h"
@@ -158,165 +158,6 @@ bool FDisplayClusterLightCardEditorHelper::FSphericalCoordinates::IsPointingAtPo
 		|| FMath::IsNearlyEqual(CoordsConformed.Inclination, PI, Margin);
 }
 
-//////////////////////////////////////////////////////////////////////////
-// FNormalMap
-
-const int32 FDisplayClusterLightCardEditorHelper::FNormalMap::NormalMapSize = 512;
-const float FDisplayClusterLightCardEditorHelper::FNormalMap::NormalMapFOV = 2.0f * FMath::RadiansToDegrees(FMath::Atan(0.55 * PI)); // Equation for FOV from desired angle from north pole;
-
-void FDisplayClusterLightCardEditorHelper::FNormalMap::Init(const FSceneViewInitOptions& InSceneViewInitOptions)
-{
-	SizeX = InSceneViewInitOptions.GetViewRect().Width();
-	SizeY = InSceneViewInitOptions.GetViewRect().Height();
-
-	ViewMatrices = FViewMatrices(InSceneViewInitOptions);
-
-	if (NormalMapTexture.IsValid())
-	{
-		NormalMapTexture->MarkAsGarbage();
-		NormalMapTexture = nullptr;
-	}
-
-	ENQUEUE_RENDER_COMMAND(InitRHIResourcesCommand)([this](FRHICommandListImmediate& RHICmdList)
-		{
-			const FRHITextureCreateDesc Desc =
-				FRHITextureCreateDesc::Create2D(TEXT("NormalMapTexture"))
-				.SetExtent(SizeX, SizeY)
-				.SetFormat(PF_FloatRGBA)
-				.SetClearValue(FClearValueBinding::Black)
-				.SetFlags(ETextureCreateFlags::RenderTargetable | ETextureCreateFlags::ShaderResource)
-				.SetInitialState(ERHIAccess::SRVMask);
-
-			RenderTargetTextureRHI = RHICreateTexture(Desc);
-		});
-}
-
-void FDisplayClusterLightCardEditorHelper::FNormalMap::Release()
-{
-	ENQUEUE_RENDER_COMMAND(ReleaseRHIResourcesCommand)([this](FRHICommandListImmediate& RHICmdList)
-		{
-			RenderTargetTextureRHI.SafeRelease();
-		});
-}
-
-bool FDisplayClusterLightCardEditorHelper::FNormalMap::GetNormalAndDistanceAtPosition(FVector Position, FVector& OutNormal, float& OutDistance) const
-{
-	auto GetPixel = [this](uint32 InX, uint32 InY)
-	{
-		uint32 ClampedX = FMath::Clamp(InX, (uint32)0, SizeX - 1);
-		uint32 ClampedY = FMath::Clamp(InY, (uint32)0, SizeY - 1);
-
-		return CachedNormalData[ClampedY * SizeX + ClampedX].GetFloats();
-	};
-
-	const FVector ViewPos = FVector(ViewMatrices.GetViewMatrix().TransformFVector4(FVector4(Position, 1)));
-	const FVector ProjectedViewPos = FDisplayClusterMeshProjectionRenderer::ProjectViewPosition(ViewPos, EDisplayClusterMeshProjectionType::Azimuthal);
-
-	const FVector4 ScreenPos = ViewMatrices.GetProjectionMatrix().TransformFVector4(FVector4(ProjectedViewPos, 1));
-
-	if (ScreenPos.W != 0.0)
-	{
-		const float InvW = (ScreenPos.W > 0.0f ? 1.0f : -1.0f) / ScreenPos.W;
-		const float Y = (GProjectionSignY > 0.0f) ? ScreenPos.Y : 1.0f - ScreenPos.Y;
-		const FVector2D PixelPos = FVector2D((0.5f + ScreenPos.X * 0.5f * InvW) * SizeX, (0.5f - Y * 0.5f * InvW) * SizeY);
-
-		// Perform a bilinear interpolation on the computed pixel position to ensure a continuous normal regardless of the resolution of the normal map
-		const uint32 PixelX = FMath::Floor(PixelPos.X - 0.5f);
-		const uint32 PixelY = FMath::Floor(PixelPos.Y - 0.5f);
-		const float PixelXFrac = FMath::Frac(PixelPos.X);
-		const float PixelYFrac = FMath::Frac(PixelPos.X);
-
-		FLinearColor NormalData;
-		NormalData = FMath::Lerp(
-			FMath::Lerp(GetPixel(PixelX, PixelY), GetPixel(PixelX + 1, PixelY), PixelXFrac),
-			FMath::Lerp(GetPixel(PixelX, PixelY + 1), GetPixel(PixelX + 1, PixelY + 1), PixelXFrac),
-			PixelYFrac);
-
-		const FVector NormalVector = 2.f * FVector(NormalData.R, NormalData.G, NormalData.B) - 1.f;
-		OutNormal = NormalVector.GetSafeNormal();
-
-		// Make sure the depth value is not 0, as that will cause a divide by zero when transformed, resulting in an NaN distance being returned
-		const float Depth = FMath::Max(0.001f, NormalData.A);
-
-		FVector4 DepthPos = ViewMatrices.GetInvProjectionMatrix().TransformFVector4(FVector4(ScreenPos.X * InvW, ScreenPos.Y * InvW, Depth, 1.0f));
-		DepthPos /= DepthPos.W;
-
-		const FVector UnprojectedDepthPos = FDisplayClusterMeshProjectionRenderer::UnprojectViewPosition(DepthPos, EDisplayClusterMeshProjectionType::Azimuthal);
-		OutDistance = UnprojectedDepthPos.Length();
-
-		return true;
-	}
-	else
-	{
-		OutNormal = FVector::ZeroVector;
-		OutDistance = 0.0;
-		return false;
-	}
-}
-
-void FDisplayClusterLightCardEditorHelper::FNormalMap::MorphProceduralMesh(UProceduralMeshComponent* InProceduralMeshComponent) const
-{
-	const FVector ViewOrigin = ViewMatrices.GetViewOrigin();
-	const FVector ViewDirection = ViewMatrices.GetViewMatrix().TransformVector(FVector::ZAxisVector);
-	const float MaxAngle = 1.0f / ViewMatrices.GetProjectionMatrix().M[0][0];
-
-	FProcMeshSection& Section = *InProceduralMeshComponent->GetProcMeshSection(0);
-
-	for (int32 Index = 0; Index < Section.ProcVertexBuffer.Num(); ++Index)
-	{
-		FProcMeshVertex& Vertex = Section.ProcVertexBuffer[Index];
-
-		const FVector VertexPosition = Vertex.Position;
-		const FVector VertexWorldPosition = VertexPosition + ViewOrigin;
-		const FVector VertexDirection = Vertex.Position.GetSafeNormal();
-		const float VertexAngle = FMath::Acos(VertexDirection | ViewDirection);
-
-		if (VertexAngle < MaxAngle)
-		{
-			FVector Normal;
-			float Depth;
-			GetNormalAndDistanceAtPosition(VertexWorldPosition, Normal, Depth);
-
-			const FVector NewPosition = VertexDirection * Depth;
-			Vertex.Position = NewPosition;
-
-			const FMatrix RadialBasis = FRotationMatrix::MakeFromX(VertexDirection);
-
-			const FVector WorldNormal = RadialBasis.TransformVector(Normal);
-			Vertex.Normal = WorldNormal;
-		}
-	}
-
-	InProceduralMeshComponent->SetProcMeshSection(0, Section);
-}
-
-UTexture2D* FDisplayClusterLightCardEditorHelper::FNormalMap::GenerateNormalMapTexture(const FString& TextureName)
-{
-	if (NormalMapTexture.IsValid())
-	{
-		NormalMapTexture->MarkAsGarbage();
-		NormalMapTexture = nullptr;
-	}
-
-	if (CachedNormalData.Num())
-	{
-		FCreateTexture2DParameters Params;
-		Params.bDeferCompression = true;
-
-		TArray<FColor> Bitmap;
-		Bitmap.AddZeroed(CachedNormalData.Num());
-
-		for (int32 Index = 0; Index < CachedNormalData.Num(); ++Index)
-		{
-			Bitmap[Index] = CachedNormalData[Index].GetFloats().ToFColor(false);
-		}
-
-		NormalMapTexture = FImageUtils::CreateTexture2D(SizeX, SizeY, Bitmap, GetTransientPackage(), TextureName, RF_Transient, Params);
-	}
-
-	return GetNormalMapTexture();
-}
-
 
 //////////////////////////////////////////////////////////////////////////
 // FDisplayClusterLightCardEditorHelper
@@ -398,16 +239,13 @@ void FDisplayClusterLightCardEditorHelper::SetLevelInstanceRootActor(ADisplayClu
 
 const UTexture2D* FDisplayClusterLightCardEditorHelper::GetNormalMapTexture(bool bShowNorthMap)
 {
-	FNormalMap& NormalMap = bShowNorthMap ? NorthNormalMap : SouthNormalMap;
-	UTexture2D* NormalMapTexture = NormalMap.GetNormalMapTexture();
-
-	if (!NormalMapTexture)
+	if (ADisplayClusterRootActor* RootActor = UpdateRootActor())
 	{
-		const FString TextureName = bShowNorthMap ? TEXT("DisplayClusterLightCardEditorHelper.NorthNormalMap") : TEXT("DisplayClusterLightCardEditorHelper.SouthNormalMap");
-		NormalMapTexture = NormalMap.GenerateNormalMapTexture(TextureName);
+		UTexture2D* NormalMapTexture = bShowNorthMap ? RootActor->GetStageGeometryComponent()->GetNorthGeometryMapTexture() : RootActor->GetStageGeometryComponent()->GetSouthGeometryMapTexture();
+		return NormalMapTexture;
 	}
 
-	return NormalMapTexture;
+	return nullptr;
 }
 
 void FDisplayClusterLightCardEditorHelper::MoveActorsToPixel(
@@ -644,20 +482,9 @@ bool FDisplayClusterLightCardEditorHelper::CalculateNormalAndPositionInDirection
 	}
 	else
 	{
-		const FVector Position = InViewOrigin + InDirection; // We fabricate a position in the right direction
-
 		float Distance;
 
-		// We find the normal and distance from origin
-		if (Position.Z < InViewOrigin.Z)
-		{
-			SouthNormalMap.GetNormalAndDistanceAtPosition(Position, OutRelativeNormal, Distance);
-		}
-		else
-		{
-			NorthNormalMap.GetNormalAndDistanceAtPosition(Position, OutRelativeNormal, Distance);
-		}
-
+		CachedRootActor->GetStageGeometryComponent()->GetStageDistanceAndNormal(InDirection, Distance, OutRelativeNormal);
 		Distance = CalculateFinalLightCardDistance(Distance, InDesiredDistanceFromFlush);
 
 		// Calculate world position
@@ -1222,7 +1049,7 @@ FDisplayClusterLightCardEditorHelper::FSphericalCoordinates FDisplayClusterLight
 	FSphericalCoordinates DeltaCoords;
 
 	// If we are in a cartesian coordinate system and are constraining to an axis, perform the constraint calculations in
-	// cartesian coordinates, then convert to spherical coordinates at the end. Otherwise, perform all calculations in spherical coodinates
+	// cartesian coordinates, then convert to spherical coordinates at the end. Otherwise, perform all calculations in spherical coordinates
 	if (CoordinateSystem == ECoordinateSystem::Cartesian && DragAxis != EAxisList::Type::XYZ)
 	{
 		// For consistency, project the cursor direction vector onto the sphere the light card is currently on. Gives a good balance of approximating
@@ -1257,19 +1084,11 @@ FDisplayClusterLightCardEditorHelper::FSphericalCoordinates FDisplayClusterLight
 	}
 	else
 	{
-		FVector Normal;
-		float Distance;
-
-		// If the light card is in the southern hemisphere of the view origin, use the southern normal map; otherwise, use the north normal map
-		if (LightCardLocation.Z < 0.0f)
-		{
-			SouthNormalMap.GetNormalAndDistanceAtPosition(Actor->GetStageActorTransform(true).GetTranslation(), Normal, Distance);
-		}
-		else
-		{
-			NorthNormalMap.GetNormalAndDistanceAtPosition(Actor->GetStageActorTransform(true).GetTranslation(), Normal, Distance);
-		}
+		FVector FlushPosition;
+		FVector FlushNormal;
 
+		CachedRootActor->GetFlushPositionAndNormal(Actor->GetStageActorTransform().GetTranslation(), FlushPosition, FlushNormal);
+		const float Distance = (FlushPosition - CachedRootActor->GetCommonViewPoint()->GetComponentLocation()).Size();
 		const FSphericalCoordinates RequestedCoords(LocalDirection * Distance);
 
 		DeltaCoords = RequestedCoords - ActorCoords;
@@ -1347,36 +1166,16 @@ void FDisplayClusterLightCardEditorHelper::InternalMoveActorTo(
 		return;
 	}
 
-	const FVector Origin = GetProjectionOrigin();
-	const FVector LightCardPosition = Origin + Position.AsCartesian();
-
-	FVector DesiredNormal;
-	float DesiredDistance;
-
-	// If the light card is in the southern hemisphere of the view origin, use the southern normal map; otherwise, use the north normal map
-	if (LightCardPosition.Z < Origin.Z)
-	{
-		SouthNormalMap.GetNormalAndDistanceAtPosition(LightCardPosition, DesiredNormal, DesiredDistance);
-	}
-	else
-	{
-		NorthNormalMap.GetNormalAndDistanceAtPosition(LightCardPosition, DesiredNormal, DesiredDistance);
-	}
-
 	// Remove actor rotation from the position before setting the coordinate. Note that we don't need to do this for the normal map,
 	// which already takes the root actor rotation into account as part of its view matrix.
 	const FQuat RootRotation = CachedRootActor->GetTransform().GetRotation().Inverse();
 	const FVector InverseRotatedPosition = RootRotation.RotateVector(Position.AsCartesian());
-
-	SetActorCoordinates(Actor, FSphericalCoordinates(InverseRotatedPosition));
 	
-	const double FinalDistance = CalculateFinalLightCardDistance(DesiredDistance);
-	Actor->SetDistanceFromCenter(FinalDistance);
+	SetActorCoordinates(Actor, FSphericalCoordinates(InverseRotatedPosition));
+	Actor->UpdateStageActorTransform();
 
-	const FRotator Rotation = FRotationMatrix::MakeFromX(-DesiredNormal).Rotator();
-	
-	Actor->SetPitch(Rotation.Pitch);
-	Actor->SetYaw(Rotation.Yaw);
+	ADisplayClusterRootActor* RootActor = (Actor->IsProxy() || !LevelInstanceRootActor.IsValid()) ? CachedRootActor.Get() : LevelInstanceRootActor.Get();
+	RootActor->MakeStageActorFlushToWall(Actor.AsActor());
 
 #if WITH_EDITOR
 	if (!Actor->IsProxy())
@@ -1503,84 +1302,6 @@ void FDisplayClusterLightCardEditorHelper::InternalDragActors(const TArray<FDisp
 	}
 }
 
-void FDisplayClusterLightCardEditorHelper::GetNormalMapSceneViewInitOptions(const FVector& ViewDirection, FSceneViewInitOptions& OutViewInitOptions)
-{
-	OutViewInitOptions.ViewLocation = ProjectionOriginComponent.IsValid() ? ProjectionOriginComponent->GetComponentLocation() : FVector::ZeroVector;
-	OutViewInitOptions.ViewRotation = ViewDirection.Rotation();
-	OutViewInitOptions.ViewOrigin = OutViewInitOptions.ViewLocation;
-
-	if (ADisplayClusterRootActor* RootActor = UpdateRootActor())
-	{
-		OutViewInitOptions.ViewRotation = FRotator(RootActor->GetActorRotation().Quaternion() * OutViewInitOptions.ViewRotation.Quaternion());
-	}
-
-	OutViewInitOptions.SetViewRectangle(FIntRect(0, 0, FNormalMap::NormalMapSize, FNormalMap::NormalMapSize));
-
-	AWorldSettings* WorldSettings = nullptr;
-	ADisplayClusterRootActor* RootActor = UpdateRootActor();
-
-	if (RootActor && RootActor->GetWorld())
-	{
-		WorldSettings = RootActor->GetWorld()->GetWorldSettings();
-	}
-
-	if (WorldSettings != nullptr)
-	{
-		OutViewInitOptions.WorldToMetersScale = WorldSettings->WorldToMeters;
-	}
-
-	// Rotate view 90 degrees
-	OutViewInitOptions.ViewRotationMatrix = FInverseRotationMatrix(OutViewInitOptions.ViewRotation) * FMatrix(
-		FPlane(0, 0, 1, 0),
-		FPlane(1, 0, 0, 0),
-		FPlane(0, 1, 0, 0),
-		FPlane(0, 0, 0, 1));
-
-	const float MinZ = GNearClippingPlane;
-	const float MaxZ = FMath::Max(RootActorBoundingRadius, MinZ);
-
-	// Avoid zero ViewFOV's which cause divide by zero's in projection matrix
-	const float MatrixFOV = FMath::Max(0.001f, FNormalMap::NormalMapFOV) * (float)PI / 360.0f;
-
-	const float XAxisMultiplier = 1.0f;
-	const float YAxisMultiplier = 1.0f;
-
-	if ((bool)ERHIZBuffer::IsInverted)
-	{
-		OutViewInitOptions.ProjectionMatrix = FReversedZPerspectiveMatrix(
-			MatrixFOV,
-			MatrixFOV,
-			XAxisMultiplier,
-			YAxisMultiplier,
-			MinZ,
-			MaxZ);
-	}
-	else
-	{
-		OutViewInitOptions.ProjectionMatrix = FPerspectiveMatrix(
-			MatrixFOV,
-			MatrixFOV,
-			XAxisMultiplier,
-			YAxisMultiplier,
-			MinZ,
-			MaxZ);
-	}
-
-	OutViewInitOptions.FOV = FNormalMap::NormalMapFOV;
-
-#if WITH_EDITOR
-	if (ViewportClient.IsValid())
-	{
-		FEditorViewportClient& ViewportClientRef = *ViewportClient.Pin().Get();
-
-		OutViewInitOptions.SceneViewStateInterface = ViewportClientRef.ViewState.GetReference();
-		OutViewInitOptions.ViewElementDrawer = &ViewportClientRef;
-		OutViewInitOptions.OverrideFarClippingPlaneDistance = ViewportClientRef.GetFarClipPlaneOverride();
-		OutViewInitOptions.EditorViewBitflag = (uint64)1 << ViewportClientRef.ViewIndex; // send the bit for this view - each actor will check it's visibility bits against this
-	}
-#endif
-}
-
 void FDisplayClusterLightCardEditorHelper::SetActorCoordinates(const FDisplayClusterWeakStageActorPtr& Actor, const FSphericalCoordinates& SphericalCoords) const
 {
 	Actor->SetDistanceFromCenter(SphericalCoords.Radius);
@@ -1700,15 +1421,7 @@ bool FDisplayClusterLightCardEditorHelper::UpdateNormalMaps()
 	// Update this so we render from the latest projection point
 	UpdateProjectionOriginComponent();
 
-	if (!RenderNormalMap(NorthNormalMap, true))
-	{
-		return false;
-	}
-
-	if (!RenderNormalMap(SouthNormalMap, false))
-	{
-		return false;
-	}
+	RootActor->GetStageGeometryComponent()->Invalidate(true);
 
 	bNormalMapInvalid = false;
 	UpdateNormalMapMesh();
@@ -1763,70 +1476,10 @@ void FDisplayClusterLightCardEditorHelper::UpdateNormalMapMesh()
 			}
 		}
 
-		NorthNormalMap.MorphProceduralMesh(NormalMapMeshComponent.Get());
-		SouthNormalMap.MorphProceduralMesh(NormalMapMeshComponent.Get());
+		CachedRootActor->GetStageGeometryComponent()->MorphProceduralMesh(NormalMapMeshComponent.Get());
 	}
 }
 
-bool FDisplayClusterLightCardEditorHelper::RenderNormalMap(FNormalMap& NormalMap, bool bIsNorthMap)
-{
-	ADisplayClusterRootActor* RootActor = UpdateRootActor();
-	if (!RootActor)
-	{
-		return false;
-	}
-
-	UWorld* World = RootActor->GetWorld();
-	if (!World)
-	{
-		return false;
-	}
-
-	// Rotate the normal map's direction to account for the actor's world rotation
-	FRotator InverseRootRotation = RootActor->GetTransform().Rotator().GetInverse();
-	FVector NormalMapDirection = InverseRootRotation.RotateVector(bIsNorthMap ? FVector::UpVector : FVector::DownVector);
-
-	// Only render primitive components from the stage actor for the normal map
-	FDisplayClusterMeshProjectionPrimitiveFilter PrimitiveFilter;
-	PrimitiveFilter.ShouldRenderPrimitiveDelegate = FDisplayClusterMeshProjectionPrimitiveFilter::FPrimitiveFilter::CreateWeakLambda(
-		RootActor, [RootActor](const UPrimitiveComponent* PrimitiveComponent)
-		{
-			return PrimitiveComponent->GetOwner() == RootActor;
-		});
-
-	FDisplayClusterMeshProjectionRenderSettings RenderSettings;
-	RenderSettings.RenderType = EDisplayClusterMeshProjectionOutput::Normals;
-	RenderSettings.ProjectionType = EDisplayClusterMeshProjectionType::Azimuthal;
-	RenderSettings.PrimitiveFilter = PrimitiveFilter;
-	RenderSettings.NormalCorrectionMatrix = FMatrix44f(FRotationMatrix::Make(InverseRootRotation));
-
-#if WITH_EDITOR
-	if (ViewportClient.IsValid())
-	{
-		const FEditorViewportClient& ViewportClientRef = *ViewportClient.Pin().Get();
-		RenderSettings.EngineShowFlags = ViewportClientRef.EngineShowFlags;
-	}
-#endif
-
-	FSceneViewInitOptions ViewInitOptions;
-	GetNormalMapSceneViewInitOptions(NormalMapDirection, RenderSettings.ViewInitOptions);
-
-	NormalMap.Init(RenderSettings.ViewInitOptions);
-
-	NormalMap.Canvas = MakeShared<FCanvas>(&NormalMap, nullptr, World, World->Scene->GetFeatureLevel(), FCanvas::CDM_DeferDrawing, 1.0f);
-	NormalMap.Canvas->Clear(FLinearColor::Black);
-
-	IDisplayClusterScenePreview::Get().Render(RendererId, RenderSettings, *NormalMap.Canvas);
-	NormalMap.Canvas->Flush_GameThread();
-
-	NormalMap.ReadFloat16Pixels(NormalMap.GetCachedNormalData());
-	NormalMap.Release();
-
-	FlushRenderingCommands();
-
-	return true;
-}
-
 void FDisplayClusterLightCardEditorHelper::OnWorldCleanup(UWorld* World, bool bSessionEnded, bool bCleanupResources)
 {
 	if (!NormalMeshScene.IsValid() || !CachedRootActor.IsValid() || World != CachedRootActor->GetWorld())
diff --git a/Engine/Plugins/Runtime/nDisplay/Source/DisplayClusterScenePreview/Public/DisplayClusterLightCardEditorHelper.h b/Engine/Plugins/Runtime/nDisplay/Source/DisplayClusterScenePreview/Public/DisplayClusterLightCardEditorHelper.h
--- a/Engine/Plugins/Runtime/nDisplay/Source/DisplayClusterScenePreview/Public/DisplayClusterLightCardEditorHelper.h
+++ b/Engine/Plugins/Runtime/nDisplay/Source/DisplayClusterScenePreview/Public/DisplayClusterLightCardEditorHelper.h
@@ -67,69 +67,7 @@ public:
 		double Inclination = 0; // radians    0 to PI (when conforming)
 		double Azimuth = 0;     // radians  -PI to PI (when conforming)
 	};
-
-private:
-	/** Custom render target that stores the normal data for the stage */
-	class FNormalMap : public FRenderTarget
-	{
-	public:
-		/** The size of the normal map */
-		static const int32 NormalMapSize;
-
-		/**
-		 * The field of view to render the normal map with. Using the azimuthal projection,
-		 * this is set so that the entire 360 degree scene is rendered
-		 */
-		static const float NormalMapFOV;
-
-		/** Initializes the normal map render target using the specified scene view options */
-		void Init(const FSceneViewInitOptions& InSceneViewInitOptions);
-
-		/** Releases the normal map render target's resources */
-		void Release();
-
-		/** Gets the size of the render target */
-		virtual FIntPoint GetSizeXY() const override { return FIntPoint(SizeX, SizeY); }
-
-		/** Gets a reference to the normal map data array, which stores the normal vector in the RGB components (color = 0.5 * Normal + 0.5) and the depth in the A component */
-		TArray<FFloat16Color>& GetCachedNormalData() { return CachedNormalData; }
-
-		/** Gets the normal vector and distance at the specified world location. The normal and distance are bilinearly interpolated from the nearest pixels in the normal map */
-		bool GetNormalAndDistanceAtPosition(FVector Position, FVector& OutNormal, float& OutDistance) const;
-
-		/** Morphs the vertices of the specified prodedural mesh to match the normal map */
-		void MorphProceduralMesh(UProceduralMeshComponent* InProceduralMeshComponent) const;
-
-		/** Generates a texture object that can be used to visualize the normal map */
-		UTexture2D* GenerateNormalMapTexture(const FString& TextureName);
-
-		/** Gets the normal map visualization texture, or null if it hasn't been generated */
-		UTexture2D* GetNormalMapTexture() const { return NormalMapTexture.IsValid() ? NormalMapTexture.Get() : nullptr; }
-
-	public:
-		/** True if this is currently waiting for a render to complete. */
-		bool bIsPendingRender = true;
-
-		/** The canvas currently being used to render this map. */
-		TSharedPtr<FCanvas> Canvas;
-
-	private:
-		/** The view matrices used when the normal map was last rendered */
-		FViewMatrices ViewMatrices;
-
-		/** The cached normal map data from the last normal map render */
-		TArray<FFloat16Color> CachedNormalData;
-
-		/** A texture that contains the normal map, for visualization purposes */
-		TWeakObjectPtr<UTexture2D> NormalMapTexture;
-
-		/** The width of the normal map. */
-		uint32 SizeX = 0;
-
-		/** The height of the normal map. */
-		uint32 SizeY = 0;
-	};
-
+	
 public:
 	/** Create a projection helper that automatically creates its own preview renderer for normal map generation. */
 	DISPLAYCLUSTERSCENEPREVIEW_API FDisplayClusterLightCardEditorHelper();
@@ -383,9 +321,6 @@ private:
 	void InternalDragActors(const TArray<FDisplayClusterWeakStageActorPtr>& Actors, const FIntPoint& PixelPos, const FSceneView& View,
 		ECoordinateSystem CoordinateSystem, const FVector& DragWidgetOffset, EAxisList::Type DragAxis, FDisplayClusterWeakStageActorPtr PrimaryActor);
 
-	/** Gets the scene view init options to use when rendering the normal map cache */
-	void GetNormalMapSceneViewInitOptions(const FVector& ViewDirection, FSceneViewInitOptions& OutViewInitOptions);
-
 	/** Sets the actor position to the given spherical coordinates */
 	void SetActorCoordinates(const FDisplayClusterWeakStageActorPtr& Actor, const FSphericalCoordinates& SphericalCoords) const;
 
@@ -407,9 +342,6 @@ private:
 	/** Update the normal map mesh. Only call this once both maps are valid. */
 	void UpdateNormalMapMesh();
 
-	/** Renders the viewport's normal map and stores the texture data to be used later */
-	bool RenderNormalMap(FNormalMap& NormalMap, bool bIsNorthMap);
-
 	/** Called when a world beings cleanup. We destroy our preview scene if the world containing the root actor is cleaned up. */
 	void OnWorldCleanup(UWorld* World, bool bSessionEnded, bool bCleanupResources);
 
@@ -447,12 +379,6 @@ private:
 	/** The component of the root actor that is acting as the projection origin. Can be either the root component (stage origin) or a view origin component */
 	TWeakObjectPtr<USceneComponent> ProjectionOriginComponent;
 
-	/** The render target used to render a map of the screens' normals for the northern hemisphere of the view */
-	FNormalMap NorthNormalMap;
-
-	/** The render target used to render a map of the screens' normals for the southern hemisphere of the view */
-	FNormalMap SouthNormalMap;
-
 	/** The projection mode of the view this is helping. */
 	EDisplayClusterMeshProjectionType ProjectionMode;
 
diff --git a/Engine/Plugins/Runtime/nDisplayModularFeatures/Source/DisplayClusterLightCardExtender/Private/StageActor/IDisplayClusterStageActor.cpp b/Engine/Plugins/Runtime/nDisplayModularFeatures/Source/DisplayClusterLightCardExtender/Private/StageActor/IDisplayClusterStageActor.cpp
--- a/Engine/Plugins/Runtime/nDisplayModularFeatures/Source/DisplayClusterLightCardExtender/Private/StageActor/IDisplayClusterStageActor.cpp
+++ b/Engine/Plugins/Runtime/nDisplayModularFeatures/Source/DisplayClusterLightCardExtender/Private/StageActor/IDisplayClusterStageActor.cpp
@@ -137,6 +137,11 @@ void IDisplayClusterStageActor::UpdatePositionalParamsFromTransform()
 		Actor->GetTransform(), GetOrigin(), GetRadialOffset());
 
 		SetPositionalParams(PositionalParams);
+
+		if (IsAlwaysFlushToWall())
+		{
+			// TODO: Implement this once there is a way to determine which root actor owns this stage actor. 
+		}
 	}
 }
 
diff --git a/Engine/Plugins/Runtime/nDisplayModularFeatures/Source/DisplayClusterLightCardExtender/Public/StageActor/IDisplayClusterStageActor.h b/Engine/Plugins/Runtime/nDisplayModularFeatures/Source/DisplayClusterLightCardExtender/Public/StageActor/IDisplayClusterStageActor.h
--- a/Engine/Plugins/Runtime/nDisplayModularFeatures/Source/DisplayClusterLightCardExtender/Public/StageActor/IDisplayClusterStageActor.h
+++ b/Engine/Plugins/Runtime/nDisplayModularFeatures/Source/DisplayClusterLightCardExtender/Public/StageActor/IDisplayClusterStageActor.h
@@ -56,6 +56,10 @@ public:
 	virtual void SetScale(const FVector2D& InScale) = 0;
 	virtual FVector2D GetScale() const = 0;
 
+	virtual void SetAlwaysFlushToWall(bool bInAlwaysFlushToWall) { }
+	virtual bool IsAlwaysFlushToWall() const { return false; }
+
+
 	/** Return property names defined for this stage actor for use in property notifies */
 	virtual const TSet<FName>& GetPositionalPropertyNames() const;
 
