nDisplay: support enhanced input + input replication

diff --git a/Engine/Plugins/Runtime/nDisplay/Source/DisplayCluster/DisplayCluster.Build.cs b/Engine/Plugins/Runtime/nDisplay/Source/DisplayCluster/DisplayCluster.Build.cs
--- a/Engine/Plugins/Runtime/nDisplay/Source/DisplayCluster/DisplayCluster.Build.cs
+++ b/Engine/Plugins/Runtime/nDisplay/Source/DisplayCluster/DisplayCluster.Build.cs
@@ -22,7 +22,8 @@ public class DisplayCluster : ModuleRules
 				"DisplayClusterConfiguration",
 				"DisplayClusterLightCardEditorShaders",
 				"DisplayClusterLightCardExtender",
-				"Engine"
+				"Engine",
+				"EnhancedInput",
 			});
 
 		PrivateDependencyModuleNames.AddRange(
diff --git a/Engine/Plugins/Runtime/nDisplay/Source/DisplayCluster/Private/Game/EngineClasses/Basics/DisplayClusterPlayerInput.cpp b/Engine/Plugins/Runtime/nDisplay/Source/DisplayCluster/Private/Game/EngineClasses/Basics/DisplayClusterPlayerInput.cpp
--- a/Engine/Plugins/Runtime/nDisplay/Source/DisplayCluster/Private/Game/EngineClasses/Basics/DisplayClusterPlayerInput.cpp
+++ b/Engine/Plugins/Runtime/nDisplay/Source/DisplayCluster/Private/Game/EngineClasses/Basics/DisplayClusterPlayerInput.cpp
@@ -1,9 +1,15 @@
 // Copyright Epic Games, Inc. All Rights Reserved.
 
 #include "DisplayClusterPlayerInput.h"
+
+#include "DisplayClusterConfigurationStrings.h"
+#include "DisplayClusterConfigurationTypes.h"
+#include "DisplayClusterEnums.h"
+
 #include "Components/InputComponent.h"
 
 #include "Cluster/IPDisplayClusterClusterManager.h"
+#include "Config/IPDisplayClusterConfigManager.h"
 
 #include "Misc/DisplayClusterGlobals.h"
 #include "Misc/DisplayClusterHelpers.h"
@@ -15,29 +21,35 @@
 
 UDisplayClusterPlayerInput::UDisplayClusterPlayerInput()
 	: Super()
+	, bReplicatePrimary(false) // no replication by default
 {
+	// Replication makes sense in 'cluster' mode only
+	if (GDisplayCluster && GDisplayCluster->GetOperationMode() == EDisplayClusterOperationMode::Cluster)
+	{
+		if (const IPDisplayClusterConfigManager* const ConfigMgr = GDisplayCluster->GetPrivateConfigMgr())
+		{
+			if (const UDisplayClusterConfigurationData* ConfigData = ConfigMgr->GetConfig())
+			{
+				const FString& SyncPolicyType = ConfigData->Cluster->Sync.InputSyncPolicy.Type;
+				UE_LOG(LogDisplayClusterGame, Log, TEXT("Native input sync policy: %s"), *SyncPolicyType);
+
+				// Optionally activate native input synchronization
+				if (SyncPolicyType.Equals(DisplayClusterConfigurationStrings::config::cluster::input_sync::InputSyncPolicyReplicatePrimary, ESearchCase::IgnoreCase))
+				{
+					// Ok, replication is needed
+					bReplicatePrimary = true;
+				}
+			}
+		}
+	}
 }
 
 void UDisplayClusterPlayerInput::ProcessInputStack(const TArray<UInputComponent*>& InputComponentStack, const float DeltaTime, const bool bGamePaused)
 {
-	UE_LOG(LogDisplayClusterGame, Verbose, TEXT("Processing input stack..."));
-
-	if (IPDisplayClusterClusterManager* const ClusterMgr = GDisplayCluster->GetPrivateClusterMgr())
+	// Input replication
+	if (bReplicatePrimary)
 	{
-		TMap<FString, FString> KeyStates;
-
-		if (ClusterMgr->IsPrimary())
-		{
-			// Export key states data for the stack
-			SerializeKeyStateMap(KeyStates);
-			ClusterMgr->ImportNativeInputData(MoveTemp(KeyStates));
-		}
-		else
-		{
-			// Import key states data to the stack
-			ClusterMgr->ExportNativeInputData(KeyStates);
-			DeserializeKeyStateMap(KeyStates);
-		}
+		ProcessPolicy_ReplicatePrimary(InputComponentStack, DeltaTime, bGamePaused);
 	}
 
 	Super::ProcessInputStack(InputComponentStack, DeltaTime, bGamePaused);
@@ -130,3 +142,26 @@ bool UDisplayClusterPlayerInput::DeserializeKeyStateMap(const TMap<FString, FStr
 
 	return true;
 }
+
+void UDisplayClusterPlayerInput::ProcessPolicy_ReplicatePrimary(const TArray<UInputComponent*>& InputComponentStack, const float DeltaTime, const bool bGamePaused)
+{
+	UE_LOG(LogDisplayClusterGame, Verbose, TEXT("Processing input stack..."));
+
+	if (IPDisplayClusterClusterManager* const ClusterMgr = GDisplayCluster->GetPrivateClusterMgr())
+	{
+		TMap<FString, FString> KeyStates;
+
+		if (ClusterMgr->IsPrimary())
+		{
+			// Export key states data for the stack
+			SerializeKeyStateMap(KeyStates);
+			ClusterMgr->ImportNativeInputData(MoveTemp(KeyStates));
+		}
+		else
+		{
+			// Import key states data to the stack
+			ClusterMgr->ExportNativeInputData(KeyStates);
+			DeserializeKeyStateMap(KeyStates);
+		}
+	}
+}
diff --git a/Engine/Plugins/Runtime/nDisplay/Source/DisplayCluster/Private/Game/EngineClasses/Scene/DisplayClusterRootActor.cpp b/Engine/Plugins/Runtime/nDisplay/Source/DisplayCluster/Private/Game/EngineClasses/Scene/DisplayClusterRootActor.cpp
--- a/Engine/Plugins/Runtime/nDisplay/Source/DisplayCluster/Private/Game/EngineClasses/Scene/DisplayClusterRootActor.cpp
+++ b/Engine/Plugins/Runtime/nDisplay/Source/DisplayCluster/Private/Game/EngineClasses/Scene/DisplayClusterRootActor.cpp
@@ -22,7 +22,6 @@
 #include "IDisplayClusterConfiguration.h"
 #include "DisplayClusterConfigurationTypes.h"
 
-#include "DisplayClusterPlayerInput.h"
 #include "Blueprints/DisplayClusterBlueprint.h"
 #include "Blueprints/DisplayClusterBlueprintGeneratedClass.h"
 
@@ -771,32 +770,6 @@ void ADisplayClusterRootActor::BeginPlay()
 
 	// Store current operation mode
 	OperationMode = GDisplayCluster->GetOperationMode();
-	if (OperationMode == EDisplayClusterOperationMode::Cluster)
-	{
-		FString SyncPolicyType;
-
-		// Read native input synchronization settings
-		IPDisplayClusterConfigManager* const ConfigMgr = GDisplayCluster->GetPrivateConfigMgr();
-		if (ConfigMgr)
-		{
-			UDisplayClusterConfigurationData* ConfigData = ConfigMgr->GetConfig();
-			if (ConfigData)
-			{
-				SyncPolicyType = ConfigData->Cluster->Sync.InputSyncPolicy.Type;
-				UE_LOG(LogDisplayClusterGame, Log, TEXT("Native input sync policy: %s"), *SyncPolicyType);
-			}
-		}
-
-		// Optionally activate native input synchronization
-		if (SyncPolicyType.Equals(DisplayClusterConfigurationStrings::config::cluster::input_sync::InputSyncPolicyReplicatePrimary, ESearchCase::IgnoreCase))
-		{
-			APlayerController* const PlayerController = GetWorld()->GetFirstPlayerController();
-			if (PlayerController)
-			{
-				PlayerController->PlayerInput = NewObject<UDisplayClusterPlayerInput>(PlayerController);
-			}
-		}
-	}
 
 	InitializeRootActor();
 }
diff --git a/Engine/Plugins/Runtime/nDisplay/Source/DisplayCluster/Public/DisplayClusterPlayerInput.h b/Engine/Plugins/Runtime/nDisplay/Source/DisplayCluster/Public/DisplayClusterPlayerInput.h
--- a/Engine/Plugins/Runtime/nDisplay/Source/DisplayCluster/Public/DisplayClusterPlayerInput.h
+++ b/Engine/Plugins/Runtime/nDisplay/Source/DisplayCluster/Public/DisplayClusterPlayerInput.h
@@ -2,15 +2,19 @@
 
 #pragma once
 
-#include "GameFramework/PlayerInput.h"
+#include "EnhancedPlayerInput.h"
 #include "DisplayClusterPlayerInput.generated.h"
 
 
 /**
- * Object within PlayerController that processes player input.
+ * An object within PlayerController that processes player input.
+ *
+ * This one is nDisplay specific implementation. It's responsible for replication
+ * of UE native input within a cluster to support simulation determinism. Various
+ * input sync policies might be implemented here.
  */
 UCLASS(Within=PlayerController, config=Input, transient)
-class DISPLAYCLUSTER_API UDisplayClusterPlayerInput : public UPlayerInput
+class DISPLAYCLUSTER_API UDisplayClusterPlayerInput : public UEnhancedPlayerInput
 {
 	GENERATED_BODY()
 
@@ -21,7 +25,19 @@ public:
 	/** Process the frame's input events given the current input component stack. */
 	virtual void ProcessInputStack(const TArray<UInputComponent*>& InputComponentStack, const float DeltaTime, const bool bGamePaused) override;
 
-private:
+protected:
+	// Performs sync logic where all input data from the primary node is replicated to all other nodes
+	virtual void ProcessPolicy_ReplicatePrimary(const TArray<UInputComponent*>& InputComponentStack, const float DeltaTime, const bool bGamePaused);
+
+protected:
+	// Aux function that serializes whole input data into a map
 	bool SerializeKeyStateMap(TMap<FString, FString>& OutKeyStateMap);
+	// Aux function that deserializes whole input data from a map
 	bool DeserializeKeyStateMap(const TMap<FString, FString>& InKeyStateMap);
+
+private:
+	// Currently we have this bool only wich provides 2 cases: sync or not sync. It covers all
+	// functionality currently available . However in the future, this might be refactored in a way
+	// that we can have multiple input sync policies including custom ones.
+	bool bReplicatePrimary;
 };
diff --git a/Engine/Plugins/Runtime/nDisplay/Source/DisplayClusterEditor/Private/Settings/DisplayClusterEditorSettings.cpp b/Engine/Plugins/Runtime/nDisplay/Source/DisplayClusterEditor/Private/Settings/DisplayClusterEditorSettings.cpp
--- a/Engine/Plugins/Runtime/nDisplay/Source/DisplayClusterEditor/Private/Settings/DisplayClusterEditorSettings.cpp
+++ b/Engine/Plugins/Runtime/nDisplay/Source/DisplayClusterEditor/Private/Settings/DisplayClusterEditorSettings.cpp
@@ -23,23 +23,32 @@ void UDisplayClusterEditorSettings::PostEditChangeProperty(FPropertyChangedEvent
 	if (PropertyChangedEvent.Property != nullptr)
 	{
 		// Since nDisplay is Windows only, save configs that depends on nDisplay plugin assets in Windows specific config file
-		static const FString DefaultEnginePath = FString::Printf(TEXT("%sDefaultEngine.ini"), *FPaths::SourceConfigDir());
-		static const FString DefaultGamePath   = FString::Printf(TEXT("%sDefaultGame.ini"),   *FPaths::SourceConfigDir());
+		const FString DefaultEnginePath = FString::Printf(TEXT("%sDefaultEngine.ini"), *FPaths::SourceConfigDir());
+		const FString DefaultGamePath   = FString::Printf(TEXT("%sDefaultGame.ini"),   *FPaths::SourceConfigDir());
+		const FString DefaultInputPath  = FString::Printf(TEXT("%sDefaultInput.ini"),  *FPaths::SourceConfigDir());
 
 		FName PropertyName(PropertyChangedEvent.Property->GetFName());
 
 		if (PropertyName == GET_MEMBER_NAME_CHECKED(UDisplayClusterEditorSettings, bEnabled))
 		{
-			if (FPlatformFileManager::Get().GetPlatformFile().IsReadOnly(*DefaultEnginePath))
+			// Reset read-only flag from the config files
 			{
-				FPlatformFileManager::Get().GetPlatformFile().SetReadOnly(*DefaultEnginePath, false);
-			}
+				const FString* ConfigFiles[] = {
+					&DefaultEnginePath,
+					&DefaultGamePath,
+					&DefaultInputPath
+				};
 
-			if (FPlatformFileManager::Get().GetPlatformFile().IsReadOnly(*DefaultGamePath))
-			{
-				FPlatformFileManager::Get().GetPlatformFile().SetReadOnly(*DefaultGamePath, false);
+				for (const FString* ConfigFile : ConfigFiles)
+				{
+					if (FPlatformFileManager::Get().GetPlatformFile().IsReadOnly(**ConfigFile))
+					{
+						FPlatformFileManager::Get().GetPlatformFile().SetReadOnly(**ConfigFile, false);
+					}
+				}
 			}
 
+			// Process nDisplay 'enable' command
 			if (bEnabled)
 			{
 				// DefaultEngine.ini
@@ -49,7 +58,11 @@ void UDisplayClusterEditorSettings::PostEditChangeProperty(FPropertyChangedEvent
 
 				// DefaultGame.ini
 				GConfig->SetString(TEXT("/Script/EngineSettings.GeneralProjectSettings"), TEXT("bUseBorderlessWindow"), TEXT("True"), DefaultGamePath);
+
+				// DefaultInput.ini
+				GConfig->SetString(TEXT("/Script/Engine.InputSettings"), TEXT("DefaultPlayerInputClass"), TEXT("/Script/DisplayCluster.DisplayClusterPlayerInput"), DefaultInputPath);
 			}
+			// Process nDisplay 'disable' command
 			else
 			{
 				// DefaultEngine.ini
@@ -59,10 +72,15 @@ void UDisplayClusterEditorSettings::PostEditChangeProperty(FPropertyChangedEvent
 
 				// DefaultGame.ini
 				GConfig->SetString(TEXT("/Script/EngineSettings.GeneralProjectSettings"), TEXT("bUseBorderlessWindow"), TEXT("False"), DefaultGamePath);
+
+				// DefaultInput.ini
+				GConfig->SetString(TEXT("/Script/Engine.InputSettings"), TEXT("DefaultPlayerInputClass"), TEXT("/Script/EnhancedInput.EnhancedPlayerInput"), DefaultInputPath);
 			}
 
+			// Save changes to the files
 			GConfig->Flush(false, DefaultEnginePath);
 			GConfig->Flush(false, DefaultGamePath);
+			GConfig->Flush(false, DefaultInputPath);
 		}
 	}
 }
diff --git a/Engine/Plugins/Runtime/nDisplay/nDisplay.uplugin b/Engine/Plugins/Runtime/nDisplay/nDisplay.uplugin
--- a/Engine/Plugins/Runtime/nDisplay/nDisplay.uplugin
+++ b/Engine/Plugins/Runtime/nDisplay/nDisplay.uplugin
@@ -201,6 +201,10 @@
 			"Name": "ConcertSyncClient",
 			"Enabled": true
 		},
+		{
+			"Name": "EnhancedInput",
+			"Enabled": true
+		},
 		{
 			"Name": "MediaIOFramework",
 			"Enabled": true
diff --git a/Engine/Plugins/VirtualProduction/Switchboard/Source/Switchboard/switchboard/devices/ndisplay/plugin_ndisplay.py b/Engine/Plugins/VirtualProduction/Switchboard/Source/Switchboard/switchboard/devices/ndisplay/plugin_ndisplay.py
--- a/Engine/Plugins/VirtualProduction/Switchboard/Source/Switchboard/switchboard/devices/ndisplay/plugin_ndisplay.py
+++ b/Engine/Plugins/VirtualProduction/Switchboard/Source/Switchboard/switchboard/devices/ndisplay/plugin_ndisplay.py
@@ -687,6 +687,11 @@ class DevicenDisplay(DeviceUnreal):
             "[/Script/EngineSettings.GeneralProjectSettings]:"
             "bUseBorderlessWindow=True")
 
+        ini_input = (
+            "-ini:Input:"
+            "[/Script/Engine.InputSettings]:"
+            "DefaultPlayerInputClass=/Script/DisplayCluster.DisplayClusterPlayerInput")
+
         # VP roles
         vproles, missing_roles = self.get_vproles()
 
@@ -776,6 +781,7 @@ class DevicenDisplay(DeviceUnreal):
             f'Log={self.log_filename}',   # log file
             f'{ini_engine}',              # Engine ini injections
             f'{ini_game}',                # Game ini injections
+            f'{ini_input}',               # Input ini injections
             f'{unattended}',              # -unattended
             f'{no_screen_messages}',      # -NoScreenMessages
             f'{disable_ensures}',         # -handleensurepercent=0
diff --git a/Templates/TP_nDisplayBP/Config/DefaultInput.ini b/Templates/TP_nDisplayBP/Config/DefaultInput.ini
--- a/Templates/TP_nDisplayBP/Config/DefaultInput.ini
+++ b/Templates/TP_nDisplayBP/Config/DefaultInput.ini
@@ -53,6 +53,7 @@ DoubleClickTime=0.200000
 +AxisMappings=(AxisName="TurnRate",Scale=1.000000,Key=MouseX)
 +AxisMappings=(AxisName="LookUpRate",Scale=1.000000,Key=Down)
 +AxisMappings=(AxisName="LookUpRate",Scale=-1.000000,Key=MouseY)
+DefaultPlayerInputClass=/Script/DisplayCluster.DisplayClusterPlayerInput
 DefaultTouchInterface=/Engine/MobileResources/HUD/DefaultVirtualJoysticks.DefaultVirtualJoysticks
 ConsoleKey=None
 -ConsoleKeys=Tilde
