Update this text to match the declaration that uses the same key in the PropertyEditor module
#rnx

diff --git a/Engine/Plugins/Runtime/nDisplay/Source/DisplayClusterColorGrading/Private/DetailView/SDisplayClusterColorGradingDetailTreeRow.cpp b/Engine/Plugins/Runtime/nDisplay/Source/DisplayClusterColorGrading/Private/DetailView/SDisplayClusterColorGradingDetailTreeRow.cpp
--- a/Engine/Plugins/Runtime/nDisplay/Source/DisplayClusterColorGrading/Private/DetailView/SDisplayClusterColorGradingDetailTreeRow.cpp
+++ b/Engine/Plugins/Runtime/nDisplay/Source/DisplayClusterColorGrading/Private/DetailView/SDisplayClusterColorGradingDetailTreeRow.cpp
@@ -407,7 +407,7 @@ FReply SDisplayClusterColorGradingDetailTreeRow::OnMouseButtonUp(const FGeometry
 
 			MenuBuilder.AddMenuEntry(
 				NSLOCTEXT("PropertyView", "CopyPropertyDisplayName", "Copy Display Name"),
-				NSLOCTEXT("PropertyView", "CopyPropertyDisplayName_ToolTip", "Copy this property display name"),
+				NSLOCTEXT("PropertyView", "CopyPropertyDisplayName_ToolTip", "Copy the display name of this property to the system clipboard."),
 				FSlateIcon(FCoreStyle::Get().GetStyleSetName(), "GenericCommands.Copy"),
 				FExecuteAction::CreateSP(this, &SDisplayClusterColorGradingDetailTreeRow::CopyPropertyName));
 		}
