nDisplay: fixed mGPU optimiztion issues
 - removed ""Multi GPU Mode" config option
 - reworked mGPU logic
 - added new CVARs for mGPU
- fixed crash when deleting viewport in DCRA editor

diff --git a/Engine/Plugins/Runtime/nDisplay/Source/DisplayCluster/Private/Render/Viewport/Configuration/DisplayClusterViewportConfiguration.cpp b/Engine/Plugins/Runtime/nDisplay/Source/DisplayCluster/Private/Render/Viewport/Configuration/DisplayClusterViewportConfiguration.cpp
--- a/Engine/Plugins/Runtime/nDisplay/Source/DisplayCluster/Private/Render/Viewport/Configuration/DisplayClusterViewportConfiguration.cpp
+++ b/Engine/Plugins/Runtime/nDisplay/Source/DisplayCluster/Private/Render/Viewport/Configuration/DisplayClusterViewportConfiguration.cpp
@@ -22,18 +22,38 @@
 #include "Misc/DisplayClusterGlobals.h"
 #include "Misc/DisplayClusterLog.h"
 
-int32 GDisplayClusterOverrideMultiGPUMode = -1;
-static FAutoConsoleVariableRef CVarDisplayClusterOverrideMultiGPUMode(
-	TEXT("DC.OverrideMultiGPUMode"),
-	GDisplayClusterOverrideMultiGPUMode,
-	TEXT("Override Multi GPU Mode setting from component (-1 == no override, or EDisplayClusterConfigurationRenderMGPUMode enum)"),
+///////////////////////////////////////////////////////////////////
+int32 GDisplayClusterCrossGPUTransferEnable = 0;
+static FAutoConsoleVariableRef CDisplayClusterCrossGPUTransferEnable(
+	TEXT("nDisplay.render.CrossGPUTransfer.Enable"),
+	GDisplayClusterCrossGPUTransferEnable,
+	TEXT("Enable cross-GPU transfers using nDisplay implementation (0 - disable, default) \n")
+	TEXT("That replaces the default cross-GPU transfers using UE Core for the nDisplay viewports viewfamilies.\n"),
+	ECVF_RenderThreadSafe
+);
+
+int32 GDisplayClusterCrossGPUTransferLockSteps = 1;
+static FAutoConsoleVariableRef CDisplayClusterCrossGPUTransferLockSteps(
+	TEXT("nDisplay.render.CrossGPUTransfer.LockSteps"),
+	GDisplayClusterCrossGPUTransferLockSteps,
+	TEXT("The bLockSteps parameter is simply passed to the FTransferResourceParams structure. (0 - disable)\n")
+	TEXT("Whether the GPUs must handshake before and after the transfer. Required if the texture rect is being written to in several render passes.\n")
+	TEXT("Otherwise, minimal synchronization will be used.\n"),
+	ECVF_RenderThreadSafe
+);
+
+int32 GDisplayClusterCrossGPUTransferPullData = 1;
+static FAutoConsoleVariableRef CVarDisplayClusterCrossGPUTransferPullData(
+	TEXT("nDisplay.render.CrossGPUTransfer.PullData"),
+	GDisplayClusterCrossGPUTransferPullData,
+	TEXT("The bPullData parameter is simply passed to the FTransferResourceParams structure. (0 - disable)\n")
+	TEXT("Whether the data is read by the dest GPU, or written by the src GPU (not allowed if the texture is a backbuffer)\n"),
 	ECVF_RenderThreadSafe
 );
 
 ///////////////////////////////////////////////////////////////////
 // FDisplayClusterViewportConfiguration
 ///////////////////////////////////////////////////////////////////
-
 bool FDisplayClusterViewportConfiguration::SetRootActor(ADisplayClusterRootActor* InRootActorPtr)
 {
 	check(IsInGameThread());
@@ -257,26 +277,10 @@ void FDisplayClusterViewportConfiguration::ImplUpdateRenderFrameConfiguration(co
 		break;
 	}
 
-	// Performance: Allow change global MGPU settings
-	int32 ModeOverride = FMath::Min(GDisplayClusterOverrideMultiGPUMode,
-		(int32)EDisplayClusterConfigurationRenderMGPUMode::Optimized_DisabledLockSteps);
-
-	switch (ModeOverride >= 0 ? (EDisplayClusterConfigurationRenderMGPUMode)ModeOverride : InRenderFrameConfiguration.MultiGPUMode)
-	{
-	case EDisplayClusterConfigurationRenderMGPUMode::None:
-		RenderFrameSettings.MultiGPUMode = EDisplayClusterMultiGPUMode::None;
-		break;
-	case EDisplayClusterConfigurationRenderMGPUMode::Optimized_DisabledLockSteps:
-		RenderFrameSettings.MultiGPUMode = EDisplayClusterMultiGPUMode::Optimized_DisabledLockSteps;
-		break;
-	case EDisplayClusterConfigurationRenderMGPUMode::Optimized_EnabledLockSteps:
-		RenderFrameSettings.MultiGPUMode = EDisplayClusterMultiGPUMode::Optimized_EnabledLockSteps;
-		break;
-	case EDisplayClusterConfigurationRenderMGPUMode::Enabled:
-	default:
-		RenderFrameSettings.MultiGPUMode = EDisplayClusterMultiGPUMode::Enabled;
-		break;
-	};
+	// Performance: nDisplay has its own implementation of cross-GPU transfer.
+	RenderFrameSettings.CrossGPUTransfer.bEnable    = GDisplayClusterCrossGPUTransferEnable != 0;
+	RenderFrameSettings.CrossGPUTransfer.bLockSteps = GDisplayClusterCrossGPUTransferLockSteps != 0;
+	RenderFrameSettings.CrossGPUTransfer.bPullData  = GDisplayClusterCrossGPUTransferPullData != 0;
 
 	// Performance: Allow to use parent ViewFamily from parent viewport 
 	// (icvfx has child viewports: lightcard and chromakey with prj_view matrices copied from parent viewport. May sense to use same viewfamily?)
diff --git a/Engine/Plugins/Runtime/nDisplay/Source/DisplayCluster/Private/Render/Viewport/Configuration/DisplayClusterViewportConfigurationHelpers.cpp b/Engine/Plugins/Runtime/nDisplay/Source/DisplayCluster/Private/Render/Viewport/Configuration/DisplayClusterViewportConfigurationHelpers.cpp
--- a/Engine/Plugins/Runtime/nDisplay/Source/DisplayCluster/Private/Render/Viewport/Configuration/DisplayClusterViewportConfigurationHelpers.cpp
+++ b/Engine/Plugins/Runtime/nDisplay/Source/DisplayCluster/Private/Render/Viewport/Configuration/DisplayClusterViewportConfigurationHelpers.cpp
@@ -174,6 +174,8 @@ void FDisplayClusterViewportConfigurationHelpers::UpdateBaseViewportSetting(FDis
 		DstViewport.RenderSettings.CameraId = InConfigurationViewport.Camera;
 		DstViewport.RenderSettings.Rect = DstViewport.GetValidRect(InConfigurationViewport.Region.ToRect(), TEXT("Configuration Region"));
 
+		DstViewport.RenderSettings.bEnableCrossGPUTransfer = InConfigurationViewport.RenderSettings.bEnableCrossGPUTransfer;
+
 		DstViewport.RenderSettings.GPUIndex = InConfigurationViewport.GPUIndex;
 		DstViewport.RenderSettings.OverlapOrder = InConfigurationViewport.OverlapOrder;
 
diff --git a/Engine/Plugins/Runtime/nDisplay/Source/DisplayCluster/Private/Render/Viewport/Containers/DisplayClusterViewportProxyData.cpp b/Engine/Plugins/Runtime/nDisplay/Source/DisplayCluster/Private/Render/Viewport/Containers/DisplayClusterViewportProxyData.cpp
--- a/Engine/Plugins/Runtime/nDisplay/Source/DisplayCluster/Private/Render/Viewport/Containers/DisplayClusterViewportProxyData.cpp
+++ b/Engine/Plugins/Runtime/nDisplay/Source/DisplayCluster/Private/Render/Viewport/Containers/DisplayClusterViewportProxyData.cpp
@@ -58,7 +58,19 @@ void FDisplayClusterViewportProxyData::UpdateProxy_RenderThread() const
 	DstViewportProxy->PostRenderSettings.SetParameters(PostRenderSettings);
 
 	DstViewportProxy->ProjectionPolicy = ProjectionPolicy;
-	DstViewportProxy->Contexts         = Contexts;
+	
+	// The RenderThreadData for DstViewportProxy has been updated in DisplayClusterViewportManagerViewExtension on the rendering thread.
+	// Therefore, the RenderThreadData values from the game thread must be overridden by current data from the render thread.
+	{
+		const TArray<FDisplayClusterViewport_Context> CurrentContexts = DstViewportProxy->Contexts;
+		DstViewportProxy->Contexts = Contexts;
+
+		int32 ContextAmmount = FMath::Min(CurrentContexts.Num(), Contexts.Num());
+		for (int32 ContextIndex = 0; ContextIndex < ContextAmmount; ContextIndex++)
+		{
+			DstViewportProxy->Contexts[ContextIndex].RenderThreadData = CurrentContexts[ContextIndex].RenderThreadData;
+		}
+	}
 
 	// Update viewport proxy resources from container
 	DstViewportProxy->RenderTargets    = RenderTargets;
diff --git a/Engine/Plugins/Runtime/nDisplay/Source/DisplayCluster/Private/Render/Viewport/DisplayClusterViewport.cpp b/Engine/Plugins/Runtime/nDisplay/Source/DisplayCluster/Private/Render/Viewport/DisplayClusterViewport.cpp
--- a/Engine/Plugins/Runtime/nDisplay/Source/DisplayCluster/Private/Render/Viewport/DisplayClusterViewport.cpp
+++ b/Engine/Plugins/Runtime/nDisplay/Source/DisplayCluster/Private/Render/Viewport/DisplayClusterViewport.cpp
@@ -62,9 +62,10 @@ FDisplayClusterViewport::~FDisplayClusterViewport()
 
 	ViewportProxy.Reset();
 
+	// Handle projection policy EndScene event
 	HandleEndScene();
 
-	// Handle projection policy event
+	// Release Projection policy after HandleEndScene()
 	ProjectionPolicy.Reset();
 	UninitializedProjectionPolicy.Reset();
 
@@ -240,12 +241,17 @@ void FDisplayClusterViewport::SetupSceneView(uint32 ContextNum, class UWorld* Wo
 		return;
 	}
 
-	// Setup MGPU features:
 	if(Contexts[ContextNum].GPUIndex >= 0)
 	{
+		// Use custom GPUIndex for render
 		InOutView.bOverrideGPUMask = true;
-		InOutView.GPUMask = FRHIGPUMask::FromIndex(FMath::Min((uint32)Contexts[ContextNum].GPUIndex, GNumExplicitGPUsForRendering - 1));
-		InOutView.bAllowCrossGPUTransfer = (Contexts[ContextNum].bAllowGPUTransferOptimization == false);
+		InOutView.GPUMask = FRHIGPUMask::FromIndex((uint32)Contexts[ContextNum].GPUIndex);
+	}
+
+	if (Contexts[ContextNum].bOverrideCrossGPUTransfer || !RenderSettings.bEnableCrossGPUTransfer)
+	{
+		// Disable native cross-GPU transfers inside Renderer.
+		InOutView.bAllowCrossGPUTransfer = false;
 	}
 
 	// Disable raytracing for lightcard and chromakey
@@ -534,22 +540,22 @@ bool FDisplayClusterViewport::UpdateFrameContexts(const uint32 InStereoViewIndex
 
 		FDisplayClusterViewport_Context Context(ContextIt, StereoscopicPass, StereoViewIndex);
 
-		int32 ContextGPUIndex = (ContextIt > 0 && RenderSettings.StereoGPUIndex >= 0) ? RenderSettings.StereoGPUIndex : RenderSettings.GPUIndex;
-		Context.GPUIndex = ContextGPUIndex;
+		Context.GPUIndex = INDEX_NONE;
 
-
-		if (InFrameSettings.bIsRenderingInEditor)
+		// The nDisplay can use its own cross-GPU transfer
+		if (InFrameSettings.CrossGPUTransfer.bEnable)
 		{
-			// Disable MultiGPU feature
-			Context.GPUIndex = INDEX_NONE;
+			Context.bOverrideCrossGPUTransfer = true;
+		}
 
-			if (InFrameSettings.bAllowMultiGPURenderingInEditor)
+		const int32 MaxExplicitGPUIndex = GNumExplicitGPUsForRendering - 1;
+		if (MaxExplicitGPUIndex > 0)
+		{
+			if (InFrameSettings.bIsRenderingInEditor)
 			{
 				// Experimental: allow mGPU for preview rendering:
-				if (GNumExplicitGPUsForRendering > 1 && bDisableRender == false)
+				if (InFrameSettings.bAllowMultiGPURenderingInEditor && !bDisableRender)
 				{
-					int32 MaxExplicitGPUIndex = GNumExplicitGPUsForRendering - 1;
-
 					int32 MinGPUIndex = FMath::Min(InFrameSettings.PreviewMinGPUIndex, MaxExplicitGPUIndex);
 					int32 MaxGPUIndex = FMath::Min(InFrameSettings.PreviewMaxGPUIndex, MaxExplicitGPUIndex);
 
@@ -560,35 +566,13 @@ bool FDisplayClusterViewport::UpdateFrameContexts(const uint32 InStereoViewIndex
 					{
 						PreviewGPUIndex = MinGPUIndex;
 					}
-
-					Context.bAllowGPUTransferOptimization = true;
-					Context.bEnabledGPUTransferLockSteps = false;
 				}
 			}
-		}
-		else
-		{
-			// Control mGPU:
-			switch (InFrameSettings.MultiGPUMode)
+			else
 			{
-			case EDisplayClusterMultiGPUMode::None:
-				Context.bAllowGPUTransferOptimization = false;
-				Context.GPUIndex = INDEX_NONE;
-				break;
-
-			case EDisplayClusterMultiGPUMode::Optimized_EnabledLockSteps:
-				Context.bAllowGPUTransferOptimization = true;
-				Context.bEnabledGPUTransferLockSteps = true;
-				break;
-
-			case EDisplayClusterMultiGPUMode::Optimized_DisabledLockSteps:
-				Context.bAllowGPUTransferOptimization = true;
-				Context.bEnabledGPUTransferLockSteps = false;
-				break;
-
-			default:
-				Context.bAllowGPUTransferOptimization = false;
-				break;
+				// Set custom GPU index for this view
+				const int32 CustomMultiGPUIndex = (ContextIt > 0 && RenderSettings.StereoGPUIndex >= 0) ? RenderSettings.StereoGPUIndex : RenderSettings.GPUIndex;
+				Context.GPUIndex = FMath::Min(CustomMultiGPUIndex, MaxExplicitGPUIndex);
 			}
 		}
 
diff --git a/Engine/Plugins/Runtime/nDisplay/Source/DisplayCluster/Private/Render/Viewport/DisplayClusterViewportManager.cpp b/Engine/Plugins/Runtime/nDisplay/Source/DisplayCluster/Private/Render/Viewport/DisplayClusterViewportManager.cpp
--- a/Engine/Plugins/Runtime/nDisplay/Source/DisplayCluster/Private/Render/Viewport/DisplayClusterViewportManager.cpp
+++ b/Engine/Plugins/Runtime/nDisplay/Source/DisplayCluster/Private/Render/Viewport/DisplayClusterViewportManager.cpp
@@ -80,6 +80,9 @@ FDisplayClusterViewportManager::FDisplayClusterViewportManager()
 	PostProcessManager  = MakeShared<FDisplayClusterViewportPostProcessManager, ESPMode::ThreadSafe>(*this);
 	LightCardManager = MakeShared<FDisplayClusterViewportLightCardManager, ESPMode::ThreadSafe>(*this);
 
+	// Create DC ViewExtension to handle special features
+	ViewportManagerViewExtension = FSceneViewExtensions::NewExtension<FDisplayClusterViewportManagerViewExtension>(this);
+
 	// initialize proxy
 	ViewportManagerProxy->Initialize(*this);
 
@@ -102,6 +105,8 @@ FDisplayClusterViewportManager::~FDisplayClusterViewportManager()
 		LightCardManager.Reset();
 	}
 
+	ViewportManagerViewExtension.Reset();
+
 	Configuration.Reset();
 
 	if (ViewportManagerProxy.IsValid())
@@ -454,12 +459,6 @@ bool FDisplayClusterViewportManager::BeginNewFrame(FViewport* InViewport, UWorld
 {
 	check(IsInGameThread());
 
-	if (!ViewportManagerViewExtension.IsValid())
-	{
-		// Create DC ViewExtension to handle special features
-		ViewportManagerViewExtension = FSceneViewExtensions::NewExtension<FDisplayClusterViewportManagerViewExtension>(this);
-	}
-
 	OutRenderFrame.ViewportManager = this;
 
 	// Handle world runtime update
diff --git a/Engine/Plugins/Runtime/nDisplay/Source/DisplayCluster/Private/Render/Viewport/DisplayClusterViewportManagerProxy.cpp b/Engine/Plugins/Runtime/nDisplay/Source/DisplayCluster/Private/Render/Viewport/DisplayClusterViewportManagerProxy.cpp
--- a/Engine/Plugins/Runtime/nDisplay/Source/DisplayCluster/Private/Render/Viewport/DisplayClusterViewportManagerProxy.cpp
+++ b/Engine/Plugins/Runtime/nDisplay/Source/DisplayCluster/Private/Render/Viewport/DisplayClusterViewportManagerProxy.cpp
@@ -35,13 +35,6 @@ static TAutoConsoleVariable<int32> CVarWarpBlendEnabled(
 	ECVF_RenderThreadSafe
 );
 
-static TAutoConsoleVariable<int32> CVarCrossGPUTransfersEnabled(
-	TEXT("nDisplay.render.CrossGPUTransfers"),
-	1,
-	TEXT("(0 = disabled)\n"),
-	ECVF_RenderThreadSafe
-);
-
 // Enable/Disable ClearTexture for Frame RTT
 static TAutoConsoleVariable<int32> CVarClearFrameRTTEnabled(
 	TEXT("nDisplay.render.ClearFrameRTTEnabled"),
@@ -186,7 +179,7 @@ void FDisplayClusterViewportManagerProxy::ImplUpdateRenderFrameSettings(const FD
 	FDisplayClusterRenderFrameSettings* Settings = new FDisplayClusterRenderFrameSettings(InRenderFrameSettings);
 
 	// Send frame settings to renderthread
-	ENQUEUE_RENDER_COMMAND(DeleteDisplayClusterViewportProxy)(
+	ENQUEUE_RENDER_COMMAND(DisplayClusterUpdateRenderFrameSettings)(
 		[ViewportManagerProxy = SharedThis(this), Settings](FRHICommandListImmediate& RHICmdList)
 	{
 		ViewportManagerProxy->RenderFrameSettings = *Settings;
@@ -208,7 +201,7 @@ void FDisplayClusterViewportManagerProxy::ImplUpdateViewports(const TArray<FDisp
 	}
 
 	// Send viewports settings to renderthread
-	ENQUEUE_RENDER_COMMAND(DeleteDisplayClusterViewportProxy)(
+	ENQUEUE_RENDER_COMMAND(DisplayClusterUpdateViewports)(
 		[ProxiesData = std::move(ViewportProxiesData)](FRHICommandListImmediate& RHICmdList)
 	{
 		for (FDisplayClusterViewportProxyData* It : ProxiesData)
@@ -224,7 +217,7 @@ DECLARE_GPU_STAT_NAMED(nDisplay_ViewportManager_RenderFrame, TEXT("nDisplay View
 
 void FDisplayClusterViewportManagerProxy::ImplRenderFrame(FViewport* InViewport)
 {
-	ENQUEUE_RENDER_COMMAND(DeleteDisplayClusterViewportProxy)(
+	ENQUEUE_RENDER_COMMAND(DisplayClusterRenderFrame_Setup)(
 		[InViewportManagerProxy = SharedThis(this), InViewport](FRHICommandListImmediate& RHICmdList)
 	{
 		SCOPED_GPU_STAT(RHICmdList, nDisplay_ViewportManager_RenderFrame);
@@ -237,8 +230,12 @@ void FDisplayClusterViewportManagerProxy::ImplRenderFrame(FViewport* InViewport)
 		{
 			ViewportManagerProxy->PostProcessManager->HandleRenderFrameSetup_RenderThread(RHICmdList, ViewportManagerProxy);
 		}
+	});
 
-		bool bWarpBlendEnabled = ViewportManagerProxy->RenderFrameSettings.bAllowWarpBlend && CVarWarpBlendEnabled.GetValueOnRenderThread() != 0;
+	ENQUEUE_RENDER_COMMAND(DisplayClusterRenderFrame_CrossGPUTransfer)(
+		[InViewportManagerProxy = SharedThis(this), InViewport](FRHICommandListImmediate& RHICmdList)
+	{
+		const FDisplayClusterViewportManagerProxy* ViewportManagerProxy = &InViewportManagerProxy.Get();
 
 		// mGPU not used for in-editor rendering
 		if(ViewportManagerProxy->RenderFrameSettings.bIsRenderingInEditor == false || ViewportManagerProxy->RenderFrameSettings.bAllowMultiGPURenderingInEditor)
@@ -252,6 +249,12 @@ void FDisplayClusterViewportManagerProxy::ImplRenderFrame(FViewport* InViewport)
 		IDisplayCluster::Get().GetCallbacks().OnDisplayClusterPostCrossGpuTransfer_RenderThread().Broadcast(RHICmdList, ViewportManagerProxy, InViewport);
 		// Latency processing
 		IDisplayCluster::Get().GetCallbacks().OnDisplayClusterProcessLatency_RenderThread().Broadcast(RHICmdList, ViewportManagerProxy, InViewport);
+	});
+
+	ENQUEUE_RENDER_COMMAND(DisplayClusterRenderFrame_WarpBlend)(
+		[InViewportManagerProxy = SharedThis(this), InViewport](FRHICommandListImmediate& RHICmdList)
+	{
+		const FDisplayClusterViewportManagerProxy* ViewportManagerProxy = &InViewportManagerProxy.Get();
 
 		// Update viewports resources: overlay, vp-overla, blur, nummips, etc
 		ViewportManagerProxy->UpdateDeferredResources_RenderThread(RHICmdList);
@@ -261,6 +264,8 @@ void FDisplayClusterViewportManagerProxy::ImplRenderFrame(FViewport* InViewport)
 			ViewportManagerProxy->PostProcessManager->HandleBeginUpdateFrameResources_RenderThread(RHICmdList, ViewportManagerProxy);
 		}
 
+		const bool bWarpBlendEnabled = ViewportManagerProxy->RenderFrameSettings.bAllowWarpBlend && CVarWarpBlendEnabled.GetValueOnRenderThread() != 0;
+
 		// Update the frame resources: post-processing, warping, and finally resolving everything to the frame resource
 		ViewportManagerProxy->UpdateFrameResources_RenderThread(RHICmdList, bWarpBlendEnabled);
 
@@ -472,53 +477,43 @@ void FDisplayClusterViewportManagerProxy::DoCrossGPUTransfers_RenderThread(FRHIC
 	check(IsInRenderingThread());
 
 #if WITH_MGPU
-
-	if ((CVarCrossGPUTransfersEnabled.GetValueOnRenderThread() == 0))
-	{
-		return;
-	}
-
 	// Copy the view render results to all GPUs that are native to the viewport.
 	TArray<FTransferResourceParams> TransferResources;
 
 	for (const TSharedPtr<FDisplayClusterViewportProxy, ESPMode::ThreadSafe>& ViewportProxyIt : ClusterNodeViewportProxies)
 	{
-		bool bShouldCrossGPUTransfersViewportRenderResources = true;
-
 		// Skip a frozen viewport that has already been transferred between GPUs
 		// The first time freezing should do the transfer (RenderTargets must be assigned on the first pass)
-		if (ViewportProxyIt->RenderSettings.bFreezeRendering && ViewportProxyIt->RenderTargets.Num() == 0)
+		const bool bSkipThisViewport = ViewportProxyIt->RenderSettings.bFreezeRendering && ViewportProxyIt->RenderTargets.Num() == 0;
+		if (ViewportProxyIt->RenderSettings.bEnableCrossGPUTransfer && !bSkipThisViewport)
 		{
-			bShouldCrossGPUTransfersViewportRenderResources = false;
-		}
-
-		if (bShouldCrossGPUTransfersViewportRenderResources)
-		{
-			for (FDisplayClusterViewport_Context& ViewportContext : ViewportProxyIt->Contexts)
+			for (const FDisplayClusterViewport_Context& ViewportContext : ViewportProxyIt->Contexts)
 			{
-				if (ViewportContext.bAllowGPUTransferOptimization && ViewportContext.GPUIndex >= 0)
-				{
-					// Use optimized cross GPU transfer for this context
+				const int32 SrcGPUIndex = ViewportContext.RenderThreadData.GPUIndex;
 
-					FRenderTarget* RenderTarget = ViewportProxyIt->RenderTargets[ViewportContext.ContextNum];
-					FRHITexture2D* TextureRHI = ViewportProxyIt->RenderTargets[ViewportContext.ContextNum]->GetViewportRenderTargetResourceRHI();
-
-					FRHIGPUMask RenderTargetGPUMask = (GNumExplicitGPUsForRendering > 1 && RenderTarget) ? RenderTarget->GetGPUMask(RHICmdList) : FRHIGPUMask::GPU0();
-					FRHIGPUMask ContextGPUMask = FRHIGPUMask::FromIndex(ViewportContext.GPUIndex);
-
-					if (ContextGPUMask != RenderTargetGPUMask)
+				// Perform an MGPU transfer for a specific viewport context:
+				if (ViewportContext.bOverrideCrossGPUTransfer && SrcGPUIndex >= 0)
+				{
+					// Clamp the view rect by the rendertarget rect to prevent issues when resizing the viewport.
+					const FIntRect TransferRect = ViewportContext.RenderTargetRect;
+					if (TransferRect.Width() > 0 && TransferRect.Height() > 0)
 					{
-						// Clamp the view rect by the rendertarget rect to prevent issues when resizing the viewport.
-						const FIntRect TransferRect = ViewportContext.RenderTargetRect;
+						// Use optimized cross GPU transfer for this context
+						FRenderTarget* RenderTarget = ViewportProxyIt->RenderTargets[ViewportContext.ContextNum];
+						const FRHIGPUMask RenderTargetGPUMask = (GNumExplicitGPUsForRendering > 1 && RenderTarget) ? RenderTarget->GetGPUMask(RHICmdList) : FRHIGPUMask::GPU0();
 
-						if (TransferRect.Width() > 0 && TransferRect.Height() > 0)
+						if (FRHITexture2D* TextureRHI = ViewportProxyIt->RenderTargets[ViewportContext.ContextNum]->GetViewportRenderTargetResourceRHI())
 						{
-							for (uint32 RenderTargetGPUIndex : RenderTargetGPUMask)
+							for (uint32 DestGPUIndex : RenderTargetGPUMask)
 							{
-								if (!ContextGPUMask.Contains(RenderTargetGPUIndex))
+								if (DestGPUIndex != SrcGPUIndex)
 								{
-									FTransferResourceParams ResourceParams(TextureRHI, TransferRect, ContextGPUMask.GetFirstIndex(), RenderTargetGPUIndex, true, ViewportContext.bEnabledGPUTransferLockSteps);
-									TransferResources.Add(ResourceParams);
+									TransferResources.Add(FTransferResourceParams(
+										TextureRHI, TransferRect,
+										SrcGPUIndex, DestGPUIndex,
+										RenderFrameSettings.CrossGPUTransfer.bPullData,
+										RenderFrameSettings.CrossGPUTransfer.bLockSteps
+									));
 								}
 							}
 						}
diff --git a/Engine/Plugins/Runtime/nDisplay/Source/DisplayCluster/Private/Render/Viewport/DisplayClusterViewportManagerViewExtension.cpp b/Engine/Plugins/Runtime/nDisplay/Source/DisplayCluster/Private/Render/Viewport/DisplayClusterViewportManagerViewExtension.cpp
--- a/Engine/Plugins/Runtime/nDisplay/Source/DisplayCluster/Private/Render/Viewport/DisplayClusterViewportManagerViewExtension.cpp
+++ b/Engine/Plugins/Runtime/nDisplay/Source/DisplayCluster/Private/Render/Viewport/DisplayClusterViewportManagerViewExtension.cpp
@@ -3,7 +3,7 @@
 #include "DisplayClusterViewportManagerViewExtension.h"
 #include "DisplayClusterSceneViewExtensions.h"
 #include "Render/Viewport/IDisplayClusterViewportManager.h"
-#include "Render/Viewport/IDisplayClusterViewportManagerProxy.h"
+#include "Render/Viewport/DisplayClusterViewportManagerProxy.h"
 #include "SceneRendering.h"
 
 #include "IDisplayCluster.h"
@@ -69,13 +69,14 @@ void FDisplayClusterViewportManagerViewExtension::PreRenderViewFamily_RenderThre
 		if (const FSceneView* SceneView = InViewFamily.Views[ViewIndex])
 		{
 			FViewportProxy NewViewportProxy;
-			NewViewportProxy.ViewportProxy = ViewportManagerProxy->FindViewport_RenderThread(SceneView->StereoViewIndex, &NewViewportProxy.ViewportProxyContext.ContextNum);
-			NewViewportProxy.ViewportProxyContext.ViewFamilyProfileDescription = InViewFamily.ProfileDescription;
 
-			NewViewportProxy.ViewIndex = ViewIndex;
-
-			if (NewViewportProxy.IsEnabled())
+			if (IDisplayClusterViewportProxy* ViewportProxyPtr = ViewportManagerProxy->FindViewport_RenderThread(SceneView->StereoViewIndex, &NewViewportProxy.ViewportProxyContext.ContextNum))
 			{
+				NewViewportProxy.ViewportProxyContext.ViewFamilyProfileDescription = InViewFamily.ProfileDescription;
+
+				NewViewportProxy.ViewportProxy = (static_cast<FDisplayClusterViewportProxy*>(ViewportProxyPtr))->AsShared();
+				NewViewportProxy.ViewIndex = ViewIndex;
+
 				Viewports.Add(NewViewportProxy);
 			}
 		}
diff --git a/Engine/Plugins/Runtime/nDisplay/Source/DisplayCluster/Private/Render/Viewport/DisplayClusterViewportManagerViewExtension.h b/Engine/Plugins/Runtime/nDisplay/Source/DisplayCluster/Private/Render/Viewport/DisplayClusterViewportManagerViewExtension.h
--- a/Engine/Plugins/Runtime/nDisplay/Source/DisplayCluster/Private/Render/Viewport/DisplayClusterViewportManagerViewExtension.h
+++ b/Engine/Plugins/Runtime/nDisplay/Source/DisplayCluster/Private/Render/Viewport/DisplayClusterViewportManagerViewExtension.h
@@ -56,15 +56,14 @@ private:
 
 	struct FViewportProxy
 	{
-		inline bool IsEnabled() const
-		{
-			return ViewportProxy != nullptr;
-		}
-
-		class IDisplayClusterViewportProxy* ViewportProxy = nullptr;
-		FDisplayClusterViewportProxy_Context ViewportProxyContext;
+		// The ViewportProxy must exist until this object is removed.
+		TSharedPtr<class IDisplayClusterViewportProxy, ESPMode::ThreadSafe> ViewportProxy;
 
+		// The index of view in viewfamily collection
 		int32 ViewIndex = 0;
+
+		// Context for viewport proxy object
+		FDisplayClusterViewportProxy_Context ViewportProxyContext;
 	};
 
 	TArray<FViewportProxy> Viewports;
diff --git a/Engine/Plugins/Runtime/nDisplay/Source/DisplayCluster/Private/Render/Viewport/DisplayClusterViewportProxy.cpp b/Engine/Plugins/Runtime/nDisplay/Source/DisplayCluster/Private/Render/Viewport/DisplayClusterViewportProxy.cpp
--- a/Engine/Plugins/Runtime/nDisplay/Source/DisplayCluster/Private/Render/Viewport/DisplayClusterViewportProxy.cpp
+++ b/Engine/Plugins/Runtime/nDisplay/Source/DisplayCluster/Private/Render/Viewport/DisplayClusterViewportProxy.cpp
@@ -769,19 +769,17 @@ void FDisplayClusterViewportProxy::PostRenderViewFamily_RenderThread(FRDGBuilder
 {
 	const uint32 InContextNum = InProxyContext.ContextNum;
 
-	// Save GPUMask for each viewport context
+#if WITH_MGPU
+	// Get the GPUIndex used to render this viewport
 	if (Contexts.IsValidIndex(InContextNum))
 	{
-		if (GPUMask.Num() < Contexts.Num())
-		{
-			GPUMask.AddDefaulted(Contexts.Num() - GPUMask.Num());
-		}
-
 		checkSlow(InSceneView.bIsViewInfo);
 		const FViewInfo& ViewInfo = static_cast<const FViewInfo&>(InSceneView);
 
-		GPUMask[InContextNum] = ViewInfo.GPUMask;
+		const uint32 GPUIndex = ViewInfo.GPUMask.GetFirstIndex();
+		Contexts[InContextNum].RenderThreadData.GPUIndex = (GPUIndex < GNumExplicitGPUsForRendering) ? GPUIndex : -1;
 	}
+#endif
 
 	if (!InProxyContext.ViewFamilyProfileDescription.IsEmpty())
 	{
diff --git a/Engine/Plugins/Runtime/nDisplay/Source/DisplayCluster/Private/Render/Viewport/DisplayClusterViewportProxy.h b/Engine/Plugins/Runtime/nDisplay/Source/DisplayCluster/Private/Render/Viewport/DisplayClusterViewportProxy.h
--- a/Engine/Plugins/Runtime/nDisplay/Source/DisplayCluster/Private/Render/Viewport/DisplayClusterViewportProxy.h
+++ b/Engine/Plugins/Runtime/nDisplay/Source/DisplayCluster/Private/Render/Viewport/DisplayClusterViewportProxy.h
@@ -177,9 +177,6 @@ protected:
 	// Viewport contexts (left/center/right eyes)
 	mutable TArray<FDisplayClusterViewport_Context> Contexts;
 
-	/** The GPU nodes on which to render this view. */
-	TArray<FRHIGPUMask> GPUMask;
-
 	// View family render to this resources
 	TArray<FDisplayClusterViewportRenderTargetResource*> RenderTargets;
 
diff --git a/Engine/Plugins/Runtime/nDisplay/Source/DisplayCluster/Private/Render/Viewport/RenderFrame/DisplayClusterRenderFrameSettings.h b/Engine/Plugins/Runtime/nDisplay/Source/DisplayCluster/Private/Render/Viewport/RenderFrame/DisplayClusterRenderFrameSettings.h
--- a/Engine/Plugins/Runtime/nDisplay/Source/DisplayCluster/Private/Render/Viewport/RenderFrame/DisplayClusterRenderFrameSettings.h
+++ b/Engine/Plugins/Runtime/nDisplay/Source/DisplayCluster/Private/Render/Viewport/RenderFrame/DisplayClusterRenderFrameSettings.h
@@ -12,8 +12,23 @@ struct FDisplayClusterRenderFrameSettings
 	// customize mono\stereo render modes
 	EDisplayClusterRenderFrameMode RenderMode = EDisplayClusterRenderFrameMode::Mono;
 
-	// Control mGPU for whole cluster
-	EDisplayClusterMultiGPUMode MultiGPUMode = EDisplayClusterMultiGPUMode::Enabled;
+	// nDisplay has its own implementation of cross-GPU transfer.
+	struct FCrossGPUTransfer
+	{
+		// Enable cross-GPU transfers using nDisplay
+		// That replaces the default cross-GPU transfers using UE Core for the nDisplay viewports viewfamilies.
+		bool bEnable = false;
+
+		// The bLockSteps parameter is simply passed to the FTransferResourceParams structure.
+		// Whether the GPUs must handshake before and after the transfer. Required if the texture rect is being written to in several render passes.
+		// Otherwise, minimal synchronization will be used.
+		bool bLockSteps = false;
+
+		// The bPullData parameter is simply passed to the FTransferResourceParams structure.
+		// Whether the data is read by the dest GPU, or written by the src GPU (not allowed if the texture is a backbuffer)
+		bool bPullData = true;
+
+	} CrossGPUTransfer;
 
 	// Some frame postprocess require additional render targetable resources
 	bool bShouldUseAdditionalFrameTargetableResource = false;
diff --git a/Engine/Plugins/Runtime/nDisplay/Source/DisplayCluster/Public/Render/Viewport/Containers/DisplayClusterViewport_Context.h b/Engine/Plugins/Runtime/nDisplay/Source/DisplayCluster/Public/Render/Viewport/Containers/DisplayClusterViewport_Context.h
--- a/Engine/Plugins/Runtime/nDisplay/Source/DisplayCluster/Public/Render/Viewport/Containers/DisplayClusterViewport_Context.h
+++ b/Engine/Plugins/Runtime/nDisplay/Source/DisplayCluster/Public/Render/Viewport/Containers/DisplayClusterViewport_Context.h
@@ -41,10 +41,11 @@ public:
 	// Rendering data, for internal usage
 
 	// GPU index for this context render target
-	int32 GPUIndex = -1;
+	int32 GPUIndex = INDEX_NONE;
 
-	bool bAllowGPUTransferOptimization = false;
-	bool bEnabledGPUTransferLockSteps = true;
+	/* Enables nDisplay's native implementation of cross-GPU transfer.
+	 * This disables cross-GPU transfer by default for nDisplay viewports in FSceneViewFamily structure. **/
+	bool bOverrideCrossGPUTransfer = false;
 
 	// Location and size on a render target texture
 	FIntRect RenderTargetRect;
@@ -63,4 +64,13 @@ public:
 
 	// Disable render for this viewport (Overlay)
 	bool bDisableRender = false;
+
+	struct FRenderThreadData
+	{
+		// GPUIndex used to render this context.
+		int32 GPUIndex = INDEX_NONE;
+	};
+
+	// This data updated only on rendering thread
+	FRenderThreadData RenderThreadData;
 };
diff --git a/Engine/Plugins/Runtime/nDisplay/Source/DisplayCluster/Public/Render/Viewport/Containers/DisplayClusterViewport_RenderSettings.h b/Engine/Plugins/Runtime/nDisplay/Source/DisplayCluster/Public/Render/Viewport/Containers/DisplayClusterViewport_RenderSettings.h
--- a/Engine/Plugins/Runtime/nDisplay/Source/DisplayCluster/Public/Render/Viewport/Containers/DisplayClusterViewport_RenderSettings.h
+++ b/Engine/Plugins/Runtime/nDisplay/Source/DisplayCluster/Public/Render/Viewport/Containers/DisplayClusterViewport_RenderSettings.h
@@ -64,6 +64,14 @@ public:
 	// Is this viewport being captured by a media capture device?
 	bool bIsBeingCaptured = false;
 
+	/** Enable cross-GPU transfer for this viewport.
+	  * It may be disabled in some configurations. For example, when using offscreen rendering with TextureShare,
+	  * cross-gpu transfer can be disabled for this viewport to improve performance, because when transfer is called,
+	  * it freezes the GPUs until synchronization is reached.
+	  * (TextureShare uses its own implementation of the crossGPU transfer for the shared textures.)
+	  */
+	bool bEnableCrossGPUTransfer = true;
+
 	// Performance, Multi-GPU: Asign GPU for viewport rendering. The Value '-1' used to default gpu mapping
 	int32 GPUIndex = -1;
 
diff --git a/Engine/Plugins/Runtime/nDisplay/Source/DisplayCluster/Public/Render/Viewport/RenderFrame/DisplayClusterRenderFrameEnums.h b/Engine/Plugins/Runtime/nDisplay/Source/DisplayCluster/Public/Render/Viewport/RenderFrame/DisplayClusterRenderFrameEnums.h
--- a/Engine/Plugins/Runtime/nDisplay/Source/DisplayCluster/Public/Render/Viewport/RenderFrame/DisplayClusterRenderFrameEnums.h
+++ b/Engine/Plugins/Runtime/nDisplay/Source/DisplayCluster/Public/Render/Viewport/RenderFrame/DisplayClusterRenderFrameEnums.h
@@ -37,20 +37,3 @@ enum class EDisplayClusterRenderFamilyMode : uint8
 	// Use rules to merge views to minimal num of families (separate by: buffer_ratio, viewExtension, max RTT size)
 	MergeAnyPossible,
 };
-
-
-// Control multiGPU for cluster
-enum class EDisplayClusterMultiGPUMode : uint8
-{
-	// Disable multi GPU rendering
-	None = 0,
-
-	// Use default crossGPU transfer
-	Enabled,
-
-	// Performance (Experimental): Use optimized transfer once per frame with bLockStepGPUs=true
-	Optimized_EnabledLockSteps,
-
-	// Performance (Experimental): Use optimized transfer once per frame with bLockStepGPUs=false 
-	Optimized_DisabledLockSteps,
-};
diff --git a/Engine/Plugins/Runtime/nDisplay/Source/DisplayClusterConfiguration/Public/DisplayClusterConfigurationTypes_Enums.h b/Engine/Plugins/Runtime/nDisplay/Source/DisplayClusterConfiguration/Public/DisplayClusterConfigurationTypes_Enums.h
--- a/Engine/Plugins/Runtime/nDisplay/Source/DisplayClusterConfiguration/Public/DisplayClusterConfigurationTypes_Enums.h
+++ b/Engine/Plugins/Runtime/nDisplay/Source/DisplayClusterConfiguration/Public/DisplayClusterConfigurationTypes_Enums.h
@@ -143,22 +143,6 @@ enum class EDisplayClusterConfigurationCameraMotionBlurMode : uint8
 	Override     UMETA(DisplayName = "ICVFX Camera Blur Off"),
 };
 
-UENUM()
-enum class EDisplayClusterConfigurationRenderMGPUMode : uint8
-{
-	// Disable multi GPU rendering
-	None         UMETA(DisplayName = "Disabled"),
-
-	// Use default UE crossGPU transfer
-	Enabled      UMETA(DisplayName = "Enabled"),
-
-	// Performance (Experimental): Use optimized transfer once per frame with bLockStepGPUs=true
-	Optimized_EnabledLockSteps  UMETA(DisplayName = "Optimization: Enabled Lockstep"),
-
-	// Performance (Experimental): Use optimized transfer once per frame with bLockStepGPUs=false 
-	Optimized_DisabledLockSteps  UMETA(DisplayName = "Optimization: Disabled Lockstep"),
-};
-
 UENUM()
 enum class EDisplayClusterConfigurationViewportOverscanMode : uint8
 {
diff --git a/Engine/Plugins/Runtime/nDisplay/Source/DisplayClusterConfiguration/Public/DisplayClusterConfigurationTypes_Viewport.h b/Engine/Plugins/Runtime/nDisplay/Source/DisplayClusterConfiguration/Public/DisplayClusterConfigurationTypes_Viewport.h
--- a/Engine/Plugins/Runtime/nDisplay/Source/DisplayClusterConfiguration/Public/DisplayClusterConfigurationTypes_Viewport.h
+++ b/Engine/Plugins/Runtime/nDisplay/Source/DisplayClusterConfiguration/Public/DisplayClusterConfigurationTypes_Viewport.h
@@ -83,6 +83,14 @@ struct DISPLAYCLUSTERCONFIGURATION_API FDisplayClusterConfigurationViewport_Rend
 	GENERATED_BODY()
 
 public:
+	/* Enable cross-GPU transfer for this viewport.
+	 * It may be disabled in some configurations. For example, when using offscreen rendering with TextureShare,
+	 * cross-gpu transfer can be disabled for this viewport to improve performance, because when transfer is called, 
+	 * it freezes the GPUs until synchronization is reached.
+	 * (TextureShare uses its own implementation of the crossGPU transfer for the shared textures.) */
+	UPROPERTY(EditAnywhere, BlueprintReadWrite, Category = "Rendering", meta = (DisplayName = "Enable Cross-GPU Transfer"))
+	bool bEnableCrossGPUTransfer = true;
+
 	/** Specify which GPU should render the second Stereo eye */
 	UPROPERTY(EditAnywhere, Category = "Stereo", meta = (DisplayName = "Stereo GPU Index"))
 	int StereoGPUIndex = -1;
@@ -240,10 +248,6 @@ struct DISPLAYCLUSTERCONFIGURATION_API FDisplayClusterConfigurationRenderFrame
 	GENERATED_BODY()
 
 public:
-	// Performance: Allow change global MGPU settings
-	UPROPERTY(EditAnywhere, Category = "Configuration", meta = (DisplayName = "Multi GPU Mode"))
-	EDisplayClusterConfigurationRenderMGPUMode MultiGPUMode = EDisplayClusterConfigurationRenderMGPUMode::Enabled;
-
 	// Performance: Allow merge multiple viewports on single RTT with atlasing (required for bAllowViewFamilyMergeOptimization)
 	// [not implemented yet] Experimental
 	UPROPERTY()
