nDisplay: FDisplayClusterViewportManagerViewExtension needs to be created only after the GEngine variable has been initialized.

diff --git a/Engine/Plugins/Runtime/nDisplay/Source/DisplayCluster/Private/Render/Viewport/DisplayClusterViewportManager.cpp b/Engine/Plugins/Runtime/nDisplay/Source/DisplayCluster/Private/Render/Viewport/DisplayClusterViewportManager.cpp
--- a/Engine/Plugins/Runtime/nDisplay/Source/DisplayCluster/Private/Render/Viewport/DisplayClusterViewportManager.cpp
+++ b/Engine/Plugins/Runtime/nDisplay/Source/DisplayCluster/Private/Render/Viewport/DisplayClusterViewportManager.cpp
@@ -80,9 +80,6 @@ FDisplayClusterViewportManager::FDisplayClusterViewportManager()
 	PostProcessManager  = MakeShared<FDisplayClusterViewportPostProcessManager, ESPMode::ThreadSafe>(*this);
 	LightCardManager = MakeShared<FDisplayClusterViewportLightCardManager, ESPMode::ThreadSafe>(*this);
 
-	// Create DC ViewExtension to handle special features
-	ViewportManagerViewExtension = FSceneViewExtensions::NewExtension<FDisplayClusterViewportManagerViewExtension>(this);
-
 	// initialize proxy
 	ViewportManagerProxy->Initialize(*this);
 
@@ -483,6 +480,12 @@ bool FDisplayClusterViewportManager::BeginNewFrame(FViewport* InViewport, UWorld
 		}
 	}
 
+	// Create DC ViewExtension to handle special features
+	if (!ViewportManagerViewExtension.IsValid())
+	{
+		ViewportManagerViewExtension = FSceneViewExtensions::NewExtension<FDisplayClusterViewportManagerViewExtension>(this);
+	}
+
 	// Before new frame
 	if (PostProcessManager.IsValid())
 	{
