nDisplay: Fixed crash that can occur when editor shuts down with the nDisplay plugin enabled. Issue was caused by attempting to access a UObject's static class after the engine had already cleaned up the static classes as part of shut down.

diff --git a/Engine/Plugins/Runtime/nDisplay/Source/DisplayClusterLightCardEditor/Private/DisplayClusterLightCardEditorModule.cpp b/Engine/Plugins/Runtime/nDisplay/Source/DisplayClusterLightCardEditor/Private/DisplayClusterLightCardEditorModule.cpp
--- a/Engine/Plugins/Runtime/nDisplay/Source/DisplayClusterLightCardEditor/Private/DisplayClusterLightCardEditorModule.cpp
+++ b/Engine/Plugins/Runtime/nDisplay/Source/DisplayClusterLightCardEditor/Private/DisplayClusterLightCardEditorModule.cpp
@@ -20,6 +20,13 @@
 
 #include "ISettingsModule.h"
 
+#define REGISTER_OBJECT_LAYOUT(ObjectType, CustomizationType) { \
+	const FName LayoutName = ObjectType::StaticClass()->GetFName(); \
+	RegisteredClassLayoutNames.Add(LayoutName); \
+	PropertyModule.RegisterCustomClassLayout(LayoutName, \
+		FOnGetDetailCustomizationInstance::CreateStatic(&CustomizationType::MakeInstance)); \
+}
+
 #define LOCTEXT_NAMESPACE "DisplayClusterLightCardEditor"
 
 void FDisplayClusterLightCardEditorModule::StartupModule()
@@ -103,10 +110,8 @@ void FDisplayClusterLightCardEditorModule::UnregisterSettings()
 void FDisplayClusterLightCardEditorModule::RegisterDetailCustomizations()
 {
 	FPropertyEditorModule& PropertyModule = FModuleManager::LoadModuleChecked<FPropertyEditorModule>("PropertyEditor");
-	PropertyModule.RegisterCustomClassLayout(
-		ADisplayClusterLightCardActor::StaticClass()->GetFName(), 
-		FOnGetDetailCustomizationInstance::CreateStatic(&FDisplayClusterLightCardActorDetails::MakeInstance)
-	);
+
+	REGISTER_OBJECT_LAYOUT(ADisplayClusterLightCardActor, FDisplayClusterLightCardActorDetails);
 }
 
 void FDisplayClusterLightCardEditorModule::UnregisterDetailCustomizations()
@@ -114,8 +119,13 @@ void FDisplayClusterLightCardEditorModule::UnregisterDetailCustomizations()
 	FPropertyEditorModule* PropertyModule = FModuleManager::GetModulePtr<FPropertyEditorModule>("PropertyEditor");
 	if (PropertyModule)
 	{
-		PropertyModule->UnregisterCustomClassLayout(ADisplayClusterLightCardActor::StaticClass()->GetFName());
+		for (const FName& ClassName : RegisteredClassLayoutNames)
+		{
+			PropertyModule->UnregisterCustomClassLayout(ClassName);
+		}
 	}
+
+	RegisteredClassLayoutNames.Empty();
 }
 
 void FDisplayClusterLightCardEditorModule::RegisterOperatorApp()
diff --git a/Engine/Plugins/Runtime/nDisplay/Source/DisplayClusterLightCardEditor/Private/DisplayClusterLightCardEditorModule.h b/Engine/Plugins/Runtime/nDisplay/Source/DisplayClusterLightCardEditor/Private/DisplayClusterLightCardEditorModule.h
--- a/Engine/Plugins/Runtime/nDisplay/Source/DisplayClusterLightCardEditor/Private/DisplayClusterLightCardEditorModule.h
+++ b/Engine/Plugins/Runtime/nDisplay/Source/DisplayClusterLightCardEditor/Private/DisplayClusterLightCardEditorModule.h
@@ -31,6 +31,9 @@ private:
 	void HandleModuleChanged(FName InModuleName, EModuleChangeReason InChangeReason);
 	
 private:
+	/** Stores a list of class names that custom layouts have been registered for */
+	TArray<FName> RegisteredClassLayoutNames;
+
 	FDelegateHandle ModuleChangedHandle;
 	FDelegateHandle OperatorAppHandle;
 };
