Force a re-register of components after transaction has posted to trigger refresh on editor preview.

diff --git a/Engine/Plugins/Runtime/nDisplay/Source/DisplayCluster/Private/Game/EngineClasses/Scene/DisplayClusterRootActorEditor.cpp b/Engine/Plugins/Runtime/nDisplay/Source/DisplayCluster/Private/Game/EngineClasses/Scene/DisplayClusterRootActorEditor.cpp
--- a/Engine/Plugins/Runtime/nDisplay/Source/DisplayCluster/Private/Game/EngineClasses/Scene/DisplayClusterRootActorEditor.cpp
+++ b/Engine/Plugins/Runtime/nDisplay/Source/DisplayCluster/Private/Game/EngineClasses/Scene/DisplayClusterRootActorEditor.cpp
@@ -29,6 +29,7 @@
 #include "Misc/DisplayClusterLog.h"
 #include "Misc/DisplayClusterStrings.h"
 
+#include "Misc/TransactionObjectEvent.h"
 #include "Render/Viewport/Configuration/DisplayClusterViewportConfigurationHelpers_ICVFX.h"
 #include "Render/Viewport/DisplayClusterViewportStrings.h"
 #include "Render/Viewport/IDisplayClusterViewportManager.h"
@@ -701,6 +702,21 @@ static FName Name_RelativeLocation = USceneComponent::GetRelativeLocationPropert
 static FName Name_RelativeRotation = USceneComponent::GetRelativeRotationPropertyName();
 static FName Name_RelativeScale3D = USceneComponent::GetRelativeScale3DPropertyName();
 
+void ADisplayClusterRootActor::PostTransacted(const FTransactionObjectEvent& TransactionEvent)
+{
+	if (bRequiresComponentRefresh
+		&& TransactionEvent.GetEventType() == ETransactionObjectEventType::Finalized)
+	{
+		if ((GEditor && GEditor->bIsSimulatingInEditor && GetWorld() != nullptr) || ReregisterComponentsWhenModified())
+		{
+			UnregisterAllComponents();
+			ReregisterAllComponents();
+		}
+		bRequiresComponentRefresh = false;
+	}
+	Super::PostTransacted(TransactionEvent);
+}
+
 void ADisplayClusterRootActor::PostEditChangeChainProperty(FPropertyChangedChainEvent& PropertyChainEvent)
 {
 	const FProperty* RootProperty = PropertyChainEvent.PropertyChain.GetActiveNode()->GetValue();
@@ -715,6 +731,7 @@ void ADisplayClusterRootActor::PostEditChangeChainProperty(FPropertyChangedChain
 			// Do not propagate the PostEditChangeProperty because we are in an interactive edit of a suboject
 			// and we do not need to rerun any construction scripts.
 			bIsInteractiveEditingSubobject = true;
+			bRequiresComponentRefresh = true;
 		}
 	}
 
diff --git a/Engine/Plugins/Runtime/nDisplay/Source/DisplayCluster/Public/DisplayClusterRootActor.h b/Engine/Plugins/Runtime/nDisplay/Source/DisplayCluster/Public/DisplayClusterRootActor.h
--- a/Engine/Plugins/Runtime/nDisplay/Source/DisplayCluster/Public/DisplayClusterRootActor.h
+++ b/Engine/Plugins/Runtime/nDisplay/Source/DisplayCluster/Public/DisplayClusterRootActor.h
@@ -27,6 +27,7 @@
 
 #if WITH_EDITOR
 class IDisplayClusterConfiguratorBlueprintEditor;
+class FTransactionObjectEvent;
 #endif
 
 class IDisplayClusterStageActor;
@@ -493,6 +494,7 @@ protected:
 	virtual void PostEditChangeChainProperty(FPropertyChangedChainEvent& PropertyChainEvent) override;
 	virtual void PostEditChangeProperty(FPropertyChangedEvent& PropertyChangedEvent) override;
 	virtual void PostEditMove(bool bFinished) override;
+	virtual void PostTransacted(const FTransactionObjectEvent& TransactionEvent) override;
 
 	/** Called when the asset has been reloaded in the editor. */
 	void HandleAssetReload(const EPackageReloadPhase InPackageReloadPhase, FPackageReloadedEvent* InPackageReloadedEvent);
@@ -505,6 +507,9 @@ private:
 	/** Flag to indicate if the user is currently interacting with a subobject of CurrentConfigData. */
 	bool bIsInteractiveEditingSubobject = false;
 
+	/** Flag to indicate if we need to reregister components. */
+	bool bRequiresComponentRefresh = false;
+
 	bool bIsSelectedInEditor = false;
 
 	/** When the preview render should be directed to the post process render target. */
