Display Cluster:  added CVar "DC.MultiGPU" to allow live toggling of MGPU.  The old CVar that used to allow for this had been removed (DC.OverrideMultiGPUMode), and it's useful for debugging.

#rnx

diff --git a/Engine/Plugins/Runtime/nDisplay/Source/DisplayCluster/Private/Render/Viewport/DisplayClusterViewport.cpp b/Engine/Plugins/Runtime/nDisplay/Source/DisplayCluster/Private/Render/Viewport/DisplayClusterViewport.cpp
--- a/Engine/Plugins/Runtime/nDisplay/Source/DisplayCluster/Private/Render/Viewport/DisplayClusterViewport.cpp
+++ b/Engine/Plugins/Runtime/nDisplay/Source/DisplayCluster/Private/Render/Viewport/DisplayClusterViewport.cpp
@@ -25,6 +25,14 @@
 
 #include "Misc/DisplayClusterLog.h"
 
+int32 GDisplayClusterMultiGPUEnable = 1;
+static FAutoConsoleVariableRef CVarDisplayClusterMultiGPUEnable(
+	TEXT("DC.MultiGPU"),
+	GDisplayClusterMultiGPUEnable,
+	TEXT("Enable MultiGPU for Display Cluster rendering.  Useful to disable for debugging.  (Default = 1)"),
+	ECVF_Default
+);
+
 ///////////////////////////////////////////////////////////////////////////////////////
 //          FDisplayClusterViewport
 ///////////////////////////////////////////////////////////////////////////////////////
@@ -548,7 +556,8 @@ bool FDisplayClusterViewport::UpdateFrameContexts(const uint32 InStereoViewIndex
 			Context.bOverrideCrossGPUTransfer = true;
 		}
 
-		const int32 MaxExplicitGPUIndex = GNumExplicitGPUsForRendering - 1;
+		const int32 MaxExplicitGPUIndex = GDisplayClusterMultiGPUEnable ? GNumExplicitGPUsForRendering - 1 : 0;
+
 		if (MaxExplicitGPUIndex > 0)
 		{
 			if (InFrameSettings.bIsRenderingInEditor)
