nDisplay: Add Chromakey Card actors which can be spawned from the ICVFX panel.
Use class display name in actor spawning when possible.
Fix recent items not always using the correct icons.

diff --git a/Engine/Plugins/Runtime/nDisplay/Source/DisplayCluster/Private/Blueprints/DisplayClusterBlueprintLib.cpp b/Engine/Plugins/Runtime/nDisplay/Source/DisplayCluster/Private/Blueprints/DisplayClusterBlueprintLib.cpp
--- a/Engine/Plugins/Runtime/nDisplay/Source/DisplayCluster/Private/Blueprints/DisplayClusterBlueprintLib.cpp
+++ b/Engine/Plugins/Runtime/nDisplay/Source/DisplayCluster/Private/Blueprints/DisplayClusterBlueprintLib.cpp
@@ -4,9 +4,11 @@
 #include "Blueprints/DisplayClusterBlueprintAPIImpl.h"
 #include "UObject/Package.h"
 
+#include "DisplayClusterChromakeyCardActor.h"
 #include "DisplayClusterLightCardActor.h"
 #include "DisplayClusterRootActor.h"
 #include "DisplayClusterConfigurationTypes.h"
+#include "Components/DisplayClusterICVFXCameraComponent.h"
 
 #if WITH_EDITOR
 #include "ScopedTransaction.h"
@@ -139,7 +141,7 @@ void UDisplayClusterBlueprintLib::DuplicateLightCards(TArray<ADisplayClusterLigh
 	}
 }
 
-void UDisplayClusterBlueprintLib::FindLightCardsForRootActor(ADisplayClusterRootActor* RootActor, TSet<ADisplayClusterLightCardActor*>& OutLightCards)
+void UDisplayClusterBlueprintLib::FindLightCardsForRootActor(const ADisplayClusterRootActor* RootActor, TSet<ADisplayClusterLightCardActor*>& OutLightCards)
 {
 	if (!RootActor)
 	{
@@ -187,4 +189,59 @@ void UDisplayClusterBlueprintLib::FindLightCardsForRootActor(ADisplayClusterRoot
 	}
 }
 
+void UDisplayClusterBlueprintLib::FindChromakeyCardsForRootActor(const ADisplayClusterRootActor* RootActor,
+                                                                 TSet<ADisplayClusterChromakeyCardActor*>& OutChromakeyCards)
+{
+	if (!RootActor)
+	{
+		return;
+	}
+
+	TArray<UDisplayClusterICVFXCameraComponent*> CameraComponents;
+	RootActor->GetComponents(CameraComponents);
+
+	for (const UDisplayClusterICVFXCameraComponent* Camera : CameraComponents)
+	{
+		const FDisplayClusterConfigurationICVFX_VisibilityList& RootActorChromakeyCards = Camera->GetCameraSettingsICVFX().Chromakey.ChromakeyRenderTexture.ShowOnlyList;
+
+		for (const TSoftObjectPtr<AActor>& ChromakeyCardActor : RootActorChromakeyCards.Actors)
+		{
+			if (!ChromakeyCardActor.IsValid() || !ChromakeyCardActor->IsA<ADisplayClusterChromakeyCardActor>())
+			{
+				continue;
+			}
+
+			OutChromakeyCards.Add(Cast<ADisplayClusterChromakeyCardActor>(ChromakeyCardActor.Get()));
+		}
+
+		// If there are any layers that are specified as light card layers, iterate over all actors in the world and 
+		// add any that are members of any of the light card layers to the list. Only add an actor once, even if it is
+		// in multiple layers
+		if (RootActorChromakeyCards.ActorLayers.IsEmpty())
+		{
+			continue;
+		}
+
+		if (const UWorld* World = RootActor->GetWorld())
+		{
+			for (TActorIterator<ADisplayClusterChromakeyCardActor> ActorIt(World); ActorIt; ++ActorIt)
+			{
+				if (!IsValid(*ActorIt))
+				{
+					continue;
+				}
+
+				for (const FActorLayer& ActorLayer : RootActorChromakeyCards.ActorLayers)
+				{
+					if (ActorIt->Layers.Contains(ActorLayer.Name))
+					{
+						OutChromakeyCards.Add(*ActorIt);
+						break;
+					}
+				}
+			}
+		}
+	}
+}
+
 #undef LOCTEXT_NAMESPACE
diff --git a/Engine/Plugins/Runtime/nDisplay/Source/DisplayCluster/Private/Game/EngineClasses/Scene/DisplayClusterChromakeyCardActor.cpp b/Engine/Plugins/Runtime/nDisplay/Source/DisplayCluster/Private/Game/EngineClasses/Scene/DisplayClusterChromakeyCardActor.cpp
new file mode 100644
--- /dev/null
+++ b/Engine/Plugins/Runtime/nDisplay/Source/DisplayCluster/Private/Game/EngineClasses/Scene/DisplayClusterChromakeyCardActor.cpp
@@ -0,0 +1,144 @@
+// Copyright Epic Games, Inc. All Rights Reserved.
+
+#include "DisplayClusterChromakeyCardActor.h"
+
+#include "DisplayClusterRootActor.h"
+#include "Components/DisplayClusterICVFXCameraComponent.h"
+
+#if WITH_EDITOR
+#include "Layers/LayersSubsystem.h"
+#endif
+
+ADisplayClusterChromakeyCardActor::ADisplayClusterChromakeyCardActor(const FObjectInitializer& ObjectInitializer)
+	: Super(ObjectInitializer)
+{
+}
+
+void ADisplayClusterChromakeyCardActor::Tick(float DeltaSeconds)
+{
+	Super::Tick(DeltaSeconds);
+
+#if WITH_EDITOR
+	UpdateChromakeySettings();
+#endif
+}
+
+void ADisplayClusterChromakeyCardActor::AddToChromakeyLayer(ADisplayClusterRootActor* InRootActor)
+{
+	check(InRootActor);
+	SetRootActorOwner(InRootActor);
+	
+	TArray<UDisplayClusterICVFXCameraComponent*> ICVFXComponents;
+	InRootActor->GetComponents(ICVFXComponents);
+
+	if (ICVFXComponents.Num() > 0)
+	{
+		static const FString LayerPrefix = TEXT("Chromakey");
+			
+		FName ChromakeyLayerName = NAME_None;
+
+		// Find an existing layer name.
+		for (UDisplayClusterICVFXCameraComponent* ICVFXComponent : ICVFXComponents)
+		{
+			FDisplayClusterConfigurationICVFX_VisibilityList& VisibilityList = ICVFXComponent->CameraSettings.Chromakey.ChromakeyRenderTexture.ShowOnlyList;
+			
+			if (const FActorLayer* ExistingLayer = VisibilityList.ActorLayers.FindByPredicate([](const FActorLayer& Layer)
+			{
+				return Layer.Name.ToString().StartsWith(LayerPrefix);
+			}))
+			{
+				ChromakeyLayerName = ExistingLayer->Name;
+				break;
+			}
+		}
+
+		if (ChromakeyLayerName == NAME_None)
+		{
+			// Find a unique name so it's Chromakey_XXXX
+			do
+			{
+				const FGuid LayerGuid = FGuid::NewGuid();
+				const FString GuidStr = LayerGuid.ToString().Left(4);
+				ChromakeyLayerName = *FString::Printf(TEXT("%s_%s"), *LayerPrefix, *GuidStr);
+			}
+			while (FindObject<UObject>(GetWorld(), *ChromakeyLayerName.ToString()) != nullptr);
+		}
+
+		check(!ChromakeyLayerName.IsNone());
+		bool bLayerAdded = false;
+#if WITH_EDITOR
+		ULayersSubsystem* LayersSubsystem = GEditor ? GEditor->GetEditorSubsystem<ULayersSubsystem>() : nullptr;
+		if (LayersSubsystem)
+		{
+			LayersSubsystem->AddActorsToLayer(TArray<AActor*>{ this }, ChromakeyLayerName);
+			bLayerAdded = true;
+		}
+#endif
+		if (!bLayerAdded)
+		{
+			// Likely -game and we need to add the layer directly
+			Layers.AddUnique(ChromakeyLayerName);
+		}
+
+		// Make sure all ICVFX components on this root actor have this chromakey layer
+		for (UDisplayClusterICVFXCameraComponent* ICVFXComponent : ICVFXComponents)
+		{
+			FDisplayClusterConfigurationICVFX_VisibilityList& VisibilityList = ICVFXComponent->CameraSettings.Chromakey.ChromakeyRenderTexture.ShowOnlyList;
+
+			if (!VisibilityList.ActorLayers.FindByPredicate([ChromakeyLayerName](const FActorLayer& Layer)
+			{
+				return Layer.Name == ChromakeyLayerName;
+			}))
+			{
+#if WITH_EDITOR
+				ICVFXComponent->Modify();
+#endif
+				VisibilityList.ActorLayers.AddDefaulted_GetRef().Name = ChromakeyLayerName;
+			}
+		}
+	}
+}
+
+bool ADisplayClusterChromakeyCardActor::IsReferencedByICVFXCamera(UDisplayClusterICVFXCameraComponent* InCamera) const
+{
+	const FDisplayClusterConfigurationICVFX_CameraSettings& CameraSettings = InCamera->GetCameraSettingsICVFX();
+	if (CameraSettings.Chromakey.ChromakeyRenderTexture.ShowOnlyList.Actors.Contains(this))
+	{
+		return true;
+	}
+
+	for (const FName& ThisLayer : Layers)
+	{
+		if (const FActorLayer* ExistingLayer = CameraSettings.Chromakey.ChromakeyRenderTexture.ShowOnlyList.ActorLayers.FindByPredicate([ThisLayer](const FActorLayer& Layer)
+		{
+			return Layer.Name == ThisLayer;
+		}))
+		{
+			return true;
+		}
+	}
+
+	return false;
+}
+
+void ADisplayClusterChromakeyCardActor::UpdateChromakeySettings()
+{
+	if (RootActorOwner.IsValid())
+	{
+		int32 Count = 0;
+		FLinearColor ChromaColor(ForceInitToZero);
+		TArray<UDisplayClusterICVFXCameraComponent*> ICVFXComponents;
+		RootActorOwner->GetComponents(ICVFXComponents);
+
+		for (UDisplayClusterICVFXCameraComponent* Component : ICVFXComponents)
+		{
+			if (IsReferencedByICVFXCamera(Component))
+			{
+				ChromaColor += Component->GetCameraSettingsICVFX().Chromakey.ChromakeyColor;
+				Count++;
+			}
+		}
+
+		Color = Count > 0 ? ChromaColor / Count : FLinearColor::Green;
+	}
+}
diff --git a/Engine/Plugins/Runtime/nDisplay/Source/DisplayCluster/Public/Blueprints/DisplayClusterBlueprintLib.h b/Engine/Plugins/Runtime/nDisplay/Source/DisplayCluster/Public/Blueprints/DisplayClusterBlueprintLib.h
--- a/Engine/Plugins/Runtime/nDisplay/Source/DisplayCluster/Public/Blueprints/DisplayClusterBlueprintLib.h
+++ b/Engine/Plugins/Runtime/nDisplay/Source/DisplayCluster/Public/Blueprints/DisplayClusterBlueprintLib.h
@@ -7,6 +7,7 @@
 #include "Blueprints/IDisplayClusterBlueprintAPI.h"
 #include "DisplayClusterBlueprintLib.generated.h"
 
+class ADisplayClusterChromakeyCardActor;
 class ADisplayClusterLightCardActor;
 class ADisplayClusterRootActor;
 
@@ -34,5 +35,9 @@ public:
 
 	/** Gets a list of all light card actors on the level linked to the specified root actor. */
 	UFUNCTION(BlueprintPure, meta = (DisplayName = "DisplayCluster Module API"), Category = "NDisplay")
-	static DISPLAYCLUSTER_API void FindLightCardsForRootActor(ADisplayClusterRootActor* RootActor, TSet<ADisplayClusterLightCardActor*>& OutLightCards);
+	static DISPLAYCLUSTER_API void FindLightCardsForRootActor(const ADisplayClusterRootActor* RootActor, TSet<ADisplayClusterLightCardActor*>& OutLightCards);
+
+	/** Gets a list of all chromakey card actors on the level linked to the specified root actor. */
+	UFUNCTION(BlueprintPure, meta = (DisplayName = "DisplayCluster Module API"), Category = "NDisplay")
+	static DISPLAYCLUSTER_API void FindChromakeyCardsForRootActor(const ADisplayClusterRootActor* RootActor, TSet<ADisplayClusterChromakeyCardActor*>& OutChromakeyCards);
 };
diff --git a/Engine/Plugins/Runtime/nDisplay/Source/DisplayCluster/Public/DisplayClusterChromakeyCardActor.h b/Engine/Plugins/Runtime/nDisplay/Source/DisplayCluster/Public/DisplayClusterChromakeyCardActor.h
new file mode 100644
--- /dev/null
+++ b/Engine/Plugins/Runtime/nDisplay/Source/DisplayCluster/Public/DisplayClusterChromakeyCardActor.h
@@ -0,0 +1,28 @@
+// Copyright Epic Games, Inc. All Rights Reserved.
+
+#pragma once
+
+#include "DisplayClusterLightCardActor.h"
+
+#include "DisplayClusterChromakeyCardActor.generated.h"
+
+class UDisplayClusterICVFXCameraComponent;
+
+UCLASS(Blueprintable, NotPlaceable, DisplayName = "Chromakey Card", HideCategories = (Tick, Physics, Collision, Replication, Cooking, Input, Actor))
+class DISPLAYCLUSTER_API ADisplayClusterChromakeyCardActor : public ADisplayClusterLightCardActor
+{
+	GENERATED_BODY()
+public:
+	ADisplayClusterChromakeyCardActor(const FObjectInitializer& ObjectInitializer);
+
+	virtual void Tick(float DeltaSeconds) override;
+
+	/** Setup chroma key for the owning root actor */
+	void AddToChromakeyLayer(ADisplayClusterRootActor* InRootActor);
+	
+	/** Checks if the given ICVFX camera has chroma key settings supporting this actor */
+	bool IsReferencedByICVFXCamera(UDisplayClusterICVFXCameraComponent* InCamera) const;
+	
+protected:
+	void UpdateChromakeySettings();
+};
diff --git a/Engine/Plugins/Runtime/nDisplay/Source/DisplayClusterConfigurator/Private/DisplayClusterConfiguratorStyle.cpp b/Engine/Plugins/Runtime/nDisplay/Source/DisplayClusterConfigurator/Private/DisplayClusterConfiguratorStyle.cpp
--- a/Engine/Plugins/Runtime/nDisplay/Source/DisplayClusterConfigurator/Private/DisplayClusterConfiguratorStyle.cpp
+++ b/Engine/Plugins/Runtime/nDisplay/Source/DisplayClusterConfigurator/Private/DisplayClusterConfiguratorStyle.cpp
@@ -79,6 +79,9 @@ void FDisplayClusterConfiguratorStyle::Initialize()
 
 		Set("ClassIcon.DisplayClusterLightCardActor.Flag", new IMAGE_BRUSH_SVG("Icons/LightCard/LightCardFlag", Icon16x16));
 		Set("ClassThumbnail.DisplayClusterLightCardActor.Flag", new IMAGE_BRUSH_SVG("Icons/UVLightCard/LightCardFlag", Icon64x64));
+
+		Set("ClassIcon.DisplayClusterChromakeyCardActor", new IMAGE_BRUSH_SVG("Icons/LightCard/LightCardChromakey", Icon16x16));
+		Set("ClassThumbnail.DisplayClusterChromakeyCardActor", new IMAGE_BRUSH_SVG("Icons/UVLightCard/LightCardChromakey", Icon64x64));
 	}
 
 	// Config Editor Tabs
diff --git a/Engine/Plugins/Runtime/nDisplay/Source/DisplayClusterLightCardEditor/Private/DisplayClusterLightCardEditor.cpp b/Engine/Plugins/Runtime/nDisplay/Source/DisplayClusterLightCardEditor/Private/DisplayClusterLightCardEditor.cpp
--- a/Engine/Plugins/Runtime/nDisplay/Source/DisplayClusterLightCardEditor/Private/DisplayClusterLightCardEditor.cpp
+++ b/Engine/Plugins/Runtime/nDisplay/Source/DisplayClusterLightCardEditor/Private/DisplayClusterLightCardEditor.cpp
@@ -2,6 +2,7 @@
 
 #include "DisplayClusterLightCardEditor.h"
 
+#include "DisplayClusterChromakeyCardActor.h"
 #include "DisplayClusterLightCardEditorCommands.h"
 #include "DisplayClusterLightCardEditorUtils.h"
 #include "DisplayClusterLightCardEditorStyle.h"
@@ -24,6 +25,7 @@
 #include "DisplayClusterRootActor.h"
 #include "Blueprints/DisplayClusterBlueprintLib.h"
 #include "Components/DisplayClusterCameraComponent.h"
+#include "Components/DisplayClusterICVFXCameraComponent.h"
 
 #include "ContentBrowserModule.h"
 #include "FileHelpers.h"
@@ -133,6 +135,13 @@ TArray<AActor*> FDisplayClusterLightCardEditor::FindAllManagedActors() const
 		TSet<ADisplayClusterLightCardActor*> RootActorLightCardActors;
 		UDisplayClusterBlueprintLib::FindLightCardsForRootActor(ActiveRootActor.Get(), RootActorLightCardActors);
 
+		// Find chromakey cards
+		{
+			TSet<ADisplayClusterChromakeyCardActor*> ChromakeyCardActors;
+			UDisplayClusterBlueprintLib::FindChromakeyCardsForRootActor(ActiveRootActor.Get(), ChromakeyCardActors);
+			RootActorLightCardActors.Append(reinterpret_cast<TSet<ADisplayClusterLightCardActor*>&>(ChromakeyCardActors));
+		}
+		
 		for (TObjectPtr<ADisplayClusterLightCardActor> LightCardActor : RootActorLightCardActors)
 		{
 			ManagedActors.Add(LightCardActor);
@@ -210,8 +219,12 @@ AActor* FDisplayClusterLightCardEditor::SpawnActor(TSubclassOf<AActor> InActorCl
 	
 	FDisplayClusterLightCardEditorHelper::FSpawnActorArgs SpawnArgs;
 	{
+		const UClass* ClassToUse = InTemplate ? InTemplate->GetClass() : InActorClass.Get();
+		FString ActorNameStr = InActorName.IsNone() ? ClassToUse->GetDisplayNameText().ToString() : InActorName.ToString();
+		ActorNameStr.RemoveSpacesInline();
+		
 		SpawnArgs.ActorClass = InActorClass;
-		SpawnArgs.ActorName = InActorName;
+		SpawnArgs.ActorName = *ActorNameStr;
 		SpawnArgs.RootActor = ActiveRootActor.Get();
 		SpawnArgs.Template = InTemplate;
 		SpawnArgs.Level = InLevel;
diff --git a/Engine/Plugins/Runtime/nDisplay/Source/DisplayClusterLightCardEditor/Private/Settings/DisplayClusterLightCardEditorSettings.cpp b/Engine/Plugins/Runtime/nDisplay/Source/DisplayClusterLightCardEditor/Private/Settings/DisplayClusterLightCardEditorSettings.cpp
--- a/Engine/Plugins/Runtime/nDisplay/Source/DisplayClusterLightCardEditor/Private/Settings/DisplayClusterLightCardEditorSettings.cpp
+++ b/Engine/Plugins/Runtime/nDisplay/Source/DisplayClusterLightCardEditor/Private/Settings/DisplayClusterLightCardEditorSettings.cpp
@@ -39,6 +39,11 @@ FText FDisplayClusterLightCardEditorRecentItem::GetItemDisplayName() const
 	}
 	if (const UObject* Object = ObjectPath.LoadSynchronous())
 	{
+		if (const UClass* Class = Cast<UClass>(Object))
+		{
+			return Class->GetDisplayNameText();
+		}
+		
 		FString ObjectName = Object->GetName();
 		ObjectName.RemoveFromEnd(TEXT("_C"));
 		return FText::FromString(MoveTemp(ObjectName));
@@ -49,6 +54,17 @@ FText FDisplayClusterLightCardEditorRecentItem::GetItemDisplayName() const
 
 const FSlateBrush* FDisplayClusterLightCardEditorRecentItem::GetSlateBrush() const
 {
+	if (ItemType == Type_Dynamic)
+	{
+		if (const UObject* Object = ObjectPath.LoadSynchronous())
+		{
+			if (const UClass* Class = Cast<UClass>(Object))
+			{
+				return FSlateIconFinder::FindIconBrushForClass(Class);
+			}
+		}
+	}
+
 	if (ItemType == Type_LightCard || ItemType == Type_Flag || ItemType == Type_Dynamic)
 	{
 		return FSlateIconFinder::FindIconBrushForClass(ADisplayClusterLightCardActor::StaticClass());
diff --git a/Engine/Plugins/Runtime/nDisplay/Source/DisplayClusterScenePreview/Private/DisplayClusterLightCardEditorHelper.cpp b/Engine/Plugins/Runtime/nDisplay/Source/DisplayClusterScenePreview/Private/DisplayClusterLightCardEditorHelper.cpp
--- a/Engine/Plugins/Runtime/nDisplay/Source/DisplayClusterScenePreview/Private/DisplayClusterLightCardEditorHelper.cpp
+++ b/Engine/Plugins/Runtime/nDisplay/Source/DisplayClusterScenePreview/Private/DisplayClusterLightCardEditorHelper.cpp
@@ -7,6 +7,7 @@
 
 #include "DisplayClusterConfigurationTypes.h"
 
+#include "DisplayClusterChromakeyCardActor.h"
 #include "DisplayClusterLightCardActor.h"
 #include "DisplayClusterRootActor.h"
 #include "Components/DisplayClusterCameraComponent.h"
@@ -868,7 +869,12 @@ AActor* FDisplayClusterLightCardEditorHelper::SpawnStageActor(const FSpawnActorA
 			}
 		}
 
-		if (!bIsPreview)
+		if (ADisplayClusterChromakeyCardActor* ChromakeyCardActor = Cast<ADisplayClusterChromakeyCardActor>(NewActor))
+		{
+			// Only add to chroma key layer, not light card layer
+			ChromakeyCardActor->AddToChromakeyLayer(RootActor);
+		}
+		else if (!bIsPreview)
 		{
 			AddLightCardsToRootActor({ NewLightCard }, RootActor, InSpawnArgs.AddLightCardArgs);
 		}
