nDisplay: fixed mpcdi rendering preview and cleanup

diff --git a/Engine/Plugins/Runtime/nDisplay/Source/DisplayCluster/Private/Render/Viewport/DisplayClusterViewport.cpp b/Engine/Plugins/Runtime/nDisplay/Source/DisplayCluster/Private/Render/Viewport/DisplayClusterViewport.cpp
--- a/Engine/Plugins/Runtime/nDisplay/Source/DisplayCluster/Private/Render/Viewport/DisplayClusterViewport.cpp
+++ b/Engine/Plugins/Runtime/nDisplay/Source/DisplayCluster/Private/Render/Viewport/DisplayClusterViewport.cpp
@@ -73,7 +73,7 @@ FDisplayClusterViewport::~FDisplayClusterViewport()
 	// Handle projection policy EndScene event
 	HandleEndScene();
 
-	// Release Projection policy after HandleEndScene()
+	// Handle projection policy event
 	ProjectionPolicy.Reset();
 	UninitializedProjectionPolicy.Reset();
 
