nDisplay: iteration over world actors looking for hidden primitives now only happens when not in
game.

As a consequence, it will now not render visualization components. That was a hack added in the past to support Text3d, but Text3d will be updated to avoid this issue instead.

Since the iteration is also needed for in-editor preview, parallelized it for higher performance. For CitySampleVP it went from 18 ms to 2 ms.

Added a missing member variable initialization.

Added a convenience function to programmatically freeze/unfreeze the outers.

diff --git a/Engine/Plugins/Runtime/nDisplay/Source/DisplayCluster/Private/Game/EngineClasses/Scene/DisplayClusterRootActor.cpp b/Engine/Plugins/Runtime/nDisplay/Source/DisplayCluster/Private/Game/EngineClasses/Scene/DisplayClusterRootActor.cpp
--- a/Engine/Plugins/Runtime/nDisplay/Source/DisplayCluster/Private/Game/EngineClasses/Scene/DisplayClusterRootActor.cpp
+++ b/Engine/Plugins/Runtime/nDisplay/Source/DisplayCluster/Private/Game/EngineClasses/Scene/DisplayClusterRootActor.cpp
@@ -2,6 +2,7 @@
 
 #include "DisplayClusterRootActor.h"
 
+#include "Async/ParallelFor.h"
 #include "Components/SceneComponent.h"
 #include "Components/DisplayClusterOriginComponent.h"
 #include "Components/DisplayClusterCameraComponent.h"
@@ -45,19 +46,13 @@
 #include "Components/DisplayClusterStageGeometryComponent.h"
 #include "UObject/Package.h"
 
+
 #if WITH_EDITOR
 #include "IConcertSyncClientModule.h"
 #include "IConcertClientWorkspace.h"
 #include "IConcertSyncClient.h"
 #endif
 
-static TAutoConsoleVariable<int32> CVarShowVisualizationComponents(
-	TEXT("nDisplay.render.show.visualizationcomponents"),
-	1,
-	TEXT("Whether to show or hide visualization mesh components \n")
-	TEXT("0 : Hides visualization components \n")
-	TEXT("1 : Shows visualization components\n")
-);
 
 ADisplayClusterRootActor::ADisplayClusterRootActor(const FObjectInitializer& ObjectInitializer)
 	: Super(ObjectInitializer)
@@ -396,15 +391,10 @@ int ADisplayClusterRootActor::GetInnerFrustumPriority(const FString& InnerFrustu
 }
 
 template <typename TComp>
-void ImplCollectChildVisualizationOrHiddenComponents(TSet<FPrimitiveComponentId>& OutPrimitives, TComp* pComp, bool bCollectVisualizationComponents = true, bool bCollectHiddenComponents = true)
+void ImplCollectChildHiddenComponents(TSet<FPrimitiveComponentId>& OutPrimitives, TComp* pComp)
 {
 #if WITH_EDITOR
 
-	if (!bCollectVisualizationComponents && !bCollectHiddenComponents)
-	{
-		return;
-	}
-
 	USceneComponent* SceneComp = Cast<USceneComponent>(pComp);
 	if (SceneComp)
 	{
@@ -412,8 +402,7 @@ void ImplCollectChildVisualizationOrHiddenComponents(TSet<FPrimitiveComponentId>
 		SceneComp->GetChildrenComponents(false, Childrens);
 		for (USceneComponent* ChildIt : Childrens)
 		{
-			if ((bCollectVisualizationComponents && ChildIt->IsVisualizationComponent()) 
-				|| (bCollectHiddenComponents && ChildIt->bHiddenInGame))
+			if (ChildIt->bHiddenInGame)
 			{
 				UPrimitiveComponent* PrimComp = Cast<UPrimitiveComponent>(ChildIt);
 				if (PrimComp)
@@ -451,7 +440,7 @@ void ADisplayClusterRootActor::GetTypedPrimitives(TSet<FPrimitiveComponentId>& O
 
 						if (bCollectChildrenVisualizationComponent)
 						{
-							ImplCollectChildVisualizationOrHiddenComponents(OutPrimitives, CompIt);
+							ImplCollectChildHiddenComponents(OutPrimitives, CompIt);
 						}
 						break;
 					}
@@ -467,7 +456,7 @@ void ADisplayClusterRootActor::GetTypedPrimitives(TSet<FPrimitiveComponentId>& O
 
 				if (bCollectChildrenVisualizationComponent)
 				{
-					ImplCollectChildVisualizationOrHiddenComponents(OutPrimitives, CompIt);
+					ImplCollectChildHiddenComponents(OutPrimitives, CompIt);
 				}
 			}
 
@@ -485,6 +474,8 @@ bool ADisplayClusterRootActor::FindPrimitivesByName(const TArray<FString>& InNam
 // Gather components not rendered in game
 bool ADisplayClusterRootActor::GetHiddenInGamePrimitives(TSet<FPrimitiveComponentId>& OutPrimitives)
 {
+	TRACE_CPUPROFILER_EVENT_SCOPE(DCRootActor_GetHiddenInGamePrimitives);
+
 	check(IsInGameThread());
 
 	OutPrimitives.Empty();
@@ -506,8 +497,6 @@ bool ADisplayClusterRootActor::GetHiddenInGamePrimitives(TSet<FPrimitiveComponen
 
 #if WITH_EDITOR
 
-	const bool bHideVisualizationComponents = !CVarShowVisualizationComponents.GetValueOnGameThread();
-
 	// Hide visualization and hidden components from RootActor
 	{
 		TArray<UPrimitiveComponent*> PrimitiveComponents;
@@ -515,40 +504,107 @@ bool ADisplayClusterRootActor::GetHiddenInGamePrimitives(TSet<FPrimitiveComponen
 
 		for (UPrimitiveComponent* CompIt : PrimitiveComponents)
 		{
-			if ((bHideVisualizationComponents && CompIt->IsVisualizationComponent()) || CompIt->bHiddenInGame)
+			if (CompIt->bHiddenInGame)
 			{
 				OutPrimitives.Add(CompIt->ComponentId);
 			}
 
-			ImplCollectChildVisualizationOrHiddenComponents(OutPrimitives, CompIt, bHideVisualizationComponents);
+			ImplCollectChildHiddenComponents(OutPrimitives, CompIt);
 		}
 	}
 
 	// Hide visualization and hidden components from scene
-	if (const UWorld* CurrentWorld = GetWorld())
+
+	const UWorld* CurrentWorld = GetWorld();
+
+	if (CurrentWorld && !CurrentWorld->IsGameWorld()) // Note: We should only have to do this for preview in editor (not already game worlds).
 	{
-		// Iterate over all actors, looking for editor components.
-		for (const TWeakObjectPtr<AActor>& WeakActor : FActorRange(CurrentWorld))
+		TArray<AActor*> Actors;
+		{
+			int32 NumActors = 0;
+
+			for (const ULevel* Level : CurrentWorld->GetLevels())
+			{
+				if (Level && Level->bIsVisible)
+				{
+					NumActors += Level->Actors.Num();
+				}
+			}
+
+			// Presize the array
+			Actors.Reserve(NumActors);
+
+			// Fill the array
+			for (const ULevel* Level : CurrentWorld->GetLevels())
+			{
+				if (Level && Level->bIsVisible)
+				{
+					Actors.Append(Level->Actors);
+				}
+			}
+		}
+
+		// Create as many threads as cores, minus an arbitrary number of reserved cores to avoid starving other subsystems
+		int32 NumIterThreads;
 		{
-			if (AActor* Actor = WeakActor.Get())
+			constexpr int32 ReservedCores = 4;
+			static const int32 NumberOfCores = FPlatformMisc::NumberOfCores();
+
+			NumIterThreads = NumberOfCores - ReservedCores;
+
+			// Make sure there are enough actors to make it worth a new thread
+			constexpr int32 MinActorsPerThread = 64;
+			NumIterThreads = FMath::Min(NumIterThreads, Actors.Num() / MinActorsPerThread);
+
+			// There should be at least one thread
+			NumIterThreads = FMath::Max(NumIterThreads, 1);
+		}
+
+		// Allocate primitive sets for each thread
+		TArray<TSet<FPrimitiveComponentId>> PrimitiveComponentsArray;
+		PrimitiveComponentsArray.AddDefaulted(NumIterThreads);
+
+		// Start the iteration parallel threads
+		ParallelFor(NumIterThreads, [NumIterThreads, CurrentWorld, &Actors, &PrimitiveComponentsArray](int32 Index)
 			{
-				// do not render hidden in game actors
-				const bool bActorHideInGame = Actor->IsHidden();
+				// Using inline allocator for efficiency
+				constexpr int32 MaxExpectedComponentsPerActor = 64;
+				TArray<UPrimitiveComponent*, TInlineAllocator<MaxExpectedComponentsPerActor>> PrimitiveComponents;
 
-				TArray<UPrimitiveComponent*> PrimitiveComponents;
-				Actor->GetComponents(PrimitiveComponents);
+				// The thread index is our starting actor index
+				int32 ActorIdx = Index;
 
-				for (UPrimitiveComponent* PrimComp : PrimitiveComponents)
+				while (ActorIdx < Actors.Num())
 				{
-					if ((bHideVisualizationComponents && PrimComp->IsVisualizationComponent()) 
-						|| ((bActorHideInGame || PrimComp->bHiddenInGame) && !PrimComp->bCastHiddenShadow))
+					const AActor* Actor = Actors[ActorIdx];
+
+					if (IsValid(Actor))
 					{
-						OutPrimitives.Add(PrimComp->ComponentId);
+						Actor->GetComponents(PrimitiveComponents);
+
+						for (UPrimitiveComponent* PrimComp : PrimitiveComponents)
+						{
+							if (((Actor->IsHidden() || PrimComp->bHiddenInGame) && !PrimComp->bCastHiddenShadow))
+							{
+								PrimitiveComponentsArray[Index].Add(PrimComp->ComponentId);
+							}
+
+							ImplCollectChildHiddenComponents(PrimitiveComponentsArray[Index], PrimComp);
+						}
+
+						// Empty w/o ever shrinking, for efficiency
+						PrimitiveComponents.Empty(PrimitiveComponents.Max());
 					}
 
-					ImplCollectChildVisualizationOrHiddenComponents(OutPrimitives, PrimComp, bHideVisualizationComponents);
+					// Jump to every other NumIterThreads actors.
+					ActorIdx += NumIterThreads;
 				}
-			}
+			});
+
+		// Join all the found primitives arrays
+		for (const TSet<FPrimitiveComponentId>& PrimitiveComponents : PrimitiveComponentsArray)
+		{
+			OutPrimitives.Append(PrimitiveComponents);
 		}
 	}
 
@@ -1215,6 +1271,21 @@ bool ADisplayClusterRootActor::SetReplaceTextureFlagForAllViewports(bool bReplac
 	return true;
 }
 
+bool ADisplayClusterRootActor::SetFreezeOuterViewports(bool bEnable)
+{
+	UDisplayClusterConfigurationData* ConfigData = GetConfigData();
+
+	if (!ConfigData)
+	{
+		UE_LOG(LogDisplayClusterGame, Warning, TEXT("ADisplayClusterRootActor::SetFreezeOuterViewports failed because ConfigData was null"));
+		return false;
+	}
+
+	ConfigData->StageSettings.bFreezeRenderOuterViewports = bEnable;
+
+	return true;
+}
+
 void ADisplayClusterRootActor::AddReferencedObjects(UObject* InThis, FReferenceCollector& Collector)
 {
 	ADisplayClusterRootActor* This = CastChecked<ADisplayClusterRootActor>(InThis);
diff --git a/Engine/Plugins/Runtime/nDisplay/Source/DisplayCluster/Public/DisplayClusterRootActor.h b/Engine/Plugins/Runtime/nDisplay/Source/DisplayCluster/Public/DisplayClusterRootActor.h
--- a/Engine/Plugins/Runtime/nDisplay/Source/DisplayCluster/Public/DisplayClusterRootActor.h
+++ b/Engine/Plugins/Runtime/nDisplay/Source/DisplayCluster/Public/DisplayClusterRootActor.h
@@ -144,6 +144,9 @@ public:
 	UFUNCTION(BlueprintCallable, Category = "NDisplay|Render")
 	bool SetReplaceTextureFlagForAllViewports(bool bReplace);
 
+	UFUNCTION(BlueprintCallable, Category = "NDisplay|Render")
+	bool SetFreezeOuterViewports(bool bEnable);
+
 	template <typename TComp>
 	TComp* GetComponentByName(const FString& ComponentName) const
 	{
@@ -377,7 +380,7 @@ private:
 	TMap<FString, TObjectPtr<UDisplayClusterPreviewComponent>> PreviewComponents;
 
 	UPROPERTY(Transient)
-	bool bDeferPreviewGeneration;
+	bool bDeferPreviewGeneration = false;
 #endif
 
 #if WITH_EDITOR
