nDisplay MRQ: Fixed nDisplay MRQ support for user defined eye points being broken/unsupported

diff --git a/Engine/Plugins/Runtime/nDisplay/Source/DisplayClusterMoviePipeline/Private/DisplayClusterMoviePipelineViewportPass.cpp b/Engine/Plugins/Runtime/nDisplay/Source/DisplayClusterMoviePipeline/Private/DisplayClusterMoviePipelineViewportPass.cpp
--- a/Engine/Plugins/Runtime/nDisplay/Source/DisplayClusterMoviePipeline/Private/DisplayClusterMoviePipelineViewportPass.cpp
+++ b/Engine/Plugins/Runtime/nDisplay/Source/DisplayClusterMoviePipeline/Private/DisplayClusterMoviePipelineViewportPass.cpp
@@ -521,12 +521,18 @@ IDisplayClusterViewport* UDisplayClusterMoviePipelineViewportPassBase::GetAndCal
 
 	if (IDisplayClusterViewport* DCViewport = ViewportManager->FindViewport(InViewportId))
 	{
+		// Get camera ID assigned to the viewport
+		const FString CameraId = DCViewport->GetRenderSettings().CameraId;
+
 		// Force custom viewport size and settings:
 		{
 			FDisplayClusterViewport_RenderSettings RenderSettings;
 			RenderSettings.Rect = FIntRect(FIntPoint(0, 0), InDestSize);
 			RenderSettings.CaptureMode = EDisplayClusterViewportCaptureMode::MoviePipeline;
 
+			// Store camera name for next frames (because UpdateConfiguration called once)
+			RenderSettings.CameraId = CameraId;
+
 			if (!bFrameWarpBlend)
 			{
 				// Create internal resources only when warpblend will be used
@@ -556,9 +562,6 @@ IDisplayClusterViewport* UDisplayClusterMoviePipelineViewportPassBase::GetAndCal
 
 			UDisplayClusterCameraComponent* ViewCamera = nullptr;
 			{
-				// Get camera ID assigned to the viewport
-				const FString& CameraId = DCViewport->GetRenderSettings().CameraId;
-
 				// Get camera component assigned to the viewport (or default camera if nothing assigned)
 				ViewCamera = (CameraId.IsEmpty() ?
 					DCRootActor->GetDefaultCamera() :
