nDisplay: Light Cards can now be distinctly configured as flags and will be listed in the ICVFX outliner. Only valid on newly added flags.

diff --git a/Engine/Plugins/Runtime/nDisplay/Source/DisplayCluster/Private/Game/EngineClasses/Scene/DisplayClusterLightCardActor.cpp b/Engine/Plugins/Runtime/nDisplay/Source/DisplayCluster/Private/Game/EngineClasses/Scene/DisplayClusterLightCardActor.cpp
--- a/Engine/Plugins/Runtime/nDisplay/Source/DisplayCluster/Private/Game/EngineClasses/Scene/DisplayClusterLightCardActor.cpp
+++ b/Engine/Plugins/Runtime/nDisplay/Source/DisplayCluster/Private/Game/EngineClasses/Scene/DisplayClusterLightCardActor.cpp
@@ -445,6 +445,11 @@ void ADisplayClusterLightCardActor::UpdateUVIndicator()
 #endif // WITH_EDITOR
 }
 
+void ADisplayClusterLightCardActor::SetIsLightCardFlag(bool bNewFlagValue)
+{
+	bIsLightCardFlag = bNewFlagValue;
+}
+
 void ADisplayClusterLightCardActor::ShowLightCardLabel(bool bValue, float ScaleValue, ADisplayClusterRootActor* InRootActor)
 {
 	if (IsUVActor())
diff --git a/Engine/Plugins/Runtime/nDisplay/Source/DisplayCluster/Public/DisplayClusterLightCardActor.h b/Engine/Plugins/Runtime/nDisplay/Source/DisplayCluster/Public/DisplayClusterLightCardActor.h
--- a/Engine/Plugins/Runtime/nDisplay/Source/DisplayCluster/Public/DisplayClusterLightCardActor.h
+++ b/Engine/Plugins/Runtime/nDisplay/Source/DisplayCluster/Public/DisplayClusterLightCardActor.h
@@ -96,6 +96,12 @@ public:
 	/** Updates the UV Indicator */
 	void UpdateUVIndicator();
 
+	/** Configures this light card as a flag */
+	void SetIsLightCardFlag(bool bNewFlagValue);
+
+	/** If this light card is considered a flag */
+	bool IsLightCardFlag() const { return bIsLightCardFlag; }
+	
 	/** Show or hide the light card label  */
 	void ShowLightCardLabel(bool bValue, float ScaleValue, ADisplayClusterRootActor* InRootActor);
 
@@ -287,6 +293,10 @@ protected:
 	/** The current owner of the light card */
 	TWeakObjectPtr<ADisplayClusterRootActor> RootActorOwner;
 
+	/** Indicates this light card should be considered a flag */
+	UPROPERTY()
+	bool bIsLightCardFlag = false;
+	
 private:
 	/** Stores the user translucency value when labels are displayed */
 	TOptional<int32> SavedTranslucencySortPriority;
diff --git a/Engine/Plugins/Runtime/nDisplay/Source/DisplayClusterLightCardEditor/Private/DisplayClusterLightCardEditor.cpp b/Engine/Plugins/Runtime/nDisplay/Source/DisplayClusterLightCardEditor/Private/DisplayClusterLightCardEditor.cpp
--- a/Engine/Plugins/Runtime/nDisplay/Source/DisplayClusterLightCardEditor/Private/DisplayClusterLightCardEditor.cpp
+++ b/Engine/Plugins/Runtime/nDisplay/Source/DisplayClusterLightCardEditor/Private/DisplayClusterLightCardEditor.cpp
@@ -417,6 +417,8 @@ void FDisplayClusterLightCardEditor::AddNewFlag()
 			// When adding a new lightcard, usually the desired location is in the middle of the viewport
 			CenterActorInView(NewLightCard);
 		}
+
+		NewLightCard->SetIsLightCardFlag(true);
 		
 		FDisplayClusterLightCardEditorRecentItem RecentlyPlacedItem;
 		RecentlyPlacedItem.ObjectPath = NewLightCard->GetClass();
diff --git a/Engine/Plugins/Runtime/nDisplay/Source/DisplayClusterLightCardEditor/Private/Outliner/SDisplayClusterLightCardOutliner.cpp b/Engine/Plugins/Runtime/nDisplay/Source/DisplayClusterLightCardEditor/Private/Outliner/SDisplayClusterLightCardOutliner.cpp
--- a/Engine/Plugins/Runtime/nDisplay/Source/DisplayClusterLightCardEditor/Private/Outliner/SDisplayClusterLightCardOutliner.cpp
+++ b/Engine/Plugins/Runtime/nDisplay/Source/DisplayClusterLightCardEditor/Private/Outliner/SDisplayClusterLightCardOutliner.cpp
@@ -259,15 +259,18 @@ void SDisplayClusterLightCardOutliner::CreateWorldOutliner()
 			{
 				if (const AActor* Actor = ActorItem->Actor.Get())
 				{
+					const FString ClassDisplayName = Actor->GetClass()->GetDisplayNameText().ToString();
+					
 					if (const ADisplayClusterLightCardActor* LightCardActor = Cast<ADisplayClusterLightCardActor>(Actor))
 					{
-						if (LightCardActor->bIsUVLightCard)
-						{
-							return TEXT("UV Light Card");
-						}
+						const FString LightCardType = LightCardActor->IsLightCardFlag() ? TEXT("Flag") : ClassDisplayName;
+						
+						return LightCardActor->bIsUVLightCard ?
+							FString::Printf(TEXT("UV %s"), *LightCardType) :
+							FString::Printf(TEXT("%s"), *LightCardType);
 					}
 					
-					return Actor->GetClass()->GetDisplayNameText().ToString();
+					return ClassDisplayName;
 				}
 			
 				return FString();
