Updating DisplayClusterShadersPostprocess_Blur to use shader parameter structs.

diff --git a/Engine/Plugins/Runtime/nDisplay/Source/DisplayClusterShaders/Private/Shaders/DisplayClusterShadersPostprocess_Blur.cpp b/Engine/Plugins/Runtime/nDisplay/Source/DisplayClusterShaders/Private/Shaders/DisplayClusterShadersPostprocess_Blur.cpp
--- a/Engine/Plugins/Runtime/nDisplay/Source/DisplayClusterShaders/Private/Shaders/DisplayClusterShadersPostprocess_Blur.cpp
+++ b/Engine/Plugins/Runtime/nDisplay/Source/DisplayClusterShaders/Private/Shaders/DisplayClusterShadersPostprocess_Blur.cpp
@@ -52,37 +52,29 @@ public:
 	}
 };
 
-template<uint32 ShaderType>
-class TPicpBlurPostProcessPS : public FGlobalShader
+class FPicpBlurPostProcessPS : public FGlobalShader
 {
-	DECLARE_SHADER_TYPE(TPicpBlurPostProcessPS, Global);
+public:
+	BEGIN_SHADER_PARAMETER_STRUCT(FParameters, )
+		SHADER_PARAMETER(FVector2f, SampleOffset)
+		SHADER_PARAMETER(int32, KernelRadius)
 
-	static bool ShouldCompilePermutation(const FGlobalShaderPermutationParameters& Parameters)
+		SHADER_PARAMETER_TEXTURE(Texture2D, SrcTexture)
+		SHADER_PARAMETER_SAMPLER(SamplerState, BilinearClampTextureSampler)
+	END_SHADER_PARAMETER_STRUCT()
+
+	FPicpBlurPostProcessPS() = default;
+	FPicpBlurPostProcessPS(const ShaderMetaType::CompiledShaderInitializerType& Initializer)
+		: FGlobalShader(Initializer)
 	{
-		return true;
 	}
 
-	/** Default constructor. */
-	TPicpBlurPostProcessPS() = default;
-
-	LAYOUT_FIELD(FShaderResourceParameter, SrcTextureParameter);
-	LAYOUT_FIELD(FShaderResourceParameter, BilinearClampTextureSamplerParameter);
-	LAYOUT_FIELD(FShaderParameter, SampleOffsetParameter);
-	LAYOUT_FIELD(FShaderParameter, KernelRadiusParameter);
-
-public:
-
-	/** Initialization constructor. */
-	TPicpBlurPostProcessPS(const ShaderMetaType::CompiledShaderInitializerType& Initializer)
-		: FGlobalShader(Initializer)
+	static bool ShouldCompilePermutation(const FGlobalShaderPermutationParameters& Parameters)
 	{
-		SrcTextureParameter.Bind(Initializer.ParameterMap, TEXT("SrcTexture"));
-		BilinearClampTextureSamplerParameter.Bind(Initializer.ParameterMap, TEXT("BilinearClampTextureSampler"));
-		SampleOffsetParameter.Bind(Initializer.ParameterMap, TEXT("SampleOffset"));
-		KernelRadiusParameter.Bind(Initializer.ParameterMap, TEXT("KernelRadius"));
+		return true;
 	}
 
-	static void ModifyCompilationEnvironment(const FGlobalShaderPermutationParameters& Parameters, FShaderCompilerEnvironment& OutEnvironment)
+	static void ModifyCompilationEnvironment(uint32 ShaderType, const FGlobalShaderPermutationParameters& Parameters, FShaderCompilerEnvironment& OutEnvironment)
 	{
 		FGlobalShader::ModifyCompilationEnvironment(Parameters, OutEnvironment);
 
@@ -100,15 +92,20 @@ public:
 			break;
 		}
 	}
+};
 
-	void SetParameters(FRHICommandListImmediate& RHICmdList, FRHITexture2D* SourceTexture, const FVector2f& SampleOffset, const int32 KernelRadius)
-	{
-		FRHIPixelShader* ShaderRHI = RHICmdList.GetBoundPixelShader();
+template<uint32 ShaderType>
+class TPicpBlurPostProcessPS : public FPicpBlurPostProcessPS
+{
+public:
+	DECLARE_SHADER_TYPE(TPicpBlurPostProcessPS, Global);
+	SHADER_USE_PARAMETER_STRUCT(TPicpBlurPostProcessPS, FPicpBlurPostProcessPS);
 
-		SetSamplerParameter(RHICmdList, ShaderRHI, BilinearClampTextureSamplerParameter, TStaticSamplerState<SF_Bilinear, AM_Clamp, AM_Clamp, AM_Clamp>::GetRHI());
-		SetTextureParameter(RHICmdList, ShaderRHI, SrcTextureParameter, SourceTexture);
-		SetShaderValue(RHICmdList, ShaderRHI, SampleOffsetParameter, SampleOffset);
-		SetShaderValue(RHICmdList, ShaderRHI, KernelRadiusParameter, KernelRadius);
+	using FParameters = FPicpBlurPostProcessPS::FParameters;
+
+	static void ModifyCompilationEnvironment(const FGlobalShaderPermutationParameters& Parameters, FShaderCompilerEnvironment& OutEnvironment)
+	{
+		FPicpBlurPostProcessPS::ModifyCompilationEnvironment(ShaderType, Parameters, OutEnvironment);
 	}
 };
 
@@ -138,6 +135,11 @@ static void PicpBlurPostProcess_RenderThread(
 	TShaderMapRef<FDirectProjectionVS>                VertexShader(GlobalShaderMap);
 	TShaderMapRef<TPicpBlurPostProcessPS<ShaderType>> PixelShader(GlobalShaderMap);
 
+	FPicpBlurPostProcessPS::FParameters PsParameters{};
+	PsParameters.SrcTexture = InShaderTexture;
+	PsParameters.BilinearClampTextureSampler = TStaticSamplerState<SF_Bilinear, AM_Clamp, AM_Clamp, AM_Clamp>::GetRHI();
+	PsParameters.KernelRadius = InSettings.KernelRadius;
+
 	FRHIRenderPassInfo RPInfo(OutRenderTargetableTexture, ERenderTargetActions::DontLoad_Store);
 	RHICmdList.Transition(FRHITransitionInfo(OutRenderTargetableTexture, ERHIAccess::Unknown, ERHIAccess::RTV));
 
@@ -155,7 +157,9 @@ static void PicpBlurPostProcess_RenderThread(
 		SetGraphicsPipelineState(RHICmdList, GraphicsPSOInit, 0);
 
 		// Blur X
-		PixelShader->SetParameters(RHICmdList, InShaderTexture, FVector2f(InSettings.KernelScale / TargetSizeXY.X, 0.0f), InSettings.KernelRadius);
+		PsParameters.SampleOffset = FVector2f(InSettings.KernelScale / TargetSizeXY.X, 0.0f);
+		SetShaderParameters(RHICmdList, PixelShader, PixelShader.GetPixelShader(), PsParameters);
+
 		FPixelShaderUtils::DrawFullscreenQuad(RHICmdList, 1);
 	
 		RHICmdList.EndRenderPass();
@@ -176,7 +180,10 @@ static void PicpBlurPostProcess_RenderThread(
 
 		// Blur Y
 		SetGraphicsPipelineState(RHICmdList, GraphicsPSOInit, 0);
-		PixelShader->SetParameters(RHICmdList, InShaderTexture, FVector2f(0.0f, InSettings.KernelScale / TargetSizeXY.Y), InSettings.KernelRadius);
+
+		PsParameters.SampleOffset = FVector2f(0.0f, InSettings.KernelScale / TargetSizeXY.Y);
+		SetShaderParameters(RHICmdList, PixelShader, PixelShader.GetPixelShader(), PsParameters);
+
 		FPixelShaderUtils::DrawFullscreenQuad(RHICmdList, 1);
 	}
 	RHICmdList.EndRenderPass();
