Fixing compile errors from running IWYU on the entire engine.

diff --git a/Engine/Plugins/AI/MLAdapter/Source/MLAdapter/Public/RPCWrapper/MsgPack.h b/Engine/Plugins/AI/MLAdapter/Source/MLAdapter/Public/RPCWrapper/MsgPack.h
--- a/Engine/Plugins/AI/MLAdapter/Source/MLAdapter/Public/RPCWrapper/MsgPack.h
+++ b/Engine/Plugins/AI/MLAdapter/Source/MLAdapter/Public/RPCWrapper/MsgPack.h
@@ -2,6 +2,8 @@
 
 #pragma once
 
+#include "Containers/Array.h"
+
 #if WITH_RPCLIB
 #include "rpclib_includes.h"
 
diff --git a/Engine/Plugins/Developer/Concert/ConcertApp/MultiUserServer/Source/MultiUserServer/Private/Widgets/Clients/Logging/Filter/ConcertFrontendLogFilter_BaseSetSelection.h b/Engine/Plugins/Developer/Concert/ConcertApp/MultiUserServer/Source/MultiUserServer/Private/Widgets/Clients/Logging/Filter/ConcertFrontendLogFilter_BaseSetSelection.h
--- a/Engine/Plugins/Developer/Concert/ConcertApp/MultiUserServer/Source/MultiUserServer/Private/Widgets/Clients/Logging/Filter/ConcertFrontendLogFilter_BaseSetSelection.h
+++ b/Engine/Plugins/Developer/Concert/ConcertApp/MultiUserServer/Source/MultiUserServer/Private/Widgets/Clients/Logging/Filter/ConcertFrontendLogFilter_BaseSetSelection.h
@@ -7,6 +7,8 @@
 #include "Widgets/Util/Filter/ConcertFrontendFilter.h"
 #include "Framework/MultiBox/MultiBoxBuilder.h"
 
+struct FConcertLogEntry;
+
 #define LOCTEXT_NAMESPACE "UnrealMultiUserUI.BaseSetSelection"
 
 namespace UE::MultiUserServer
diff --git a/Engine/Plugins/Developer/Concert/ConcertSync/ConcertSyncServer/Source/ConcertSyncServer/ConcertSyncServer.Build.cs b/Engine/Plugins/Developer/Concert/ConcertSync/ConcertSyncServer/Source/ConcertSyncServer/ConcertSyncServer.Build.cs
--- a/Engine/Plugins/Developer/Concert/ConcertSync/ConcertSyncServer/Source/ConcertSyncServer/ConcertSyncServer.Build.cs
+++ b/Engine/Plugins/Developer/Concert/ConcertSync/ConcertSyncServer/Source/ConcertSyncServer/ConcertSyncServer.Build.cs
@@ -16,6 +16,7 @@ namespace UnrealBuildTool.Rules
 					"CoreUObject",
 					"Concert",
 					"ConcertSyncCore",
+					"Launch",
 				}
 			);
 
diff --git a/Engine/Plugins/Developer/Concert/ConcertSync/ConcertSyncServer/Source/ConcertSyncServer/Public/ConcertSyncServerLoop.inl b/Engine/Plugins/Developer/Concert/ConcertSync/ConcertSyncServer/Source/ConcertSyncServer/Public/ConcertSyncServerLoop.inl
--- a/Engine/Plugins/Developer/Concert/ConcertSync/ConcertSyncServer/Source/ConcertSyncServer/Public/ConcertSyncServerLoop.inl
+++ b/Engine/Plugins/Developer/Concert/ConcertSync/ConcertSyncServer/Source/ConcertSyncServer/Public/ConcertSyncServerLoop.inl
@@ -1,5 +1,7 @@
 // Copyright Epic Games, Inc. All Rights Reserved.
 
+#pragma once
+
 #include "ConcertSyncServerLoop.h"
 #include "ConcertSettings.h"
 #include "IConcertServer.h"
diff --git a/Engine/Plugins/Developer/Concert/ConcertSync/ConcertSyncTest/Source/ConcertSyncTest/Private/HistoryEdition/HistoryTestUtil.h b/Engine/Plugins/Developer/Concert/ConcertSync/ConcertSyncTest/Source/ConcertSyncTest/Private/HistoryEdition/HistoryTestUtil.h
--- a/Engine/Plugins/Developer/Concert/ConcertSync/ConcertSyncTest/Source/ConcertSyncTest/Private/HistoryEdition/HistoryTestUtil.h
+++ b/Engine/Plugins/Developer/Concert/ConcertSync/ConcertSyncTest/Source/ConcertSyncTest/Private/HistoryEdition/HistoryTestUtil.h
@@ -3,6 +3,7 @@
 #pragma once
 
 #include "CoreMinimal.h"
+#include "HistoryEdition/ActivityDependencyGraph.h"
 #include "HistoryEdition/ActivityGraphIDs.h"
 #include "HistoryEdition/ActivityNode.h"
 #include "Misc/AutomationTest.h"
diff --git a/Engine/Plugins/Editor/GuidedTutorials/Source/IntroTutorials/Private/STutorialsBrowser.h b/Engine/Plugins/Editor/GuidedTutorials/Source/IntroTutorials/Private/STutorialsBrowser.h
--- a/Engine/Plugins/Editor/GuidedTutorials/Source/IntroTutorials/Private/STutorialsBrowser.h
+++ b/Engine/Plugins/Editor/GuidedTutorials/Source/IntroTutorials/Private/STutorialsBrowser.h
@@ -3,6 +3,7 @@
 #pragma once
 
 #include "CoreMinimal.h"
+#include "EditorTutorial.h"
 #include "SlateFwd.h"
 #include "Layout/Visibility.h"
 #include "Widgets/DeclarativeSyntaxSupport.h"
diff --git a/Engine/Plugins/Editor/WaveformEditor/Source/WaveformTransformationsWidgets/Private/WaveformTransformationDurationRenderer.h b/Engine/Plugins/Editor/WaveformEditor/Source/WaveformTransformationsWidgets/Private/WaveformTransformationDurationRenderer.h
--- a/Engine/Plugins/Editor/WaveformEditor/Source/WaveformTransformationsWidgets/Private/WaveformTransformationDurationRenderer.h
+++ b/Engine/Plugins/Editor/WaveformEditor/Source/WaveformTransformationsWidgets/Private/WaveformTransformationDurationRenderer.h
@@ -2,6 +2,9 @@
 
 #pragma once
 
+#include "Layout/Geometry.h"
+#include "Rendering/DrawElements.h"
+#include "Styling/AppStyle.h"
 #include "WaveformTransformationRendererBase.h"
 
 
diff --git a/Engine/Plugins/Enterprise/DataprepEditor/Source/DataprepCore/Public/DataprepBindingCommandChange.h b/Engine/Plugins/Enterprise/DataprepEditor/Source/DataprepCore/Public/DataprepBindingCommandChange.h
--- a/Engine/Plugins/Enterprise/DataprepEditor/Source/DataprepCore/Public/DataprepBindingCommandChange.h
+++ b/Engine/Plugins/Enterprise/DataprepEditor/Source/DataprepCore/Public/DataprepBindingCommandChange.h
@@ -2,6 +2,7 @@
 
 #pragma once
 
+#include "Delegates/Delegate.h"
 #include "Misc/Change.h"
 
 template<class ObservedObjectType, typename DelatgateToBindType>
diff --git a/Engine/Plugins/Enterprise/DataprepEditor/Source/DataprepCore/Public/DataprepContentProducer.h b/Engine/Plugins/Enterprise/DataprepEditor/Source/DataprepCore/Public/DataprepContentProducer.h
--- a/Engine/Plugins/Enterprise/DataprepEditor/Source/DataprepCore/Public/DataprepContentProducer.h
+++ b/Engine/Plugins/Enterprise/DataprepEditor/Source/DataprepCore/Public/DataprepContentProducer.h
@@ -8,6 +8,7 @@
 #include "IDataprepProgressReporter.h"
 
 #include "Engine/World.h"
+#include "UObject/Package.h"
 
 #include "DataprepContentProducer.generated.h"
 
diff --git a/Engine/Plugins/Enterprise/DataprepEditor/Source/DataprepEditor/Private/DataprepEditorOutlinerMode.h b/Engine/Plugins/Enterprise/DataprepEditor/Source/DataprepEditor/Private/DataprepEditorOutlinerMode.h
--- a/Engine/Plugins/Enterprise/DataprepEditor/Source/DataprepEditor/Private/DataprepEditorOutlinerMode.h
+++ b/Engine/Plugins/Enterprise/DataprepEditor/Source/DataprepEditor/Private/DataprepEditorOutlinerMode.h
@@ -5,6 +5,7 @@
 #include "ActorMode.h"
 #include "ActorTreeItem.h"
 #include "ComponentTreeItem.h"
+#include "DataprepEditor.h"
 
 namespace DataprepEditorSceneOutlinerUtils
 {
diff --git a/Engine/Plugins/Enterprise/DataprepEditor/Source/DataprepEditor/Private/Widgets/DataprepGraph/DataprepActionWidgetsUtils.h b/Engine/Plugins/Enterprise/DataprepEditor/Source/DataprepEditor/Private/Widgets/DataprepGraph/DataprepActionWidgetsUtils.h
--- a/Engine/Plugins/Enterprise/DataprepEditor/Source/DataprepEditor/Private/Widgets/DataprepGraph/DataprepActionWidgetsUtils.h
+++ b/Engine/Plugins/Enterprise/DataprepEditor/Source/DataprepEditor/Private/Widgets/DataprepGraph/DataprepActionWidgetsUtils.h
@@ -3,6 +3,7 @@
 #pragma once
 
 #include "CoreMinimal.h"
+#include "UObject/Class.h"
 
 namespace DataprepActionWidgetsUtils
 {
diff --git a/Engine/Plugins/Enterprise/DatasmithCADImporter/Source/DatasmithDispatcher/Public/DatasmithDispatcherConfig.h b/Engine/Plugins/Enterprise/DatasmithCADImporter/Source/DatasmithDispatcher/Public/DatasmithDispatcherConfig.h
--- a/Engine/Plugins/Enterprise/DatasmithCADImporter/Source/DatasmithDispatcher/Public/DatasmithDispatcherConfig.h
+++ b/Engine/Plugins/Enterprise/DatasmithCADImporter/Source/DatasmithDispatcher/Public/DatasmithDispatcherConfig.h
@@ -2,6 +2,8 @@
 
 #pragma once
 
+#include "HAL/Platform.h"
+
 namespace DatasmithDispatcher
 {
 namespace Config
diff --git a/Engine/Plugins/Enterprise/DatasmithFBXImporter/Source/DatasmithFBXTranslator/Public/DatasmithFBXHashUtils.h b/Engine/Plugins/Enterprise/DatasmithFBXImporter/Source/DatasmithFBXTranslator/Public/DatasmithFBXHashUtils.h
--- a/Engine/Plugins/Enterprise/DatasmithFBXImporter/Source/DatasmithFBXTranslator/Public/DatasmithFBXHashUtils.h
+++ b/Engine/Plugins/Enterprise/DatasmithFBXImporter/Source/DatasmithFBXTranslator/Public/DatasmithFBXHashUtils.h
@@ -5,6 +5,8 @@
 #include "CoreTypes.h"
 #include "Containers/Array.h"
 #include "Containers/UnrealString.h"
+#include "Math/Transform.h"
+#include "Math/Vector.h"
 #include "Misc/SecureHash.h"
 
 class FDatasmithFBXHashUtils
diff --git a/Engine/Plugins/Enterprise/DatasmithImporter/Source/DatasmithImporter/Private/ObjectElements/DatasmithUElementsUtils.h b/Engine/Plugins/Enterprise/DatasmithImporter/Source/DatasmithImporter/Private/ObjectElements/DatasmithUElementsUtils.h
--- a/Engine/Plugins/Enterprise/DatasmithImporter/Source/DatasmithImporter/Private/ObjectElements/DatasmithUElementsUtils.h
+++ b/Engine/Plugins/Enterprise/DatasmithImporter/Source/DatasmithImporter/Private/ObjectElements/DatasmithUElementsUtils.h
@@ -2,6 +2,10 @@
 #pragma once
 
 #include "CoreMinimal.h"
+#include "DatasmithVariantElements.h"
+#include "IDatasmithSceneElements.h"
+
+class IDatasmithLevelVariantSetsElement;
 
 class FDatasmithUElementsUtils
 {
diff --git a/Engine/Plugins/Experimental/AR/PointCloud/Source/PointCloud/Private/PointCloudBuffers.h b/Engine/Plugins/Experimental/AR/PointCloud/Source/PointCloud/Private/PointCloudBuffers.h
--- a/Engine/Plugins/Experimental/AR/PointCloud/Source/PointCloud/Private/PointCloudBuffers.h
+++ b/Engine/Plugins/Experimental/AR/PointCloud/Source/PointCloud/Private/PointCloudBuffers.h
@@ -3,6 +3,8 @@
 #pragma once
 
 #include "Containers/ResourceArray.h"
+#include "RenderResource.h"
+#include "RHI.h"
 
 class FPointCloudVertexResourceArray :
 	public FResourceArrayInterface
diff --git a/Engine/Plugins/Experimental/Animation/PoseSearch/Source/Runtime/Private/PoseSearchEigenHelper.h b/Engine/Plugins/Experimental/Animation/PoseSearch/Source/Runtime/Private/PoseSearchEigenHelper.h
--- a/Engine/Plugins/Experimental/Animation/PoseSearch/Source/Runtime/Private/PoseSearchEigenHelper.h
+++ b/Engine/Plugins/Experimental/Animation/PoseSearch/Source/Runtime/Private/PoseSearchEigenHelper.h
@@ -3,6 +3,7 @@
 #pragma once
 
 #include "CoreTypes.h"
+#include "Containers/UnrealString.h"
 
 #if UE_BUILD_DEBUG && WITH_EDITOR
 #define UE_POSE_SEARCH_EIGEN_DEBUG 1
diff --git a/Engine/Plugins/Experimental/RenderGrid/Source/RenderGridEditor/Private/UI/Components/SRenderGridDragHandle.h b/Engine/Plugins/Experimental/RenderGrid/Source/RenderGridEditor/Private/UI/Components/SRenderGridDragHandle.h
--- a/Engine/Plugins/Experimental/RenderGrid/Source/RenderGridEditor/Private/UI/Components/SRenderGridDragHandle.h
+++ b/Engine/Plugins/Experimental/RenderGrid/Source/RenderGridEditor/Private/UI/Components/SRenderGridDragHandle.h
@@ -2,6 +2,7 @@
 
 #pragma once
 
+#include "Widgets/SBoxPanel.h"
 #include "Widgets/SCompoundWidget.h"
 #include "Widgets/Layout/SBox.h"
 #include "Widgets/Images/SImage.h"
diff --git a/Engine/Plugins/Experimental/ScreenReader/Source/ScreenReader/Public/Navigation/ScreenReaderNavigationPolicy.h b/Engine/Plugins/Experimental/ScreenReader/Source/ScreenReader/Public/Navigation/ScreenReaderNavigationPolicy.h
--- a/Engine/Plugins/Experimental/ScreenReader/Source/ScreenReader/Public/Navigation/ScreenReaderNavigationPolicy.h
+++ b/Engine/Plugins/Experimental/ScreenReader/Source/ScreenReader/Public/Navigation/ScreenReaderNavigationPolicy.h
@@ -1,9 +1,10 @@
 // Copyright Epic Games, Inc. All Rights Reserved.
 
 #pragma once
+
 #include "CoreMinimal.h"
+#include "GenericPlatform/Accessibility/GenericAccessibleInterfaces.h"
 
-class IAccessibleWidget;
 /**
  * An interface that controls how a screen reader user navigates around the accessible widget hierarchy and what accessible widgets can
  * be navigated to from a source widget.
diff --git a/Engine/Plugins/Media/BinkMedia/Source/BinkMediaPlayerEditor/Private/Models/BinkMediaPlayerEditorTexture.h b/Engine/Plugins/Media/BinkMedia/Source/BinkMediaPlayerEditor/Private/Models/BinkMediaPlayerEditorTexture.h
--- a/Engine/Plugins/Media/BinkMedia/Source/BinkMediaPlayerEditor/Private/Models/BinkMediaPlayerEditorTexture.h
+++ b/Engine/Plugins/Media/BinkMedia/Source/BinkMediaPlayerEditor/Private/Models/BinkMediaPlayerEditorTexture.h
@@ -5,6 +5,7 @@
 
 #include "BinkMediaPlayer.h"
 #include "RenderingThread.h"
+#include "Slate/SlateTextures.h"
 
 struct FBinkMediaPlayerEditorTexture : public FTickableGameObject 
 {
diff --git a/Engine/Plugins/Media/BlackmagicMedia/Source/BlackmagicMedia/BlackmagicMedia.Build.cs b/Engine/Plugins/Media/BlackmagicMedia/Source/BlackmagicMedia/BlackmagicMedia.Build.cs
--- a/Engine/Plugins/Media/BlackmagicMedia/Source/BlackmagicMedia/BlackmagicMedia.Build.cs
+++ b/Engine/Plugins/Media/BlackmagicMedia/Source/BlackmagicMedia/BlackmagicMedia.Build.cs
@@ -11,6 +11,7 @@ namespace UnrealBuildTool.Rules
 		public BlackmagicMedia(ReadOnlyTargetRules Target) : base(Target)
 		{
 			PCHUsage = PCHUsageMode.UseExplicitOrSharedPCHs;
+			IWYUSupport = IWYUSupport.None;
 
 			DynamicallyLoadedModuleNames.AddRange(
 				new string[] {
diff --git a/Engine/Plugins/Runtime/HTTPChunkInstaller/Source/Public/ChunkSetup.h b/Engine/Plugins/Runtime/HTTPChunkInstaller/Source/Public/ChunkSetup.h
--- a/Engine/Plugins/Runtime/HTTPChunkInstaller/Source/Public/ChunkSetup.h
+++ b/Engine/Plugins/Runtime/HTTPChunkInstaller/Source/Public/ChunkSetup.h
@@ -2,6 +2,17 @@
 
 #pragma once
 
+#include "HAL/Event.h"
+#include "HAL/FileManager.h"
+#include "HAL/PlatformFile.h"
+#include "HAL/PlatformFileManager.h"
+#include "HAL/Runnable.h"
+#include "Interfaces/IBuildManifest.h"
+#include "Interfaces/IBuildPatchServicesModule.h"
+#include "Misc/CoreDelegates.h"
+#include "Misc/Paths.h"
+#include "Stats/Stats.h"
+
 // Helper class to find all pak/mainfests files.
 class FFileSearchVisitor : public IPlatformFile::FDirectoryVisitor
 {
diff --git a/Engine/Plugins/Runtime/nDisplay/Source/DisplayCluster/Private/Network/Service/ClusterEventsJson/DisplayClusterClusterEventsJsonStrings.h b/Engine/Plugins/Runtime/nDisplay/Source/DisplayCluster/Private/Network/Service/ClusterEventsJson/DisplayClusterClusterEventsJsonStrings.h
--- a/Engine/Plugins/Runtime/nDisplay/Source/DisplayCluster/Private/Network/Service/ClusterEventsJson/DisplayClusterClusterEventsJsonStrings.h
+++ b/Engine/Plugins/Runtime/nDisplay/Source/DisplayCluster/Private/Network/Service/ClusterEventsJson/DisplayClusterClusterEventsJsonStrings.h
@@ -2,6 +2,7 @@
 
 #pragma once
 
+#include "HAL/Platform.h"
 
 /**
  * JSON cluster event protocol strings
diff --git a/Engine/Plugins/Runtime/nDisplay/Source/DisplayCluster/Private/Network/Session/IDisplayClusterSessionPacketHandler.h b/Engine/Plugins/Runtime/nDisplay/Source/DisplayCluster/Private/Network/Session/IDisplayClusterSessionPacketHandler.h
--- a/Engine/Plugins/Runtime/nDisplay/Source/DisplayCluster/Private/Network/Session/IDisplayClusterSessionPacketHandler.h
+++ b/Engine/Plugins/Runtime/nDisplay/Source/DisplayCluster/Private/Network/Session/IDisplayClusterSessionPacketHandler.h
@@ -2,8 +2,10 @@
 
 #pragma once
 
-struct FDisplayClusterSessionInfo;
+#include "Templates/SharedPointerFwd.h"
+#include <type_traits>
 
+struct FDisplayClusterSessionInfo;
 
 /**
  * Packet handler interface for all incoming packets
diff --git a/Engine/Plugins/Runtime/nDisplay/Source/DisplayCluster/Public/Config/IDisplayClusterConfigManager.h b/Engine/Plugins/Runtime/nDisplay/Source/DisplayCluster/Public/Config/IDisplayClusterConfigManager.h
--- a/Engine/Plugins/Runtime/nDisplay/Source/DisplayCluster/Public/Config/IDisplayClusterConfigManager.h
+++ b/Engine/Plugins/Runtime/nDisplay/Source/DisplayCluster/Public/Config/IDisplayClusterConfigManager.h
@@ -5,6 +5,7 @@
 class  UDisplayClusterConfigurationData;
 class  UDisplayClusterConfigurationClusterNode;
 class  UDisplayClusterConfigurationViewport;
+class  FString;
 struct FDisplayClusterConfigurationPostprocess;
 struct FDisplayClusterConfigurationProjection;
 
diff --git a/Engine/Plugins/Runtime/nDisplay/Source/DisplayCluster/Public/Render/Viewport/RenderFrame/DisplayClusterRenderFrame.h b/Engine/Plugins/Runtime/nDisplay/Source/DisplayCluster/Public/Render/Viewport/RenderFrame/DisplayClusterRenderFrame.h
--- a/Engine/Plugins/Runtime/nDisplay/Source/DisplayCluster/Public/Render/Viewport/RenderFrame/DisplayClusterRenderFrame.h
+++ b/Engine/Plugins/Runtime/nDisplay/Source/DisplayCluster/Public/Render/Viewport/RenderFrame/DisplayClusterRenderFrame.h
@@ -5,8 +5,11 @@
 #include "CoreMinimal.h"
 
 #include "DisplayClusterRenderFrameEnums.h"
+#include "Render/Viewport/Containers/DisplayClusterViewport_RenderSettings.h"
 #include "SceneViewExtension.h"
 
+class IDisplayClusterViewportManager;
+
 // Render frame container
 class FDisplayClusterRenderFrame
 {
diff --git a/Engine/Plugins/Runtime/nDisplay/Source/DisplayClusterConfiguration/Public/DisplayClusterConfigurationStrings.h b/Engine/Plugins/Runtime/nDisplay/Source/DisplayClusterConfiguration/Public/DisplayClusterConfigurationStrings.h
--- a/Engine/Plugins/Runtime/nDisplay/Source/DisplayClusterConfiguration/Public/DisplayClusterConfigurationStrings.h
+++ b/Engine/Plugins/Runtime/nDisplay/Source/DisplayClusterConfiguration/Public/DisplayClusterConfigurationStrings.h
@@ -2,6 +2,8 @@
 
 #pragma once
 
+#include "HAL/Platform.h"
+
 namespace DisplayClusterConfigurationStrings
 {
 	// GUI
diff --git a/Engine/Plugins/Runtime/nDisplay/Source/DisplayClusterConfiguration/Public/DisplayClusterConfigurationVersion.h b/Engine/Plugins/Runtime/nDisplay/Source/DisplayClusterConfiguration/Public/DisplayClusterConfigurationVersion.h
--- a/Engine/Plugins/Runtime/nDisplay/Source/DisplayClusterConfiguration/Public/DisplayClusterConfigurationVersion.h
+++ b/Engine/Plugins/Runtime/nDisplay/Source/DisplayClusterConfiguration/Public/DisplayClusterConfigurationVersion.h
@@ -2,6 +2,8 @@
 
 #pragma once
 
+#include "HAL/Platform.h"
+
 enum class EDisplayClusterConfigurationVersion : uint8
 {
 	Unknown,     // Unknown version or not a config file
diff --git a/Engine/Plugins/Runtime/nDisplay/Source/DisplayClusterConfigurator/Private/Views/OutputMapping/Alignment/DisplayClusterConfiguratorNodeAlignmentHelper.h b/Engine/Plugins/Runtime/nDisplay/Source/DisplayClusterConfigurator/Private/Views/OutputMapping/Alignment/DisplayClusterConfiguratorNodeAlignmentHelper.h
--- a/Engine/Plugins/Runtime/nDisplay/Source/DisplayClusterConfigurator/Private/Views/OutputMapping/Alignment/DisplayClusterConfiguratorNodeAlignmentHelper.h
+++ b/Engine/Plugins/Runtime/nDisplay/Source/DisplayClusterConfigurator/Private/Views/OutputMapping/Alignment/DisplayClusterConfiguratorNodeAlignmentHelper.h
@@ -5,6 +5,7 @@
 #include "CoreMinimal.h"
 #include "UObject/WeakObjectPtr.h"
 #include "UObject/WeakObjectPtrTemplates.h"
+#include "Views/OutputMapping/EdNodes/DisplayClusterConfiguratorBaseNode.h"
 
 class UDisplayClusterConfiguratorBaseNode;
 
diff --git a/Engine/Plugins/Runtime/nDisplay/Source/DisplayClusterConfigurator/Private/Views/OutputMapping/EdNodes/DisplayClusterConfiguratorBaseNode.cpp b/Engine/Plugins/Runtime/nDisplay/Source/DisplayClusterConfigurator/Private/Views/OutputMapping/EdNodes/DisplayClusterConfiguratorBaseNode.cpp
--- a/Engine/Plugins/Runtime/nDisplay/Source/DisplayClusterConfigurator/Private/Views/OutputMapping/EdNodes/DisplayClusterConfiguratorBaseNode.cpp
+++ b/Engine/Plugins/Runtime/nDisplay/Source/DisplayClusterConfigurator/Private/Views/OutputMapping/EdNodes/DisplayClusterConfiguratorBaseNode.cpp
@@ -2,6 +2,7 @@
 
 #include "Views/OutputMapping/EdNodes/DisplayClusterConfiguratorBaseNode.h"
 
+#include "Views/OutputMapping/Alignment/DisplayClusterConfiguratorNodeAlignmentHelper.h"
 #include "Views/OutputMapping/IDisplayClusterConfiguratorViewOutputMapping.h"
 #include "Views/TreeViews/IDisplayClusterConfiguratorViewTree.h"
 #include "Views/TreeViews/IDisplayClusterConfiguratorTreeItem.h"
diff --git a/Engine/Plugins/Runtime/nDisplay/Source/DisplayClusterConfigurator/Private/Views/OutputMapping/EdNodes/DisplayClusterConfiguratorBaseNode.h b/Engine/Plugins/Runtime/nDisplay/Source/DisplayClusterConfigurator/Private/Views/OutputMapping/EdNodes/DisplayClusterConfiguratorBaseNode.h
--- a/Engine/Plugins/Runtime/nDisplay/Source/DisplayClusterConfigurator/Private/Views/OutputMapping/EdNodes/DisplayClusterConfiguratorBaseNode.h
+++ b/Engine/Plugins/Runtime/nDisplay/Source/DisplayClusterConfigurator/Private/Views/OutputMapping/EdNodes/DisplayClusterConfiguratorBaseNode.h
@@ -5,7 +5,7 @@
 #include "Views/OutputMapping/IDisplayClusterConfiguratorOutputMappingItem.h"
 #include "EdGraph/EdGraphNode.h"
 
-#include "Views/OutputMapping/Alignment/DisplayClusterConfiguratorNodeAlignmentHelper.h"
+//#include "Views/OutputMapping/Alignment/DisplayClusterConfiguratorNodeAlignmentHelper.h"
 
 #include "DisplayClusterConfiguratorBaseNode.generated.h"
 
@@ -13,6 +13,10 @@ class FDisplayClusterConfiguratorBlueprintEditor;
 class IDisplayClusterConfiguratorOutputMappingSlot;
 class IDisplayClusterConfiguratorTreeItem;
 
+struct FNodeAlignmentAnchors;
+struct FNodeAlignmentPair;
+struct FNodeAlignmentParams;
+
 namespace DisplayClusterConfiguratorGraphLayers
 {
 	// Layer size is set to 1000 because the graph nodes rely on SOverlay widgets for rendering, and every slot in an overlay
diff --git a/Engine/Plugins/Runtime/nDisplay/Source/DisplayClusterConfigurator/Private/Views/OutputMapping/EdNodes/DisplayClusterConfiguratorHostNode.cpp b/Engine/Plugins/Runtime/nDisplay/Source/DisplayClusterConfigurator/Private/Views/OutputMapping/EdNodes/DisplayClusterConfiguratorHostNode.cpp
--- a/Engine/Plugins/Runtime/nDisplay/Source/DisplayClusterConfigurator/Private/Views/OutputMapping/EdNodes/DisplayClusterConfiguratorHostNode.cpp
+++ b/Engine/Plugins/Runtime/nDisplay/Source/DisplayClusterConfigurator/Private/Views/OutputMapping/EdNodes/DisplayClusterConfiguratorHostNode.cpp
@@ -5,6 +5,7 @@
 #include "DisplayClusterConfiguratorBlueprintEditor.h"
 #include "DisplayClusterConfigurationTypes.h"
 #include "ClusterConfiguration/DisplayClusterConfiguratorClusterUtils.h"
+#include "Views/OutputMapping/Alignment/DisplayClusterConfiguratorNodeAlignmentHelper.h"
 #include "Views/OutputMapping/IDisplayClusterConfiguratorViewOutputMapping.h"
 #include "Views/TreeViews/IDisplayClusterConfiguratorViewTree.h"
 #include "Views/TreeViews/IDisplayClusterConfiguratorTreeItem.h"
diff --git a/Engine/Plugins/Runtime/nDisplay/Source/DisplayClusterConfigurator/Private/Views/OutputMapping/GraphNodes/SDisplayClusterConfiguratorBaseNode.cpp b/Engine/Plugins/Runtime/nDisplay/Source/DisplayClusterConfigurator/Private/Views/OutputMapping/GraphNodes/SDisplayClusterConfiguratorBaseNode.cpp
--- a/Engine/Plugins/Runtime/nDisplay/Source/DisplayClusterConfigurator/Private/Views/OutputMapping/GraphNodes/SDisplayClusterConfiguratorBaseNode.cpp
+++ b/Engine/Plugins/Runtime/nDisplay/Source/DisplayClusterConfigurator/Private/Views/OutputMapping/GraphNodes/SDisplayClusterConfiguratorBaseNode.cpp
@@ -7,6 +7,7 @@
 #include "Framework/Application/SlateApplication.h"
 #include "Views/TreeViews/IDisplayClusterConfiguratorViewTree.h"
 #include "Views/TreeViews/IDisplayClusterConfiguratorTreeItem.h"
+#include "Views/OutputMapping/Alignment/DisplayClusterConfiguratorNodeAlignmentHelper.h"
 #include "Views/OutputMapping/IDisplayClusterConfiguratorViewOutputMapping.h"
 #include "Views/OutputMapping/EdNodes/DisplayClusterConfiguratorBaseNode.h"
 #include "Views/OutputMapping/Widgets/SDisplayClusterConfiguratorResizer.h"
diff --git a/Engine/Plugins/Runtime/nDisplay/Source/DisplayClusterScenePreview/Private/DisplayClusterScenePreviewModule.h b/Engine/Plugins/Runtime/nDisplay/Source/DisplayClusterScenePreview/Private/DisplayClusterScenePreviewModule.h
--- a/Engine/Plugins/Runtime/nDisplay/Source/DisplayClusterScenePreview/Private/DisplayClusterScenePreviewModule.h
+++ b/Engine/Plugins/Runtime/nDisplay/Source/DisplayClusterScenePreview/Private/DisplayClusterScenePreviewModule.h
@@ -6,6 +6,7 @@
 
 #include "Containers/Ticker.h"
 #include "DisplayClusterMeshProjectionRenderer.h"
+#include "Engine/TextureRenderTarget2D.h"
 #include "UObject/StrongObjectPtr.h"
 
 class ADisplayClusterRootActor;
diff --git a/Engine/Plugins/VirtualProduction/DMX/DMXPixelMapping/Source/DMXPixelMappingEditor/Private/Customizations/DMXPixelMappingDetailCustomization_FixtureGroup.h b/Engine/Plugins/VirtualProduction/DMX/DMXPixelMapping/Source/DMXPixelMappingEditor/Private/Customizations/DMXPixelMappingDetailCustomization_FixtureGroup.h
--- a/Engine/Plugins/VirtualProduction/DMX/DMXPixelMapping/Source/DMXPixelMappingEditor/Private/Customizations/DMXPixelMappingDetailCustomization_FixtureGroup.h
+++ b/Engine/Plugins/VirtualProduction/DMX/DMXPixelMapping/Source/DMXPixelMappingEditor/Private/Customizations/DMXPixelMappingDetailCustomization_FixtureGroup.h
@@ -4,6 +4,7 @@
 
 #include "CoreMinimal.h"
 #include "IDetailCustomization.h"
+#include "Toolkits/DMXPixelMappingToolkit.h"
 #include "Types/SlateEnums.h"
 
 struct FDMXEntityFixturePatchRef;
diff --git a/Engine/Plugins/VirtualProduction/DMX/DMXPixelMapping/Source/DMXPixelMappingEditor/Private/Views/SDMXPixelMappingDetailsView.cpp b/Engine/Plugins/VirtualProduction/DMX/DMXPixelMapping/Source/DMXPixelMappingEditor/Private/Views/SDMXPixelMappingDetailsView.cpp
--- a/Engine/Plugins/VirtualProduction/DMX/DMXPixelMapping/Source/DMXPixelMappingEditor/Private/Views/SDMXPixelMappingDetailsView.cpp
+++ b/Engine/Plugins/VirtualProduction/DMX/DMXPixelMapping/Source/DMXPixelMappingEditor/Private/Views/SDMXPixelMappingDetailsView.cpp
@@ -18,6 +18,7 @@
 #include "Customizations/DMXPixelMappingDetailCustomization_Matrix.h"
 #include "Customizations/DMXPixelMappingDetailCustomization_MatrixCell.h"
 
+#include "PropertyEditorDelegates.h"
 #include "PropertyEditorModule.h"
 #include "Modules/ModuleManager.h"
 
diff --git a/Engine/Plugins/VirtualProduction/DMX/DMXProtocol/Source/DMXProtocol/Public/Interfaces/IDMXProtocolPacket.h b/Engine/Plugins/VirtualProduction/DMX/DMXProtocol/Source/DMXProtocol/Public/Interfaces/IDMXProtocolPacket.h
--- a/Engine/Plugins/VirtualProduction/DMX/DMXProtocol/Source/DMXProtocol/Public/Interfaces/IDMXProtocolPacket.h
+++ b/Engine/Plugins/VirtualProduction/DMX/DMXProtocol/Source/DMXProtocol/Public/Interfaces/IDMXProtocolPacket.h
@@ -2,6 +2,7 @@
 
 #pragma once
 
+#include "Templates/SharedPointer.h"
 #include "Serialization/BufferArchive.h"
 
 struct IDMXProtocolPacket
diff --git a/Engine/Plugins/VirtualProduction/LevelSnapshots/Source/LevelSnapshots/Public/Restorability/Interfaces/ICustomObjectSnapshotSerializer.h b/Engine/Plugins/VirtualProduction/LevelSnapshots/Source/LevelSnapshots/Public/Restorability/Interfaces/ICustomObjectSnapshotSerializer.h
--- a/Engine/Plugins/VirtualProduction/LevelSnapshots/Source/LevelSnapshots/Public/Restorability/Interfaces/ICustomObjectSnapshotSerializer.h
+++ b/Engine/Plugins/VirtualProduction/LevelSnapshots/Source/LevelSnapshots/Public/Restorability/Interfaces/ICustomObjectSnapshotSerializer.h
@@ -2,6 +2,9 @@
 
 #pragma once
 
+#include "HAL/Platform.h"
+#include "Misc/AssertionMacros.h"
+
 class UObject;
 struct FPropertySelectionMap;
 
diff --git a/Engine/Plugins/VirtualProduction/LevelSnapshots/Source/LevelSnapshots/Public/Restorability/Interfaces/IPropertyComparer.h b/Engine/Plugins/VirtualProduction/LevelSnapshots/Source/LevelSnapshots/Public/Restorability/Interfaces/IPropertyComparer.h
--- a/Engine/Plugins/VirtualProduction/LevelSnapshots/Source/LevelSnapshots/Public/Restorability/Interfaces/IPropertyComparer.h
+++ b/Engine/Plugins/VirtualProduction/LevelSnapshots/Source/LevelSnapshots/Public/Restorability/Interfaces/IPropertyComparer.h
@@ -2,6 +2,8 @@
 
 #pragma once
 
+#include "HAL/Platform.h"
+
 class AActor;
 class FProperty;
 class UObject;
diff --git a/Engine/Plugins/VirtualProduction/LevelSnapshots/Source/LevelSnapshots/Public/Restorability/Params/PropertyComparisonParams.h b/Engine/Plugins/VirtualProduction/LevelSnapshots/Source/LevelSnapshots/Public/Restorability/Params/PropertyComparisonParams.h
--- a/Engine/Plugins/VirtualProduction/LevelSnapshots/Source/LevelSnapshots/Public/Restorability/Params/PropertyComparisonParams.h
+++ b/Engine/Plugins/VirtualProduction/LevelSnapshots/Source/LevelSnapshots/Public/Restorability/Params/PropertyComparisonParams.h
@@ -6,6 +6,7 @@
 
 class AActor;
 class FProperty;
+class ULevelSnapshot;
 class UObject;
 class UStruct;
 
diff --git a/Engine/Plugins/VirtualProduction/RemoteControl/Source/RemoteControl/Private/RemoteControlPropertyHandle.h b/Engine/Plugins/VirtualProduction/RemoteControl/Source/RemoteControl/Private/RemoteControlPropertyHandle.h
--- a/Engine/Plugins/VirtualProduction/RemoteControl/Source/RemoteControl/Private/RemoteControlPropertyHandle.h
+++ b/Engine/Plugins/VirtualProduction/RemoteControl/Source/RemoteControl/Private/RemoteControlPropertyHandle.h
@@ -3,6 +3,7 @@
 #pragma once
 
 #include "IRemoteControlPropertyHandle.h"
+#include "UObject/UnrealType.h"
 #include "UObject/WeakFieldPtr.h"
 
 /**
diff --git a/Engine/Plugins/VirtualProduction/RemoteControl/Source/RemoteControlUI/Public/IHasProtocolExtensibility.h b/Engine/Plugins/VirtualProduction/RemoteControl/Source/RemoteControlUI/Public/IHasProtocolExtensibility.h
--- a/Engine/Plugins/VirtualProduction/RemoteControl/Source/RemoteControlUI/Public/IHasProtocolExtensibility.h
+++ b/Engine/Plugins/VirtualProduction/RemoteControl/Source/RemoteControlUI/Public/IHasProtocolExtensibility.h
@@ -4,6 +4,8 @@
 
 #include "CoreMinimal.h"
 
+class SWidget;
+
 /**
  * A simple interface provides protocol specific utilities.
  */
diff --git a/Engine/Plugins/VirtualProduction/SequencerPlaylists/Source/SequencerPlaylists/Private/SequencerPlaylistsWidgets.h b/Engine/Plugins/VirtualProduction/SequencerPlaylists/Source/SequencerPlaylists/Private/SequencerPlaylistsWidgets.h
--- a/Engine/Plugins/VirtualProduction/SequencerPlaylists/Source/SequencerPlaylists/Private/SequencerPlaylistsWidgets.h
+++ b/Engine/Plugins/VirtualProduction/SequencerPlaylists/Source/SequencerPlaylists/Private/SequencerPlaylistsWidgets.h
@@ -5,6 +5,7 @@
 #include "CoreMinimal.h"
 #include "EditorUndoClient.h"
 #include "PropertyEditorDelegates.h"
+#include "SequencerPlaylistItem.h"
 #include "Widgets/DeclarativeSyntaxSupport.h"
 #include "Widgets/SCompoundWidget.h"
 #include "Widgets/Views/SListView.h"
diff --git a/Engine/Plugins/VirtualProduction/TimedDataMonitor/Source/TimedDataMonitorEditor/Private/SFrameTime.h b/Engine/Plugins/VirtualProduction/TimedDataMonitor/Source/TimedDataMonitorEditor/Private/SFrameTime.h
--- a/Engine/Plugins/VirtualProduction/TimedDataMonitor/Source/TimedDataMonitorEditor/Private/SFrameTime.h
+++ b/Engine/Plugins/VirtualProduction/TimedDataMonitor/Source/TimedDataMonitorEditor/Private/SFrameTime.h
@@ -3,8 +3,12 @@
 #pragma once
 
 #include "CoreMinimal.h"
+#include "Engine/Engine.h"
+#include "Misc/App.h"
 #include "Rendering/DrawElements.h"
 #include "Styling/AppStyle.h"
+#include "Styling/CoreStyle.h"
+#include "Styling/SlateTypes.h"
 #include "TimedDataMonitorEditorSettings.h"
 #include "Widgets/DeclarativeSyntaxSupport.h"
 #include "Widgets/SLeafWidget.h"
diff --git a/Engine/Source/Developer/AutomationDriver/Private/IStepExecutor.h b/Engine/Source/Developer/AutomationDriver/Private/IStepExecutor.h
--- a/Engine/Source/Developer/AutomationDriver/Private/IStepExecutor.h
+++ b/Engine/Source/Developer/AutomationDriver/Private/IStepExecutor.h
@@ -2,6 +2,7 @@
 #pragma once
 
 #include "Async/AsyncResult.h"
+#include "Delegates/Delegate.h"
 #include "Misc/Timespan.h"
 
 struct FStepResult
diff --git a/Engine/Source/Developer/AutomationDriver/Private/MetaData/DriverUniqueTagMetaData.h b/Engine/Source/Developer/AutomationDriver/Private/MetaData/DriverUniqueTagMetaData.h
--- a/Engine/Source/Developer/AutomationDriver/Private/MetaData/DriverUniqueTagMetaData.h
+++ b/Engine/Source/Developer/AutomationDriver/Private/MetaData/DriverUniqueTagMetaData.h
@@ -1,7 +1,7 @@
 // Copyright Epic Games, Inc. All Rights Reserved.
 #pragma once
 
-#include "CoreMinimal.h"
+#include "Types/ISlateMetaData.h"
 
 class FDriverUniqueTagMetaData
 	: public ISlateMetaData
diff --git a/Engine/Source/Developer/AutomationDriver/Public/AutomationDriverTypeDefs.h b/Engine/Source/Developer/AutomationDriver/Public/AutomationDriverTypeDefs.h
--- a/Engine/Source/Developer/AutomationDriver/Public/AutomationDriverTypeDefs.h
+++ b/Engine/Source/Developer/AutomationDriver/Public/AutomationDriverTypeDefs.h
@@ -1,6 +1,8 @@
 // Copyright Epic Games, Inc. All Rights Reserved.
 #pragma once
 
+#include "Templates/SharedPointerFwd.h"
+
 class IAutomationDriver;
 class IDriverElement;
 class IDriverElementCollection;
diff --git a/Engine/Source/Developer/AutomationDriver/Public/DriverConfiguration.h b/Engine/Source/Developer/AutomationDriver/Public/DriverConfiguration.h
--- a/Engine/Source/Developer/AutomationDriver/Public/DriverConfiguration.h
+++ b/Engine/Source/Developer/AutomationDriver/Public/DriverConfiguration.h
@@ -1,6 +1,9 @@
 // Copyright Epic Games, Inc. All Rights Reserved.
+
 #pragma once
 
+#include "Misc/Timespan.h"
+
 /**
  * A model class that holds all the configuration options for an AutomationDriver
  */
diff --git a/Engine/Source/Developer/AutomationDriver/Public/IDriverSequence.h b/Engine/Source/Developer/AutomationDriver/Public/IDriverSequence.h
--- a/Engine/Source/Developer/AutomationDriver/Public/IDriverSequence.h
+++ b/Engine/Source/Developer/AutomationDriver/Public/IDriverSequence.h
@@ -7,6 +7,7 @@
 #include "GenericPlatform/GenericApplicationMessageHandler.h"
 
 class IElementLocator;
+template<typename ResultType> class TAsyncResult;
 
 class IAsyncActionSequence
 {
diff --git a/Engine/Source/Runtime/Core/Public/UObject/LoadTimesObjectVersion.h b/Engine/Source/Runtime/Core/Public/UObject/LoadTimesObjectVersion.h
--- a/Engine/Source/Runtime/Core/Public/UObject/LoadTimesObjectVersion.h
+++ b/Engine/Source/Runtime/Core/Public/UObject/LoadTimesObjectVersion.h
@@ -1,6 +1,10 @@
 // Copyright Epic Games, Inc. All Rights Reserved.
 #pragma once
 
+#include "HAL/Platform.h"
+
+struct FGuid;
+
 // Custom serialization version for changes made in LoadTimes stream
 struct CORE_API FLoadTimesObjectVersion
 {
diff --git a/Engine/Source/Runtime/Core/Public/UObject/ReflectionCaptureObjectVersion.h b/Engine/Source/Runtime/Core/Public/UObject/ReflectionCaptureObjectVersion.h
--- a/Engine/Source/Runtime/Core/Public/UObject/ReflectionCaptureObjectVersion.h
+++ b/Engine/Source/Runtime/Core/Public/UObject/ReflectionCaptureObjectVersion.h
@@ -1,6 +1,10 @@
 // Copyright Epic Games, Inc. All Rights Reserved.
 #pragma once
 
+#include "HAL/Platform.h"
+
+struct FGuid;
+
 // Custom serialization version for changes made for a private stream
 struct CORE_API FReflectionCaptureObjectVersion
 {
diff --git a/Engine/Source/Runtime/Datasmith/CADKernel/Private/CADKernel/Geo/Sampling/SurfacicSampling.h b/Engine/Source/Runtime/Datasmith/CADKernel/Private/CADKernel/Geo/Sampling/SurfacicSampling.h
--- a/Engine/Source/Runtime/Datasmith/CADKernel/Private/CADKernel/Geo/Sampling/SurfacicSampling.h
+++ b/Engine/Source/Runtime/Datasmith/CADKernel/Private/CADKernel/Geo/Sampling/SurfacicSampling.h
@@ -1,8 +1,10 @@
 // Copyright Epic Games, Inc. All Rights Reserved.
+
 #pragma once
 
 #include "CADKernel/Core/Types.h"
 #include "CADKernel/Geo/GeoEnum.h"
+#include "CADKernel/Geo/GeoPoint.h"
 #include "CADKernel/Math/Point.h"
 
 namespace UE::CADKernel
diff --git a/Engine/Source/Runtime/Datasmith/CADKernel/Private/CADKernel/Mesh/MeshEnum.h b/Engine/Source/Runtime/Datasmith/CADKernel/Private/CADKernel/Mesh/MeshEnum.h
--- a/Engine/Source/Runtime/Datasmith/CADKernel/Private/CADKernel/Mesh/MeshEnum.h
+++ b/Engine/Source/Runtime/Datasmith/CADKernel/Private/CADKernel/Mesh/MeshEnum.h
@@ -1,6 +1,8 @@
 // Copyright Epic Games, Inc. All Rights Reserved.
 #pragma once
 
+#include "HAL/Platform.h"
+
 namespace UE::CADKernel
 {
 /**
diff --git a/Engine/Source/Runtime/Datasmith/CADKernel/Private/CADKernel/Mesh/Meshers/IsoTriangulator/IsoCell.h b/Engine/Source/Runtime/Datasmith/CADKernel/Private/CADKernel/Mesh/Meshers/IsoTriangulator/IsoCell.h
--- a/Engine/Source/Runtime/Datasmith/CADKernel/Private/CADKernel/Mesh/Meshers/IsoTriangulator/IsoCell.h
+++ b/Engine/Source/Runtime/Datasmith/CADKernel/Private/CADKernel/Mesh/Meshers/IsoTriangulator/IsoCell.h
@@ -4,6 +4,7 @@
 #include "CADKernel/Core/Types.h"
 #include "CADKernel/Mesh/MeshEnum.h"
 #include "CADKernel/Mesh/Meshers/IsoTriangulator/IsoNode.h"
+#include "CADKernel/Mesh/Meshers/IsoTriangulator/IsoSegment.h"
 #include "CADKernel/Mesh/Meshers/IsoTriangulator/IntersectionSegmentTool.h"
 #include "CADKernel/Mesh/Structure/Grid.h"
 #include "CADKernel/Topo/TopologicalFace.h"
diff --git a/Engine/Source/Runtime/Datasmith/CADKernel/Public/CADKernel/Topo/TopologicalLink.h b/Engine/Source/Runtime/Datasmith/CADKernel/Public/CADKernel/Topo/TopologicalLink.h
--- a/Engine/Source/Runtime/Datasmith/CADKernel/Public/CADKernel/Topo/TopologicalLink.h
+++ b/Engine/Source/Runtime/Datasmith/CADKernel/Public/CADKernel/Topo/TopologicalLink.h
@@ -3,6 +3,7 @@
 
 #include "CADKernel/Core/Entity.h" 
 #include "CADKernel/Math/Point.h"
+#include "CADKernel/Topo/TopologicalEntity.h"
 #include "CADKernel/UI/Message.h"
 
 namespace UE::CADKernel
diff --git a/Engine/Source/Runtime/Datasmith/CADKernel/Public/CADKernel/Utils/Cache.h b/Engine/Source/Runtime/Datasmith/CADKernel/Public/CADKernel/Utils/Cache.h
--- a/Engine/Source/Runtime/Datasmith/CADKernel/Public/CADKernel/Utils/Cache.h
+++ b/Engine/Source/Runtime/Datasmith/CADKernel/Public/CADKernel/Utils/Cache.h
@@ -1,6 +1,10 @@
 // Copyright Epic Games, Inc. All Rights Reserved.
 #pragma once
 
+#include "CADKernel/Core/Types.h"
+#include "HAL/Platform.h"
+#include "Serialization/Archive.h"
+
 namespace UE::CADKernel
 {
 template<class ObjectType>
diff --git a/Engine/Source/Runtime/Datasmith/DirectLink/Public/DirectLinkSerialMethods.h b/Engine/Source/Runtime/Datasmith/DirectLink/Public/DirectLinkSerialMethods.h
--- a/Engine/Source/Runtime/Datasmith/DirectLink/Public/DirectLinkSerialMethods.h
+++ b/Engine/Source/Runtime/Datasmith/DirectLink/Public/DirectLinkSerialMethods.h
@@ -5,6 +5,7 @@
 #include "Containers/Array.h"
 #include "CoreTypes.h"
 #include "Math/Quat.h"
+#include "Math/Transform.h"
 #include "Math/Vector.h"
 #include "Misc/SecureHash.h"
 #include "Serialization/Archive.h"
diff --git a/Engine/Source/Runtime/IPC/Public/IPC.h b/Engine/Source/Runtime/IPC/Public/IPC.h
--- a/Engine/Source/Runtime/IPC/Public/IPC.h
+++ b/Engine/Source/Runtime/IPC/Public/IPC.h
@@ -5,7 +5,9 @@
 ==============================================================================*/
 
 #pragma once
+
 #include "CoreMinimal.h"
+#include "HAL/PlatformProcess.h"
 
 namespace IPC
 {
diff --git a/Engine/Source/Runtime/UELibrary/Public/UELibraryAPI.h b/Engine/Source/Runtime/UELibrary/Public/UELibraryAPI.h
--- a/Engine/Source/Runtime/UELibrary/Public/UELibraryAPI.h
+++ b/Engine/Source/Runtime/UELibrary/Public/UELibraryAPI.h
@@ -2,6 +2,10 @@
 
 #pragma once
 
+#include "HAL/Platform.h"
+
+#if PLATFORM_WINDOWS
+
 #ifndef _WIN64
 	#error "UELibrary is currently only supported under 64-bit Windows"
 #endif
@@ -68,3 +72,5 @@ UELIBRARYAPI LRESULT UELibrary_WndProc(HWND hWnd, UINT message, WPARAM wParam, L
  */
 UELIBRARYAPI int UELibrary_Shutdown();
 }
+
+#endif // PLATFORM_WINDOWS
