nDisplay: Added shared cross GPU media io.

This allows for sharing viewports between render nodes with zero or 1 frame latency. This is an alternative to mGPU that in many cases can be more efficient. It should be used w/o SLI enabled. Two nDisplay nodes would run on the same PC, one with r.GraphicsAdapter 0 and the other with r.GraphicsAdapter 1, which can be set in the nDisplay config, and overridden from a newly added option in Switchboard.

diff --git a/Engine/Plugins/Runtime/nDisplay/Source/DisplayClusterMedia/DisplayClusterMedia.Build.cs b/Engine/Plugins/Runtime/nDisplay/Source/DisplayClusterMedia/DisplayClusterMedia.Build.cs
--- a/Engine/Plugins/Runtime/nDisplay/Source/DisplayClusterMedia/DisplayClusterMedia.Build.cs
+++ b/Engine/Plugins/Runtime/nDisplay/Source/DisplayClusterMedia/DisplayClusterMedia.Build.cs
@@ -17,6 +17,13 @@ public class DisplayClusterMedia : ModuleRules
 				"DisplayClusterShaders",
 			});
 
+		PublicDependencyModuleNames.AddRange(
+			new string[] {
+				"Media",
+				"MediaAssets",
+				"MediaIOCore",
+			});
+
 		PrivateDependencyModuleNames.AddRange(
 			new string[] {
 				"Core",
@@ -24,8 +31,6 @@ public class DisplayClusterMedia : ModuleRules
 				"DisplayCluster",
 				"DisplayClusterConfiguration",
 				"Engine",
-				"MediaAssets",
-				"MediaIOCore",
 				"Renderer",
 				"RenderCore",
 				"RHI",
@@ -35,5 +40,11 @@ public class DisplayClusterMedia : ModuleRules
 		{
 			PrivateDependencyModuleNames.Add("UnrealEd");
 		}
+
+		if (Target.Platform.IsInGroup(UnrealPlatformGroup.Windows))
+		{
+			AddEngineThirdPartyPrivateStaticDependencies(Target, "DX12");
+			PrivateDependencyModuleNames.Add("D3D12RHI");
+		}
 	}
 }
diff --git a/Engine/Plugins/Runtime/nDisplay/Source/DisplayClusterMedia/Private/Capture/DisplayClusterMediaCaptureBase.cpp b/Engine/Plugins/Runtime/nDisplay/Source/DisplayClusterMedia/Private/Capture/DisplayClusterMediaCaptureBase.cpp
--- a/Engine/Plugins/Runtime/nDisplay/Source/DisplayClusterMedia/Private/Capture/DisplayClusterMediaCaptureBase.cpp
+++ b/Engine/Plugins/Runtime/nDisplay/Source/DisplayClusterMedia/Private/Capture/DisplayClusterMediaCaptureBase.cpp
@@ -72,7 +72,7 @@ void FDisplayClusterMediaCaptureBase::ExportMediaData(FRDGBuilder& GraphBuilder,
 	{
 		MediaCapture->SetValidSourceGPUMask(GraphBuilder.RHICmdList.GetGPUMask());
 
-		const FIntPoint SrcRegionSize = TextureInfo.Region.Size();
+		LastSrcRegionSize = FIntSize(TextureInfo.Region.Size());
 		MediaCapture->CaptureImmediate_RenderThread(GraphBuilder, SrcTexture);
 	}
 }
@@ -99,7 +99,15 @@ void FDisplayClusterMediaCaptureBase::OnPostClusterTick()
 bool FDisplayClusterMediaCaptureBase::StartMediaCapture()
 {
 	FRHICaptureResourceDescription Descriptor;
-	Descriptor.ResourceSize = GetCaptureSize();
+
+	if (LastSrcRegionSize.load().ToIntPoint() == FIntPoint::ZeroValue)
+	{
+		Descriptor.ResourceSize = GetCaptureSize();
+	}
+	else
+	{
+		Descriptor.ResourceSize = LastSrcRegionSize.load().ToIntPoint();
+	}
 	
 	FMediaCaptureOptions MediaCaptureOptions;
 	MediaCaptureOptions.NumberOfFramesToCapture = -1;
diff --git a/Engine/Plugins/Runtime/nDisplay/Source/DisplayClusterMedia/Private/Capture/DisplayClusterMediaCaptureBase.h b/Engine/Plugins/Runtime/nDisplay/Source/DisplayClusterMedia/Private/Capture/DisplayClusterMediaCaptureBase.h
--- a/Engine/Plugins/Runtime/nDisplay/Source/DisplayClusterMedia/Private/Capture/DisplayClusterMediaCaptureBase.h
+++ b/Engine/Plugins/Runtime/nDisplay/Source/DisplayClusterMedia/Private/Capture/DisplayClusterMediaCaptureBase.h
@@ -45,6 +45,23 @@ protected:
 
 	virtual FIntPoint GetCaptureSize() const = 0;
 
+private:
+
+	// Trivial version of FIntPoint so that it can be std::atomic
+	struct FIntSize
+	{
+		int32 X = 0;
+		int32 Y = 0;
+
+		FIntSize(int32 InX, int32 InY) : X(InX), Y(InY) {}
+		FIntSize(const FIntPoint& IntPoint) : X(IntPoint.X), Y(IntPoint.Y) {}
+
+		FIntPoint ToIntPoint()
+		{
+			return FIntPoint(X, Y);
+		}
+	};
+
 private:
 	//~ Begin GC by AddReferencedObjects
 	UMediaOutput*           MediaOutput  = nullptr;
@@ -56,4 +73,7 @@ private:
 
 	// Used to control the rate at which we try to restart the capture
 	double LastRestartTimestamp = 0;
+
+	// Last region size of the texture being exported. Used to restart the capture when in error.
+	std::atomic<FIntSize> LastSrcRegionSize { FIntSize(0,0) };
 };
diff --git a/Engine/Plugins/Runtime/nDisplay/Source/DisplayClusterMedia/Private/Capture/DisplayClusterMediaCaptureCamera.cpp b/Engine/Plugins/Runtime/nDisplay/Source/DisplayClusterMedia/Private/Capture/DisplayClusterMediaCaptureCamera.cpp
--- a/Engine/Plugins/Runtime/nDisplay/Source/DisplayClusterMedia/Private/Capture/DisplayClusterMediaCaptureCamera.cpp
+++ b/Engine/Plugins/Runtime/nDisplay/Source/DisplayClusterMedia/Private/Capture/DisplayClusterMediaCaptureCamera.cpp
@@ -30,7 +30,16 @@ FDisplayClusterMediaCaptureCamera::FDisplayClusterMediaCaptureCamera(const FStri
 			{
 				if (ICVFXCamera->GetName() == CameraId)
 				{
-					CameraResolution = ICVFXCamera->GetCameraSettingsICVFX().CustomFrustum.InnerFrustumResolution;
+					const FDisplayClusterConfigurationICVFX_CameraSettings& CameraSettings = ICVFXCamera->GetCameraSettingsICVFX();
+
+					if (CameraSettings.CustomFrustum.bEnable)
+					{
+						CameraResolution = CameraSettings.CustomFrustum.EstimatedOverscanResolution;
+					}
+					else
+					{
+						CameraResolution = CameraSettings.CustomFrustum.InnerFrustumResolution;
+					}
 				}
 			}
 		}
diff --git a/Engine/Plugins/Runtime/nDisplay/Source/DisplayClusterMedia/Private/Capture/SharedMemoryMediaCapture.cpp b/Engine/Plugins/Runtime/nDisplay/Source/DisplayClusterMedia/Private/Capture/SharedMemoryMediaCapture.cpp
new file mode 100644
--- /dev/null
+++ b/Engine/Plugins/Runtime/nDisplay/Source/DisplayClusterMedia/Private/Capture/SharedMemoryMediaCapture.cpp
@@ -0,0 +1,332 @@
+// Copyright Epic Games, Inc. All Rights Reserved.
+
+#include "SharedMemoryMediaCapture.h"
+#include "SharedMemoryMediaOutput.h"
+
+#include "GenericPlatform/GenericPlatformMemory.h"
+#include "Internationalization/TextLocalizationResource.h"
+#include "RenderGraphUtils.h"
+
+#include "DisplayClusterMediaLog.h"
+#include "SharedMemoryMediaPlatform.h"
+#include "SharedMemoryMediaTypes.h"
+
+DECLARE_GPU_STAT(SharedMemory_Capture);
+
+
+bool USharedMemoryMediaCapture::InitializeCapture()
+{
+	// Validate Media Output type
+
+	USharedMemoryMediaOutput* SharedMemoryMediaOutput = Cast<USharedMemoryMediaOutput>(MediaOutput);
+
+	if (!SharedMemoryMediaOutput)
+	{
+		UE_LOG(LogDisplayClusterMedia, Error, TEXT("Invalid MediaOutput, cannot InitializeCapture"));
+		return false;
+	}
+
+	// Get an RHI type specific implementation.
+	if (!PlatformData.IsValid())
+	{
+		const ERHIInterfaceType RhiInterfaceType = GDynamicRHI->GetInterfaceType();
+
+		PlatformData = FSharedMemoryMediaPlatform::CreateInstanceForRhi(RhiInterfaceType);
+
+		if (!PlatformData.IsValid())
+		{
+			UE_LOG(LogDisplayClusterMedia, Error, TEXT("Unfortunately, SharedMemoryMedia doesn't support the current RHI type '%s'"),
+				*FSharedMemoryMediaPlatform::GetRhiTypeString(RhiInterfaceType));
+
+			return false;
+		}
+	}
+
+	const SIZE_T SharedMemorySize = sizeof(FSharedMemoryMediaFrameMetadata);
+
+	for (int32 BufferIdx = 0; BufferIdx < NUMBUFFERS; ++BufferIdx)
+	{
+		// Generate the shared memory Guid from the MediaOutput user set unique name.
+		const FGuid Guid = UE::SharedMemoryMedia::GenerateSharedMemoryGuid(SharedMemoryMediaOutput->UniqueName, BufferIdx);
+
+		const FString SharedMemoryRegionName = Guid.ToString(EGuidFormats::DigitsWithHyphensInBraces);
+
+		// Open existing shared memory region, in case it exists:
+
+		const uint32 AccessMode = FPlatformMemory::ESharedMemoryAccess::Read | FPlatformMemory::ESharedMemoryAccess::Write;
+
+		FPlatformMemory::FSharedMemoryRegion* SharedMemoryRegion = FPlatformMemory::MapNamedSharedMemoryRegion(
+			*SharedMemoryRegionName, false /* bCreate */, AccessMode, SharedMemorySize
+		);
+
+		// If it doesn't exist, then we allocate and zero-initialize it.
+		if (!SharedMemoryRegion)
+		{
+			// Create
+			SharedMemoryRegion = FPlatformMemory::MapNamedSharedMemoryRegion(
+				*SharedMemoryRegionName,
+				true /* bCreate */,
+				AccessMode,
+				SharedMemorySize
+			);
+
+			// Zero
+			if (SharedMemoryRegion)
+			{
+				check(SharedMemoryRegion->GetAddress());
+				FMemory::Memzero(SharedMemoryRegion->GetAddress(), SharedMemoryRegion->GetSize());
+
+				// Except some special data
+				FSharedMemoryMediaFrameMetadata* Data = static_cast<FSharedMemoryMediaFrameMetadata*>(SharedMemoryRegion->GetAddress());
+				Data->Receiver.FrameNumberAcked = ~0;
+
+				UE_LOG(LogDisplayClusterMedia, Verbose, TEXT("Created SharedMemoryRegion[%d] = %s"), BufferIdx, *SharedMemoryRegionName);
+			}
+		}
+
+		// Verify that the shared memory creation succeeded
+		if (!SharedMemoryRegion || !SharedMemoryRegion->GetAddress())
+		{
+			SetState(EMediaCaptureState::Error);
+			return false;
+		}
+
+		SharedMemory[BufferIdx] = SharedMemoryRegion;
+	}
+
+	for (uint32 BufferIdx = 0; BufferIdx < NUMBUFFERS; BufferIdx++)
+	{
+		// Initialize fences
+		if (!TextureReadyFences[BufferIdx])
+		{
+			TextureReadyFences[BufferIdx] = RHICreateGPUFence(*FString::Printf(TEXT("SharedMemoryMediaOutputFence_%d"), BufferIdx));
+		}
+	}			
+
+	SetState(EMediaCaptureState::Capturing);
+
+	return true;
+}
+
+void USharedMemoryMediaCapture::StopCaptureImpl(bool bAllowPendingFrameToBeProcess)
+{
+	// Note: This gets called by StopCapture which already changed the state to Stopped and called FlushRenderingCommands.
+
+	using namespace UE::SharedMemoryMedia;
+
+	// Since FlushRenderingCommands was already called by StopCapture, we can safely release all the resources
+
+	// ... if we don't have any running tasks.
+	while (RunningTasksCount > 0)
+	{
+		FPlatformProcess::SleepNoStats(SpinWaitTimeSeconds);
+	}
+
+	check(RunningTasksCount == 0);
+
+	for (int32 BufferIdx = 0; BufferIdx < NUMBUFFERS; ++BufferIdx)
+	{
+		check(!bTextureReadyFenceBusy[BufferIdx]);
+		TextureReadyFences[BufferIdx].SafeRelease();
+
+		if (SharedMemory[BufferIdx])
+		{
+			// Let the receivers know that we're closed
+			{
+				FSharedMemoryMediaFrameMetadata* SharedMemoryData = static_cast<FSharedMemoryMediaFrameMetadata*>(SharedMemory[BufferIdx]->GetAddress());
+
+				if (SharedMemoryData)
+				{
+					SharedMemoryData->Sender.Magic = 0;
+					SharedMemoryData->Sender.TextureGuid = ZeroGuid;
+				}
+			}
+
+			FPlatformMemory::UnmapNamedSharedMemoryRegion(SharedMemory[BufferIdx]);
+			SharedMemory[BufferIdx] = nullptr;
+		}
+
+		SharedCrossGpuTextures[BufferIdx].SafeRelease(); // This should release the platform specific resources
+
+		if (PlatformData.IsValid())
+		{
+			PlatformData->ReleaseSharedCrossGpuTexture(BufferIdx);
+		}
+
+		SharedCrossGpuTextureGuids[BufferIdx] = FGuid();
+	}
+
+	// Free platform specific resources
+	PlatformData.Reset();
+}
+
+bool USharedMemoryMediaCapture::ShouldCaptureRHIResource() const
+{
+	return true;
+}
+
+FIntPoint USharedMemoryMediaCapture::GetCustomOutputSize(const FIntPoint& InSize) const
+{
+	// We pass back the desired size
+	return InSize;
+}
+
+EMediaCaptureResourceType USharedMemoryMediaCapture::GetCustomOutputResourceType() const
+{
+	return EMediaCaptureResourceType::Texture;
+}
+
+void USharedMemoryMediaCapture::OnCustomCapture_RenderingThread(
+	FRDGBuilder& GraphBuilder, 
+	const FCaptureBaseData& InBaseData, 
+	TSharedPtr<FMediaCaptureUserData, ESPMode::ThreadSafe> InUserData, 
+	FRDGTextureRef InSourceTexture, 
+	FRDGTextureRef OutputTexture, 
+	const FRHICopyTextureInfo& CopyInfo, 
+	FVector2D CropU, 
+	FVector2D CropV)
+{
+	RDG_GPU_STAT_SCOPE(GraphBuilder, SharedMemory_Capture)
+	TRACE_CPUPROFILER_EVENT_SCOPE(USharedMemoryMediaCapture::OnCustomCapture_RenderingThread);
+
+	// Initialize shared gpu textures if needed.
+
+	check(PlatformData.IsValid());
+
+	for (uint32 Idx = 0; Idx < NUMBUFFERS; Idx++)
+	{
+		if (!SharedCrossGpuTextures[Idx].IsValid())
+		{
+			const FGuid Guid = FGuid::NewGuid();
+
+			SharedCrossGpuTextures[Idx] = PlatformData->CreateSharedCrossGpuTexture(
+				InSourceTexture->Desc.Format, 
+				CopyInfo.Size.X, 
+				CopyInfo.Size.Y, 
+				Guid, 
+				Idx
+			);
+
+			check(SharedCrossGpuTextures[Idx].IsValid());
+
+			SharedCrossGpuTextureGuids[Idx] = Guid;
+			UE_LOG(LogDisplayClusterMedia, Verbose, TEXT("Created SharedGpuTextureGuid[%d] = %s"), Idx, *SharedCrossGpuTextureGuids[Idx].ToString());
+		}
+	}
+
+	// Add the copy texture pass
+	AddCopyToSharedGpuTexturePass(GraphBuilder, InSourceTexture, GFrameCounterRenderThread % NUMBUFFERS);
+}
+
+BEGIN_SHADER_PARAMETER_STRUCT(FCopyToSharedGpuTexturePass, )
+	RDG_TEXTURE_ACCESS(SrcTexture, ERHIAccess::CopySrc)
+	RDG_TEXTURE_ACCESS(DstTexture, ERHIAccess::CopyDest)
+END_SHADER_PARAMETER_STRUCT()
+
+void USharedMemoryMediaCapture::AddCopyToSharedGpuTexturePass(FRDGBuilder& GraphBuilder, FRDGTextureRef InSourceTexture, uint32 SharedTextureIdx)
+{
+	using namespace UE::SharedMemoryMedia;
+
+	FCopyToSharedGpuTexturePass* PassParameters = GraphBuilder.AllocParameters<FCopyToSharedGpuTexturePass>();
+
+	PassParameters->SrcTexture = InSourceTexture;
+
+	PassParameters->DstTexture = GraphBuilder.RegisterExternalTexture(
+		CreateRenderTarget(SharedCrossGpuTextures[SharedTextureIdx], *FString::Printf(TEXT("SharedCrossGpuTextures_%d"), SharedTextureIdx)));
+
+	GraphBuilder.AddPass(
+		RDG_EVENT_NAME("Copy%sToSharedGpuTexture", InSourceTexture->Name),
+		PassParameters,
+		ERDGPassFlags::Copy,
+		[InSourceTexture, SharedTextureIdx, this](FRHICommandList& RHICmdList)
+		{
+			// bTextureReadyFenceBusy will also signal that the resource is safe to reuse
+			if (bTextureReadyFenceBusy[SharedTextureIdx])
+			{
+				TRACE_CPUPROFILER_EVENT_SCOPE(SharedMemMediaOutputFenceBusy);
+
+				UE_LOG(LogDisplayClusterMedia, Warning, TEXT("bTextureReadyFenceBusy[%d] for frame %d was busy, so we wait"), SharedTextureIdx, GFrameCounterRenderThread);
+
+				while (bTextureReadyFenceBusy[SharedTextureIdx])
+				{
+					FPlatformProcess::SleepNoStats(SpinWaitTimeSeconds);
+				}
+			}
+
+			// This flag will be cleared by the async task when the receiver is done with the shared cross gpu texture
+			bTextureReadyFenceBusy[SharedTextureIdx] = true;
+
+			// Do the copy
+			FRHICopyTextureInfo CopyInfo;
+			RHICmdList.CopyTexture(InSourceTexture->GetRHI(), SharedCrossGpuTextures[SharedTextureIdx], CopyInfo);
+
+			// Write GPU fence
+			RHICmdList.WriteGPUFence(TextureReadyFences[SharedTextureIdx]);
+
+			// Spawn a thread that via shared ram will notify receiver that data is ready
+			// It will also verify that the data has been consumed (with a timeout).
+			RunningTasksCount++;
+			UE::Tasks::Launch(UE_SOURCE_LOCATION, [FrameNumber = GFrameCounterRenderThread, SharedTextureIdx, this]()
+			{
+				FRunWhenGoingOutOfScope DecrementRunningTasksCount([this]()
+					{
+						RunningTasksCount--;
+					});
+
+				const FString CopyThreadName = FString::Printf(TEXT("SharedMemMediaOutputGpuTextureInTransitForFrame_%d"), FrameNumber);
+				TRACE_CPUPROFILER_EVENT_SCOPE_TEXT(*CopyThreadName);
+
+				// Wait for fence that indicates that the gpu texture has the data
+				{
+					TRACE_CPUPROFILER_EVENT_SCOPE(WaitForGpuTextureReadyFence);
+
+					while (TextureReadyFences[SharedTextureIdx] && !TextureReadyFences[SharedTextureIdx]->Poll())
+					{
+						FPlatformProcess::SleepNoStats(0);
+					}
+				}
+
+				// Update shared memory metadata to indicate to the receiver that there is new data
+
+				FSharedMemoryMediaFrameMetadata* SharedMemoryData = static_cast<FSharedMemoryMediaFrameMetadata*>(SharedMemory[SharedTextureIdx]->GetAddress());
+
+				FSharedMemoryMediaFrameMetadata FrameMetadata;
+				FrameMetadata.Sender.FrameNumber = FrameNumber;
+				FrameMetadata.Sender.TextureGuid = SharedCrossGpuTextureGuids[SharedTextureIdx];
+
+				// We only send the sender structure
+				FMemory::Memcpy(SharedMemoryData, &FrameMetadata, sizeof(FSharedMemoryMediaFrameMetadata::FSender));
+
+				// Wait for FrameNumber ack
+
+				if (SharedMemoryData->Receiver.KeepAliveShiftRegister)
+				{
+					TRACE_CPUPROFILER_EVENT_SCOPE(WaitForGpuTextureAck)
+
+					const double StartTimeSeconds = FPlatformTime::Seconds();
+					constexpr double TimeoutSeconds = 0.5;
+
+					while (SharedMemoryData->Receiver.FrameNumberAcked < FrameNumber && SharedMemoryData->Receiver.KeepAliveShiftRegister)
+					{
+						FPlatformProcess::SleepNoStats(SpinWaitTimeSeconds);
+
+						if ((FPlatformTime::Seconds() - StartTimeSeconds) > TimeoutSeconds)
+						{
+							// @todo use proper log category
+							UE_LOG(LogDisplayClusterMedia, Warning, TEXT("FSharedMemoryMediaCapture timed out waiting for its receiver to ack frame %d"), FrameNumber);
+
+							break;
+						}
+					};
+				}
+
+				// Shift the keep alive. The bit depth of the keep alive is the number of frames the receiver has to re-set it before it expires.
+				SharedMemoryData->Receiver.KeepAliveShiftRegister >>= 1;
+
+				// Clear fence and flag that we're ready for a new frame
+
+				TextureReadyFences[SharedTextureIdx]->Clear();
+				bTextureReadyFenceBusy[SharedTextureIdx] = false;
+			});
+		});
+}
diff --git a/Engine/Plugins/Runtime/nDisplay/Source/DisplayClusterMedia/Private/Capture/SharedMemoryMediaOutput.cpp b/Engine/Plugins/Runtime/nDisplay/Source/DisplayClusterMedia/Private/Capture/SharedMemoryMediaOutput.cpp
new file mode 100644
--- /dev/null
+++ b/Engine/Plugins/Runtime/nDisplay/Source/DisplayClusterMedia/Private/Capture/SharedMemoryMediaOutput.cpp
@@ -0,0 +1,44 @@
+// Copyright Epic Games, Inc. All Rights Reserved.
+
+#include "SharedMemoryMediaOutput.h"
+
+#include "MediaOutput.h"
+#include "SharedMemoryMediaCapture.h"
+#include "SharedMemoryMediaTypes.h"
+
+bool USharedMemoryMediaOutput::Validate(FString& OutFailureReason) const
+{
+	if (!Super::Validate(OutFailureReason))
+	{
+		return false;
+	}
+	
+	return true;
+}
+
+FIntPoint USharedMemoryMediaOutput::GetRequestedSize() const
+{
+	return RequestCaptureSourceSize;
+}
+
+EPixelFormat USharedMemoryMediaOutput::GetRequestedPixelFormat() const
+{
+	// We are copying to a shared cross gpu texture allocated outside the MF, so this PF needs to be valid but is not relevant.
+	// @todo is there a way to not have the MF render target allocated since we're not using it ?
+	return EPixelFormat::PF_A8;
+}
+
+EMediaCaptureConversionOperation USharedMemoryMediaOutput::GetConversionOperation(EMediaCaptureSourceType InSourceType) const
+{
+	return EMediaCaptureConversionOperation::CUSTOM;
+}
+
+UMediaCapture* USharedMemoryMediaOutput::CreateMediaCaptureImpl()
+{
+	UMediaCapture* Result = NewObject<USharedMemoryMediaCapture>();
+	if (Result)
+	{
+		Result->SetMediaOutput(this);
+	}
+	return Result;
+}
diff --git a/Engine/Plugins/Runtime/nDisplay/Source/DisplayClusterMedia/Private/DisplayClusterMediaModule.cpp b/Engine/Plugins/Runtime/nDisplay/Source/DisplayClusterMedia/Private/DisplayClusterMediaModule.cpp
--- a/Engine/Plugins/Runtime/nDisplay/Source/DisplayClusterMedia/Private/DisplayClusterMediaModule.cpp
+++ b/Engine/Plugins/Runtime/nDisplay/Source/DisplayClusterMedia/Private/DisplayClusterMediaModule.cpp
@@ -20,7 +20,9 @@
 #include "Capture/DisplayClusterMediaCaptureViewport.h"
 #include "Input/DisplayClusterMediaInputNode.h"
 #include "Input/DisplayClusterMediaInputViewport.h"
+#include "Input/SharedMemoryMediaPlayerFactory.h"
 
+#include "IMediaModule.h"
 #include "Misc/CoreDelegates.h"
 
 
@@ -30,12 +32,34 @@ void FDisplayClusterMediaModule::StartupModule()
 
 	IDisplayCluster::Get().GetCallbacks().OnDisplayClusterCustomPresentSet().AddRaw(this, &FDisplayClusterMediaModule::OnCustomPresentSet);
 	FCoreDelegates::OnEnginePreExit.AddRaw(this, &FDisplayClusterMediaModule::OnEnginePreExit);
+
+	// Create Player Factories
+	PlayerFactories.Add(MakeUnique<FSharedMemoryMediaPlayerFactory>());
+
+
+	// register share memory player factory
+	if (IMediaModule* MediaModule = FModuleManager::LoadModulePtr<IMediaModule>("Media"))
+	{
+		for (TUniquePtr<IMediaPlayerFactory>& Factory : PlayerFactories)
+		{
+			MediaModule->RegisterPlayerFactory(*Factory);
+		}
+	}
 }
 
 void FDisplayClusterMediaModule::ShutdownModule()
 {
 	UE_LOG(LogDisplayClusterMedia, Log, TEXT("Shutting down module 'DisplayClusterMedia'..."));
 
+	// unregister share memory player factory
+	if (IMediaModule* MediaModule = FModuleManager::LoadModulePtr<IMediaModule>("Media"))
+	{
+		for (TUniquePtr<IMediaPlayerFactory>& Factory : PlayerFactories)
+		{
+			MediaModule->UnregisterPlayerFactory(*Factory);
+		}
+	}
+
 	IDisplayCluster::Get().GetCallbacks().OnDisplayClusterCustomPresentSet().RemoveAll(this);
 	FCoreDelegates::OnEnginePreExit.RemoveAll(this);
 }
diff --git a/Engine/Plugins/Runtime/nDisplay/Source/DisplayClusterMedia/Private/DisplayClusterMediaModule.h b/Engine/Plugins/Runtime/nDisplay/Source/DisplayClusterMedia/Private/DisplayClusterMediaModule.h
--- a/Engine/Plugins/Runtime/nDisplay/Source/DisplayClusterMedia/Private/DisplayClusterMediaModule.h
+++ b/Engine/Plugins/Runtime/nDisplay/Source/DisplayClusterMedia/Private/DisplayClusterMediaModule.h
@@ -15,7 +15,7 @@ class FDisplayClusterMediaInputBase;
 class FDisplayClusterMediaInputICVFX;
 class FDisplayClusterMediaInputNode;
 class FDisplayClusterMediaInputViewport;
-
+class IMediaPlayerFactory;
 
 /**
  * Media module
@@ -63,4 +63,7 @@ private:
 
 	// Latency queue
 	FDisplayClusterFrameQueue FrameQueue;
+
+	// Media player factories
+	TArray<TUniquePtr<IMediaPlayerFactory>> PlayerFactories;
 };
diff --git a/Engine/Plugins/Runtime/nDisplay/Source/DisplayClusterMedia/Private/Input/DisplayClusterMediaInputBase.cpp b/Engine/Plugins/Runtime/nDisplay/Source/DisplayClusterMedia/Private/Input/DisplayClusterMediaInputBase.cpp
--- a/Engine/Plugins/Runtime/nDisplay/Source/DisplayClusterMedia/Private/Input/DisplayClusterMediaInputBase.cpp
+++ b/Engine/Plugins/Runtime/nDisplay/Source/DisplayClusterMedia/Private/Input/DisplayClusterMediaInputBase.cpp
@@ -29,6 +29,7 @@ FDisplayClusterMediaInputBase::FDisplayClusterMediaInputBase(const FString& InMe
 		if (MediaTexture)
 		{
 			MediaTexture->NewStyleOutput = true;
+			MediaTexture->SetRenderMode(UMediaTexture::ERenderMode::JustInTime);
 			MediaTexture->SetMediaPlayer(MediaPlayer);
 			MediaTexture->UpdateResource();
 		}
@@ -70,7 +71,9 @@ void FDisplayClusterMediaInputBase::Stop()
 
 void FDisplayClusterMediaInputBase::ImportMediaData(FRHICommandListImmediate& RHICmdList, const FMediaTextureInfo& TextureInfo)
 {
-	FRHITexture* const SrcTexture = MediaTexture->GetResource()->GetTextureRHI();
+	MediaTexture->JustInTimeRender();
+
+	FRHITexture* const SrcTexture = MediaTexture->GetResource() ? MediaTexture->GetResource()->GetTextureRHI() : nullptr;
 	FRHITexture* const DstTexture = TextureInfo.Texture;
 
 	if (SrcTexture && DstTexture)
diff --git a/Engine/Plugins/Runtime/nDisplay/Source/DisplayClusterMedia/Private/Input/SharedMemoryMediaPlayer.cpp b/Engine/Plugins/Runtime/nDisplay/Source/DisplayClusterMedia/Private/Input/SharedMemoryMediaPlayer.cpp
new file mode 100644
--- /dev/null
+++ b/Engine/Plugins/Runtime/nDisplay/Source/DisplayClusterMedia/Private/Input/SharedMemoryMediaPlayer.cpp
@@ -0,0 +1,602 @@
+// Copyright Epic Games, Inc. All Rights Reserved.
+
+#include "SharedMemoryMediaPlayer.h"
+
+#include "Engine/Texture2D.h"
+#include "IMediaEventSink.h"
+#include "IMediaTextureSampleConverter.h"
+#include "Misc/Timespan.h"
+#include "Tasks/Task.h"
+#include "TextureResource.h"
+
+#include "DisplayClusterMediaLog.h"
+#include "SharedMemoryMediaOutput.h"
+#include "SharedMemoryMediaPlatform.h"
+#include "SharedMemoryMediaSourceOptions.h"
+#include "SharedMemoryMediaSample.h"
+#include "SharedMemoryMediaSamples.h"
+#include "SharedMemoryMediaTextureSampleConverter.h"
+#include "SharedMemoryMediaTypes.h"
+
+
+#define LOCTEXT_NAMESPACE "FSharedMemoryMediaPlayer"
+
+DECLARE_GPU_STAT_NAMED(SharedMemoryMedia_CopyToSampleCommon, TEXT("SharedMemoryMedia_CopyToSampleCommon"));
+DECLARE_GPU_STAT_NAMED(SharedMemoryMedia_UnlockTexture, TEXT("SharedMemoryMedia_UnlockTexture"));
+DECLARE_GPU_STAT_NAMED(SharedMemoryMedia_LockTexture, TEXT("SharedMemoryMedia_LockTexture"));
+
+using namespace UE::SharedMemoryMedia;
+
+int32 GDisplayClusterSharedMemoryLatency = -1;
+static FAutoConsoleVariableRef CVarDisplayClusterSharedMemoryLatency(
+	TEXT("DC.SharedMemoryMedia.Latency"),
+	GDisplayClusterSharedMemoryLatency,
+	TEXT("Forces frame latency on all shared memory media players. Valid values are -1 (don't force), 0 and 1."),
+	ECVF_RenderThreadSafe
+);
+
+
+FSharedMemoryMediaPlayer::FSharedMemoryMediaPlayer()
+	: Super()
+{
+	Samples = new FSharedMemoryMediaSamples();
+}
+
+FSharedMemoryMediaPlayer::~FSharedMemoryMediaPlayer()
+{
+	Close();
+	delete Samples;
+}
+
+bool FSharedMemoryMediaPlayer::Open(const FString& Url, const IMediaOptions* Options)
+{
+	// Get an RHI type specific implementation.
+	if (!PlatformData.IsValid())
+	{
+		const ERHIInterfaceType RhiInterfaceType = GDynamicRHI->GetInterfaceType();
+
+		PlatformData = FSharedMemoryMediaPlatform::CreateInstanceForRhi(RhiInterfaceType);
+
+		if (!PlatformData.IsValid())
+		{
+			UE_LOG(LogDisplayClusterMedia, Error, TEXT("Unfortunately, SharedMemoryMedia doesn't support the current RHI type '%s'"),
+				*FSharedMemoryMediaPlatform::GetRhiTypeString(RhiInterfaceType));
+
+			return false;
+		}
+	}
+
+	// Cache the Url for later query responses
+	OpenUrl = Url;
+
+	// Grab unique name
+	{
+		const FString EmptyUniqueName;
+		UniqueName = Options->GetMediaOption(SharedMemoryMediaOption::UniqueName, EmptyUniqueName);
+
+		if (!UniqueName.Len())
+		{
+			UE_LOG(LogDisplayClusterMedia, Error, TEXT("SharedMemoryMediaSource must have a UniqueName that is not empty, and should match the MediaOutput's UniqueName"));
+			return false;
+		}
+	}
+
+	// Grab zero latency option
+	{
+		bZeroLatency = Options->GetMediaOption(SharedMemoryMediaOption::ZeroLatency, true);
+	}
+
+	// Initialize gpu fences
+	for (int32 BufferIdx = 0; BufferIdx < NUMSHAREDMEM; BufferIdx++)
+	{
+		if (!FrameAckFences[BufferIdx])
+		{
+			FrameAckFences[BufferIdx] = RHICreateGPUFence(TEXT("SharedMemoryMediaFrameAckFence"));
+			bFrameAckFenceBusy[BufferIdx] = false;
+		}
+	}
+
+	return true;
+}
+
+bool FSharedMemoryMediaPlayer::Open(const TSharedRef<FArchive, ESPMode::ThreadSafe>& Archive, const FString& OriginalUrl, const IMediaOptions* Options)
+{
+	return false;
+}
+
+void FSharedMemoryMediaPlayer::Close()
+{
+	OpenUrl.Empty();
+	Samples->FlushSamples();
+
+	// We flush rendering commands so that we can safely release the resources
+	FlushRenderingCommands();
+
+	// Release resources related to the indexed textures
+	for (int32 BufferIdx = 0; BufferIdx < NUMSHAREDMEM; ++BufferIdx)
+	{
+		// Free cross gpu texture
+		SharedCrossGpuTextures[BufferIdx].SafeRelease();
+
+		// Free shared memory
+		if (SharedMemory[BufferIdx])
+		{
+			// Let the sender know that we're not monitoring this and that it should not wait for acks
+			{
+				FSharedMemoryMediaFrameMetadata* Data = static_cast<FSharedMemoryMediaFrameMetadata*>(SharedMemory[BufferIdx]->GetAddress());
+
+				if (Data)
+				{
+					Data->Receiver.KeepAliveShiftRegister = 0;
+				}
+			}
+
+			// Close the shared memory
+			FPlatformMemory::UnmapNamedSharedMemoryRegion(SharedMemory[BufferIdx]);
+			SharedMemory[BufferIdx] = nullptr;
+		}
+
+		// Free gpu fences
+		{
+			// Fences should not be busy since we flush rendering commands.
+			check(!bFrameAckFenceBusy[BufferIdx]);
+
+			FrameAckFences[BufferIdx].SafeRelease();
+		}
+	}
+
+	// Free sample common texture
+	SampleCommonTexture = nullptr;
+
+	// Free platform specific resources
+	PlatformData.Reset();
+}
+
+FGuid FSharedMemoryMediaPlayer::GetPlayerPluginGUID() const
+{
+	return UE::SharedMemoryMedia::PlayerGuid;
+}
+
+FString FSharedMemoryMediaPlayer::GetStats() const
+{
+	return FString();
+}
+
+FTimespan FSharedMemoryMediaPlayer::GetTime() const
+{
+	// We fabricate a time based on the frame number, to control what sample is used for which frame.
+	return FTimespan(GFrameCounter);
+}
+
+void FSharedMemoryMediaPlayer::TickFetch(FTimespan DeltaTime, FTimespan Timecode)
+{
+	Super::TickFetch(DeltaTime, Timecode);
+
+	// Make sure the sender's shared memory metadata is available. If not, don't bother creating a media sample.
+	//
+	for (int32 MemIdx = 0; MemIdx < NUMSHAREDMEM; ++MemIdx)
+	{
+		if (!SharedMemory[MemIdx])
+		{
+			const FGuid SharedMemoryGuid = UE::SharedMemoryMedia::GenerateSharedMemoryGuid(UniqueName, MemIdx);
+			const FString SharedMemoryRegionName = SharedMemoryGuid.ToString(EGuidFormats::DigitsWithHyphensInBraces);
+
+			const SIZE_T SharedMemorySize = sizeof(FSharedMemoryMediaFrameMetadata);
+
+			// Open existing shared memory region, in case it exists:
+
+			const uint32 AccessMode = FPlatformMemory::ESharedMemoryAccess::Read | FPlatformMemory::ESharedMemoryAccess::Write;
+
+			SharedMemory[MemIdx] = FPlatformMemory::MapNamedSharedMemoryRegion(
+				*SharedMemoryRegionName, false /* bCreate */, AccessMode, SharedMemorySize
+			);
+
+			if (SharedMemory[MemIdx])
+			{
+				UE_LOG(LogDisplayClusterMedia, Verbose, TEXT("Opened SharedMemory named '%s'"), *SharedMemoryRegionName);
+
+				FSharedMemoryMediaFrameMetadata* Data = static_cast<FSharedMemoryMediaFrameMetadata*>(SharedMemory[MemIdx]->GetAddress());
+				check(Data);
+			}
+			else
+			{
+				UE_LOG(LogDisplayClusterMedia, Warning, TEXT("Could not open SharedMemory named '%s'"), *SharedMemoryRegionName);
+			}
+		}
+
+		// No point in continuing if there is no shared memory region available
+		if (!SharedMemory[MemIdx])
+		{
+			return;
+		}
+	}
+
+	// Since we can't send a sample without a valid texture, we are going to wait until we get a valid shared gpu texture
+	{
+		const uint32 ExpectedFrameNumber = InputTextureFrameNumberForFrameNumber(GFrameCounter);
+		const int32 MemIdx = ExpectedFrameNumber % NUMSHAREDMEM;
+
+		FSharedMemoryMediaFrameMetadata* SharedMemoryData = static_cast<FSharedMemoryMediaFrameMetadata*>(SharedMemory[MemIdx]->GetAddress());
+
+		check(SharedMemoryData);
+
+		if (SharedMemoryData->Sender.Magic != FSharedMemoryMediaFrameMetadata::MAGIC)
+		{
+			return;
+		}
+
+		if (SharedMemoryData->Sender.TextureGuid == UE::SharedMemoryMedia::ZeroGuid)
+		{
+			return;
+		}
+	}
+
+	// Update shared textures
+	//
+	for (uint32 MemIdx = 0; MemIdx < NUMSHAREDMEM; MemIdx++)
+	{
+		// If our sender is not there, keep idling
+
+		check(SharedMemory[MemIdx]);
+		FSharedMemoryMediaFrameMetadata* SharedMemoryMetadataPtr = static_cast<FSharedMemoryMediaFrameMetadata*>(SharedMemory[MemIdx]->GetAddress());
+		check(SharedMemoryMetadataPtr);
+
+		const FSharedMemoryMediaFrameMetadata::FSender SenderData = SharedMemoryMetadataPtr->Sender;
+
+		if (SenderData.Magic != FSharedMemoryMediaFrameMetadata::MAGIC)
+		{
+			continue;
+		}
+
+		// See if we already have the gpu texture
+		if (SharedCrossGpuTextures[MemIdx].IsValid())
+		{
+			// We have a texture, now see if it needs to be updated based on the Guid
+
+			if (SenderData.TextureGuid == SharedCrossGpuTextureDescriptions[MemIdx].Guid)
+			{
+				// Same guid as before, so we continue normally
+				continue;
+			}
+
+			// The texture guid has changed, so we need to update it if there is a valid one.
+
+			if (SenderData.TextureGuid == UE::SharedMemoryMedia::ZeroGuid)
+			{
+				// Guid isn't valid anymore, so we keep idling, we expect that there will be a valid one soon.
+				continue;
+			}
+
+			// We flush rendering commands so that we can safely update
+			FlushRenderingCommands();
+
+			// Release our current gpu texture
+			SharedCrossGpuTextures[MemIdx].SafeRelease();
+
+			// Release our sample common texture as well, since its description will likely not match the cross gpu texture anymore
+			SampleCommonTexture = nullptr;
+
+			// No need for the sender to wait for us while we figure out the new textures.
+			SharedMemoryMetadataPtr->Receiver.KeepAliveShiftRegister = 0;
+
+			// We now exit this block, which will treat the invalid gpu texture case.
+		}
+
+		// Reset the cross gpu texture description
+		SharedCrossGpuTextureDescriptions[MemIdx] = FSharedMemoryMediaTextureDescription();
+
+		// Check if the shared memory has the cross gpu texture data that we need
+
+		const FGuid& SharedGpuTextureGuid = SenderData.TextureGuid;
+
+		// Guid of zero is not valid
+		if (SharedGpuTextureGuid == UE::SharedMemoryMedia::ZeroGuid)
+		{
+			UE_LOG(LogDisplayClusterMedia, Warning, TEXT("Invalid zero guid"));
+			continue;
+		}
+
+		// Now that we have the shared texture guid from the sender, use it to open the associated texture
+
+		SharedCrossGpuTextures[MemIdx] = PlatformData->OpenSharedCrossGpuTextureByGuid(SharedGpuTextureGuid, SharedCrossGpuTextureDescriptions[MemIdx]);
+
+		if (!SharedCrossGpuTextures[MemIdx].IsValid())
+		{
+			continue;
+		}
+
+		// Since the texture is valid, set the Guid in its description.
+		SharedCrossGpuTextureDescriptions[MemIdx].Guid = SenderData.TextureGuid;
+	}
+
+	// Don't create samples if we don't have shared gpu memory handles yet
+	for (uint32 MemIdx = 0; MemIdx < NUMSHAREDMEM; MemIdx++)
+	{
+		if (!SharedCrossGpuTextures[MemIdx].IsValid())
+		{
+			return;
+		}
+	}
+
+	// The descriptions of all buffers should match
+	for (uint32 MemIdx = 1; MemIdx < NUMSHAREDMEM; MemIdx++)
+	{
+		if (!SharedCrossGpuTextureDescriptions[MemIdx - 1].IsEquivalentTo(SharedCrossGpuTextureDescriptions[MemIdx]))
+		{
+			return;
+		}
+	}
+
+	// Initialize sample texture if needed
+	if (!SampleCommonTexture.IsValid())
+	{
+		SampleCommonTexture.Reset(UTexture2D::CreateTransient(
+			SharedCrossGpuTextureDescriptions[0].Width,
+			SharedCrossGpuTextureDescriptions[0].Height,
+			SharedCrossGpuTextureDescriptions[0].Format,
+			FName(TEXT("SampleCommonTexture"))
+		));
+
+		SampleCommonTexture->UpdateResource();
+	}
+
+	check(SampleCommonTexture.IsValid() && SampleCommonTexture->GetResource());
+
+	// Don't create samples if we don't have the sample common resource yet
+	if (!SampleCommonTexture->GetResource()->GetTextureRHI()) // Could be empty the first time
+	{
+		return;
+	}
+
+	// Let sender know that we're monitoring the shared memory by setting the Keepalive shift register
+	for (int32 MemIdx = 0; MemIdx < NUMSHAREDMEM; ++MemIdx)
+	{
+		FSharedMemoryMediaFrameMetadata* SharedMetadataPtr = static_cast<FSharedMemoryMediaFrameMetadata*>(SharedMemory[MemIdx]->GetAddress());
+
+		// Let sender know that we're monitoring this shared memory.
+		// We expect the sender to not wait for acks if nobody is monitoring it.
+		SharedMetadataPtr->Receiver.KeepAliveShiftRegister = ~0; // all ones
+	}
+
+	// Create a new media sample and populated with any needed data
+	{
+		Samples->CurrentSample = MakeShared<FSharedMemoryMediaSample>();
+		Samples->CurrentSample->Player = AsShared();
+		Samples->CurrentSample->Texture = SampleCommonTexture->GetResource()->GetTextureRHI();
+		Samples->CurrentSample->Dim = FIntPoint(SharedCrossGpuTextureDescriptions[0].Width, SharedCrossGpuTextureDescriptions[0].Height);
+		Samples->CurrentSample->Stride = SharedCrossGpuTextureDescriptions[0].Stride;
+
+		// Ticks are "frames" in our player. W/o constraining the sample timing, it would be possible for the
+		// render thread to purge the sample of the next frame, in particular when the game thread stalls.
+		Samples->CurrentSample->Time = FTimespan(GFrameCounter);
+	}
+}
+
+void FSharedMemoryMediaPlayer::TickInput(FTimespan DeltaTime, FTimespan Timecode)
+{
+	Super::TickInput(DeltaTime, Timecode);
+
+	// @todo update our state
+}
+
+bool FSharedMemoryMediaPlayer::GetPlayerFeatureFlag(EFeatureFlag Flag) const
+{
+	if (Super::GetPlayerFeatureFlag(Flag))
+	{
+		return true;
+	}
+
+	return false;
+}
+
+IMediaSamples& FSharedMemoryMediaPlayer::GetSamples()
+{
+	return *Samples;
+}
+
+IMediaCache& FSharedMemoryMediaPlayer::GetCache()
+{
+	return *this;
+}
+
+IMediaControls& FSharedMemoryMediaPlayer::GetControls() 
+{
+	return *this;
+}
+
+FString FSharedMemoryMediaPlayer::GetInfo() const 
+{
+	return FString();
+}
+
+IMediaTracks& FSharedMemoryMediaPlayer::GetTracks() 
+{
+	return *this;
+}
+
+FString FSharedMemoryMediaPlayer::GetUrl() const 
+{
+	return OpenUrl;
+}
+
+IMediaView& FSharedMemoryMediaPlayer::GetView()
+{
+	return *this;
+}
+
+void FSharedMemoryMediaPlayer::JustInTimeSampleRender()
+{
+	FRHICommandListImmediate& RHICmdList = FRHICommandListExecutor::GetImmediateCommandList();
+
+	// We only allow this function to run once per frame.
+	if (LastFrameNumberThatUpdatedJustInTime == GFrameCounterRenderThread)
+	{
+		UE_LOG(LogDisplayClusterMedia, Warning, 
+			TEXT("FSharedMemoryMediaPlayer::JustInTimeSampleRender called more than once in GFrameCounterRenderThread %llu"), 
+			GFrameCounterRenderThread
+		);
+
+		return;
+	}
+
+	LastFrameNumberThatUpdatedJustInTime = GFrameCounterRenderThread;
+
+	UE_LOG(LogDisplayClusterMedia, VeryVerbose, TEXT("FSharedMemoryMediaPlayer::JustInTimeSampleRender called for GFrameCounterRenderThread %llu"), 
+		GFrameCounterRenderThread);
+
+	for (int32 MemIdx = 0; MemIdx < NUMSHAREDMEM; ++MemIdx)
+	{
+		// No point in continuing if there is no shared memory region available
+		if (!SharedMemory[MemIdx])
+		{
+			return;
+		}
+	}
+
+	const uint32 ExpectedFrameNumber = InputTextureFrameNumberForFrameNumber(GFrameCounterRenderThread);
+	const uint32 SharedMemoryIdx = ExpectedFrameNumber % NUMSHAREDMEM;
+
+	FSharedMemoryMediaFrameMetadata* SharedMemoryData = static_cast<FSharedMemoryMediaFrameMetadata*>(SharedMemory[SharedMemoryIdx]->GetAddress());
+	check(SharedMemoryData);
+
+	// Enqueue a lambda that will wait for the cross gpu texture for this frame have the data populated.
+	// It will determine this by polling the sender's shared memory metadata.
+	// To avoid any hangs, it will give up after some time.
+	RHICmdList.EnqueueLambda(
+		[
+			FrameNumber = GFrameCounterRenderThread, 
+			SharedMemoryIdx, 
+			ExpectedFrameNumber, 
+			SharedMemoryData, 
+			this
+		](FRHICommandListImmediate& RHICmdList)
+		{
+			TRACE_CPUPROFILER_EVENT_SCOPE(SharedMemMediaWaitForSharedGpuTexture);
+
+			alignas(128) FSharedMemoryMediaFrameMetadata::FSender SenderMetadata;
+
+			const double StartTimeSeconds = FPlatformTime::Seconds();
+			constexpr double TimeoutSeconds = 0.5;
+
+			while (true)
+			{
+				FMemory::Memcpy(&SenderMetadata, SharedMemoryData, sizeof(FSharedMemoryMediaFrameMetadata::FSender));
+
+				if (SenderMetadata.Magic == FSharedMemoryMediaFrameMetadata::MAGIC)
+				{
+					if (SenderMetadata.FrameNumber >= ExpectedFrameNumber)
+					{
+						if (SenderMetadata.FrameNumber > ExpectedFrameNumber)
+						{
+							UE_LOG(LogDisplayClusterMedia, Warning, TEXT("Using too recent frame. Expected %u and used %u"), 
+								ExpectedFrameNumber, SenderMetadata.FrameNumber);
+						}
+
+						break;
+					}
+				}
+
+				FPlatformProcess::SleepNoStats(SpinWaitTimeSeconds);
+
+				if ((FPlatformTime::Seconds() - StartTimeSeconds) > TimeoutSeconds)
+				{
+					UE_LOG(LogDisplayClusterMedia, Warning, TEXT("FSharedMemoryMediaPlayer timed out waiting for ExpectedFrameNumber %u for frame %llu, only saw up to frame %u"),
+						ExpectedFrameNumber, FrameNumber, SenderMetadata.FrameNumber);
+
+					// Note: It would be desirable to stop the copy texture from happening, but at this point it has already been enqueued.
+
+					break;
+				}
+			}
+		}
+	);
+
+	// Copy to sample common texture
+
+	check(SampleCommonTexture.IsValid() && SampleCommonTexture->GetResource());
+
+	// But first wait for the fence to be cleared, in case it hasn't (which is not usual).
+	if (bFrameAckFenceBusy[SharedMemoryIdx])
+	{
+		UE_LOG(LogDisplayClusterMedia, Warning, TEXT("bFrameAckFenceBusy[%d] was busy for ExpectedFrameNumber %d"),
+			SharedMemoryIdx, ExpectedFrameNumber);
+
+		TRACE_CPUPROFILER_EVENT_SCOPE(SharedMemMediaWaitForFrameAckFenceBusyToClear);
+
+		while (bFrameAckFenceBusy[SharedMemoryIdx]) // The flag is guaranteed to eventually clear by the async task
+		{
+			FPlatformProcess::SleepNoStats(SpinWaitTimeSeconds);
+		}
+	}
+
+	// Enqueue the copy from cross gpu texture to sample common texture
+	if (SharedCrossGpuTextures[SharedMemoryIdx].IsValid())
+	{
+		FRHICopyTextureInfo CopyInfo;
+
+		SCOPED_GPU_STAT(RHICmdList, SharedMemoryMedia_CopyToSampleCommon);
+		SCOPED_DRAW_EVENT(RHICmdList, SharedMemoryMedia_CopyToSampleCommon);
+
+		RHICmdList.CopyTexture(SharedCrossGpuTextures[SharedMemoryIdx], SampleCommonTexture->GetResource()->GetTextureRHI(), CopyInfo);
+	}
+	else
+	{
+		UE_LOG(LogDisplayClusterMedia, Error, TEXT("SharedCrossGpuTextures[%d] was unexpectedly invalid"), SharedMemoryIdx);
+	}
+
+	// Write gpu fence to indicate we're done with the shared cross gpu texture
+	{
+		// This is our atomic flag that avoids this thread from re-using the fence before it is cleared.
+		// It will be cleared by the async task.
+		bFrameAckFenceBusy[SharedMemoryIdx] = true;
+
+		RHICmdList.WriteGPUFence(FrameAckFences[SharedMemoryIdx]);
+	}
+
+	// spawn async task that will wait on FrameAckFence, ack the frame to the Media Capture, and flag that it is clear for re-use
+	UE::Tasks::Launch(UE_SOURCE_LOCATION,
+		[SharedMemoryData, ExpectedFrameNumber, SharedMemoryIdx, this]()
+		{
+			TRACE_CPUPROFILER_EVENT_SCOPE(SharedMemoryMediaFrameAckTask);
+
+			do
+			{
+				const bool bFenceReached = !FrameAckFences[SharedMemoryIdx] || FrameAckFences[SharedMemoryIdx]->Poll();
+
+				if (bFenceReached)
+				{
+					break;
+				}
+
+				FPlatformProcess::SleepNoStats(SpinWaitTimeSeconds);
+
+			} while (true);
+
+			// We clear the fence and signal that it can be re-used.
+			FrameAckFences[SharedMemoryIdx]->Clear();
+			bFrameAckFenceBusy[SharedMemoryIdx] = false;
+
+			// Ack the frame to the Media Capture, indicating that the cross gpu texture can be re-used.
+			SharedMemoryData->Receiver.FrameNumberAcked = ExpectedFrameNumber;
+		}
+	);
+}
+
+uint32 FSharedMemoryMediaPlayer::InputTextureFrameNumberForFrameNumber(uint32 FrameNumber) const
+{
+	int32 Latency = bZeroLatency ? 0 : 1;
+
+	if (GDisplayClusterSharedMemoryLatency == 0)
+	{
+		Latency = 0;
+	}
+	else if (GDisplayClusterSharedMemoryLatency == 1)
+	{
+		Latency = 1;
+	}
+
+	return FrameNumber - Latency;
+}
+
+
+#undef LOCTEXT_NAMESPACE
diff --git a/Engine/Plugins/Runtime/nDisplay/Source/DisplayClusterMedia/Private/Input/SharedMemoryMediaPlayer.h b/Engine/Plugins/Runtime/nDisplay/Source/DisplayClusterMedia/Private/Input/SharedMemoryMediaPlayer.h
new file mode 100644
--- /dev/null
+++ b/Engine/Plugins/Runtime/nDisplay/Source/DisplayClusterMedia/Private/Input/SharedMemoryMediaPlayer.h
@@ -0,0 +1,258 @@
+// Copyright Epic Games, Inc. All Rights Reserved.
+
+#pragma once
+
+#include "CoreMinimal.h"
+#include "GenericPlatform/GenericPlatformMemory.h"
+#include "MediaIOCorePlayerBase.h"
+#include "MediaTexture.h"
+#include "SharedMemoryMediaSource.h"
+#include "SharedMemoryMediaTypes.h"
+#include "UObject/StrongObjectPtr.h"
+
+
+class UTexture2D;
+class FSharedMemoryMediaPlatform;
+
+/**
+ * Implements a media player using shared cross gpu textures. See USharedMemoryMediaCapture docstring for complementary information.
+ * It will open shared system memory based on its UniqueName setting, which much match that of the corresponding Media Output.
+ * This data will inform the Guid of the shared cross gpu textures and other metadata, and it will use this to open them
+ * and to detect when the textures are still valid.
+ * 
+ * It will look at the frame numbers in the metadata when deciding if the data that it is looking for is ready.
+ * 
+ * It informs the sender that it is actively reading the data using a keep alive shift register that it must continually refresh.
+ * Also, it acknowledges that it is done with a given gpu texture by setting the frame number in said shared system memory.
+ * 
+ * Even if the sender closes the textures, the open handles in the player will keep them valid and avoid undue crashes.
+ */
+class FSharedMemoryMediaPlayer 
+	: public IMediaPlayer
+	, protected IMediaCache
+	, protected IMediaControls
+	, protected IMediaTracks
+	, protected IMediaView
+	, public TSharedFromThis<FSharedMemoryMediaPlayer>
+{
+	using Super = IMediaPlayer;
+
+	// @todo remove these friends and create needed getters/setters
+	friend class FSharedMemoryMediaSample;
+	friend class FSharedMemoryMediaTextureSampleConverter;
+
+public:
+
+	FSharedMemoryMediaPlayer();
+
+	virtual ~FSharedMemoryMediaPlayer();
+
+public:
+
+	//~ Begin IMediaPlayer interface
+	virtual void Close() override;
+	
+	virtual FGuid GetPlayerPluginGUID() const override;
+
+	virtual bool Open(const FString& Url, const IMediaOptions* Options) override;
+	virtual bool Open(const TSharedRef<FArchive, ESPMode::ThreadSafe>& Archive, const FString& OriginalUrl, const IMediaOptions* Options) override;
+
+	virtual void TickFetch(FTimespan DeltaTime, FTimespan Timecode) override;
+	virtual void TickInput(FTimespan DeltaTime, FTimespan Timecode) override;
+	virtual bool GetPlayerFeatureFlag(EFeatureFlag Flag) const;
+
+	virtual FString GetStats() const override;
+
+	virtual IMediaSamples& GetSamples() override;
+
+	virtual IMediaCache& GetCache() override;
+	virtual IMediaControls& GetControls() override;
+	virtual FString GetInfo() const override;
+	virtual IMediaTracks& GetTracks() override;
+	virtual FString GetUrl() const override;
+	virtual IMediaView& GetView() override;
+
+	//~ End IMediaPlayer interface
+
+public:
+	//~ IMediaCache interface
+
+	virtual bool QueryCacheState(EMediaCacheState State, TRangeSet<FTimespan>& OutTimeRanges) const override 
+	{ 
+		return true; 
+	}
+
+	virtual int32 GetSampleCount(EMediaCacheState State) const override 
+	{
+		return 1;
+	}
+
+protected:
+
+	//~ IMediaControls interface
+
+	virtual bool CanControl(EMediaControl Control) const override
+	{
+		return false;
+	}
+
+	virtual FTimespan GetDuration() const override
+	{
+		return FTimespan::MaxValue();
+	}
+
+	virtual float GetRate() const override
+	{
+		return 1.0;
+	}
+
+	virtual EMediaState GetState() const override
+	{
+		// @todo fix this
+		return EMediaState::Playing;
+	}
+
+	virtual EMediaStatus GetStatus() const override
+	{
+		return EMediaStatus::None;
+	}
+
+	virtual TRangeSet<float> GetSupportedRates(EMediaRateThinning Thinning) const override
+	{
+		TRangeSet<float> Result;
+		Result.Add(TRange<float>(1.0f));
+		return Result;
+	}
+
+	virtual FTimespan GetTime() const override;
+
+	virtual bool IsLooping() const override
+	{
+		return false;
+	}
+
+	virtual bool Seek(const FTimespan& Time) override
+	{
+		return false;
+	}
+
+	virtual bool SetLooping(bool Looping) override
+	{
+		return false;
+	}
+
+	virtual bool SetRate(float Rate) override
+	{
+		return FMath::IsNearlyEqual(Rate, 1.0f);
+	}
+
+protected:
+
+	//~ IMediaTracks interface
+
+	virtual bool GetAudioTrackFormat(int32 TrackIndex, int32 FormatIndex, FMediaAudioTrackFormat& OutFormat) const override
+	{
+		return false;
+	}
+
+	virtual int32 GetNumTracks(EMediaTrackType TrackType) const override
+	{
+		return 0;
+	}
+
+	virtual int32 GetNumTrackFormats(EMediaTrackType TrackType, int32 TrackIndex) const override
+	{
+		return 0;
+	}
+
+	virtual int32 GetSelectedTrack(EMediaTrackType TrackType) const override
+	{
+		return INDEX_NONE;
+	}
+
+	virtual FText GetTrackDisplayName(EMediaTrackType TrackType, int32 TrackIndex) const override
+	{
+		return FText::GetEmpty();
+	}
+
+	virtual int32 GetTrackFormat(EMediaTrackType TrackType, int32 TrackIndex) const override
+	{
+		return INDEX_NONE;
+	}
+
+	virtual FString GetTrackLanguage(EMediaTrackType TrackType, int32 TrackIndex) const override
+	{
+		return FString();
+	}
+
+	virtual FString GetTrackName(EMediaTrackType TrackType, int32 TrackIndex) const override
+	{
+		return FString();
+	}
+
+	virtual bool GetVideoTrackFormat(int32 TrackIndex, int32 FormatIndex, FMediaVideoTrackFormat& OutFormat) const override
+	{
+		return false;
+	}
+
+	virtual bool SelectTrack(EMediaTrackType TrackType, int32 TrackIndex) override
+	{
+		return false;
+	}
+
+	virtual bool SetTrackFormat(EMediaTrackType TrackType, int32 TrackIndex, int32 FormatIndex) override
+	{
+		return false;
+	}
+	
+protected:
+
+	/** Enqueues a last minute wait and copy of the expected texture data for the renderer to use */
+	virtual void JustInTimeSampleRender();
+
+private:
+
+	/** Offsets the frame number by the specified latency. Used to associate resources with frame numbers being evaluated */
+	uint32 InputTextureFrameNumberForFrameNumber(uint32 FrameNumber) const;
+
+private:
+
+	/** Number of resources used for pipelining the data flow */
+	static constexpr int32 NUMSHAREDMEM = UE::SharedMemoryMedia::SenderNumBuffers;
+
+	/** Platform specific data or resources */
+	TSharedPtr<FSharedMemoryMediaPlatform, ESPMode::ThreadSafe> PlatformData;
+
+	/** Set of cross Gpu textures that the Media Capture allocates and the player opens */
+	FTextureRHIRef SharedCrossGpuTextures[NUMSHAREDMEM];
+
+	/** Shared memory allocated by the Media Capture for IPC purposes with the Media Player */
+	FPlatformMemory::FSharedMemoryRegion* SharedMemory[NUMSHAREDMEM] = { 0 };
+
+	/** Indicates that we're done reading the SharedCrossGpuTexture and we can ack the frame to the Media Capture for re-use */
+	FGPUFenceRHIRef FrameAckFences[NUMSHAREDMEM];
+
+	/** Needs to be cleared before the corresponding FrameAckFence is reused. */
+	std::atomic<bool> bFrameAckFenceBusy[NUMSHAREDMEM] = { 0 };
+
+	/** For all media samples to re-use */
+	TStrongObjectPtr<UTexture2D> SampleCommonTexture; // @todo we sure we don't need 2 to avoid tearing
+
+	/** Used to ensure that JustInTimeSampleRender is only run once per frame */
+	uint32 LastFrameNumberThatUpdatedJustInTime = 0;
+
+	/** The Media Samples used by this player */
+	class FSharedMemoryMediaSamples* Samples = nullptr;
+
+	/** Cached Url currently playing. Used for queries from the outside. */
+	FString OpenUrl;
+
+	/** Description of SharedCrossGpuTextures */
+	FSharedMemoryMediaTextureDescription SharedCrossGpuTextureDescriptions[NUMSHAREDMEM];
+
+	/** Unique Name that must match the corresponding MediaOutput setting. It is used to find allocated shared memory by name */
+	FString UniqueName;
+
+	/** Zero latency option to wait for the cross gpu texture rendered on the same frame. May adversely affect fps */
+	bool bZeroLatency = true;
+};
diff --git a/Engine/Plugins/Runtime/nDisplay/Source/DisplayClusterMedia/Private/Input/SharedMemoryMediaPlayerFactory.cpp b/Engine/Plugins/Runtime/nDisplay/Source/DisplayClusterMedia/Private/Input/SharedMemoryMediaPlayerFactory.cpp
new file mode 100644
--- /dev/null
+++ b/Engine/Plugins/Runtime/nDisplay/Source/DisplayClusterMedia/Private/Input/SharedMemoryMediaPlayerFactory.cpp
@@ -0,0 +1,81 @@
+// Copyright Epic Games, Inc. All Rights Reserved.
+
+#include "SharedMemoryMediaPlayerFactory.h"
+
+#include "SharedMemoryMediaPlayer.h"
+
+
+#define LOCTEXT_NAMESPACE "SharedMemoryMediaFactory"
+
+
+FSharedMemoryMediaPlayerFactory::FSharedMemoryMediaPlayerFactory()
+{
+	// supported platforms
+	SupportedPlatforms.Add(TEXT("Windows"));
+
+	// supported schemes
+	SupportedUriSchemes.Add(TEXT("dcsm"));
+}
+
+bool FSharedMemoryMediaPlayerFactory::CanPlayUrl(const FString& Url, const IMediaOptions* Options, TArray<FText>* OutWarnings, TArray<FText>* OutErrors) const
+{
+	FString Scheme;
+	FString Location;
+
+	// extract scheme
+	if (!Url.Split(TEXT("://"), &Scheme, &Location, ESearchCase::CaseSensitive))
+	{
+		if (OutErrors != nullptr)
+		{
+			OutErrors->Add(LOCTEXT("NoSchemeFound", "No URI scheme found"));
+		}
+
+		return false;
+	}
+
+	// see if the scheme is supported
+	if (!SupportedUriSchemes.Contains(Scheme))
+	{
+		if (OutErrors != nullptr)
+		{
+			OutErrors->Add(FText::Format(LOCTEXT("SchemeNotSupported", "The URI scheme '{0}' is not supported"), FText::FromString(Scheme)));
+		}
+
+		return false;
+	}
+
+	return true;
+}
+
+TSharedPtr<IMediaPlayer, ESPMode::ThreadSafe> FSharedMemoryMediaPlayerFactory::CreatePlayer(IMediaEventSink& EventSink)
+{
+	return MakeShared<FSharedMemoryMediaPlayer>();
+}
+
+FText FSharedMemoryMediaPlayerFactory::GetDisplayName() const
+{
+	return LOCTEXT("MediaPlayerDisplayName", "Shared Memory Device Interface");
+}
+
+FName FSharedMemoryMediaPlayerFactory::GetPlayerName() const
+{
+	static FName PlayerName(TEXT("SharedMemoryMedia"));
+	return PlayerName;
+}
+
+FGuid FSharedMemoryMediaPlayerFactory::GetPlayerPluginGUID() const
+{
+	return UE::SharedMemoryMedia::PlayerGuid;
+}
+
+const TArray<FString>& FSharedMemoryMediaPlayerFactory::GetSupportedPlatforms() const
+{
+	return SupportedPlatforms;
+}
+
+bool FSharedMemoryMediaPlayerFactory::SupportsFeature(EMediaFeature Feature) const
+{
+	return Feature == EMediaFeature::VideoSamples;
+}
+
+#undef LOCTEXT_NAMESPACE
diff --git a/Engine/Plugins/Runtime/nDisplay/Source/DisplayClusterMedia/Private/Input/SharedMemoryMediaPlayerFactory.h b/Engine/Plugins/Runtime/nDisplay/Source/DisplayClusterMedia/Private/Input/SharedMemoryMediaPlayerFactory.h
new file mode 100644
--- /dev/null
+++ b/Engine/Plugins/Runtime/nDisplay/Source/DisplayClusterMedia/Private/Input/SharedMemoryMediaPlayerFactory.h
@@ -0,0 +1,40 @@
+// Copyright Epic Games, Inc. All Rights Reserved.
+
+#pragma once
+
+#include "CoreMinimal.h"
+
+#include "IMediaPlayerFactory.h"
+
+
+/**
+ *  Factory for the shared memory media player.
+ */
+class FSharedMemoryMediaPlayerFactory : public IMediaPlayerFactory
+{
+
+private:
+	/** List of platforms that the media player support. */
+	TArray<FString> SupportedPlatforms;
+
+	/** List of supported URI schemes. */
+	TArray<FString> SupportedUriSchemes;
+
+public:
+
+	FSharedMemoryMediaPlayerFactory();
+
+	//~ Begin IMediaPlayerFactory interface
+
+	virtual bool CanPlayUrl(const FString& Url, const IMediaOptions* Options, TArray<FText>* OutWarnings, TArray<FText>* OutErrors) const override;
+	virtual TSharedPtr<IMediaPlayer, ESPMode::ThreadSafe> CreatePlayer(IMediaEventSink& EventSink) override;
+	virtual FText GetDisplayName() const override;
+	virtual FName GetPlayerName() const override;
+	virtual FGuid GetPlayerPluginGUID() const override;
+	virtual const TArray<FString>& GetSupportedPlatforms() const override;
+	virtual bool SupportsFeature(EMediaFeature Feature) const override;
+
+	//~ End IMediaPlayerFactory interface
+
+};
+
diff --git a/Engine/Plugins/Runtime/nDisplay/Source/DisplayClusterMedia/Private/Input/SharedMemoryMediaSample.h b/Engine/Plugins/Runtime/nDisplay/Source/DisplayClusterMedia/Private/Input/SharedMemoryMediaSample.h
new file mode 100644
--- /dev/null
+++ b/Engine/Plugins/Runtime/nDisplay/Source/DisplayClusterMedia/Private/Input/SharedMemoryMediaSample.h
@@ -0,0 +1,136 @@
+// Copyright Epic Games, Inc. All Rights Reserved.
+
+#pragma once
+
+#include "CoreMinimal.h"
+
+#include "DisplayClusterMediaLog.h"
+#include "IMediaTextureSample.h"
+#include "SharedMemoryMediaTextureSampleConverter.h"
+
+
+class FSharedMemoryMediaPlayer;
+
+/** This is the mostly empty sample that the player generates each frame. 
+ *  Its job is mainly to provide the sample converter that will call back the player
+ *  when rendering the sample, and fill in the data just in time.
+ */
+class FSharedMemoryMediaSample : public IMediaTextureSample
+{
+
+public:
+
+	/** The player that created this sample */
+	TWeakPtr<FSharedMemoryMediaPlayer, ESPMode::ThreadSafe> Player;
+
+	/** Custom converter that will be the one checking back with the player for just in time sample render purposes */
+	TSharedPtr<FSharedMemoryMediaTextureSampleConverter, ESPMode::ThreadSafe> Converter;
+
+	/** The texture that will get its pixels just in time when playing back this sample */
+	FTextureRHIRef Texture;
+
+	/** Player time for this sample to be evaluated at */
+	FTimespan Time;
+
+	/** Texture stride */
+	uint32 Stride = 0;
+
+	/** Texture dimensions */
+	FIntPoint Dim = FIntPoint::ZeroValue;
+
+private:
+
+	/** Used in case the Texture is not valid, so we can provide pixels to the Media Framework */
+	TArray<uint8> BackupData;
+
+public:
+
+	//~ Begin IMediaTextureSample interface
+
+	const void* GetBuffer() override
+	{
+		// This will get called if there wasn't a gpu texture available. 
+		// If we return null, FMediaTextureResource::ConvertSample will crash, so we log the error but return something valid.
+
+		UE_LOG(LogDisplayClusterMedia, Warning, 
+			TEXT("FSharedMemoryMediaSample::GetBuffer was called, so it must not have had a valid texture, which is unexpected"));
+
+		const int32 NumBytes = GetStride() * GetDim().Y;
+
+		if (BackupData.Num() < NumBytes)
+		{
+			BackupData.AddDefaulted(NumBytes - BackupData.Num());
+		}
+
+		return BackupData.GetData();
+	}
+
+	FIntPoint GetDim() const override
+	{
+		return Dim;
+	}
+
+	FTimespan GetDuration() const override
+	{
+		// Ticks are "frames" in our player.
+		return FTimespan(1);
+	}
+
+	EMediaTextureSampleFormat GetFormat() const override
+	{
+		// All output types go through buffer conversion
+		return EMediaTextureSampleFormat::CharBGR10A2;
+	}
+
+	FIntPoint GetOutputDim() const override
+	{
+		return GetDim();
+	}
+
+	uint32 GetStride() const override
+	{
+		return Stride;
+	}
+
+	FMediaTimeStamp GetTime() const override
+	{
+		return FMediaTimeStamp(Time);
+	}
+
+	bool IsCacheable() const override
+	{
+		return false;
+	}
+
+	bool IsOutputSrgb() const override
+	{
+		return true;
+	}
+
+public:
+
+#if WITH_ENGINE
+
+	virtual IMediaTextureSampleConverter* GetMediaTextureSampleConverter() override
+	{
+		if (!Converter.IsValid())
+		{
+			Converter = MakeShared<FSharedMemoryMediaTextureSampleConverter>();
+		}
+
+		check(Converter.IsValid());
+
+		Converter->Player = Player.Pin().Get();
+
+		return Converter.Get();
+	}
+
+#endif // WITH_ENGINE
+
+	virtual FRHITexture* GetTexture() const override
+	{
+		return Texture;
+	}
+
+	//~ End IMediaTextureSample interface
+};
diff --git a/Engine/Plugins/Runtime/nDisplay/Source/DisplayClusterMedia/Private/Input/SharedMemoryMediaSamples.h b/Engine/Plugins/Runtime/nDisplay/Source/DisplayClusterMedia/Private/Input/SharedMemoryMediaSamples.h
new file mode 100644
--- /dev/null
+++ b/Engine/Plugins/Runtime/nDisplay/Source/DisplayClusterMedia/Private/Input/SharedMemoryMediaSamples.h
@@ -0,0 +1,57 @@
+// Copyright Epic Games, Inc. All Rights Reserved.
+
+#pragma once
+
+#include "CoreMinimal.h"
+#include "IMediaSamples.h"
+
+/** This is the mostly empty media samples that only carries a single sample per frame,
+ *  that the player provides.
+ */
+class FSharedMemoryMediaSamples : public IMediaSamples
+{
+public:
+
+	FSharedMemoryMediaSamples()
+	{
+
+	}
+
+	FSharedMemoryMediaSamples(const FSharedMemoryMediaSamples&) = delete;
+	FSharedMemoryMediaSamples& operator=(const FSharedMemoryMediaSamples&) = delete;
+
+public:
+
+	//~ Begin IMediaSamples interface
+
+	virtual bool FetchVideo(TRange<FTimespan> TimeRange, TSharedPtr<IMediaTextureSample, ESPMode::ThreadSafe>& OutSample) override
+	{
+		if (!CurrentSample.IsValid())
+		{
+			return false;
+		}
+
+		OutSample = CurrentSample;
+
+		CurrentSample.Reset();
+
+		return true;
+	}
+
+	virtual bool PeekVideoSampleTime(FMediaTimeStamp& TimeStamp) override
+	{
+		return false;
+	}
+
+	virtual void FlushSamples() override
+	{
+		CurrentSample.Reset();
+	}
+
+	//~ End IMediaSamples interface
+
+public:
+
+	TSharedPtr<FSharedMemoryMediaSample, ESPMode::ThreadSafe> CurrentSample;
+};
+
diff --git a/Engine/Plugins/Runtime/nDisplay/Source/DisplayClusterMedia/Private/Input/SharedMemoryMediaSource.cpp b/Engine/Plugins/Runtime/nDisplay/Source/DisplayClusterMedia/Private/Input/SharedMemoryMediaSource.cpp
new file mode 100644
--- /dev/null
+++ b/Engine/Plugins/Runtime/nDisplay/Source/DisplayClusterMedia/Private/Input/SharedMemoryMediaSource.cpp
@@ -0,0 +1,57 @@
+// Copyright Epic Games, Inc. All Rights Reserved.
+
+#include "SharedMemoryMediaSource.h"
+
+#include "MediaIOCorePlayerBase.h"
+#include "SharedMemoryMediaSourceOptions.h"
+
+
+bool USharedMemoryMediaSource::GetMediaOption(const FName& Key, bool DefaultValue) const
+{
+	if (Key == SharedMemoryMediaOption::ZeroLatency)
+	{
+		return bZeroLatency;
+	}
+
+	return Super::GetMediaOption(Key, DefaultValue);
+}
+
+int64 USharedMemoryMediaSource::GetMediaOption(const FName& Key, int64 DefaultValue) const
+{
+	return Super::GetMediaOption(Key, DefaultValue);
+}
+
+FString USharedMemoryMediaSource::GetMediaOption(const FName& Key, const FString& DefaultValue) const
+{
+	if (Key == SharedMemoryMediaOption::UniqueName)
+	{
+		return UniqueName;
+	}
+
+	return Super::GetMediaOption(Key, DefaultValue);
+}
+
+bool USharedMemoryMediaSource::HasMediaOption(const FName& Key) const
+{
+	if (Key == SharedMemoryMediaOption::UniqueName)
+	{
+		return true;
+	}
+
+	if (Key == SharedMemoryMediaOption::ZeroLatency)
+	{
+		return true;
+	}
+
+	return Super::HasMediaOption(Key);
+}
+
+FString USharedMemoryMediaSource::GetUrl() const
+{
+	return FString::Printf(TEXT("dcsm://%s"), *UniqueName);
+}
+
+bool USharedMemoryMediaSource::Validate() const
+{
+	return true;
+}
diff --git a/Engine/Plugins/Runtime/nDisplay/Source/DisplayClusterMedia/Private/Input/SharedMemoryMediaSourceOptions.h b/Engine/Plugins/Runtime/nDisplay/Source/DisplayClusterMedia/Private/Input/SharedMemoryMediaSourceOptions.h
new file mode 100644
--- /dev/null
+++ b/Engine/Plugins/Runtime/nDisplay/Source/DisplayClusterMedia/Private/Input/SharedMemoryMediaSourceOptions.h
@@ -0,0 +1,12 @@
+// Copyright Epic Games, Inc. All Rights Reserved.
+
+#pragma once
+
+#include "CoreMinimal.h"
+
+
+namespace SharedMemoryMediaOption
+{
+	static const FName UniqueName("UniqueName");
+	static const FName ZeroLatency("ZeroLatency");
+}
diff --git a/Engine/Plugins/Runtime/nDisplay/Source/DisplayClusterMedia/Private/Input/SharedMemoryMediaTextureSampleConverter.h b/Engine/Plugins/Runtime/nDisplay/Source/DisplayClusterMedia/Private/Input/SharedMemoryMediaTextureSampleConverter.h
new file mode 100644
--- /dev/null
+++ b/Engine/Plugins/Runtime/nDisplay/Source/DisplayClusterMedia/Private/Input/SharedMemoryMediaTextureSampleConverter.h
@@ -0,0 +1,36 @@
+// Copyright Epic Games, Inc. All Rights Reserved.
+
+#pragma once
+
+#include "IMediaTextureSampleConverter.h"
+
+class FSharedMemoryMediaPlayer;
+
+class FSharedMemoryMediaTextureSampleConverter : public IMediaTextureSampleConverter
+{
+public:
+
+	/** The player that this converter will check back with for just in time sample render purposes */
+	FSharedMemoryMediaPlayer* Player = nullptr;
+
+public:
+
+	//~ Begin IMediaTextureSampleConverter interface
+	virtual uint32 GetConverterInfoFlags() const override
+	{
+		return ConverterInfoFlags_PreprocessOnly;
+	}
+
+	virtual bool Convert(FTexture2DRHIRef& InDstTexture, const FConversionHints& Hints) override
+	{
+		if (!Player)
+		{
+			return false;
+		}
+
+		Player->JustInTimeSampleRender();
+
+		return true;
+	}
+	//~ End IMediaTextureSampleConverter interface
+};
diff --git a/Engine/Plugins/Runtime/nDisplay/Source/DisplayClusterMedia/Private/Platforms/Windows/SharedMemoryMediaPlatformWindowsD3D12.cpp b/Engine/Plugins/Runtime/nDisplay/Source/DisplayClusterMedia/Private/Platforms/Windows/SharedMemoryMediaPlatformWindowsD3D12.cpp
new file mode 100644
--- /dev/null
+++ b/Engine/Plugins/Runtime/nDisplay/Source/DisplayClusterMedia/Private/Platforms/Windows/SharedMemoryMediaPlatformWindowsD3D12.cpp
@@ -0,0 +1,185 @@
+// Copyright Epic Games, Inc. All Rights Reserved.
+
+#include "SharedMemoryMediaPlatformWindowsD3D12.h"
+
+#include "ID3D12DynamicRHI.h"
+
+#include "DisplayClusterMediaLog.h"
+
+
+// This should run when the module starts and register the creation function for this rhi platform.
+bool FSharedMemoryMediaPlatformWindowsD3D12::bRegistered = FSharedMemoryMediaPlatform::RegisterPlatformForRhi(
+	ERHIInterfaceType::D3D12, &FSharedMemoryMediaPlatformWindowsD3D12::CreateInstance);
+
+TSharedPtr<FSharedMemoryMediaPlatform, ESPMode::ThreadSafe> FSharedMemoryMediaPlatformWindowsD3D12::CreateInstance()
+{
+	return MakeShared<FSharedMemoryMediaPlatformWindowsD3D12, ESPMode::ThreadSafe>();
+}
+
+FSharedMemoryMediaPlatformWindowsD3D12::~FSharedMemoryMediaPlatformWindowsD3D12()
+{
+	// Close shared handles
+	for (int32 BufferIdx = 0; BufferIdx < UE::SharedMemoryMedia::SenderNumBuffers; ++BufferIdx)
+	{
+		::CloseHandle(SharedHandle[BufferIdx]);
+	}
+}
+
+FTextureRHIRef FSharedMemoryMediaPlatformWindowsD3D12::CreateSharedCrossGpuTexture(EPixelFormat Format, int32 Width, int32 Height, const FGuid& Guid, uint32 BufferIdx)
+{
+	check(BufferIdx < UE::SharedMemoryMedia::SenderNumBuffers);
+	check(!SharedHandle[BufferIdx]);
+
+	ID3D12Device* D3D12Device = static_cast<ID3D12Device*>(GDynamicRHI->RHIGetNativeDevice());
+	check(D3D12Device);
+
+	D3D12_RESOURCE_DESC SharedCrossGpuTextureDesc = CD3DX12_RESOURCE_DESC::Tex2D(
+		DXGI_FORMAT(GPixelFormats[Format].PlatformFormat),
+		Width,
+		Height,
+		1, // arraySize
+		1, // mipLevels
+		1, // sampleCount
+		0, // sampleQuality
+		D3D12_RESOURCE_FLAG_ALLOW_CROSS_ADAPTER | D3D12_RESOURCE_FLAG_ALLOW_SIMULTANEOUS_ACCESS, // flags
+		D3D12_TEXTURE_LAYOUT_ROW_MAJOR // layout
+	);
+
+	CD3DX12_HEAP_PROPERTIES HeapProperties(D3D12_HEAP_TYPE_DEFAULT);
+
+	ID3D12Resource* SharedGpuResource = nullptr;
+
+	HRESULT HResult = D3D12Device->CreateCommittedResource(
+		&HeapProperties,
+		D3D12_HEAP_FLAG_SHARED | D3D12_HEAP_FLAG_SHARED_CROSS_ADAPTER,
+		&SharedCrossGpuTextureDesc,
+		D3D12_RESOURCE_STATE_COPY_DEST,
+		nullptr, // pOptimizedClearValue
+		IID_PPV_ARGS(&SharedGpuResource)
+	);
+
+	check(HResult == S_OK);
+	check(SharedGpuResource);
+
+	const FString SharedGpuTextureName = FString::Printf(TEXT("Global\\%s"),
+		*Guid.ToString(EGuidFormats::DigitsWithHyphensInBraces));
+
+	HResult = D3D12Device->CreateSharedHandle(
+		SharedGpuResource,                    // pObject
+		nullptr,                              // pAttributes
+		GENERIC_ALL,                          // Access
+		*SharedGpuTextureName,                // Name
+		&SharedHandle[BufferIdx]              // pHandle
+	);
+
+	check(HResult == S_OK);
+	check(SharedHandle[BufferIdx]);
+
+	return GetID3D12DynamicRHI()->RHICreateTexture2DFromResource(
+		Format,
+		TexCreate_Dynamic | TexCreate_DisableSRVCreation,
+		FClearValueBinding::None,
+		SharedGpuResource
+	);
+}
+
+void FSharedMemoryMediaPlatformWindowsD3D12::ReleaseSharedCrossGpuTexture(uint32 BufferIdx)
+{
+	check(BufferIdx < UE::SharedMemoryMedia::SenderNumBuffers);
+
+	::CloseHandle(SharedHandle[BufferIdx]);
+}
+
+FTextureRHIRef FSharedMemoryMediaPlatformWindowsD3D12::OpenSharedCrossGpuTextureByGuid(const FGuid& Guid, FSharedMemoryMediaTextureDescription& OutTextureDescription)
+{
+	ID3D12Device* D3D12Device = static_cast<ID3D12Device*>(GDynamicRHI->RHIGetNativeDevice());
+	check(D3D12Device);
+
+	const FString SharedGpuTextureName = FString::Printf(TEXT("Global\\%s"), *Guid.ToString(EGuidFormats::DigitsWithHyphensInBraces));
+
+	ID3D12Resource* SharedCrossGpuTexture = nullptr;
+
+	HANDLE NamedSharedGpuTextureHandle = nullptr;
+	HRESULT HResult = D3D12Device->OpenSharedHandleByName(*SharedGpuTextureName, GENERIC_ALL, &NamedSharedGpuTextureHandle);
+
+	if (FAILED(HResult))
+	{
+		UE_LOG(LogDisplayClusterMedia, Error, TEXT("D3D12Device->OpenSharedHandleByName(%s) failed: 0x%X - %s"),
+			*SharedGpuTextureName, HResult, *GetD3D12ComErrorDescription(HResult));
+
+		return nullptr;
+	}
+
+	HResult = D3D12Device->OpenSharedHandle(NamedSharedGpuTextureHandle, IID_PPV_ARGS(&SharedCrossGpuTexture));
+
+	::CloseHandle(NamedSharedGpuTextureHandle);
+
+	if (FAILED(HResult))
+	{
+		UE_LOG(LogDisplayClusterMedia, Error, TEXT("D3D12Device->OpenSharedHandle(0x%X) failed: 0x%X - %s"),
+			NamedSharedGpuTextureHandle, HResult, *GetD3D12ComErrorDescription(HResult));
+
+		return nullptr;
+	}
+
+	check(SharedCrossGpuTexture);
+
+	// Query texture dimensions and use them to describe downstream resources
+	{
+		D3D12_RESOURCE_DESC SharedGpuTextureDesc = SharedCrossGpuTexture->GetDesc();
+
+		OutTextureDescription.Height = SharedGpuTextureDesc.Height;
+		OutTextureDescription.Width = SharedGpuTextureDesc.Width;
+
+		// Find EPixelFormat from platform format
+		for (int32 FormatIdx = 0; FormatIdx < PF_MAX; FormatIdx++)
+		{
+			if (GPixelFormats[FormatIdx].PlatformFormat == SharedGpuTextureDesc.Format)
+			{
+				OutTextureDescription.Format = EPixelFormat(FormatIdx);
+				OutTextureDescription.BytesPerPixel = GPixelFormats[FormatIdx].BlockBytes;
+				OutTextureDescription.Stride = OutTextureDescription.Width * OutTextureDescription.BytesPerPixel;
+				break;
+			}
+		}
+
+		if (OutTextureDescription.Format == EPixelFormat::PF_Unknown)
+		{
+			UE_LOG(LogDisplayClusterMedia, Error, TEXT("Could not find a known pixel format for SharedCrossGpuTexture DXGI_FORMAT %d."),
+				SharedGpuTextureDesc.Format);
+
+			::CloseHandle(NamedSharedGpuTextureHandle);
+
+			return nullptr;
+		}
+	}
+
+	return GetID3D12DynamicRHI()->RHICreateTexture2DFromResource(
+		OutTextureDescription.Format,
+		TexCreate_Dynamic | TexCreate_DisableSRVCreation,
+		FClearValueBinding::None,
+		SharedCrossGpuTexture
+	);
+}
+
+
+const FString FSharedMemoryMediaPlatformWindowsD3D12::GetD3D12ComErrorDescription(HRESULT Hresult)
+{
+	constexpr uint32 BufSize = 1024;
+	WIDECHAR Buffer[BufSize];
+
+	if (::FormatMessageW(FORMAT_MESSAGE_FROM_SYSTEM,
+		nullptr,
+		Hresult,
+		MAKELANGID(LANG_NEUTRAL, SUBLANG_NEUTRAL),
+		Buffer,
+		sizeof(Buffer) / sizeof(*Buffer),
+		nullptr))
+	{
+		return Buffer;
+	}
+	else
+	{
+		return FString::Printf(TEXT("[Could not find a d3d12 error description for HRESULT %d]"), Hresult);
+	}
+}
diff --git a/Engine/Plugins/Runtime/nDisplay/Source/DisplayClusterMedia/Private/Platforms/Windows/SharedMemoryMediaPlatformWindowsD3D12.h b/Engine/Plugins/Runtime/nDisplay/Source/DisplayClusterMedia/Private/Platforms/Windows/SharedMemoryMediaPlatformWindowsD3D12.h
new file mode 100644
--- /dev/null
+++ b/Engine/Plugins/Runtime/nDisplay/Source/DisplayClusterMedia/Private/Platforms/Windows/SharedMemoryMediaPlatformWindowsD3D12.h
@@ -0,0 +1,41 @@
+// Copyright Epic Games, Inc. All Rights Reserved.
+
+#pragma once
+
+#include "RHI.h"
+
+#include "Windows/AllowWindowsPlatformTypes.h"
+#include "d3d12.h"
+#include "Windows/HideWindowsPlatformTypes.h"
+
+#include "SharedMemoryMediaPlatform.h"
+
+/** 
+ * D3D12 RHI specific implementation of FSharedMemoryMediaPlatform. 
+ */
+class FSharedMemoryMediaPlatformWindowsD3D12 : public FSharedMemoryMediaPlatform
+{
+private:
+
+	/** Handles to the shared handles of the cross gpu resources */
+	HANDLE SharedHandle[UE::SharedMemoryMedia::SenderNumBuffers]{ 0 };
+
+	/** Dummy bool to get the registration call made */
+	static bool bRegistered;
+
+	/** Factory function to create this rhi platform instance. It is registed with and called by FSharedMemoryMediaPlatform */
+	static TSharedPtr<FSharedMemoryMediaPlatform> CreateInstance();
+
+	/** Convenience function to stringify an HRESULT */
+	static const FString GetD3D12ComErrorDescription(HRESULT Hresult);
+
+public:
+
+	//~ Begin FSharedMemoryMediaPlatform interface
+	virtual FTextureRHIRef CreateSharedCrossGpuTexture(EPixelFormat Format, int32 Width, int32 Height, const FGuid& Guid, uint32 BufferIdx) override;
+	virtual void ReleaseSharedCrossGpuTexture(uint32 BufferIdx) override;
+	virtual FTextureRHIRef OpenSharedCrossGpuTextureByGuid(const FGuid& Guid, FSharedMemoryMediaTextureDescription& OutTextureDescription) override;
+	//~ End FSharedMemoryMediaPlatform interface
+
+	virtual ~FSharedMemoryMediaPlatformWindowsD3D12();
+};
diff --git a/Engine/Plugins/Runtime/nDisplay/Source/DisplayClusterMedia/Private/SharedMemoryMediaPlatform.cpp b/Engine/Plugins/Runtime/nDisplay/Source/DisplayClusterMedia/Private/SharedMemoryMediaPlatform.cpp
new file mode 100644
--- /dev/null
+++ b/Engine/Plugins/Runtime/nDisplay/Source/DisplayClusterMedia/Private/SharedMemoryMediaPlatform.cpp
@@ -0,0 +1,43 @@
+// Copyright Epic Games, Inc. All Rights Reserved.
+
+#include "SharedMemoryMediaPlatform.h"
+
+TMap<ERHIInterfaceType, FSharedMemoryMediaPlatform::CreateSharedMemoryMediaPlatform> FSharedMemoryMediaPlatform::PlatformCreators;
+
+FString FSharedMemoryMediaPlatform::GetRhiTypeString(const ERHIInterfaceType RhiType)
+{
+	static TMap<ERHIInterfaceType, FString> Rhis;
+
+	if (!Rhis.Num())
+	{
+		Rhis.Emplace(ERHIInterfaceType::Hidden, TEXT("Hidden"));
+		Rhis.Emplace(ERHIInterfaceType::Null, TEXT("Null"));
+		Rhis.Emplace(ERHIInterfaceType::D3D11, TEXT("D3D11"));
+		Rhis.Emplace(ERHIInterfaceType::D3D12, TEXT("D3D12"));
+		Rhis.Emplace(ERHIInterfaceType::Vulkan, TEXT("Vulkan"));
+		Rhis.Emplace(ERHIInterfaceType::Metal, TEXT("Metal"));
+		Rhis.Emplace(ERHIInterfaceType::Agx, TEXT("Agx"));
+		Rhis.Emplace(ERHIInterfaceType::OpenGL, TEXT("OpenGL"));
+	};
+
+	const FString* RhiString = Rhis.Find(RhiType);
+
+	return RhiString ? *RhiString : TEXT("Unknown");
+}
+
+bool FSharedMemoryMediaPlatform::RegisterPlatformForRhi(ERHIInterfaceType RhiType, CreateSharedMemoryMediaPlatform PlatformCreator)
+{
+	PlatformCreators.Add(RhiType, PlatformCreator);
+	return true;
+}
+
+TSharedPtr<FSharedMemoryMediaPlatform, ESPMode::ThreadSafe> FSharedMemoryMediaPlatform::CreateInstanceForRhi(ERHIInterfaceType RhiType)
+{
+	if (const CreateSharedMemoryMediaPlatform* Creator = PlatformCreators.Find(RhiType))
+	{
+		return (*Creator)();
+	}
+
+	return nullptr;
+}
+
diff --git a/Engine/Plugins/Runtime/nDisplay/Source/DisplayClusterMedia/Private/SharedMemoryMediaPlatform.h b/Engine/Plugins/Runtime/nDisplay/Source/DisplayClusterMedia/Private/SharedMemoryMediaPlatform.h
new file mode 100644
--- /dev/null
+++ b/Engine/Plugins/Runtime/nDisplay/Source/DisplayClusterMedia/Private/SharedMemoryMediaPlatform.h
@@ -0,0 +1,52 @@
+// Copyright Epic Games, Inc. All Rights Reserved.
+
+#pragma once
+
+#include "RHI.h"
+
+#include "SharedMemoryMediaTypes.h"
+
+
+/** This is the interface to the platform specific functions needed by shared memory media. 
+ *  It is mostly about the platform specific implementation of cross gpu textures.
+ */
+class FSharedMemoryMediaPlatform
+{
+public:
+
+	/** This is the type definition of the rhi platform factory function that each rhi platfom implementation can register with.  */
+	typedef TSharedPtr<FSharedMemoryMediaPlatform>(*CreateSharedMemoryMediaPlatform)();
+
+public:
+
+	/** TMap of the collection of registered rhi platform factory functions */
+	static TMap<ERHIInterfaceType, CreateSharedMemoryMediaPlatform> PlatformCreators;
+
+public:
+
+	/** Helper function to get a stringified ERHIInterfaceType */
+	static FString GetRhiTypeString(const ERHIInterfaceType RhiType);
+
+	/** Rhi platform implementations call this function to register their factory creation function */
+	static bool RegisterPlatformForRhi(ERHIInterfaceType RhiType, CreateSharedMemoryMediaPlatform PlatformCreator);
+
+	/** This factory function will create an instance of the rhi platform specific implementation, if it has been registered. */
+	static TSharedPtr<FSharedMemoryMediaPlatform, ESPMode::ThreadSafe> CreateInstanceForRhi(ERHIInterfaceType RhiType);
+
+	/** Creates a cross gpu texture */
+	virtual FTextureRHIRef CreateSharedCrossGpuTexture(EPixelFormat Format, int32 Width, int32 Height, const FGuid& Guid, uint32 BufferIdx)
+	{
+		return nullptr;
+	}
+
+	/** Opens a cross gpu texture specified by a Guid */
+	virtual FTextureRHIRef OpenSharedCrossGpuTextureByGuid(const FGuid& Guid, FSharedMemoryMediaTextureDescription& OutTextureDescription)
+	{
+		return nullptr;
+	}
+
+	/** Release any platform specific resources related to the indexed texture */
+	virtual void ReleaseSharedCrossGpuTexture(uint32 BufferIdx) {};
+
+};
+
diff --git a/Engine/Plugins/Runtime/nDisplay/Source/DisplayClusterMedia/Private/SharedMemoryMediaTypes.h b/Engine/Plugins/Runtime/nDisplay/Source/DisplayClusterMedia/Private/SharedMemoryMediaTypes.h
new file mode 100644
--- /dev/null
+++ b/Engine/Plugins/Runtime/nDisplay/Source/DisplayClusterMedia/Private/SharedMemoryMediaTypes.h
@@ -0,0 +1,114 @@
+// Copyright Epic Games, Inc. All Rights Reserved.
+
+#pragma once
+
+#include "CoreMinimal.h"
+#include "Internationalization/TextKey.h"
+
+namespace UE::SharedMemoryMedia
+{
+	// When the implementation requires spin waiting on a fence or polling shared system memory,
+	// it will sleep by the amount below in a loop until the condition is met or a given timout occurs.
+	constexpr float SpinWaitTimeSeconds = 50 * 1e-6;
+
+	/** 
+	 * This defines the number of textures used for communication. Having more that one allows for overlapping
+	 * sends and minimize waits on frame acks for texture re-use. Use 3 for best performance, or 2 for
+	 * a potential compromise with resource usage.
+	 */
+	constexpr int32 SenderNumBuffers = 3;
+
+	// Guid that identifies the SharedMemoryMediaPlayer.
+	static const FGuid PlayerGuid = FGuid(0xAF8C5107, 0x13CF992C, 0x9FD4EBFE, 0x2E3E049F);
+
+	/**
+	 * Used to generate a Guid based on the UniqueName that sender and receiver uses to open the same
+	 * shared system memory used for IPC. There will be one shard memory region per buffer index.
+	 * There are SenderNumBuffers buffer indices.
+	 * A given frame number uses a buffer index defined by:
+	 *   BufferIdx = FrameNumber % SenderNumBuffers.
+	 */
+	static FGuid GenerateSharedMemoryGuid(FString UniqueName, uint32 BufferIdx)
+	{
+		// Avoid the hard to spot user error of differing spaces in strings.
+		UniqueName.RemoveSpacesInline();
+
+		const uint32 UniqueNameHash = TextKeyUtil::HashString(UniqueName);
+		return FGuid(0xDA90FCAF, 0x7034A2AA, UniqueNameHash, 0xEB12036B + BufferIdx);
+	};
+
+	// Convenience ZeroGuid value.
+	static const FGuid ZeroGuid = FGuid(0, 0, 0, 0);
+
+	/** Use this to call a lambda when the variable gets destroyed. */
+	template <typename LAMBDA>
+	struct FRunWhenGoingOutOfScope
+	{
+		FRunWhenGoingOutOfScope(LAMBDA&& InLambda) : Lambda(InLambda) { }
+
+		~FRunWhenGoingOutOfScope()
+		{
+			Lambda();
+		}
+
+	private:
+
+		// Holds the lambda to execute in the destructor of this object
+		LAMBDA Lambda;
+	};
+}
+
+struct FSharedMemoryMediaTextureDescription
+{
+	uint32 Width = 0;
+	uint32 Height = 0;
+	uint32 BytesPerPixel = 0;
+	uint32 Stride = 0;
+	EPixelFormat Format = EPixelFormat::PF_Unknown;
+	FGuid Guid = UE::SharedMemoryMedia::ZeroGuid;
+
+	bool IsEquivalentTo(const FSharedMemoryMediaTextureDescription& Other) const
+	{
+		return Width == Other.Width
+			&& Height == Other.Height
+			&& BytesPerPixel == Other.BytesPerPixel
+			&& Stride == Other.Stride
+			&& Format == Other.Format;
+	}
+};
+
+/**
+ * This structure represents the IPC between sender and receiver of the shared memory media.
+ */
+struct FSharedMemoryMediaFrameMetadata
+{
+	static constexpr uint32 MAGIC = 0x534D4D4D; // SMMM - Shared Memory Media Metadata
+
+	// Data that the sender wants to communicate
+	struct FSender
+	{
+		// Used to determine that the rest of the metadata is valid.
+		uint32 Magic = MAGIC;
+
+		// Each shared cross gpu texture is assigned a Guid. The player can use this to locate and open them.
+		FGuid TextureGuid;
+
+		// This is the sender's frame number of the pixel data in the texture.
+		uint32 FrameNumber = 0;
+	} Sender;
+
+	// Data that the receiver wants to communicate (although KeepAliveShiftRegister is shifted by the Sender)
+	struct FReceiver
+	{
+		// Frame number that the receiver is done reading. This signals the sender that it can re-use the texture.
+		uint32 FrameNumberAcked = 0;
+
+		/** 
+		 * The receiver sets to ones this register. This signals the sender that it should wait for the related frame number ack before re-using the texture.
+		 * The sender will shift this register to the right each frame, as a way to time out the session in case the receiver disappears.
+		 */
+		uint8 KeepAliveShiftRegister = 0;
+	} Receiver;
+};
+
+
diff --git a/Engine/Plugins/Runtime/nDisplay/Source/DisplayClusterMedia/Public/SharedMemoryMediaCapture.h b/Engine/Plugins/Runtime/nDisplay/Source/DisplayClusterMedia/Public/SharedMemoryMediaCapture.h
new file mode 100644
--- /dev/null
+++ b/Engine/Plugins/Runtime/nDisplay/Source/DisplayClusterMedia/Public/SharedMemoryMediaCapture.h
@@ -0,0 +1,82 @@
+// Copyright Epic Games, Inc. All Rights Reserved.
+
+#pragma once
+
+#include "MediaCapture.h"
+
+#include "GenericPlatform/GenericPlatformMemory.h"
+#include "SharedMemoryMediaOutput.h"
+#include "SharedMemoryMediaTypes.h"
+
+#include "SharedMemoryMediaCapture.generated.h"
+
+class FSharedMemoryMediaPlatform;
+
+/**
+ * Output Media for SharedMemory.
+ * 
+ * The pixels are captured into shared cross gpu textures, that a player can read.
+ * Inter-process communication happens via shared system memory. The metadata exchanged is documented 
+ * in the FSharedMemoryMediaFrameMetadata structure. The shared memory can be located via a UniqueName
+ * that must match in the Media Output and corresponding Media Source settings.
+ * 
+ * It is mostly intended for use in nDisplay render nodes, which are designed to be frame-locked.
+ */
+UCLASS(BlueprintType)
+class DISPLAYCLUSTERMEDIA_API USharedMemoryMediaCapture : public UMediaCapture
+{
+	GENERATED_BODY()
+
+	//~ Begin UMediaCapture interface
+
+protected:
+
+	virtual bool ShouldCaptureRHIResource() const override;
+	virtual bool InitializeCapture() override;
+	virtual void StopCaptureImpl(bool bAllowPendingFrameToBeProcess) override;
+
+	/** For custom conversion, methods that need to be overridden */
+	virtual FIntPoint GetCustomOutputSize(const FIntPoint& InSize) const override;
+	virtual EMediaCaptureResourceType GetCustomOutputResourceType() const override;
+
+	virtual void OnCustomCapture_RenderingThread(
+		FRDGBuilder& GraphBuilder, 
+		const FCaptureBaseData& InBaseData, 
+		TSharedPtr<FMediaCaptureUserData, ESPMode::ThreadSafe> InUserData, 
+		FRDGTextureRef InSourceTexture, 
+		FRDGTextureRef OutputTexture, 
+		const FRHICopyTextureInfo& CopyInfo, 
+		FVector2D CropU, 
+		FVector2D CropV) override;
+
+	//~ End UMediaCapture interface
+
+	/** Adds to the graph builder the copy from the input texture to the shared gpu texture that corresponds to the current frame number */
+	void AddCopyToSharedGpuTexturePass(FRDGBuilder& GraphBuilder, FRDGTextureRef InSourceTexture, uint32 SharedTextureIdx);
+
+private:
+
+	/** Number of shared buffers */
+	static constexpr int32 NUMBUFFERS = UE::SharedMemoryMedia::SenderNumBuffers;
+
+	/** Shared RAM used for IPC between media capture and player */
+	FPlatformMemory::FSharedMemoryRegion* SharedMemory[NUMBUFFERS] = { 0 };
+
+	/** Shared Cross Gpu Textures. The texture exchange happens using these. */
+	FTextureRHIRef SharedCrossGpuTextures[NUMBUFFERS];
+
+	/** Fence for an async task to know when the cross gpu texture has the data to be shared */
+	FGPUFenceRHIRef TextureReadyFences[NUMBUFFERS];
+	
+	/** Flags that the associated TextureReadyFence can be re-used. Otherwise the player may still be reading it. */
+	std::atomic<bool> bTextureReadyFenceBusy[NUMBUFFERS] = { 0 };
+
+	/** Platform specific data or resources */
+	TSharedPtr<FSharedMemoryMediaPlatform, ESPMode::ThreadSafe> PlatformData;
+
+	/** Guids associated with the cross gpu textures in use. These are communicated to the player for it to open them by name. */
+	FGuid SharedCrossGpuTextureGuids[NUMBUFFERS];
+
+	/** Counter of running tasks used to detect when to release resources */
+	std::atomic<int32> RunningTasksCount{ 0 };
+};
diff --git a/Engine/Plugins/Runtime/nDisplay/Source/DisplayClusterMedia/Public/SharedMemoryMediaOutput.h b/Engine/Plugins/Runtime/nDisplay/Source/DisplayClusterMedia/Public/SharedMemoryMediaOutput.h
new file mode 100644
--- /dev/null
+++ b/Engine/Plugins/Runtime/nDisplay/Source/DisplayClusterMedia/Public/SharedMemoryMediaOutput.h
@@ -0,0 +1,39 @@
+// Copyright Epic Games, Inc. All Rights Reserved.
+
+#pragma once
+
+#include "MediaOutput.h"
+
+#include "MediaIOCoreDefinitions.h"
+#include "UObject/ObjectMacros.h"
+
+#include "SharedMemoryMediaOutput.generated.h"
+
+
+/**
+ * Output information for a SharedMemory media capture.
+ */
+UCLASS(BlueprintType)
+class DISPLAYCLUSTERMEDIA_API USharedMemoryMediaOutput : public UMediaOutput
+{
+	GENERATED_BODY()
+
+public:
+
+	/** Shared memory will be opened by using this name. Should be unique per media output. */
+	UPROPERTY(EditAnywhere, BlueprintReadWrite, Category = "Output")
+	FString UniqueName = TEXT("UniqueName");
+
+public:
+	
+	//~ Begin UMediaOutput interface
+	virtual bool Validate(FString& FailureReason) const override;
+	virtual FIntPoint GetRequestedSize() const override;
+	virtual EPixelFormat GetRequestedPixelFormat() const override;
+	virtual EMediaCaptureConversionOperation GetConversionOperation(EMediaCaptureSourceType InSourceType) const override;
+
+protected:
+	virtual UMediaCapture* CreateMediaCaptureImpl() override;
+	//~ End UMediaOutput interface
+
+};
diff --git a/Engine/Plugins/Runtime/nDisplay/Source/DisplayClusterMedia/Public/SharedMemoryMediaSource.h b/Engine/Plugins/Runtime/nDisplay/Source/DisplayClusterMedia/Public/SharedMemoryMediaSource.h
new file mode 100644
--- /dev/null
+++ b/Engine/Plugins/Runtime/nDisplay/Source/DisplayClusterMedia/Public/SharedMemoryMediaSource.h
@@ -0,0 +1,41 @@
+// Copyright Epic Games, Inc. All Rights Reserved.
+
+#pragma once
+
+#include "CoreMinimal.h"
+#include "MediaSource.h"
+
+#include "SharedMemoryMediaSource.generated.h"
+
+/**
+ * Media source for SharedMemory streams.
+ */
+UCLASS(BlueprintType, hideCategories=(Platforms,Object))
+class DISPLAYCLUSTERMEDIA_API USharedMemoryMediaSource : public UMediaSource
+{
+	GENERATED_BODY()
+
+public:
+
+	/** Shared memory will be found by using this name. Should match the media output setting. */
+	UPROPERTY(EditAnywhere, BlueprintReadWrite, Category = "Input")
+	FString UniqueName = TEXT("UniqueName");
+
+	/** Zero latency option to wait for the cross gpu texture rendered on the same frame. May adversely affect fps */
+	UPROPERTY(EditAnywhere, BlueprintReadWrite, Category = "Input")
+	bool bZeroLatency = true;
+
+public:
+
+	//~ Begin UMediaSource interface
+	virtual FString GetUrl() const override;
+	virtual bool Validate() const override;
+	//~ End UMediaSource interface
+
+	//~ Begin IMediaOptions interface
+	virtual bool GetMediaOption(const FName& Key, bool DefaultValue) const override;
+	virtual int64 GetMediaOption(const FName& Key, int64 DefaultValue) const override;
+	virtual FString GetMediaOption(const FName& Key, const FString& DefaultValue) const override;
+	virtual bool HasMediaOption(const FName& Key) const override;
+	//~ End IMediaOptions interface
+};
diff --git a/Engine/Plugins/Runtime/nDisplay/Source/DisplayClusterMediaEditor/DisplayClusterMediaEditor.Build.cs b/Engine/Plugins/Runtime/nDisplay/Source/DisplayClusterMediaEditor/DisplayClusterMediaEditor.Build.cs
new file mode 100644
--- /dev/null
+++ b/Engine/Plugins/Runtime/nDisplay/Source/DisplayClusterMediaEditor/DisplayClusterMediaEditor.Build.cs
@@ -0,0 +1,19 @@
+// Copyright Epic Games, Inc. All Rights Reserved.
+
+using UnrealBuildTool;
+using System.IO;
+
+public class DisplayClusterMediaEditor : ModuleRules
+{
+	public DisplayClusterMediaEditor(ReadOnlyTargetRules ROTargetRules) : base(ROTargetRules)
+	{
+		PrivateDependencyModuleNames.AddRange(
+			new string[] {
+				"Core",
+				"CoreUObject",
+				"DisplayClusterMedia",
+				"Engine",
+				"UnrealEd",
+			});
+	}
+}
diff --git a/Engine/Plugins/Runtime/nDisplay/Source/DisplayClusterMediaEditor/Private/DisplayClusterMediaEditorModule.cpp b/Engine/Plugins/Runtime/nDisplay/Source/DisplayClusterMediaEditor/Private/DisplayClusterMediaEditorModule.cpp
new file mode 100644
--- /dev/null
+++ b/Engine/Plugins/Runtime/nDisplay/Source/DisplayClusterMediaEditor/Private/DisplayClusterMediaEditorModule.cpp
@@ -0,0 +1,12 @@
+// Copyright Epic Games, Inc. All Rights Reserved.
+
+#include "Modules/ModuleManager.h"
+
+
+class FDisplayClusterMediaEditorModule : public IModuleInterface
+{
+
+};
+
+
+IMPLEMENT_MODULE(FDisplayClusterMediaEditorModule, DisplayClusterMediaEditor);
diff --git a/Engine/Plugins/Runtime/nDisplay/Source/DisplayClusterMediaEditor/Private/SharedMemoryMediaOutputFactory.cpp b/Engine/Plugins/Runtime/nDisplay/Source/DisplayClusterMediaEditor/Private/SharedMemoryMediaOutputFactory.cpp
new file mode 100644
--- /dev/null
+++ b/Engine/Plugins/Runtime/nDisplay/Source/DisplayClusterMediaEditor/Private/SharedMemoryMediaOutputFactory.cpp
@@ -0,0 +1,50 @@
+// Copyright Epic Games, Inc. All Rights Reserved.
+
+#include "SharedMemoryMediaOutputFactory.h"
+
+#include "AssetTypeCategories.h"
+#include "SharedMemoryMediaOutput.h"
+#include "EngineAnalytics.h"
+
+
+/* USharedMemoryMediaOutputFactoryNew structors
+ *****************************************************************************/
+
+USharedMemoryMediaOutputFactory::USharedMemoryMediaOutputFactory()
+{
+	SupportedClass = USharedMemoryMediaOutput::StaticClass();
+	bCreateNew = true;
+	bEditAfterNew = true;
+}
+
+
+/* UFactory overrides
+ *****************************************************************************/
+
+/**
+ * @EventName MediaFramework.SharedMemoryMediaOutputCreated
+ * @Trigger Triggered when an SharedMemory media output asset is created.
+ * @Type Client
+ * @Owner MediaIO Team
+ */
+UObject* USharedMemoryMediaOutputFactory::FactoryCreateNew(UClass* InClass, UObject* InParent, FName InName, EObjectFlags Flags, UObject* Context, FFeedbackContext* Warn)
+{
+	if (FEngineAnalytics::IsAvailable())
+	{
+		FEngineAnalytics::GetProvider().RecordEvent(TEXT("MediaFramework.SharedMemoryMediaOutputCreated"));
+	}
+	
+	return NewObject<USharedMemoryMediaOutput>(InParent, InClass, InName, Flags);
+}
+
+
+uint32 USharedMemoryMediaOutputFactory::GetMenuCategories() const
+{
+	return EAssetTypeCategories::Media;
+}
+
+
+bool USharedMemoryMediaOutputFactory::ShouldShowInNewMenu() const
+{
+	return true;
+}
diff --git a/Engine/Plugins/Runtime/nDisplay/Source/DisplayClusterMediaEditor/Private/SharedMemoryMediaOutputFactory.h b/Engine/Plugins/Runtime/nDisplay/Source/DisplayClusterMediaEditor/Private/SharedMemoryMediaOutputFactory.h
new file mode 100644
--- /dev/null
+++ b/Engine/Plugins/Runtime/nDisplay/Source/DisplayClusterMediaEditor/Private/SharedMemoryMediaOutputFactory.h
@@ -0,0 +1,27 @@
+// Copyright Epic Games, Inc. All Rights Reserved.
+
+#pragma once
+
+#include "Factories/Factory.h"
+#include "UObject/ObjectMacros.h"
+
+#include "SharedMemoryMediaOutputFactory.generated.h"
+
+
+/**
+ * Implements a factory for USharedMemoryMediaOutput objects.
+ */
+UCLASS(hidecategories=Object)
+class USharedMemoryMediaOutputFactory : public UFactory
+{
+	GENERATED_BODY()
+
+public:
+
+	USharedMemoryMediaOutputFactory();
+
+	//~ UFactory Interface
+	virtual UObject* FactoryCreateNew(UClass* InClass, UObject* InParent, FName InName, EObjectFlags Flags, UObject* Context, FFeedbackContext* Warn) override;
+	virtual uint32 GetMenuCategories() const override;
+	virtual bool ShouldShowInNewMenu() const override;
+};
diff --git a/Engine/Plugins/Runtime/nDisplay/Source/DisplayClusterMediaEditor/Private/SharedMemoryMediaSourceFactory.cpp b/Engine/Plugins/Runtime/nDisplay/Source/DisplayClusterMediaEditor/Private/SharedMemoryMediaSourceFactory.cpp
new file mode 100644
--- /dev/null
+++ b/Engine/Plugins/Runtime/nDisplay/Source/DisplayClusterMediaEditor/Private/SharedMemoryMediaSourceFactory.cpp
@@ -0,0 +1,50 @@
+// Copyright Epic Games, Inc. All Rights Reserved.
+
+#include "SharedMemoryMediaSourceFactory.h"
+
+#include "AssetTypeCategories.h"
+#include "EngineAnalytics.h"
+#include "SharedMemoryMediaSource.h"
+
+
+/* USharedMemoryMediaSourceFactory structors
+ *****************************************************************************/
+
+USharedMemoryMediaSourceFactory::USharedMemoryMediaSourceFactory()
+{
+	SupportedClass = USharedMemoryMediaSource::StaticClass();
+	bCreateNew = true;
+	bEditAfterNew = true;
+}
+
+
+/* UFactory overrides
+ *****************************************************************************/
+
+ /**
+ * @EventName MediaFramework.SharedMemoryMediaSourceCreated
+ * @Trigger Triggered when a SharedMemory media source asset is created.
+ * @Type Client
+ * @Owner MediaIO Team
+ */
+UObject* USharedMemoryMediaSourceFactory::FactoryCreateNew(UClass* InClass, UObject* InParent, FName InName, EObjectFlags Flags, UObject* Context, FFeedbackContext* Warn)
+{
+	if (FEngineAnalytics::IsAvailable())
+	{
+		FEngineAnalytics::GetProvider().RecordEvent(TEXT("MediaFramework.SharedMemoryMediaSourceCreated"));
+	}
+
+	return NewObject<USharedMemoryMediaSource>(InParent, InClass, InName, Flags);
+}
+
+
+uint32 USharedMemoryMediaSourceFactory::GetMenuCategories() const
+{
+	return EAssetTypeCategories::Media;
+}
+
+
+bool USharedMemoryMediaSourceFactory::ShouldShowInNewMenu() const
+{
+	return true;
+}
diff --git a/Engine/Plugins/Runtime/nDisplay/Source/DisplayClusterMediaEditor/Private/SharedMemoryMediaSourceFactory.h b/Engine/Plugins/Runtime/nDisplay/Source/DisplayClusterMediaEditor/Private/SharedMemoryMediaSourceFactory.h
new file mode 100644
--- /dev/null
+++ b/Engine/Plugins/Runtime/nDisplay/Source/DisplayClusterMediaEditor/Private/SharedMemoryMediaSourceFactory.h
@@ -0,0 +1,27 @@
+// Copyright Epic Games, Inc. All Rights Reserved.
+
+#pragma once
+
+#include "Factories/Factory.h"
+#include "SharedMemoryMediaSourceFactory.generated.h"
+
+
+/**
+ * Implements a factory for USharedMemoryMediaSource objects.
+ */
+UCLASS(hidecategories=Object)
+class USharedMemoryMediaSourceFactory
+	: public UFactory
+{
+	GENERATED_BODY()
+
+public:
+
+	USharedMemoryMediaSourceFactory();
+
+	//~ UFactory Interface
+
+	virtual UObject* FactoryCreateNew(UClass* InClass, UObject* InParent, FName InName, EObjectFlags Flags, UObject* Context, FFeedbackContext* Warn) override;
+	virtual uint32 GetMenuCategories() const override;
+	virtual bool ShouldShowInNewMenu() const override;
+};
diff --git a/Engine/Plugins/Runtime/nDisplay/nDisplay.uplugin b/Engine/Plugins/Runtime/nDisplay/nDisplay.uplugin
--- a/Engine/Plugins/Runtime/nDisplay/nDisplay.uplugin
+++ b/Engine/Plugins/Runtime/nDisplay/nDisplay.uplugin
@@ -152,6 +152,15 @@
 				"Linux"
 			]
 		},
+		{
+			"Name": "DisplayClusterMediaEditor",
+			"Type": "Editor",
+			"LoadingPhase": "Default",
+			"PlatformAllowList": [
+				"Win64",
+				"Linux"
+			]
+		},
 		{
 			"Name": "DisplayClusterMultiUser",
 			"Type": "UncookedOnly",
