nDisplay: SharedMemoryMedia d3d12 now correctly checks INVALID_HANDLE_VALUE.

diff --git a/Engine/Plugins/Runtime/nDisplay/Source/DisplayClusterMedia/Private/Platforms/Windows/SharedMemoryMediaPlatformWindowsD3D12.cpp b/Engine/Plugins/Runtime/nDisplay/Source/DisplayClusterMedia/Private/Platforms/Windows/SharedMemoryMediaPlatformWindowsD3D12.cpp
--- a/Engine/Plugins/Runtime/nDisplay/Source/DisplayClusterMedia/Private/Platforms/Windows/SharedMemoryMediaPlatformWindowsD3D12.cpp
+++ b/Engine/Plugins/Runtime/nDisplay/Source/DisplayClusterMedia/Private/Platforms/Windows/SharedMemoryMediaPlatformWindowsD3D12.cpp
@@ -16,19 +16,32 @@ TSharedPtr<FSharedMemoryMediaPlatform, ESPMode::ThreadSafe> FSharedMemoryMediaPl
 	return MakeShared<FSharedMemoryMediaPlatformWindowsD3D12, ESPMode::ThreadSafe>();
 }
 
+FSharedMemoryMediaPlatformWindowsD3D12::FSharedMemoryMediaPlatformWindowsD3D12()
+{
+	for (int32 Idx = 0; Idx < UE::SharedMemoryMedia::SenderNumBuffers; ++Idx)
+	{
+		SharedHandle[Idx] = INVALID_HANDLE_VALUE;
+	}
+}
+
+
 FSharedMemoryMediaPlatformWindowsD3D12::~FSharedMemoryMediaPlatformWindowsD3D12()
 {
 	// Close shared handles
 	for (int32 BufferIdx = 0; BufferIdx < UE::SharedMemoryMedia::SenderNumBuffers; ++BufferIdx)
 	{
-		::CloseHandle(SharedHandle[BufferIdx]);
+		if (SharedHandle[BufferIdx] != INVALID_HANDLE_VALUE)
+		{
+			::CloseHandle(SharedHandle[BufferIdx]);
+			SharedHandle[BufferIdx] = INVALID_HANDLE_VALUE;
+		}
 	}
 }
 
 FTextureRHIRef FSharedMemoryMediaPlatformWindowsD3D12::CreateSharedCrossGpuTexture(EPixelFormat Format, int32 Width, int32 Height, const FGuid& Guid, uint32 BufferIdx)
 {
 	check(BufferIdx < UE::SharedMemoryMedia::SenderNumBuffers);
-	check(!SharedHandle[BufferIdx]);
+	check(SharedHandle[BufferIdx] == INVALID_HANDLE_VALUE);
 
 	ID3D12Device* D3D12Device = static_cast<ID3D12Device*>(GDynamicRHI->RHIGetNativeDevice());
 	check(D3D12Device);
@@ -73,7 +86,7 @@ FTextureRHIRef FSharedMemoryMediaPlatformWindowsD3D12::CreateSharedCrossGpuTextu
 	);
 
 	check(HResult == S_OK);
-	check(SharedHandle[BufferIdx]);
+	check(SharedHandle[BufferIdx] != INVALID_HANDLE_VALUE);
 
 	return GetID3D12DynamicRHI()->RHICreateTexture2DFromResource(
 		Format,
@@ -88,6 +101,8 @@ void FSharedMemoryMediaPlatformWindowsD3D12::ReleaseSharedCrossGpuTexture(uint32
 	check(BufferIdx < UE::SharedMemoryMedia::SenderNumBuffers);
 
 	::CloseHandle(SharedHandle[BufferIdx]);
+
+	SharedHandle[BufferIdx] = INVALID_HANDLE_VALUE;
 }
 
 FTextureRHIRef FSharedMemoryMediaPlatformWindowsD3D12::OpenSharedCrossGpuTextureByGuid(const FGuid& Guid, FSharedMemoryMediaTextureDescription& OutTextureDescription)
diff --git a/Engine/Plugins/Runtime/nDisplay/Source/DisplayClusterMedia/Private/Platforms/Windows/SharedMemoryMediaPlatformWindowsD3D12.h b/Engine/Plugins/Runtime/nDisplay/Source/DisplayClusterMedia/Private/Platforms/Windows/SharedMemoryMediaPlatformWindowsD3D12.h
--- a/Engine/Plugins/Runtime/nDisplay/Source/DisplayClusterMedia/Private/Platforms/Windows/SharedMemoryMediaPlatformWindowsD3D12.h
+++ b/Engine/Plugins/Runtime/nDisplay/Source/DisplayClusterMedia/Private/Platforms/Windows/SharedMemoryMediaPlatformWindowsD3D12.h
@@ -37,5 +37,6 @@ public:
 	virtual FTextureRHIRef OpenSharedCrossGpuTextureByGuid(const FGuid& Guid, FSharedMemoryMediaTextureDescription& OutTextureDescription) override;
 	//~ End FSharedMemoryMediaPlatform interface
 
+	FSharedMemoryMediaPlatformWindowsD3D12();
 	virtual ~FSharedMemoryMediaPlatformWindowsD3D12();
 };
