nDisplay[SharedMemoryMedia]: Prevent possible static initialization race condition in platform factory.

diff --git a/Engine/Plugins/Runtime/nDisplay/Source/DisplayClusterMedia/Private/Capture/SharedMemoryMediaCapture.cpp b/Engine/Plugins/Runtime/nDisplay/Source/DisplayClusterMedia/Private/Capture/SharedMemoryMediaCapture.cpp
--- a/Engine/Plugins/Runtime/nDisplay/Source/DisplayClusterMedia/Private/Capture/SharedMemoryMediaCapture.cpp
+++ b/Engine/Plugins/Runtime/nDisplay/Source/DisplayClusterMedia/Private/Capture/SharedMemoryMediaCapture.cpp
@@ -31,12 +31,12 @@ bool USharedMemoryMediaCapture::InitializeCapture()
 	{
 		const ERHIInterfaceType RhiInterfaceType = GDynamicRHI->GetInterfaceType();
 
-		PlatformData = FSharedMemoryMediaPlatform::CreateInstanceForRhi(RhiInterfaceType);
+		PlatformData = FSharedMemoryMediaPlatformFactory::Get()->CreateInstanceForRhi(RhiInterfaceType);
 
 		if (!PlatformData.IsValid())
 		{
 			UE_LOG(LogDisplayClusterMedia, Error, TEXT("Unfortunately, SharedMemoryMedia doesn't support the current RHI type '%s'"),
-				*FSharedMemoryMediaPlatform::GetRhiTypeString(RhiInterfaceType));
+				*FSharedMemoryMediaPlatformFactory::GetRhiTypeString(RhiInterfaceType));
 
 			return false;
 		}
diff --git a/Engine/Plugins/Runtime/nDisplay/Source/DisplayClusterMedia/Private/Input/SharedMemoryMediaPlayer.cpp b/Engine/Plugins/Runtime/nDisplay/Source/DisplayClusterMedia/Private/Input/SharedMemoryMediaPlayer.cpp
--- a/Engine/Plugins/Runtime/nDisplay/Source/DisplayClusterMedia/Private/Input/SharedMemoryMediaPlayer.cpp
+++ b/Engine/Plugins/Runtime/nDisplay/Source/DisplayClusterMedia/Private/Input/SharedMemoryMediaPlayer.cpp
@@ -55,12 +55,12 @@ bool FSharedMemoryMediaPlayer::Open(const FString& Url, const IMediaOptions* Opt
 	{
 		const ERHIInterfaceType RhiInterfaceType = GDynamicRHI->GetInterfaceType();
 
-		PlatformData = FSharedMemoryMediaPlatform::CreateInstanceForRhi(RhiInterfaceType);
+		PlatformData = FSharedMemoryMediaPlatformFactory::Get()->CreateInstanceForRhi(RhiInterfaceType);
 
 		if (!PlatformData.IsValid())
 		{
 			UE_LOG(LogDisplayClusterMedia, Error, TEXT("Unfortunately, SharedMemoryMedia doesn't support the current RHI type '%s'"),
-				*FSharedMemoryMediaPlatform::GetRhiTypeString(RhiInterfaceType));
+				*FSharedMemoryMediaPlatformFactory::GetRhiTypeString(RhiInterfaceType));
 
 			return false;
 		}
diff --git a/Engine/Plugins/Runtime/nDisplay/Source/DisplayClusterMedia/Private/Platforms/Windows/SharedMemoryMediaPlatformWindowsD3D12.cpp b/Engine/Plugins/Runtime/nDisplay/Source/DisplayClusterMedia/Private/Platforms/Windows/SharedMemoryMediaPlatformWindowsD3D12.cpp
--- a/Engine/Plugins/Runtime/nDisplay/Source/DisplayClusterMedia/Private/Platforms/Windows/SharedMemoryMediaPlatformWindowsD3D12.cpp
+++ b/Engine/Plugins/Runtime/nDisplay/Source/DisplayClusterMedia/Private/Platforms/Windows/SharedMemoryMediaPlatformWindowsD3D12.cpp
@@ -8,7 +8,7 @@
 
 
 // This should run when the module starts and register the creation function for this rhi platform.
-bool FSharedMemoryMediaPlatformWindowsD3D12::bRegistered = FSharedMemoryMediaPlatform::RegisterPlatformForRhi(
+bool FSharedMemoryMediaPlatformWindowsD3D12::bRegistered = FSharedMemoryMediaPlatformFactory::Get()->RegisterPlatformForRhi(
 	ERHIInterfaceType::D3D12, &FSharedMemoryMediaPlatformWindowsD3D12::CreateInstance);
 
 TSharedPtr<FSharedMemoryMediaPlatform, ESPMode::ThreadSafe> FSharedMemoryMediaPlatformWindowsD3D12::CreateInstance()
diff --git a/Engine/Plugins/Runtime/nDisplay/Source/DisplayClusterMedia/Private/SharedMemoryMediaPlatform.cpp b/Engine/Plugins/Runtime/nDisplay/Source/DisplayClusterMedia/Private/SharedMemoryMediaPlatform.cpp
--- a/Engine/Plugins/Runtime/nDisplay/Source/DisplayClusterMedia/Private/SharedMemoryMediaPlatform.cpp
+++ b/Engine/Plugins/Runtime/nDisplay/Source/DisplayClusterMedia/Private/SharedMemoryMediaPlatform.cpp
@@ -2,9 +2,7 @@
 
 #include "SharedMemoryMediaPlatform.h"
 
-TMap<ERHIInterfaceType, FSharedMemoryMediaPlatform::CreateSharedMemoryMediaPlatform> FSharedMemoryMediaPlatform::PlatformCreators;
-
-FString FSharedMemoryMediaPlatform::GetRhiTypeString(const ERHIInterfaceType RhiType)
+FString FSharedMemoryMediaPlatformFactory::GetRhiTypeString(const ERHIInterfaceType RhiType)
 {
 	static TMap<ERHIInterfaceType, FString> Rhis;
 
@@ -25,13 +23,13 @@ FString FSharedMemoryMediaPlatform::GetRhiTypeString(const ERHIInterfaceType Rhi
 	return RhiString ? *RhiString : TEXT("Unknown");
 }
 
-bool FSharedMemoryMediaPlatform::RegisterPlatformForRhi(ERHIInterfaceType RhiType, CreateSharedMemoryMediaPlatform PlatformCreator)
+bool FSharedMemoryMediaPlatformFactory::RegisterPlatformForRhi(ERHIInterfaceType RhiType, CreateSharedMemoryMediaPlatform PlatformCreator)
 {
 	PlatformCreators.Add(RhiType, PlatformCreator);
 	return true;
 }
 
-TSharedPtr<FSharedMemoryMediaPlatform, ESPMode::ThreadSafe> FSharedMemoryMediaPlatform::CreateInstanceForRhi(ERHIInterfaceType RhiType)
+TSharedPtr<FSharedMemoryMediaPlatform, ESPMode::ThreadSafe> FSharedMemoryMediaPlatformFactory::CreateInstanceForRhi(ERHIInterfaceType RhiType)
 {
 	if (const CreateSharedMemoryMediaPlatform* Creator = PlatformCreators.Find(RhiType))
 	{
@@ -41,3 +39,10 @@ TSharedPtr<FSharedMemoryMediaPlatform, ESPMode::ThreadSafe> FSharedMemoryMediaPl
 	return nullptr;
 }
 
+FSharedMemoryMediaPlatformFactory* FSharedMemoryMediaPlatformFactory::Get()
+{
+	static FSharedMemoryMediaPlatformFactory Instance;
+	return &Instance;
+}
+
+
diff --git a/Engine/Plugins/Runtime/nDisplay/Source/DisplayClusterMedia/Private/SharedMemoryMediaPlatform.h b/Engine/Plugins/Runtime/nDisplay/Source/DisplayClusterMedia/Private/SharedMemoryMediaPlatform.h
--- a/Engine/Plugins/Runtime/nDisplay/Source/DisplayClusterMedia/Private/SharedMemoryMediaPlatform.h
+++ b/Engine/Plugins/Runtime/nDisplay/Source/DisplayClusterMedia/Private/SharedMemoryMediaPlatform.h
@@ -12,27 +12,9 @@
  */
 class FSharedMemoryMediaPlatform
 {
-public:
-
-	/** This is the type definition of the rhi platform factory function that each rhi platfom implementation can register with.  */
-	typedef TSharedPtr<FSharedMemoryMediaPlatform>(*CreateSharedMemoryMediaPlatform)();
 
 public:
 
-	/** TMap of the collection of registered rhi platform factory functions */
-	static TMap<ERHIInterfaceType, CreateSharedMemoryMediaPlatform> PlatformCreators;
-
-public:
-
-	/** Helper function to get a stringified ERHIInterfaceType */
-	static FString GetRhiTypeString(const ERHIInterfaceType RhiType);
-
-	/** Rhi platform implementations call this function to register their factory creation function */
-	static bool RegisterPlatformForRhi(ERHIInterfaceType RhiType, CreateSharedMemoryMediaPlatform PlatformCreator);
-
-	/** This factory function will create an instance of the rhi platform specific implementation, if it has been registered. */
-	static TSharedPtr<FSharedMemoryMediaPlatform, ESPMode::ThreadSafe> CreateInstanceForRhi(ERHIInterfaceType RhiType);
-
 	/** Creates a cross gpu texture */
 	virtual FTextureRHIRef CreateSharedCrossGpuTexture(EPixelFormat Format, int32 Width, int32 Height, const FGuid& Guid, uint32 BufferIdx)
 	{
@@ -47,6 +29,32 @@ public:
 
 	/** Release any platform specific resources related to the indexed texture */
 	virtual void ReleaseSharedCrossGpuTexture(uint32 BufferIdx) {};
-
 };
 
+/** Factory of registered FSharedMemoryMediaPlatform RHI implementations. Use RegisterPlatformForRhi to register them. */
+class FSharedMemoryMediaPlatformFactory
+{
+public:
+
+	/** This is the type definition of the rhi platform factory function that each rhi platfom implementation can register with.  */
+	typedef TSharedPtr<FSharedMemoryMediaPlatform>(*CreateSharedMemoryMediaPlatform)();
+
+public:
+
+	/** Rhi platform implementations call this function to register their factory creation function */
+	bool RegisterPlatformForRhi(ERHIInterfaceType RhiType, CreateSharedMemoryMediaPlatform PlatformCreator);
+
+	/** This factory function will create an instance of the rhi platform specific implementation, if it has been registered. */
+	TSharedPtr<FSharedMemoryMediaPlatform, ESPMode::ThreadSafe> CreateInstanceForRhi(ERHIInterfaceType RhiType);
+
+	/** Gets the singleton instance of this factory */
+	static FSharedMemoryMediaPlatformFactory* Get();
+
+	/** Helper function to get a stringified ERHIInterfaceType */
+	static FString GetRhiTypeString(const ERHIInterfaceType RhiType);
+
+private:
+
+	/** TMap of the collection of registered rhi platform factory functions */
+	TMap<ERHIInterfaceType, CreateSharedMemoryMediaPlatform> PlatformCreators;
+};
