nDisplay: Fix static analysis.

diff --git a/Engine/Plugins/Runtime/nDisplay/Source/DisplayClusterMedia/Private/Capture/SharedMemoryMediaCapture.cpp b/Engine/Plugins/Runtime/nDisplay/Source/DisplayClusterMedia/Private/Capture/SharedMemoryMediaCapture.cpp
--- a/Engine/Plugins/Runtime/nDisplay/Source/DisplayClusterMedia/Private/Capture/SharedMemoryMediaCapture.cpp
+++ b/Engine/Plugins/Runtime/nDisplay/Source/DisplayClusterMedia/Private/Capture/SharedMemoryMediaCapture.cpp
@@ -289,13 +289,14 @@ void USharedMemoryMediaCapture::AddCopyToSharedGpuTexturePass(FRDGBuilder& Graph
 				// Update shared memory metadata to indicate to the receiver that there is new data
 
 				FSharedMemoryMediaFrameMetadata* SharedMemoryData = static_cast<FSharedMemoryMediaFrameMetadata*>(SharedMemory[SharedTextureIdx]->GetAddress());
+				{
+					FSharedMemoryMediaFrameMetadata::FSender SenderMetadata;
+					SenderMetadata.FrameNumber = FrameNumber;
+					SenderMetadata.TextureGuid = SharedCrossGpuTextureGuids[SharedTextureIdx];
 
-				FSharedMemoryMediaFrameMetadata FrameMetadata;
-				FrameMetadata.Sender.FrameNumber = FrameNumber;
-				FrameMetadata.Sender.TextureGuid = SharedCrossGpuTextureGuids[SharedTextureIdx];
-
-				// We only send the sender structure
-				FMemory::Memcpy(SharedMemoryData, &FrameMetadata, sizeof(FSharedMemoryMediaFrameMetadata::FSender));
+					// We only send the sender structure
+					FMemory::Memcpy(&SharedMemoryData->Sender, &SenderMetadata, sizeof(FSharedMemoryMediaFrameMetadata::FSender));
+				}
 
 				// Wait for FrameNumber ack
 
diff --git a/Engine/Plugins/Runtime/nDisplay/Source/DisplayClusterMedia/Private/Platforms/Windows/SharedMemoryMediaPlatformWindowsD3D12.cpp b/Engine/Plugins/Runtime/nDisplay/Source/DisplayClusterMedia/Private/Platforms/Windows/SharedMemoryMediaPlatformWindowsD3D12.cpp
--- a/Engine/Plugins/Runtime/nDisplay/Source/DisplayClusterMedia/Private/Platforms/Windows/SharedMemoryMediaPlatformWindowsD3D12.cpp
+++ b/Engine/Plugins/Runtime/nDisplay/Source/DisplayClusterMedia/Private/Platforms/Windows/SharedMemoryMediaPlatformWindowsD3D12.cpp
@@ -114,7 +114,7 @@ FTextureRHIRef FSharedMemoryMediaPlatformWindowsD3D12::OpenSharedCrossGpuTexture
 
 	ID3D12Resource* SharedCrossGpuTexture = nullptr;
 
-	HANDLE NamedSharedGpuTextureHandle = nullptr;
+	HANDLE NamedSharedGpuTextureHandle = INVALID_HANDLE_VALUE;
 	HRESULT HResult = D3D12Device->OpenSharedHandleByName(*SharedGpuTextureName, GENERIC_ALL, &NamedSharedGpuTextureHandle);
 
 	if (FAILED(HResult))
@@ -137,6 +137,8 @@ FTextureRHIRef FSharedMemoryMediaPlatformWindowsD3D12::OpenSharedCrossGpuTexture
 		return nullptr;
 	}
 
+	NamedSharedGpuTextureHandle = INVALID_HANDLE_VALUE;
+
 	check(SharedCrossGpuTexture);
 
 	// Query texture dimensions and use them to describe downstream resources
@@ -163,7 +165,8 @@ FTextureRHIRef FSharedMemoryMediaPlatformWindowsD3D12::OpenSharedCrossGpuTexture
 			UE_LOG(LogDisplayClusterMedia, Error, TEXT("Could not find a known pixel format for SharedCrossGpuTexture DXGI_FORMAT %d."),
 				SharedGpuTextureDesc.Format);
 
-			::CloseHandle(NamedSharedGpuTextureHandle);
+			SharedCrossGpuTexture->Release();
+			SharedCrossGpuTexture = nullptr;
 
 			return nullptr;
 		}
