nDisplay: allow to render the inner views offscreen without a stub viewport

diff --git a/Engine/Plugins/Runtime/nDisplay/Source/DisplayCluster/Private/Render/Viewport/Configuration/DisplayClusterViewportConfigurationICVFX.cpp b/Engine/Plugins/Runtime/nDisplay/Source/DisplayCluster/Private/Render/Viewport/Configuration/DisplayClusterViewportConfigurationICVFX.cpp
--- a/Engine/Plugins/Runtime/nDisplay/Source/DisplayCluster/Private/Render/Viewport/Configuration/DisplayClusterViewportConfigurationICVFX.cpp
+++ b/Engine/Plugins/Runtime/nDisplay/Source/DisplayCluster/Private/Render/Viewport/Configuration/DisplayClusterViewportConfigurationICVFX.cpp
@@ -22,6 +22,8 @@
 #include "Components/DisplayClusterICVFXCameraComponent.h"
 
 #include "Misc/DisplayClusterLog.h"
+#include "Misc/Parse.h"
+
 
 ////////////////////////////////////////////////////////////////////////////////
 int32 GDisplayClusterEnableAlphaChannelRendering = 0;
@@ -408,6 +410,23 @@ void FDisplayClusterViewportConfigurationICVFX::Update()
 				}
 			}
 		}
+		else
+		{
+			// We need to be able to render inner views (ICVFX camera view) without any outers. Current pipeline has
+			// a bunch of optimizations that block inner rendering if no outers available. To avoid the limitation
+			// we force inner data initialization if no outers found.
+			static const bool bIsRenderingOffscreen = FParse::Param(FCommandLine::Get(), TEXT("RenderOffscreen"));
+
+			// Cherry-pick ICVFX cameras initialization from IF-block above
+			if (bIsRenderingOffscreen)
+			{
+				ImplGetCameras();
+				for (FDisplayClusterViewportConfigurationCameraICVFX& CameraIt : StageCameras)
+				{
+					CameraIt.Update();
+				}
+			}
+		}
 
 		ImplFinishReallocateViewports(*ViewportManager);
 	}
diff --git a/Engine/Plugins/Runtime/nDisplay/Source/DisplayCluster/Private/Render/Viewport/RenderFrame/DisplayClusterRenderFrameManager.cpp b/Engine/Plugins/Runtime/nDisplay/Source/DisplayCluster/Private/Render/Viewport/RenderFrame/DisplayClusterRenderFrameManager.cpp
--- a/Engine/Plugins/Runtime/nDisplay/Source/DisplayCluster/Private/Render/Viewport/RenderFrame/DisplayClusterRenderFrameManager.cpp
+++ b/Engine/Plugins/Runtime/nDisplay/Source/DisplayCluster/Private/Render/Viewport/RenderFrame/DisplayClusterRenderFrameManager.cpp
@@ -7,6 +7,8 @@
 #include "Render/Viewport/RenderFrame/DisplayClusterRenderFrame.h"
 #include "Render/Viewport/RenderFrame/DisplayClusterRenderFrameSettings.h"
 
+#include "Misc/Parse.h"
+
 
 ///////////////////////////////////////////////////////////////
 // FDisplayClusterRenderTargetFrame
@@ -187,7 +189,15 @@ bool FDisplayClusterRenderFrameManager::FindFrameTargetRect(FViewport* InViewpor
 
 	if (OutFrameTargetRect.Width() <= 0 || OutFrameTargetRect.Height() <= 0)
 	{
-		return false;
+		// We need to be able to render inner views (ICVFX camera view) without any outers. Current pipeline has
+		// a bunch of optimizations that block inner rendering if no outers available. To avoid the limitation
+		// we force inner data initialization if no outers found.
+		static const bool bIsRenderingOffscreen = FParse::Param(FCommandLine::Get(), TEXT("RenderOffscreen"));
+
+		if (bIsRenderingOffscreen)
+		{
+			bIsUsed = true;
+		}
 	}
 
 	return bIsUsed;
diff --git a/Engine/Plugins/Runtime/nDisplay/Source/DisplayClusterConfigurator/Private/DisplayClusterConfiguratorCompiler.cpp b/Engine/Plugins/Runtime/nDisplay/Source/DisplayClusterConfigurator/Private/DisplayClusterConfiguratorCompiler.cpp
--- a/Engine/Plugins/Runtime/nDisplay/Source/DisplayClusterConfigurator/Private/DisplayClusterConfiguratorCompiler.cpp
+++ b/Engine/Plugins/Runtime/nDisplay/Source/DisplayClusterConfigurator/Private/DisplayClusterConfiguratorCompiler.cpp
@@ -257,17 +257,12 @@ void FDisplayClusterConfiguratorKismetCompilerContext::ValidateConfiguration()
 		MessageLog.Error(*LOCTEXT("NoPrimaryNodeError", "Primary cluster node not set. Please set a primary node.").ToString());
 	}
 
-	bool bAtLeastOneViewportFound = false;
-
 	TMap<FString, UDisplayClusterConfigurationViewport*> ViewportsByName;
 	
 	for (const auto& ClusterNodes : BlueprintData->Cluster->Nodes)
 	{
 		if (ClusterNodes.Value->Viewports.Num() > 0)
 		{
-			// Pass with at least one set.
-			bAtLeastOneViewportFound = true;
-			
 			for (const auto& Viewport : ClusterNodes.Value->Viewports)
 			{
 				const FString ViewportName = Viewport.Value->GetName();
@@ -292,11 +287,15 @@ void FDisplayClusterConfiguratorKismetCompilerContext::ValidateConfiguration()
 				}
 			}
 		}
-	}
-
-	if (!bAtLeastOneViewportFound)
-	{
-		MessageLog.Warning(*LOCTEXT("NoViewportsError", "No viewports found. Please add a viewport.").ToString());
+		else
+		{
+			// It's allowed not to have any viewports when rendering offscreen. For example, we're going to render
+			// ICVFX inner views only which don't require any viewports. In other cases continue throwing a warning.
+			if (!ClusterNodes.Value->bRenderHeadless)
+			{
+				MessageLog.Warning(*LOCTEXT("NoViewportsError", "Node @@ has no viewports. Please add a viewport.").ToString(), *ClusterNodes.Key);
+			}
+		}
 	}
 }
 
