nDisplay: New media configuration approach

diff --git a/Engine/Plugins/Runtime/nDisplay/Source/DisplayCluster/Private/Render/Viewport/Configuration/DisplayClusterViewportConfigurationHelpers.cpp b/Engine/Plugins/Runtime/nDisplay/Source/DisplayCluster/Private/Render/Viewport/Configuration/DisplayClusterViewportConfigurationHelpers.cpp
--- a/Engine/Plugins/Runtime/nDisplay/Source/DisplayCluster/Private/Render/Viewport/Configuration/DisplayClusterViewportConfigurationHelpers.cpp
+++ b/Engine/Plugins/Runtime/nDisplay/Source/DisplayCluster/Private/Render/Viewport/Configuration/DisplayClusterViewportConfigurationHelpers.cpp
@@ -213,18 +213,14 @@ void FDisplayClusterViewportConfigurationHelpers::UpdateBaseViewportSetting(FDis
 	{
 		const FDisplayClusterConfigurationMedia& MediaSettings = InConfigurationViewport.RenderSettings.Media;
 
-		const FString ThisClusterNodeId = IDisplayCluster::Get().GetClusterMgr()->GetNodeId();
-		const bool bThisNodeSharesMedia = MediaSettings.IsMediaSharingUsed() && MediaSettings.MediaSharingNode.Equals(ThisClusterNodeId, ESearchCase::IgnoreCase);
-
-		// Don't render the viewport if media input assigned
-		DstViewport.RenderSettings.bSkipSceneRenderingButLeaveResourcesAvailable = MediaSettings.IsMediaSharingUsed() ?
-			!bThisNodeSharesMedia :
-			!!MediaSettings.MediaSource;
-
-		// Mark this viewport is going to be captured by a capture device
-		DstViewport.RenderSettings.bIsBeingCaptured = MediaSettings.IsMediaSharingUsed() ?
-			bThisNodeSharesMedia :
-			!!MediaSettings.MediaOutput;
+		if (MediaSettings.bEnable)
+		{
+			// Don't render the viewport if media input assigned
+			DstViewport.RenderSettings.bSkipSceneRenderingButLeaveResourcesAvailable = !!MediaSettings.MediaSource;
+
+			// Mark this viewport is going to be captured by a capture device
+			DstViewport.RenderSettings.bIsBeingCaptured = !!MediaSettings.MediaOutput;
+		}
 	}
 
 	// FDisplayClusterConfigurationViewport_ICVFX property:
diff --git a/Engine/Plugins/Runtime/nDisplay/Source/DisplayCluster/Private/Render/Viewport/Configuration/DisplayClusterViewportConfigurationHelpers_ICVFX.cpp b/Engine/Plugins/Runtime/nDisplay/Source/DisplayCluster/Private/Render/Viewport/Configuration/DisplayClusterViewportConfigurationHelpers_ICVFX.cpp
--- a/Engine/Plugins/Runtime/nDisplay/Source/DisplayCluster/Private/Render/Viewport/Configuration/DisplayClusterViewportConfigurationHelpers_ICVFX.cpp
+++ b/Engine/Plugins/Runtime/nDisplay/Source/DisplayCluster/Private/Render/Viewport/Configuration/DisplayClusterViewportConfigurationHelpers_ICVFX.cpp
@@ -565,20 +565,18 @@ void FDisplayClusterViewportConfigurationHelpers_ICVFX::UpdateCameraViewportSett
 	// Set media related configuration (runtime only for now)
 	if (IDisplayCluster::Get().GetOperationMode() == EDisplayClusterOperationMode::Cluster)
 	{
-		const FDisplayClusterConfigurationMedia& MediaSettings = InCameraComponent.CameraSettings.RenderSettings.Media;
+		const FDisplayClusterConfigurationMediaICVFX& MediaICVFXSettings = InCameraComponent.CameraSettings.RenderSettings.Media;
 
-		const FString ThisClusterNodeId = DstViewport.GetClusterNodeId();
-		const bool bThisNodeSharesMedia = MediaSettings.IsMediaSharingUsed() && MediaSettings.MediaSharingNode.Equals(ThisClusterNodeId, ESearchCase::IgnoreCase);
+		if (MediaICVFXSettings.bEnable)
+		{
+			const FString ThisClusterNodeId = DstViewport.GetClusterNodeId();
 
-		// Don't render the viewport if media input assigned
-		DstViewport.RenderSettings.bSkipSceneRenderingButLeaveResourcesAvailable = MediaSettings.IsMediaSharingUsed() ?
-			!bThisNodeSharesMedia :
-			!!MediaSettings.MediaSource;
+			// Don't render the viewport if media input assigned
+			DstViewport.RenderSettings.bSkipSceneRenderingButLeaveResourcesAvailable = MediaICVFXSettings.IsMediaInputAssigned(ThisClusterNodeId);
 
-		// Mark this viewport is going to be captured by a capture device
-		DstViewport.RenderSettings.bIsBeingCaptured = MediaSettings.IsMediaSharingUsed() ?
-			bThisNodeSharesMedia :
-			!!MediaSettings.MediaOutput;
+			// Mark this viewport is going to be captured by a capture device
+			DstViewport.RenderSettings.bIsBeingCaptured = MediaICVFXSettings.IsMediaOutputAssigned(ThisClusterNodeId);
+		}
 	}
 }
 
diff --git a/Engine/Plugins/Runtime/nDisplay/Source/DisplayClusterConfiguration/Private/DisplayClusterConfigurationTypes_Media.cpp b/Engine/Plugins/Runtime/nDisplay/Source/DisplayClusterConfiguration/Private/DisplayClusterConfigurationTypes_Media.cpp
new file mode 100644
--- /dev/null
+++ b/Engine/Plugins/Runtime/nDisplay/Source/DisplayClusterConfiguration/Private/DisplayClusterConfigurationTypes_Media.cpp
@@ -0,0 +1,54 @@
+// Copyright Epic Games, Inc. All Rights Reserved.
+
+#pragma once
+
+#include "DisplayClusterConfigurationTypes_Media.h"
+
+
+bool FDisplayClusterConfigurationMediaICVFX::IsMediaInputAssigned(const FString& InNodeId) const
+{
+	return GetMediaSource(InNodeId) != nullptr;
+}
+
+bool FDisplayClusterConfigurationMediaICVFX::IsMediaOutputAssigned(const FString& InNodeId) const
+{
+	return GetMediaOutput(InNodeId) != nullptr;
+}
+
+UMediaSource* FDisplayClusterConfigurationMediaICVFX::GetMediaSource(const FString& InNodeId) const
+{
+	// Look up for a group that contains node ID specified
+	for (const FDisplayClusterConfigurationMediaInputGroup& MediaInputGroup : MediaInputGroups)
+	{
+		const bool bFound = MediaInputGroup.ClusterNodes.ItemNames.ContainsByPredicate([InNodeId](const FString& Item)
+			{
+				return Item.Equals(InNodeId, ESearchCase::IgnoreCase);
+			});
+
+		if (bFound)
+		{
+			return MediaInputGroup.MediaSource;
+		}
+	}
+
+	return nullptr;
+}
+
+UMediaOutput* FDisplayClusterConfigurationMediaICVFX::GetMediaOutput(const FString& InNodeId) const
+{
+	// Look up for a group that contains node ID specified
+	for (const FDisplayClusterConfigurationMediaOutputGroup& MediaOutputGroup : MediaOutputGroups)
+	{
+		const bool bFound = MediaOutputGroup.ClusterNodes.ItemNames.ContainsByPredicate([InNodeId](const FString& Item)
+			{
+				return Item.Equals(InNodeId, ESearchCase::IgnoreCase);
+			});
+
+		if (bFound)
+		{
+			return MediaOutputGroup.MediaOutput;
+		}
+	}
+
+	return nullptr;
+}
diff --git a/Engine/Plugins/Runtime/nDisplay/Source/DisplayClusterConfiguration/Public/DisplayClusterConfigurationTypes_ICVFX.h b/Engine/Plugins/Runtime/nDisplay/Source/DisplayClusterConfiguration/Public/DisplayClusterConfigurationTypes_ICVFX.h
--- a/Engine/Plugins/Runtime/nDisplay/Source/DisplayClusterConfiguration/Public/DisplayClusterConfigurationTypes_ICVFX.h
+++ b/Engine/Plugins/Runtime/nDisplay/Source/DisplayClusterConfiguration/Public/DisplayClusterConfigurationTypes_ICVFX.h
@@ -323,7 +323,7 @@ public:
 
 	// Media settings
 	UPROPERTY(EditAnywhere, BlueprintReadWrite, Category = NDisplay, meta = (DisplayName = "Media"))
-	FDisplayClusterConfigurationMedia Media;
+	FDisplayClusterConfigurationMediaICVFX Media;
 
 	UPROPERTY()
 	FDisplayClusterConfigurationPostRender_BlurPostprocess PostprocessBlur;
diff --git a/Engine/Plugins/Runtime/nDisplay/Source/DisplayClusterConfiguration/Public/DisplayClusterConfigurationTypes_Media.h b/Engine/Plugins/Runtime/nDisplay/Source/DisplayClusterConfiguration/Public/DisplayClusterConfigurationTypes_Media.h
--- a/Engine/Plugins/Runtime/nDisplay/Source/DisplayClusterConfiguration/Public/DisplayClusterConfigurationTypes_Media.h
+++ b/Engine/Plugins/Runtime/nDisplay/Source/DisplayClusterConfiguration/Public/DisplayClusterConfigurationTypes_Media.h
@@ -4,6 +4,8 @@
 
 #include "CoreMinimal.h"
 
+#include "DisplayClusterConfigurationTypes_Base.h"
+
 #include "MediaPlayer.h"
 #include "MediaSource.h"
 #include "MediaTexture.h"
@@ -13,7 +15,7 @@
 
 
 /*
- * Media settings
+ * Media settings for viewports and backbuffer
  */
 USTRUCT(Blueprintable)
 struct DISPLAYCLUSTERCONFIGURATION_API FDisplayClusterConfigurationMedia
@@ -21,27 +23,100 @@ struct DISPLAYCLUSTERCONFIGURATION_API FDisplayClusterConfigurationMedia
 	GENERATED_BODY()
 
 public:
-	/** Enable/disable in-cluster media sharing */
-	UPROPERTY(EditAnywhere, BlueprintReadWrite, Category = "Media Sharing")
-	bool bMediaSharing = false;
+	/** Enable/disable media */
+	UPROPERTY(EditAnywhere, BlueprintReadWrite, Category = "Media")
+	bool bEnable = false;
+
+	/** Media source to use */
+	UPROPERTY(Instanced, EditAnywhere, BlueprintReadWrite, Category = "Media")
+	TObjectPtr<UMediaSource> MediaSource;
 
-	/** When in-cluster media sharing us used, the cluster node specified here will be used as a source (Tx node). The others will be Tx'es. */
-	UPROPERTY(EditAnywhere, BlueprintReadWrite, Category = "Media Sharing")
-	FString MediaSharingNode;
+	/** Media output to use */
+	UPROPERTY(Instanced, EditAnywhere, BlueprintReadWrite, Category = "Media")
+	TObjectPtr<UMediaOutput> MediaOutput;
+
+	UE_DEPRECATED(5.2, "This property has been deprecated")
+	UPROPERTY(meta = (DeprecatedProperty, DeprecationMessage = "This property has been deprecated."))
+	FString MediaSharingNode_DEPRECATED;
+};
+
+
+/*
+ * Input media group
+ */
+USTRUCT(Blueprintable)
+struct DISPLAYCLUSTERCONFIGURATION_API FDisplayClusterConfigurationMediaInputGroup
+{
+	GENERATED_BODY()
+
+public:
+	/** Cluster nodes that use media source below */
+	UPROPERTY(EditAnywhere, BlueprintReadWrite, Category = Media, meta = (ClusterItemType = ClusterNodes))
+	FDisplayClusterConfigurationClusterItemReferenceList ClusterNodes;
 
 	/** Media source to use */
-	UPROPERTY(EditAnywhere, BlueprintReadWrite, Category = "Media Input")
-	TObjectPtr<UMediaSource> MediaSource = nullptr;
+	UPROPERTY(Instanced, EditAnywhere, BlueprintReadWrite, Category = "Media")
+	TObjectPtr<UMediaSource> MediaSource;
+};
+
+
+/*
+ * Output media group
+ */
+USTRUCT(Blueprintable)
+struct DISPLAYCLUSTERCONFIGURATION_API FDisplayClusterConfigurationMediaOutputGroup
+{
+	GENERATED_BODY()
+
+public:
+	/** Cluster nodes that export media via MediaOutput below */
+	UPROPERTY(EditAnywhere, BlueprintReadWrite, Category = Media, meta = (ClusterItemType = ClusterNodes))
+	FDisplayClusterConfigurationClusterItemReferenceList ClusterNodes;
 
 	/** Media output to use */
-	UPROPERTY(EditAnywhere, BlueprintReadWrite, Category = "Media Output")
-	TObjectPtr<UMediaOutput> MediaOutput = nullptr;
+	UPROPERTY(Instanced, EditAnywhere, BlueprintReadWrite, Category = "Media")
+	TObjectPtr<UMediaOutput> MediaOutput;
+};
+
+
+/*
+ * Media settings for ICVFX cameras
+ */
+USTRUCT(Blueprintable)
+struct DISPLAYCLUSTERCONFIGURATION_API FDisplayClusterConfigurationMediaICVFX
+{
+	GENERATED_BODY()
+
+public:
+	/** Enable/disable media */
+	UPROPERTY(EditAnywhere, BlueprintReadWrite, Category = "Media")
+	bool bEnable = false;
+
+	/** Media output mapping */
+	UPROPERTY(EditAnywhere, BlueprintReadWrite, Category = "Media")
+	TArray<FDisplayClusterConfigurationMediaOutputGroup> MediaOutputGroups;
+
+	/** Media input mapping */
+	UPROPERTY(EditAnywhere, BlueprintReadWrite, Category = "Media")
+	TArray<FDisplayClusterConfigurationMediaInputGroup> MediaInputGroups;
 
 public:
-	/** Returns true if media sharing is used */
-	bool IsMediaSharingUsed() const
+	/** Returns true if a specific cluster node has media source assigned */
+	bool IsMediaInputAssigned(const FString& InNodeId) const;
+
+	/** Returns true if a specific cluster node has media source assigned */
+	bool IsMediaOutputAssigned(const FString& InNodeId) const;
+
+	/** Returns media source bound to a specific cluster node */
+	UMediaSource* GetMediaSource(const FString& InNodeId) const;
+
+	/** Returns media output bound to a specific cluster node */
+	UMediaOutput* GetMediaOutput(const FString& InNodeId) const;
+
+	UE_DEPRECATED(5.2, "This function has beend deprecated.")
+	bool IsMediaSharingNode(const FString& InNodeId) const
 	{
-		return bMediaSharing && MediaSource && MediaOutput && !MediaSharingNode.IsEmpty();
+		return false;
 	}
 };
 
diff --git a/Engine/Plugins/Runtime/nDisplay/Source/DisplayClusterConfigurator/Private/DisplayClusterConfiguratorModule.cpp b/Engine/Plugins/Runtime/nDisplay/Source/DisplayClusterConfigurator/Private/DisplayClusterConfiguratorModule.cpp
--- a/Engine/Plugins/Runtime/nDisplay/Source/DisplayClusterConfigurator/Private/DisplayClusterConfiguratorModule.cpp
+++ b/Engine/Plugins/Runtime/nDisplay/Source/DisplayClusterConfigurator/Private/DisplayClusterConfiguratorModule.cpp
@@ -19,7 +19,6 @@
 #include "Views/Details/Cluster/DisplayClusterConfiguratorClusterDetailsCustomization.h"
 #include "Views/Details/Cluster/DisplayClusterConfiguratorExternalImageTypeCustomization.h"
 #include "Views/Details/Cluster/DisplayClusterConfiguratorGenerateMipsCustomization.h"
-#include "Views/Details/Cluster/DisplayClusterConfiguratorMediaCustomization.h"
 #include "Views/Details/Cluster/DisplayClusterConfiguratorNodeSelectionCustomization.h"
 #include "Views/Details/Cluster/DisplayClusterConfiguratorClusterReferenceListCustomization.h"
 #include "Views/Details/Cluster/DisplayClusterConfiguratorViewportDetailsCustomization.h"
@@ -204,7 +203,6 @@ void FDisplayClusterConfiguratorModule::RegisterCustomLayouts()
 	REGISTER_PROPERTY_LAYOUT(FDisplayClusterConfigurationClusterItemReferenceList, FDisplayClusterConfiguratorClusterReferenceListCustomization);
 	REGISTER_PROPERTY_LAYOUT(FDisplayClusterConfigurationViewport_RemapData, FDisplayClusterConfiguratorViewportRemapCustomization);
 	REGISTER_PROPERTY_LAYOUT(FDisplayClusterConfigurationRectangle, FDisplayClusterConfiguratorRectangleCustomization);
-	REGISTER_PROPERTY_LAYOUT(FDisplayClusterConfigurationMedia, FDisplayClusterConfiguratorMediaCustomization);
 }
 
 void FDisplayClusterConfiguratorModule::UnregisterCustomLayouts()
diff --git a/Engine/Plugins/Runtime/nDisplay/Source/DisplayClusterConfigurator/Private/Views/Details/Cluster/DisplayClusterConfiguratorMediaCustomization.cpp b/Engine/Plugins/Runtime/nDisplay/Source/DisplayClusterConfigurator/Private/Views/Details/Cluster/DisplayClusterConfiguratorMediaCustomization.cpp
deleted file mode 100644
diff --git a/Engine/Plugins/Runtime/nDisplay/Source/DisplayClusterConfigurator/Private/Views/Details/Cluster/DisplayClusterConfiguratorMediaCustomization.h b/Engine/Plugins/Runtime/nDisplay/Source/DisplayClusterConfigurator/Private/Views/Details/Cluster/DisplayClusterConfiguratorMediaCustomization.h
deleted file mode 100644
diff --git a/Engine/Plugins/Runtime/nDisplay/Source/DisplayClusterMedia/Private/DisplayClusterMediaModule.cpp b/Engine/Plugins/Runtime/nDisplay/Source/DisplayClusterMedia/Private/DisplayClusterMediaModule.cpp
--- a/Engine/Plugins/Runtime/nDisplay/Source/DisplayClusterMedia/Private/DisplayClusterMediaModule.cpp
+++ b/Engine/Plugins/Runtime/nDisplay/Source/DisplayClusterMedia/Private/DisplayClusterMediaModule.cpp
@@ -104,41 +104,7 @@ void FDisplayClusterMediaModule::InitializeMedia()
 			{
 				const FDisplayClusterConfigurationMedia& MediaSettings = ClusterNode->Media;
 
-				const bool bShared             = MediaSettings.IsMediaSharingUsed();
-				const bool bClusterNodeMatches = MediaSettings.MediaSharingNode.Equals(ClusterNodeId, ESearchCase::IgnoreCase);
-
-				// Shared backbuffer
-				if (bShared && MediaSettings.MediaSource && MediaSettings.MediaOutput)
-				{
-					// Tx node, initialize media capture
-					if (bClusterNodeMatches)
-					{
-						const FString MediaCaptureId = RootActor->GetName() + FString("_") + ClusterNodeId + FString("_backbuffer_capture");
-
-						UE_LOG(LogDisplayClusterMedia, Log, TEXT("Initializing in-cluster (shared) backbuffer media capture for node '%s'..."), *ClusterNodeId);
-
-						CaptureNode = MakeUnique<FDisplayClusterMediaCaptureNode>(
-							MediaCaptureId, ClusterNodeId,
-							MediaSettings.MediaOutput);
-
-						AllCaptures.Add(CaptureNode.Get());
-					}
-					// Rx node, initialize media playback
-					else
-					{
-						const FString MediaInputId = RootActor->GetName() + FString("_") + ClusterNodeId + FString("_backbuffer_input");
-
-						UE_LOG(LogDisplayClusterMedia, Log, TEXT("Initializing in-cluster (shared) backbuffer media input for node '%s'..."), *ClusterNodeId);
-
-						InputNode = MakeUnique<FDisplayClusterMediaInputNode>(
-							MediaInputId, ClusterNodeId, 
-							MediaSettings.MediaSource);
-
-						AllInputs.Add(InputNode.Get());
-					}
-				}
-				// No in-cluster sharing
-				else
+				if (MediaSettings.bEnable)
 				{
 					// Media input
 					if (MediaSettings.MediaSource)
@@ -148,7 +114,7 @@ void FDisplayClusterMediaModule::InitializeMedia()
 						UE_LOG(LogDisplayClusterMedia, Log, TEXT("Initializing backbuffer media input for node '%s'..."), *ClusterNodeId);
 
 						InputNode = MakeUnique<FDisplayClusterMediaInputNode>(
-							MediaInputId, ClusterNodeId, 
+							MediaInputId, ClusterNodeId,
 							MediaSettings.MediaSource);
 
 						AllInputs.Add(InputNode.Get());
@@ -162,7 +128,7 @@ void FDisplayClusterMediaModule::InitializeMedia()
 						UE_LOG(LogDisplayClusterMedia, Log, TEXT("Initializing backbuffer media capture for node '%s'..."), *ClusterNodeId);
 
 						CaptureNode = MakeUnique<FDisplayClusterMediaCaptureNode>(
-							MediaCaptureId, ClusterNodeId, 
+							MediaCaptureId, ClusterNodeId,
 							MediaSettings.MediaOutput);
 
 						AllCaptures.Add(CaptureNode.Get());
@@ -178,45 +144,7 @@ void FDisplayClusterMediaModule::InitializeMedia()
 				{
 					const FDisplayClusterConfigurationMedia& MediaSettings = Viewport->RenderSettings.Media;
 
-					const bool bShared             = MediaSettings.IsMediaSharingUsed();
-					const bool bClusterNodeMatches = MediaSettings.MediaSharingNode.Equals(ClusterNodeId, ESearchCase::IgnoreCase);
-
-					// Shared viewport
-					if (bShared && MediaSettings.MediaSource && MediaSettings.MediaOutput)
-					{
-						// Tx node, initialize media capture
-						if (bClusterNodeMatches)
-						{
-							const FString MediaCaptureId = RootActor->GetName() + FString("_") + ClusterNodeId + FString("_") + ViewportIt.Key + FString("_viewport_capture");
-
-							UE_LOG(LogDisplayClusterMedia, Log, TEXT("Initializing in-cluster (shared) viewport capture for viewport '%s'..."), *ViewportIt.Key);
-
-							TUniquePtr<FDisplayClusterMediaCaptureViewport> NewViewportCapture = MakeUnique<FDisplayClusterMediaCaptureViewport>(
-								MediaCaptureId, ClusterNodeId, 
-								ViewportIt.Key, 
-								MediaSettings.MediaOutput);
-
-							AllCaptures.Add(NewViewportCapture.Get());
-							CaptureViewports.Emplace(MediaCaptureId, MoveTemp(NewViewportCapture));
-						}
-						// Rx node, initialize media playback
-						else
-						{
-							const FString MediaInputId = RootActor->GetName() + FString("_") + ClusterNodeId + FString("_") + ViewportIt.Key + FString("_viewport_input");
-
-							UE_LOG(LogDisplayClusterMedia, Log, TEXT("Initializing in-cluster (shared) viewport media input for viewport '%s'..."), *ViewportIt.Key);
-
-							TUniquePtr<FDisplayClusterMediaInputViewport> NewViewportInput = MakeUnique<FDisplayClusterMediaInputViewport>(
-								MediaInputId, ClusterNodeId,
-								ViewportIt.Key,
-								MediaSettings.MediaSource);
-
-							AllInputs.Add(NewViewportInput.Get());
-							InputViewports.Emplace(MediaInputId, MoveTemp(NewViewportInput));
-						}
-					}
-					// No in-cluster sharing
-					else
+					if (MediaSettings.bEnable)
 					{
 						// Media input
 						if (MediaSettings.MediaSource)
@@ -263,70 +191,42 @@ void FDisplayClusterMediaModule::InitializeMedia()
 				{
 					if (const UDisplayClusterICVFXCameraComponent* const ICVFXCamera = Cast<UDisplayClusterICVFXCameraComponent>(Component))
 					{
-						const FDisplayClusterConfigurationMedia& MediaSettings = ICVFXCamera->CameraSettings.RenderSettings.Media;
-
-						const FString ICVFXCameraName  = ICVFXCamera->GetName();
-						const FString ICVFXViewportId  = DisplayClusterMediaHelpers::GenerateICVFXViewportName(ClusterNodeId, ICVFXCameraName);
-						const FString MediaCaptureId   = RootActor->GetName() + FString("_") + ICVFXCameraName + FString("_icvfx_capture");
-						const FString MediaInputId     = RootActor->GetName() + FString("_") + ICVFXCameraName + FString("_icvfx_input");
-						const bool bShared             = MediaSettings.IsMediaSharingUsed();
-						const bool bClusterNodeMatches = MediaSettings.MediaSharingNode.Equals(ClusterNodeId, ESearchCase::IgnoreCase);
+						const FDisplayClusterConfigurationMediaICVFX& MediaSettings = ICVFXCamera->CameraSettings.RenderSettings.Media;
 
-						// Shared camera
-						if (bShared && MediaSettings.MediaSource && MediaSettings.MediaOutput)
+						if (MediaSettings.bEnable)
 						{
-							// Tx node, initialize media capture
-							if (bClusterNodeMatches)
-							{
-								UE_LOG(LogDisplayClusterMedia, Log, TEXT("Initializing in-cluster (shared) ICVFX capture for camera '%s'..."), *ICVFXCameraName);
+							const FString ICVFXCameraName = ICVFXCamera->GetName();
+							const FString ICVFXViewportId = DisplayClusterMediaHelpers::GenerateICVFXViewportName(ClusterNodeId, ICVFXCameraName);
+							const FString MediaCaptureId  = RootActor->GetName() + FString("_") + ICVFXCameraName + FString("_icvfx_capture");
+							const FString MediaInputId    = RootActor->GetName() + FString("_") + ICVFXCameraName + FString("_icvfx_input");
 
-								TUniquePtr<FDisplayClusterMediaCaptureViewport> NewICVFXCapture = MakeUnique<FDisplayClusterMediaCaptureCamera>(
-									MediaCaptureId, ClusterNodeId,
-									ICVFXCameraName, ICVFXViewportId,
-									MediaSettings.MediaOutput);
-
-								AllCaptures.Add(NewICVFXCapture.Get());
-								CaptureViewports.Emplace(MediaCaptureId, MoveTemp(NewICVFXCapture));
-							}
-							// Rx node, initialize media playback
-							else
+							// Media input only
+							if (UMediaSource* MediaSource = MediaSettings.GetMediaSource(ClusterNodeId))
 							{
-								UE_LOG(LogDisplayClusterMedia, Log, TEXT("Initializing in-cluster (shared) ICVFX media input for camera '%s'..."), *ICVFXCameraName);
+								UE_LOG(LogDisplayClusterMedia, Log, TEXT("Initializing ICVFX media input for camera '%s'..."), *ICVFXCameraName);
 
 								TUniquePtr<FDisplayClusterMediaInputViewport> NewICVFXInput = MakeUnique<FDisplayClusterMediaInputViewport>(
 									MediaInputId, ClusterNodeId,
 									ICVFXViewportId,
-									MediaSettings.MediaSource);
+									MediaSource);
 
 								AllInputs.Add(NewICVFXInput.Get());
 								InputViewports.Emplace(MediaInputId, MoveTemp(NewICVFXInput));
 							}
-						}
-						// Media input only
-						else if (MediaSettings.MediaSource)
-						{
-							UE_LOG(LogDisplayClusterMedia, Log, TEXT("Initializing ICVFX media input for camera '%s'..."), *ICVFXCameraName);
 
-							TUniquePtr<FDisplayClusterMediaInputViewport> NewICVFXInput = MakeUnique<FDisplayClusterMediaInputViewport>(
-								MediaInputId, ClusterNodeId,
-								ICVFXViewportId,
-								MediaSettings.MediaSource);
-
-							AllInputs.Add(NewICVFXInput.Get());
-							InputViewports.Emplace(MediaInputId, MoveTemp(NewICVFXInput));
-						}
-						// Media capture only
-						else if (MediaSettings.MediaOutput)
-						{
-							UE_LOG(LogDisplayClusterMedia, Log, TEXT("Initializing ICVFX capture for camera '%s'..."), *ICVFXCameraName);
+							// Media capture only
+							if (UMediaOutput* MediaOutput = MediaSettings.GetMediaOutput(ClusterNodeId))
+							{
+								UE_LOG(LogDisplayClusterMedia, Log, TEXT("Initializing ICVFX capture for camera '%s'..."), *ICVFXCameraName);
 
-							TUniquePtr<FDisplayClusterMediaCaptureViewport> NewICVFXCapture = MakeUnique<FDisplayClusterMediaCaptureCamera>(
-								MediaCaptureId, ClusterNodeId,
-								ICVFXCameraName, ICVFXViewportId,
-								MediaSettings.MediaOutput);
+								TUniquePtr<FDisplayClusterMediaCaptureViewport> NewICVFXCapture = MakeUnique<FDisplayClusterMediaCaptureCamera>(
+									MediaCaptureId, ClusterNodeId,
+									ICVFXCameraName, ICVFXViewportId,
+									MediaOutput);
 
-							AllCaptures.Add(NewICVFXCapture.Get());
-							CaptureViewports.Emplace(MediaCaptureId, MoveTemp(NewICVFXCapture));
+								AllCaptures.Add(NewICVFXCapture.Get());
+								CaptureViewports.Emplace(MediaCaptureId, MoveTemp(NewICVFXCapture));
+							}
 						}
 					}
 				}
