nDisplay: Added per-node ICVFX camera on/off.

This will allow us to have dedicated nodes that only render the given ICVFX cameras and not the others.

Also fixed deprecation warnings in Linux.

diff --git a/Engine/Plugins/Runtime/nDisplay/Source/DisplayCluster/Private/Game/EngineClasses/Scene/DisplayClusterICVFXCameraComponent.cpp b/Engine/Plugins/Runtime/nDisplay/Source/DisplayCluster/Private/Game/EngineClasses/Scene/DisplayClusterICVFXCameraComponent.cpp
--- a/Engine/Plugins/Runtime/nDisplay/Source/DisplayCluster/Private/Game/EngineClasses/Scene/DisplayClusterICVFXCameraComponent.cpp
+++ b/Engine/Plugins/Runtime/nDisplay/Source/DisplayCluster/Private/Game/EngineClasses/Scene/DisplayClusterICVFXCameraComponent.cpp
@@ -3,11 +3,16 @@
 #include "Components/DisplayClusterICVFXCameraComponent.h"
 #include "Components/DrawFrustumComponent.h"
 
+#include "Cluster/IPDisplayClusterClusterManager.h"
+
 #include "Render/Viewport/Containers/DisplayClusterViewport_CameraMotionBlur.h"
 #include "Render/Viewport/Containers/ImplDisplayClusterViewport_CustomFrustum.h"
 #include "Components/DisplayClusterCameraComponent.h"
 #include "DisplayClusterRootActor.h"
 
+#include "Misc/DisplayClusterGlobals.h"
+#include "DisplayClusterEnums.h"
+
 
 void UDisplayClusterICVFXCameraComponent::GetDesiredView(FMinimalViewInfo& DesiredView)
 {
@@ -30,6 +35,25 @@ FString UDisplayClusterICVFXCameraComponent::GetCameraUniqueId() const
 	return GetFName().ToString();
 }
 
+bool UDisplayClusterICVFXCameraComponent::IsICVFXEnabled() const
+{
+	// In runtime, we also need to check if this camera is allowed to render on a specific node
+	static const bool bIsRunningCluster = (GDisplayCluster->GetOperationMode() == EDisplayClusterOperationMode::Cluster);
+
+	if (bIsRunningCluster)
+	{
+		static const FString NodeId = GDisplayCluster->GetPrivateClusterMgr()->GetNodeId();
+		const bool bCameraIsAllowedOnThisNode = !CameraSettings.HiddenICVFXNodes.ItemNames.ContainsByPredicate([](const FString& Item)
+			{
+				return Item.Equals(NodeId, ESearchCase::IgnoreCase);
+			});
+
+		return CameraSettings.bEnable && bCameraIsAllowedOnThisNode;
+	}
+
+	return CameraSettings.bEnable;
+}
+
 #if WITH_EDITOR
 bool UDisplayClusterICVFXCameraComponent::GetEditorPreviewInfo(float DeltaTime, FMinimalViewInfo& ViewOut)
 {
diff --git a/Engine/Plugins/Runtime/nDisplay/Source/DisplayCluster/Private/Render/Viewport/Configuration/DisplayClusterViewportConfigurationICVFX.cpp b/Engine/Plugins/Runtime/nDisplay/Source/DisplayCluster/Private/Render/Viewport/Configuration/DisplayClusterViewportConfigurationICVFX.cpp
--- a/Engine/Plugins/Runtime/nDisplay/Source/DisplayCluster/Private/Render/Viewport/Configuration/DisplayClusterViewportConfigurationICVFX.cpp
+++ b/Engine/Plugins/Runtime/nDisplay/Source/DisplayCluster/Private/Render/Viewport/Configuration/DisplayClusterViewportConfigurationICVFX.cpp
@@ -527,22 +527,22 @@ void FDisplayClusterViewportConfigurationICVFX::ImplFinishReallocateViewports(FD
 
 void FDisplayClusterViewportConfigurationICVFX::ImplGetCameras()
 {
-	for (UActorComponent* ActorComponentIt : RootActor.GetComponents())
+	// Get all ICVFX camera components
+	TArray<UDisplayClusterICVFXCameraComponent*> ICVFXCameraComponents;
+	RootActor.GetComponents(ICVFXCameraComponents);
+
+	// Filter active cameras only
+	for (UDisplayClusterICVFXCameraComponent* ICVFXCameraComponent : ICVFXCameraComponents)
 	{
-		// Try to create ICVFX camera from component:
-		if (ActorComponentIt)
+		if (ICVFXCameraComponent && ICVFXCameraComponent->IsICVFXEnabled())
 		{
-			UDisplayClusterICVFXCameraComponent* CineCameraComponent = Cast<UDisplayClusterICVFXCameraComponent>(ActorComponentIt);
-			if (CineCameraComponent && CineCameraComponent->IsICVFXEnabled())
+			const FString InnerFrustumID = ICVFXCameraComponent->GetCameraUniqueId();
+			if (RootActor.IsInnerFrustumEnabled(InnerFrustumID))
 			{
-				const FString InnerFrustumID = CineCameraComponent->GetCameraUniqueId();
-				if (RootActor.IsInnerFrustumEnabled(InnerFrustumID))
+				FDisplayClusterViewportConfigurationCameraICVFX NewCamera(RootActor, *ICVFXCameraComponent);
+				if (NewCamera.Initialize())
 				{
-					FDisplayClusterViewportConfigurationCameraICVFX NewCamera(RootActor, *CineCameraComponent);
-					if (NewCamera.Initialize())
-					{
-						StageCameras.Add(NewCamera);
-					}
+					StageCameras.Add(NewCamera);
 				}
 			}
 		}
diff --git a/Engine/Plugins/Runtime/nDisplay/Source/DisplayCluster/Public/Components/DisplayClusterICVFXCameraComponent.h b/Engine/Plugins/Runtime/nDisplay/Source/DisplayCluster/Public/Components/DisplayClusterICVFXCameraComponent.h
--- a/Engine/Plugins/Runtime/nDisplay/Source/DisplayCluster/Public/Components/DisplayClusterICVFXCameraComponent.h
+++ b/Engine/Plugins/Runtime/nDisplay/Source/DisplayCluster/Public/Components/DisplayClusterICVFXCameraComponent.h
@@ -45,11 +45,8 @@ public:
 public:
 	FDisplayClusterViewport_CameraMotionBlur GetMotionBlurParameters();
 
-
-	bool IsICVFXEnabled() const
-	{
-		return CameraSettings.bEnable;
-	}
+	// Returns true if this camera is active
+	bool IsICVFXEnabled() const;
 
 	// Return unique camera name
 	FString GetCameraUniqueId() const;
@@ -91,6 +88,9 @@ private:
 	UPROPERTY(EditAnywhere, Transient, Category = "In Camera VFX", meta = (PropertyPath = "CameraSettings.HiddenICVFXViewports"))
 	FDisplayClusterEditorPropertyReference HiddenICVFXViewportsRef;
 
+	UPROPERTY(EditAnywhere, Transient, Category = "In Camera VFX", meta = (PropertyPath = "CameraSettings.HiddenICVFXNodes"))
+	FDisplayClusterEditorPropertyReference HiddenICVFXNodesRef;
+
 	UPROPERTY(EditAnywhere, Transient, Category = "In Camera VFX", meta = (PropertyPath = "CameraSettings.ExternalCameraActor"))
 	FDisplayClusterEditorPropertyReference ExternalCameraActorRef;
 
diff --git a/Engine/Plugins/Runtime/nDisplay/Source/DisplayClusterConfiguration/Public/DisplayClusterConfigurationTypes_ICVFX.h b/Engine/Plugins/Runtime/nDisplay/Source/DisplayClusterConfiguration/Public/DisplayClusterConfigurationTypes_ICVFX.h
--- a/Engine/Plugins/Runtime/nDisplay/Source/DisplayClusterConfiguration/Public/DisplayClusterConfigurationTypes_ICVFX.h
+++ b/Engine/Plugins/Runtime/nDisplay/Source/DisplayClusterConfiguration/Public/DisplayClusterConfigurationTypes_ICVFX.h
@@ -538,6 +538,10 @@ public:
 	/** A list of viewports that the inner frustum is not rendered to. */
 	UPROPERTY(EditAnywhere, BlueprintReadWrite, Category = NDisplay, meta = (ClusterItemType = Viewports, DisplayName = "Inner Frustum Hidden in Viewports", EditCondition = "bEnable"))
 	FDisplayClusterConfigurationClusterItemReferenceList HiddenICVFXViewports;
+
+	/** A list of nodes where the inner frustum is not rendered. */
+	UPROPERTY(EditAnywhere, BlueprintReadWrite, Category = NDisplay, meta = (ClusterItemType = ClusterNodes, DisplayName = "Inner Frustum Hidden on Nodes", EditCondition = "bEnable"))
+	FDisplayClusterConfigurationClusterItemReferenceList HiddenICVFXNodes;
 };
 
 USTRUCT(Blueprintable)
diff --git a/Engine/Plugins/Runtime/nDisplay/Source/DisplayClusterConfiguration/Public/DisplayClusterConfigurationTypes_Media.h b/Engine/Plugins/Runtime/nDisplay/Source/DisplayClusterConfiguration/Public/DisplayClusterConfigurationTypes_Media.h
--- a/Engine/Plugins/Runtime/nDisplay/Source/DisplayClusterConfiguration/Public/DisplayClusterConfigurationTypes_Media.h
+++ b/Engine/Plugins/Runtime/nDisplay/Source/DisplayClusterConfiguration/Public/DisplayClusterConfigurationTypes_Media.h
@@ -13,6 +13,7 @@
 
 #include "DisplayClusterConfigurationTypes_Media.generated.h"
 
+PRAGMA_DISABLE_DEPRECATION_WARNINGS
 
 /*
  * Media settings for viewports and backbuffer
@@ -35,11 +36,17 @@ public:
 	UPROPERTY(Instanced, EditAnywhere, BlueprintReadWrite, Category = "Media")
 	TObjectPtr<UMediaOutput> MediaOutput;
 
+#if WITH_EDITORONLY_DATA
+
 	UE_DEPRECATED(5.2, "This property has been deprecated")
-	UPROPERTY(meta = (DeprecatedProperty, DeprecationMessage = "This property has been deprecated."))
+	UPROPERTY()
 	FString MediaSharingNode_DEPRECATED;
+
+#endif // WITH_EDITORONLY_DATA
 };
 
+PRAGMA_ENABLE_DEPRECATION_WARNINGS
+
 
 /*
  * Input media group
