nDisplay[SharedMemoryMedia]: Minor simplification.

diff --git a/Engine/Plugins/Runtime/nDisplay/Source/DisplayClusterMedia/Private/Capture/SharedMemoryMediaCapture.cpp b/Engine/Plugins/Runtime/nDisplay/Source/DisplayClusterMedia/Private/Capture/SharedMemoryMediaCapture.cpp
--- a/Engine/Plugins/Runtime/nDisplay/Source/DisplayClusterMedia/Private/Capture/SharedMemoryMediaCapture.cpp
+++ b/Engine/Plugins/Runtime/nDisplay/Source/DisplayClusterMedia/Private/Capture/SharedMemoryMediaCapture.cpp
@@ -268,10 +268,11 @@ void USharedMemoryMediaCapture::AddCopyToSharedGpuTexturePass(FRDGBuilder& Graph
 			RunningTasksCount++;
 			UE::Tasks::Launch(UE_SOURCE_LOCATION, [FrameNumber = GFrameCounterRenderThread, SharedTextureIdx, this]()
 			{
-				FRunWhenGoingOutOfScope DecrementRunningTasksCount([this]()
-					{
-						RunningTasksCount--;
-					});
+				// Decrement RunningTasksCount when the task exits
+				ON_SCOPE_EXIT
+				{
+					RunningTasksCount--;
+				};
 
 				const FString CopyThreadName = FString::Printf(TEXT("SharedMemMediaOutputGpuTextureInTransitForFrame_%d"), FrameNumber);
 				TRACE_CPUPROFILER_EVENT_SCOPE_TEXT(*CopyThreadName);
diff --git a/Engine/Plugins/Runtime/nDisplay/Source/DisplayClusterMedia/Private/SharedMemoryMediaTypes.h b/Engine/Plugins/Runtime/nDisplay/Source/DisplayClusterMedia/Private/SharedMemoryMediaTypes.h
--- a/Engine/Plugins/Runtime/nDisplay/Source/DisplayClusterMedia/Private/SharedMemoryMediaTypes.h
+++ b/Engine/Plugins/Runtime/nDisplay/Source/DisplayClusterMedia/Private/SharedMemoryMediaTypes.h
@@ -40,22 +40,6 @@ namespace UE::SharedMemoryMedia
 	// Convenience ZeroGuid value.
 	static const FGuid ZeroGuid = FGuid(0, 0, 0, 0);
 
-	/** Use this to call a lambda when the variable gets destroyed. */
-	template <typename LAMBDA>
-	struct FRunWhenGoingOutOfScope
-	{
-		FRunWhenGoingOutOfScope(LAMBDA&& InLambda) : Lambda(InLambda) { }
-
-		~FRunWhenGoingOutOfScope()
-		{
-			Lambda();
-		}
-
-	private:
-
-		// Holds the lambda to execute in the destructor of this object
-		LAMBDA Lambda;
-	};
 }
 
 struct FSharedMemoryMediaTextureDescription
diff --git a/Engine/Plugins/Runtime/nDisplay/Source/DisplayClusterMediaEditor/Private/DisplayClusterMediaEditorModule.cpp b/Engine/Plugins/Runtime/nDisplay/Source/DisplayClusterMediaEditor/Private/DisplayClusterMediaEditorModule.cpp
--- a/Engine/Plugins/Runtime/nDisplay/Source/DisplayClusterMediaEditor/Private/DisplayClusterMediaEditorModule.cpp
+++ b/Engine/Plugins/Runtime/nDisplay/Source/DisplayClusterMediaEditor/Private/DisplayClusterMediaEditorModule.cpp
@@ -3,10 +3,5 @@
 #include "Modules/ModuleManager.h"
 
 
-class FDisplayClusterMediaEditorModule : public IModuleInterface
-{
 
-};
-
-
-IMPLEMENT_MODULE(FDisplayClusterMediaEditorModule, DisplayClusterMediaEditor);
+IMPLEMENT_MODULE(FDefaultModuleImpl, DisplayClusterMediaEditor);
