Remove unneeded test to GIntraFrameDebuggingGameThread.

diff --git a/Engine/Plugins/2D/Paper2D/Source/Paper2DEditor/Private/FlipbookEditor/FlipbookEditorViewportClient.cpp b/Engine/Plugins/2D/Paper2D/Source/Paper2DEditor/Private/FlipbookEditor/FlipbookEditorViewportClient.cpp
--- a/Engine/Plugins/2D/Paper2D/Source/Paper2DEditor/Private/FlipbookEditor/FlipbookEditorViewportClient.cpp
+++ b/Engine/Plugins/2D/Paper2D/Source/Paper2DEditor/Private/FlipbookEditor/FlipbookEditorViewportClient.cpp
@@ -102,10 +102,7 @@ void FFlipbookEditorViewportClient::Tick(float DeltaSeconds)
 
 	FPaperEditorViewportClient::Tick(DeltaSeconds);
 
-	if (!GIntraFrameDebuggingGameThread)
-	{
-		OwnedPreviewScene.GetWorld()->Tick(LEVELTICK_All, DeltaSeconds);
-	}
+	OwnedPreviewScene.GetWorld()->Tick(LEVELTICK_All, DeltaSeconds);
 }
 
 bool FFlipbookEditorViewportClient::InputKey(const FInputKeyEventArgs& EventArgs)
diff --git a/Engine/Plugins/2D/Paper2D/Source/Paper2DEditor/Private/SpriteEditor/SpriteEditorViewportClient.cpp b/Engine/Plugins/2D/Paper2D/Source/Paper2DEditor/Private/SpriteEditor/SpriteEditorViewportClient.cpp
--- a/Engine/Plugins/2D/Paper2D/Source/Paper2DEditor/Private/SpriteEditor/SpriteEditorViewportClient.cpp
+++ b/Engine/Plugins/2D/Paper2D/Source/Paper2DEditor/Private/SpriteEditor/SpriteEditorViewportClient.cpp
@@ -646,10 +646,7 @@ void FSpriteEditorViewportClient::Tick(float DeltaSeconds)
 
 	FPaperEditorViewportClient::Tick(DeltaSeconds);
 
-	if (!GIntraFrameDebuggingGameThread)
-	{
-		OwnedPreviewScene.GetWorld()->Tick(LEVELTICK_All, DeltaSeconds);
-	}
+	OwnedPreviewScene.GetWorld()->Tick(LEVELTICK_All, DeltaSeconds);
 }
 
 void FSpriteEditorViewportClient::ToggleShowSourceTexture()
diff --git a/Engine/Plugins/2D/Paper2D/Source/Paper2DEditor/Private/TileMapEditing/TileMapEditorViewportClient.cpp b/Engine/Plugins/2D/Paper2D/Source/Paper2DEditor/Private/TileMapEditing/TileMapEditorViewportClient.cpp
--- a/Engine/Plugins/2D/Paper2D/Source/Paper2DEditor/Private/TileMapEditing/TileMapEditorViewportClient.cpp
+++ b/Engine/Plugins/2D/Paper2D/Source/Paper2DEditor/Private/TileMapEditing/TileMapEditorViewportClient.cpp
@@ -214,10 +214,7 @@ void FTileMapEditorViewportClient::Tick(float DeltaSeconds)
 {
 	FPaperEditorViewportClient::Tick(DeltaSeconds);
 
-	if (!GIntraFrameDebuggingGameThread)
-	{
-		OwnedPreviewScene.GetWorld()->Tick(LEVELTICK_All, DeltaSeconds);
-	}
+	OwnedPreviewScene.GetWorld()->Tick(LEVELTICK_All, DeltaSeconds);
 }
 
 FLinearColor FTileMapEditorViewportClient::GetBackgroundColor() const
diff --git a/Engine/Plugins/2D/Paper2D/Source/Paper2DEditor/Private/TileSetEditor/SingleTileEditorViewportClient.cpp b/Engine/Plugins/2D/Paper2D/Source/Paper2DEditor/Private/TileSetEditor/SingleTileEditorViewportClient.cpp
--- a/Engine/Plugins/2D/Paper2D/Source/Paper2DEditor/Private/TileSetEditor/SingleTileEditorViewportClient.cpp
+++ b/Engine/Plugins/2D/Paper2D/Source/Paper2DEditor/Private/TileSetEditor/SingleTileEditorViewportClient.cpp
@@ -64,11 +64,7 @@ FBox FSingleTileEditorViewportClient::GetDesiredFocusBounds() const
 void FSingleTileEditorViewportClient::Tick(float DeltaSeconds)
 {
 	FPaperEditorViewportClient::Tick(DeltaSeconds);
-
-	if (!GIntraFrameDebuggingGameThread)
-	{
-		OwnedPreviewScene.GetWorld()->Tick(LEVELTICK_All, DeltaSeconds);
-	}
+	OwnedPreviewScene.GetWorld()->Tick(LEVELTICK_All, DeltaSeconds);
 }
 
 FLinearColor FSingleTileEditorViewportClient::GetBackgroundColor() const
diff --git a/Engine/Plugins/Experimental/Animation/ContextualAnimation/Source/ContextualAnimationEditor/Private/ContextualAnimPreviewScene.cpp b/Engine/Plugins/Experimental/Animation/ContextualAnimation/Source/ContextualAnimationEditor/Private/ContextualAnimPreviewScene.cpp
--- a/Engine/Plugins/Experimental/Animation/ContextualAnimation/Source/ContextualAnimationEditor/Private/ContextualAnimPreviewScene.cpp
+++ b/Engine/Plugins/Experimental/Animation/ContextualAnimation/Source/ContextualAnimationEditor/Private/ContextualAnimPreviewScene.cpp
@@ -42,8 +42,5 @@ void FContextualAnimPreviewScene::Tick(float InDeltaTime)
 		PreviewWorld->bBegunPlay = true;
 	}
 
-	if (!GIntraFrameDebuggingGameThread)
-	{
-		GetWorld()->Tick(LEVELTICK_All, InDeltaTime);
-	}
+	GetWorld()->Tick(LEVELTICK_All, InDeltaTime);
 }
\ No newline at end of file
diff --git a/Engine/Plugins/Experimental/Animation/PoseSearch/Source/Editor/Private/PoseSearchDatabasePreviewScene.cpp b/Engine/Plugins/Experimental/Animation/PoseSearch/Source/Editor/Private/PoseSearchDatabasePreviewScene.cpp
--- a/Engine/Plugins/Experimental/Animation/PoseSearch/Source/Editor/Private/PoseSearchDatabasePreviewScene.cpp
+++ b/Engine/Plugins/Experimental/Animation/PoseSearch/Source/Editor/Private/PoseSearchDatabasePreviewScene.cpp
@@ -53,10 +53,7 @@ namespace UE::PoseSearch
 			PreviewWorld->bBegunPlay = true;
 		}
 
-		if (!GIntraFrameDebuggingGameThread)
-		{
-			GetWorld()->Tick(LEVELTICK_All, InDeltaTime);
-		}
+		GetWorld()->Tick(LEVELTICK_All, InDeltaTime);
 
 		FDatabaseViewModel* ViewModel = GetEditor()->GetViewModel();
 		const UPoseSearchDatabase* Database = ViewModel->GetPoseSearchDatabase();
diff --git a/Engine/Plugins/Experimental/Dataflow/Source/DataflowEditor/Private/Dataflow/DataflowEditorViewportClient.cpp b/Engine/Plugins/Experimental/Dataflow/Source/DataflowEditor/Private/Dataflow/DataflowEditorViewportClient.cpp
--- a/Engine/Plugins/Experimental/Dataflow/Source/DataflowEditor/Private/Dataflow/DataflowEditorViewportClient.cpp
+++ b/Engine/Plugins/Experimental/Dataflow/Source/DataflowEditor/Private/Dataflow/DataflowEditorViewportClient.cpp
@@ -235,10 +235,7 @@ void FDataflowEditorViewportClient::Tick(float DeltaSeconds)
 	}
 
 	// Tick the preview scene world.
-	if (!GIntraFrameDebuggingGameThread)
-	{
-		PreviewScene->GetWorld()->Tick(IsRealtime() ? LEVELTICK_All : LEVELTICK_TimeOnly, DeltaSeconds);
-	}
+	PreviewScene->GetWorld()->Tick(IsRealtime() ? LEVELTICK_All : LEVELTICK_TimeOnly, DeltaSeconds);
 }
 
 
diff --git a/Engine/Plugins/Experimental/Mutable/Source/CustomizableObjectPopulationEditor/Private/MuCOPE/CustomizableObjectPopulationEditorViewportClient.cpp b/Engine/Plugins/Experimental/Mutable/Source/CustomizableObjectPopulationEditor/Private/MuCOPE/CustomizableObjectPopulationEditorViewportClient.cpp
--- a/Engine/Plugins/Experimental/Mutable/Source/CustomizableObjectPopulationEditor/Private/MuCOPE/CustomizableObjectPopulationEditorViewportClient.cpp
+++ b/Engine/Plugins/Experimental/Mutable/Source/CustomizableObjectPopulationEditor/Private/MuCOPE/CustomizableObjectPopulationEditorViewportClient.cpp
@@ -35,10 +35,7 @@ void FCustomizableObjectPopulationEditorViewportClient::Tick(float DeltaSeconds)
 	FEditorViewportClient::Tick(DeltaSeconds);
 
 	// Tick the preview scene world.
-	if (!GIntraFrameDebuggingGameThread)
-	{
-		PreviewScene->GetWorld()->Tick(LEVELTICK_All, DeltaSeconds);
-	}
+	PreviewScene->GetWorld()->Tick(LEVELTICK_All, DeltaSeconds);
 }
 
 
diff --git a/Engine/Plugins/Experimental/Water/Source/Editor/Private/WaterWavesEditorViewport.cpp b/Engine/Plugins/Experimental/Water/Source/Editor/Private/WaterWavesEditorViewport.cpp
--- a/Engine/Plugins/Experimental/Water/Source/Editor/Private/WaterWavesEditorViewport.cpp
+++ b/Engine/Plugins/Experimental/Water/Source/Editor/Private/WaterWavesEditorViewport.cpp
@@ -101,9 +101,6 @@ void FWaterWavesEditorViewportClient::Tick(float DeltaSeconds)
 	Super::Tick(DeltaSeconds);
 
 	// Tick the preview scene world.
-	if (!GIntraFrameDebuggingGameThread)
-	{
-		PreviewScene->GetWorld()->Tick(IsRealtime() ? LEVELTICK_All : LEVELTICK_TimeOnly, DeltaSeconds);
-	}
+	PreviewScene->GetWorld()->Tick(IsRealtime() ? LEVELTICK_All : LEVELTICK_TimeOnly, DeltaSeconds);
 }
 
diff --git a/Engine/Plugins/FX/Niagara/Source/NiagaraEditor/Private/Widgets/SNiagaraSimCacheViewport.cpp b/Engine/Plugins/FX/Niagara/Source/NiagaraEditor/Private/Widgets/SNiagaraSimCacheViewport.cpp
--- a/Engine/Plugins/FX/Niagara/Source/NiagaraEditor/Private/Widgets/SNiagaraSimCacheViewport.cpp
+++ b/Engine/Plugins/FX/Niagara/Source/NiagaraEditor/Private/Widgets/SNiagaraSimCacheViewport.cpp
@@ -68,10 +68,7 @@ void FNiagaraSimCacheViewportClient::Tick(float DeltaSeconds)
 	FEditorViewportClient::Tick(DeltaSeconds);
 
 	// Tick the preview scene world.
-	if (!GIntraFrameDebuggingGameThread)
-	{
-		PreviewScene->GetWorld()->Tick(LEVELTICK_All, DeltaSeconds);
-	}
+	PreviewScene->GetWorld()->Tick(LEVELTICK_All, DeltaSeconds);
 }
 
 void FNiagaraSimCacheViewportClient::Draw(FViewport* InViewport, FCanvas* Canvas)
diff --git a/Engine/Plugins/FX/Niagara/Source/NiagaraEditor/Private/Widgets/SNiagaraSystemViewport.cpp b/Engine/Plugins/FX/Niagara/Source/NiagaraEditor/Private/Widgets/SNiagaraSystemViewport.cpp
--- a/Engine/Plugins/FX/Niagara/Source/NiagaraEditor/Private/Widgets/SNiagaraSystemViewport.cpp
+++ b/Engine/Plugins/FX/Niagara/Source/NiagaraEditor/Private/Widgets/SNiagaraSystemViewport.cpp
@@ -109,10 +109,7 @@ void FNiagaraSystemViewportClient::Tick(float DeltaSeconds)
 	FEditorViewportClient::Tick(DeltaSeconds);
 
 	// Tick the preview scene world.
-	if (!GIntraFrameDebuggingGameThread)
-	{
-		PreviewScene->GetWorld()->Tick(LEVELTICK_All, DeltaSeconds);
-	}
+	PreviewScene->GetWorld()->Tick(LEVELTICK_All, DeltaSeconds);
 }
 
 void FNiagaraSystemViewportClient::Draw(FViewport* InViewport, FCanvas* Canvas)
@@ -907,10 +904,7 @@ void FNiagaraBaselineViewportClient::Tick(float DeltaSeconds)
 		}
 
 		// Tick the preview scene world.
-		if (!GIntraFrameDebuggingGameThread)
-		{
-			World->Tick(LEVELTICK_All, DeltaSeconds);
-		}
+		World->Tick(LEVELTICK_All, DeltaSeconds);
 	}
 }
 
diff --git a/Engine/Plugins/Runtime/HairStrands/Source/HairStrandsEditor/Private/Widgets/SGroomEditorViewport.cpp b/Engine/Plugins/Runtime/HairStrands/Source/HairStrandsEditor/Private/Widgets/SGroomEditorViewport.cpp
--- a/Engine/Plugins/Runtime/HairStrands/Source/HairStrandsEditor/Private/Widgets/SGroomEditorViewport.cpp
+++ b/Engine/Plugins/Runtime/HairStrands/Source/HairStrandsEditor/Private/Widgets/SGroomEditorViewport.cpp
@@ -84,10 +84,7 @@ void FGroomEditorViewportClient::Tick(float DeltaSeconds)
 	FEditorViewportClient::Tick(DeltaSeconds);
 	
 	// Tick the preview scene world.
-	if (!GIntraFrameDebuggingGameThread)
-	{
-		PreviewScene->GetWorld()->Tick(LEVELTICK_All, DeltaSeconds);
-	}
+	PreviewScene->GetWorld()->Tick(LEVELTICK_All, DeltaSeconds);
 }
 
 void FGroomEditorViewportClient::Draw(FViewport* InViewport, FCanvas* Canvas)
diff --git a/Engine/Plugins/Runtime/nDisplay/Source/DisplayClusterLightCardEditor/Private/Viewport/DisplayClusterLightCardEditorViewportClient.cpp b/Engine/Plugins/Runtime/nDisplay/Source/DisplayClusterLightCardEditor/Private/Viewport/DisplayClusterLightCardEditorViewportClient.cpp
--- a/Engine/Plugins/Runtime/nDisplay/Source/DisplayClusterLightCardEditor/Private/Viewport/DisplayClusterLightCardEditorViewportClient.cpp
+++ b/Engine/Plugins/Runtime/nDisplay/Source/DisplayClusterLightCardEditor/Private/Viewport/DisplayClusterLightCardEditorViewportClient.cpp
@@ -156,17 +156,14 @@ void FDisplayClusterLightCardEditorViewportClient::Tick(float DeltaSeconds)
 	CalcEditorWidgetTransform(CachedEditorWidgetWorldTransform);
 
 	// Tick the preview scene world.
-	if (!GIntraFrameDebuggingGameThread)
+	// Allow full tick only if preview simulation is enabled and we're not currently in an active SIE or PIE session
+	if (GEditor->PlayWorld == nullptr && !GEditor->bIsSimulatingInEditor)
 	{
-		// Allow full tick only if preview simulation is enabled and we're not currently in an active SIE or PIE session
-		if (GEditor->PlayWorld == nullptr && !GEditor->bIsSimulatingInEditor)
-		{
-			PreviewScene->GetWorld()->Tick(IsRealtime() ? LEVELTICK_All : LEVELTICK_TimeOnly, DeltaSeconds);
-		}
-		else
-		{
-			PreviewScene->GetWorld()->Tick(IsRealtime() ? LEVELTICK_ViewportsOnly : LEVELTICK_TimeOnly, DeltaSeconds);
-		}
+		PreviewScene->GetWorld()->Tick(IsRealtime() ? LEVELTICK_All : LEVELTICK_TimeOnly, DeltaSeconds);
+	}
+	else
+	{
+		PreviewScene->GetWorld()->Tick(IsRealtime() ? LEVELTICK_ViewportsOnly : LEVELTICK_TimeOnly, DeltaSeconds);
 	}
 
 	if (RootActorProxy.IsValid() && RootActorLevelInstance.IsValid())
diff --git a/Engine/Source/Editor/MaterialEditor/Private/SMaterialEditorViewport.cpp b/Engine/Source/Editor/MaterialEditor/Private/SMaterialEditorViewport.cpp
--- a/Engine/Source/Editor/MaterialEditor/Private/SMaterialEditorViewport.cpp
+++ b/Engine/Source/Editor/MaterialEditor/Private/SMaterialEditorViewport.cpp
@@ -111,10 +111,7 @@ void FMaterialEditorViewportClient::Tick(float DeltaSeconds)
 	FEditorViewportClient::Tick(DeltaSeconds);
 
 	// Tick the preview scene world.
-	if (!GIntraFrameDebuggingGameThread)
-	{
-		PreviewScene->GetWorld()->Tick(LEVELTICK_All, DeltaSeconds);
-	}
+	PreviewScene->GetWorld()->Tick(LEVELTICK_All, DeltaSeconds);
 }
 
 
diff --git a/Engine/Source/Editor/Persona/Private/AnimationEditorPreviewScene.cpp b/Engine/Source/Editor/Persona/Private/AnimationEditorPreviewScene.cpp
--- a/Engine/Source/Editor/Persona/Private/AnimationEditorPreviewScene.cpp
+++ b/Engine/Source/Editor/Persona/Private/AnimationEditorPreviewScene.cpp
@@ -1201,10 +1201,7 @@ void FAnimationEditorPreviewScene::Tick(float InDeltaTime)
 
 	IPersonaPreviewScene::Tick(InDeltaTime);
 
-	if (!GIntraFrameDebuggingGameThread)
-	{
-		GetWorld()->Tick(LEVELTICK_All, InDeltaTime);
-	}
+	GetWorld()->Tick(LEVELTICK_All, InDeltaTime);
 
 	if (SkeletalMeshComponent)
 	{
diff --git a/Engine/Source/Editor/StaticMeshEditor/Private/StaticMeshEditorViewportClient.cpp b/Engine/Source/Editor/StaticMeshEditor/Private/StaticMeshEditorViewportClient.cpp
--- a/Engine/Source/Editor/StaticMeshEditor/Private/StaticMeshEditorViewportClient.cpp
+++ b/Engine/Source/Editor/StaticMeshEditor/Private/StaticMeshEditorViewportClient.cpp
@@ -109,10 +109,7 @@ void FStaticMeshEditorViewportClient::Tick(float DeltaSeconds)
 	FEditorViewportClient::Tick(DeltaSeconds);
 
 	// Tick the preview scene world.
-	if (!GIntraFrameDebuggingGameThread)
-	{
-		PreviewScene->GetWorld()->Tick(LEVELTICK_All, DeltaSeconds);
-	}
+	PreviewScene->GetWorld()->Tick(LEVELTICK_All, DeltaSeconds);
 }
 
 /**
diff --git a/Engine/Source/Editor/UMGEditor/Private/WidgetBlueprintEditor.cpp b/Engine/Source/Editor/UMGEditor/Private/WidgetBlueprintEditor.cpp
--- a/Engine/Source/Editor/UMGEditor/Private/WidgetBlueprintEditor.cpp
+++ b/Engine/Source/Editor/UMGEditor/Private/WidgetBlueprintEditor.cpp
@@ -917,17 +917,14 @@ void FWidgetBlueprintEditor::Tick(float DeltaTime)
 	Super::Tick(DeltaTime);
 
 	// Tick the preview scene world.
-	if ( !GIntraFrameDebuggingGameThread )
+	// Allow full tick only if preview simulation is enabled and we're not currently in an active SIE or PIE session
+	if (bIsSimulateEnabled && GEditor->PlayWorld == nullptr && !GEditor->bIsSimulatingInEditor)
 	{
-		// Allow full tick only if preview simulation is enabled and we're not currently in an active SIE or PIE session
-		if ( bIsSimulateEnabled && GEditor->PlayWorld == nullptr && !GEditor->bIsSimulatingInEditor )
-		{
-			PreviewScene.GetWorld()->Tick(bIsRealTime ? LEVELTICK_All : LEVELTICK_TimeOnly, DeltaTime);
-		}
-		else
-		{
-			PreviewScene.GetWorld()->Tick(bIsRealTime ? LEVELTICK_ViewportsOnly : LEVELTICK_TimeOnly, DeltaTime);
-		}
+		PreviewScene.GetWorld()->Tick(bIsRealTime ? LEVELTICK_All : LEVELTICK_TimeOnly, DeltaTime);
+	}
+	else
+	{
+		PreviewScene.GetWorld()->Tick(bIsRealTime ? LEVELTICK_ViewportsOnly : LEVELTICK_TimeOnly, DeltaTime);
 	}
 
 	// Whenever animations change the generated class animations need to be updated since they are copied on compile.  This
