nDisplay[SharedMemoryMedia]: Added media source modes of operation for non nDisplay usage.

Framelocked - Matches source and local frame numbers. Always use this mode in nDisplay.
Genlocked - It doesn't match frame numbers, but it also doesn't skip frames, so will hold back the sender (up to a timeout) if it is faster than the receiver.
Freerun - It always grabs the latest frame. It may skip frames if they arrive too fast.

diff --git a/Engine/Plugins/Runtime/nDisplay/Source/DisplayClusterMedia/Private/Input/SharedMemoryMediaPlayer.cpp b/Engine/Plugins/Runtime/nDisplay/Source/DisplayClusterMedia/Private/Input/SharedMemoryMediaPlayer.cpp
--- a/Engine/Plugins/Runtime/nDisplay/Source/DisplayClusterMedia/Private/Input/SharedMemoryMediaPlayer.cpp
+++ b/Engine/Plugins/Runtime/nDisplay/Source/DisplayClusterMedia/Private/Input/SharedMemoryMediaPlayer.cpp
@@ -86,6 +86,11 @@ bool FSharedMemoryMediaPlayer::Open(const FString& Url, const IMediaOptions* Opt
 		bZeroLatency = Options->GetMediaOption(SharedMemoryMediaOption::ZeroLatency, true);
 	}
 
+	// Grab Mode of operation
+	{
+		Mode = ESharedMemoryMediaSourceMode(Options->GetMediaOption(SharedMemoryMediaOption::Mode, int64(ESharedMemoryMediaSourceMode::Framelocked)));
+	}
+
 	// Initialize gpu fences
 	for (int32 BufferIdx = 0; BufferIdx < NUMSHAREDMEM; BufferIdx++)
 	{
@@ -150,6 +155,9 @@ void FSharedMemoryMediaPlayer::Close()
 
 	// Free platform specific resources
 	PlatformData.Reset();
+
+	// Reset our state variables
+	ModeState.Reset();
 }
 
 FGuid FSharedMemoryMediaPlayer::GetPlayerPluginGUID() const
@@ -172,6 +180,12 @@ void FSharedMemoryMediaPlayer::TickFetch(FTimespan DeltaTime, FTimespan Timecode
 {
 	Super::TickFetch(DeltaTime, Timecode);
 
+	// Make sure we have platform data
+	if (!PlatformData.IsValid())
+	{
+		return;
+	}
+
 	// Make sure the sender's shared memory metadata is available. If not, don't bother creating a media sample.
 	//
 	for (int32 MemIdx = 0; MemIdx < NUMSHAREDMEM; ++MemIdx)
@@ -424,6 +438,221 @@ IMediaView& FSharedMemoryMediaPlayer::GetView()
 	return *this;
 }
 
+
+bool FSharedMemoryMediaPlayer::DetermineNextSourceFrameFramelockMode(uint64 FrameNumber, uint32& OutExpectedFrameNumber, uint32& OutSharedMemoryIdx)
+{
+	// When framelocked, the next source frame is determined by the current local frame number.
+	OutExpectedFrameNumber = InputTextureFrameNumberForFrameNumber(FrameNumber);
+	OutSharedMemoryIdx = OutExpectedFrameNumber % NUMSHAREDMEM;
+
+	return true;
+}
+
+bool FSharedMemoryMediaPlayer::DetermineNextSourceFrameGenlockMode(uint64 FrameNumber, uint32& OutExpectedFrameNumber, uint32& OutSharedMemoryIdx)
+{
+	// Look for the closest thing to the next consecutive frame.
+
+	FSharedMemoryMediaFrameMetadata* SharedMemoryData[NUMSHAREDMEM] = { 0 };
+	FSharedMemoryMediaFrameMetadata::FSender SenderMetadata[NUMSHAREDMEM];
+
+	uint32 NextFrameNumber = ~0;
+	int32 NextFrameNumberIdx = -1;
+
+	bool bAtLeastOneFrameValid = false;
+	bool bAlreadyPickedFrameFound = false;
+
+	for (int32 MemIdx = 0; MemIdx < NUMSHAREDMEM; ++MemIdx)
+	{
+		SharedMemoryData[MemIdx] = static_cast<FSharedMemoryMediaFrameMetadata*>(SharedMemory[MemIdx]->GetAddress());
+		check(SharedMemoryData[MemIdx]);
+
+		FMemory::Memcpy(&SenderMetadata[MemIdx], SharedMemoryData[MemIdx], sizeof(FSharedMemoryMediaFrameMetadata::FSender));
+
+		// Data needs to be valid
+		if (SenderMetadata[MemIdx].Magic != FSharedMemoryMediaFrameMetadata::MAGIC)
+		{
+			continue;
+		}
+
+		bAtLeastOneFrameValid = true;
+
+		if (SenderMetadata[MemIdx].FrameNumber == ModeState.Genlock.LastSourceFrameNumberPicked)
+		{
+			bAlreadyPickedFrameFound = true;
+		}
+
+		// Detect if this one is closer to being consecutive.
+		if (SenderMetadata[MemIdx].FrameNumber > ModeState.Genlock.LastSourceFrameNumberPicked)
+		{
+			const uint32 ThisDistance = SenderMetadata[MemIdx].FrameNumber - ModeState.Genlock.LastSourceFrameNumberPicked;
+			const uint32 NextDistance = NextFrameNumber - ModeState.Genlock.LastSourceFrameNumberPicked;
+
+			if (ThisDistance < NextDistance)
+			{
+				NextFrameNumber = SenderMetadata[MemIdx].FrameNumber;
+				NextFrameNumberIdx = MemIdx;
+			}
+		}
+	}
+
+	// If none were valid, return. We will not wait any time for them to arrive at this point.
+	if (!bAtLeastOneFrameValid)
+	{
+		return false;
+	}
+
+	// If none were equal greater than our last frame then no need to render. 
+	if (NextFrameNumberIdx < 0)
+	{
+		// If the previous frame wasn't found though, we should reset our state, as the source may have just(re)connected.
+		if (!bAlreadyPickedFrameFound)
+		{
+			UE_LOG(LogDisplayClusterMedia, Warning, TEXT("SharedMemoryMedia: In genlock mode, the sender's frame count seems to have reset to earlier than %u"),
+				ModeState.Genlock.LastSourceFrameNumberPicked);
+
+			ModeState.Genlock.Reset();
+		}
+
+		return false;
+	}
+
+	// No point in rendering the same frame
+	if (ModeState.Freerun.LastSourceFrameNumberPicked == NextFrameNumber)
+	{
+		return false;
+	}
+
+	OutExpectedFrameNumber = NextFrameNumber;
+	OutSharedMemoryIdx = NextFrameNumberIdx;
+
+	ModeState.Genlock.LastSourceFrameNumberPicked = OutExpectedFrameNumber;
+
+	return true;
+}
+
+bool FSharedMemoryMediaPlayer::DetermineNextSourceFrameFreerunMode(uint64 FrameNumber, uint32& OutExpectedFrameNumber, uint32& OutSharedMemoryIdx)
+{
+	// Use the latest frame, ack the others (if they haven't already been acked).
+
+	// See what is on all the reception buffers
+
+	FSharedMemoryMediaFrameMetadata* SharedMemoryData[NUMSHAREDMEM];
+	FSharedMemoryMediaFrameMetadata::FSender SenderMetadata[NUMSHAREDMEM];
+
+	int32 HighestFrameNumberIdx = -1;
+	uint32 HighestFrameNumber = 0;
+
+	int32 LowestFrameNumberIdx = -1;
+	uint32 LowestFrameNumber = ~0;
+
+	for (int32 MemIdx = 0; MemIdx < NUMSHAREDMEM; ++MemIdx)
+	{
+		SharedMemoryData[MemIdx] = static_cast<FSharedMemoryMediaFrameMetadata*>(SharedMemory[MemIdx]->GetAddress());
+		check(SharedMemoryData[MemIdx]);
+
+		FMemory::Memcpy(&SenderMetadata[MemIdx], SharedMemoryData[MemIdx], sizeof(FSharedMemoryMediaFrameMetadata::FSender));
+
+		// Data needs to be valid
+		if (SenderMetadata[MemIdx].Magic != FSharedMemoryMediaFrameMetadata::MAGIC)
+		{
+			continue;
+		}
+
+		if (HighestFrameNumber < SenderMetadata[MemIdx].FrameNumber)
+		{
+			HighestFrameNumber = SenderMetadata[MemIdx].FrameNumber;
+			HighestFrameNumberIdx = MemIdx;
+		}
+
+		if (LowestFrameNumber > SenderMetadata[MemIdx].FrameNumber)
+		{
+			LowestFrameNumber = SenderMetadata[MemIdx].FrameNumber;
+			LowestFrameNumberIdx = MemIdx;
+		}
+	}
+
+	// If none were valid, return. We will not wait any time for them to arrive at this point.
+	if (HighestFrameNumberIdx < 0)
+	{
+		return false;
+	}
+
+	check(LowestFrameNumberIdx >= 0);
+
+	// No point in rendering the same frame
+	if (ModeState.Freerun.LastSourceFrameNumberPicked == HighestFrameNumber)
+	{
+		return false;
+	}
+	else if (ModeState.Freerun.LastSourceFrameNumberPicked > HighestFrameNumber)
+	{
+		// We have picked in the past a frame number higher than the highest in the buffers.
+		// This is unexpected, something may have happened to our source.
+		UE_LOG(LogDisplayClusterMedia, Warning, TEXT("SharedMemoryMedia: In freerun mode, the sender's frame count seems to have reset from %u to %u."),
+			ModeState.Freerun.LastSourceFrameNumberPicked, HighestFrameNumber);
+
+		OutExpectedFrameNumber = LowestFrameNumber;
+		OutSharedMemoryIdx = LowestFrameNumberIdx;
+
+		ModeState.Freerun.Reset();
+		ModeState.Freerun.LastSourceFrameNumberConsideredAtIdx[OutSharedMemoryIdx] = OutExpectedFrameNumber;
+		ModeState.Freerun.LastSourceFrameNumberPicked = OutExpectedFrameNumber;
+
+		return true;
+	}
+
+	check(ModeState.Freerun.LastSourceFrameNumberPicked < HighestFrameNumber);
+
+	// ack the ones we're skipping in favor of the latest one
+	for (int32 MemIdx = 0; MemIdx < NUMSHAREDMEM; ++MemIdx)
+	{
+		const uint32 LastSourceFrameNumberConsideredAtThisIdx = ModeState.Freerun.LastSourceFrameNumberConsideredAtIdx[MemIdx];
+
+		ModeState.Freerun.LastSourceFrameNumberConsideredAtIdx[MemIdx] = SenderMetadata[MemIdx].FrameNumber;
+
+		// Don't ack the latest one just yet (we'll ack it once we're done with it)
+		if (HighestFrameNumberIdx == MemIdx)
+		{
+			continue;
+		}
+
+		// Ack the skipped frames (that haven't been considered and therefore acked before)
+		if (LastSourceFrameNumberConsideredAtThisIdx < SenderMetadata[MemIdx].FrameNumber)
+		{
+			SharedMemoryData[MemIdx]->Receiver.FrameNumberAcked = SenderMetadata[MemIdx].FrameNumber;
+		}
+	}
+
+	OutExpectedFrameNumber = HighestFrameNumber;
+	OutSharedMemoryIdx = HighestFrameNumberIdx;
+
+	ModeState.Freerun.LastSourceFrameNumberPicked = OutExpectedFrameNumber;
+
+	return true;
+
+}
+
+bool FSharedMemoryMediaPlayer::DetermineNextSourceFrame(uint64 FrameNumber, uint32& OutExpectedFrameNumber, uint32& OutSharedMemoryIdx)
+{
+	switch (Mode)
+	{
+	case ESharedMemoryMediaSourceMode::Framelocked:
+		return DetermineNextSourceFrameFramelockMode(FrameNumber, OutExpectedFrameNumber, OutSharedMemoryIdx);
+
+	case ESharedMemoryMediaSourceMode::Freerun:
+		return DetermineNextSourceFrameFreerunMode(FrameNumber, OutExpectedFrameNumber, OutSharedMemoryIdx);
+
+	case ESharedMemoryMediaSourceMode::Genlocked:
+		return DetermineNextSourceFrameGenlockMode(FrameNumber, OutExpectedFrameNumber, OutSharedMemoryIdx);
+
+	default:
+		unimplemented();
+		break;
+	}
+
+	return false;
+}
+
 void FSharedMemoryMediaPlayer::JustInTimeSampleRender()
 {
 	FRHICommandListImmediate& RHICmdList = FRHICommandListExecutor::GetImmediateCommandList();
@@ -453,8 +682,15 @@ void FSharedMemoryMediaPlayer::JustInTimeSampleRender()
 		}
 	}
 
-	const uint32 ExpectedFrameNumber = InputTextureFrameNumberForFrameNumber(GFrameCounterRenderThread);
-	const uint32 SharedMemoryIdx = ExpectedFrameNumber % NUMSHAREDMEM;
+	uint32 ExpectedFrameNumber;
+	uint32 SharedMemoryIdx;
+
+	if (!DetermineNextSourceFrame(GFrameCounterRenderThread, ExpectedFrameNumber, SharedMemoryIdx))
+	{
+		return;
+	}
+
+	check(SharedMemoryIdx < NUMSHAREDMEM);
 
 	FSharedMemoryMediaFrameMetadata* SharedMemoryData = static_cast<FSharedMemoryMediaFrameMetadata*>(SharedMemory[SharedMemoryIdx]->GetAddress());
 	check(SharedMemoryData);
diff --git a/Engine/Plugins/Runtime/nDisplay/Source/DisplayClusterMedia/Private/Input/SharedMemoryMediaPlayer.h b/Engine/Plugins/Runtime/nDisplay/Source/DisplayClusterMedia/Private/Input/SharedMemoryMediaPlayer.h
--- a/Engine/Plugins/Runtime/nDisplay/Source/DisplayClusterMedia/Private/Input/SharedMemoryMediaPlayer.h
+++ b/Engine/Plugins/Runtime/nDisplay/Source/DisplayClusterMedia/Private/Input/SharedMemoryMediaPlayer.h
@@ -210,12 +210,29 @@ protected:
 	/** Enqueues a last minute wait and copy of the expected texture data for the renderer to use */
 	virtual void JustInTimeSampleRender();
 
-private:
-
 	/** Offsets the frame number by the specified latency. Used to associate resources with frame numbers being evaluated */
 	uint32 InputTextureFrameNumberForFrameNumber(uint32 FrameNumber) const;
 
-private:
+	/** 
+	 * Determine that the next source frame number and buffer index is. This depends on the player settings.
+	 * 
+	 * @param OutExpectedFrameNumber Will contain the expected next source frame number
+	 * @param OutSharedMemoryIdx Will contain the memory idx where we should be looking for the frame
+	 * 
+	 * @return true only if the output values are valid.
+	 */
+	bool DetermineNextSourceFrame(uint64 FrameNumber, uint32& OutExpectedFrameNumber, uint32& OutSharedMemoryIdx);
+
+	/** DetermineNextSourceFrame implementration for Framelock Mode */
+	bool DetermineNextSourceFrameFramelockMode(uint64 FrameNumber, uint32& OutExpectedFrameNumber, uint32& OutSharedMemoryIdx);
+
+	/** DetermineNextSourceFrame implementration for Genlock Mode */
+	bool DetermineNextSourceFrameGenlockMode(uint64 FrameNumber, uint32& OutExpectedFrameNumber, uint32& OutSharedMemoryIdx);
+
+	/** DetermineNextSourceFrame implementration for Freerun Mode */
+	bool DetermineNextSourceFrameFreerunMode(uint64 FrameNumber, uint32& OutExpectedFrameNumber, uint32& OutSharedMemoryIdx);
+
+protected:
 
 	/** Number of resources used for pipelining the data flow */
 	static constexpr int32 NUMSHAREDMEM = UE::SharedMemoryMedia::SenderNumBuffers;
@@ -253,6 +270,53 @@ private:
 	/** Unique Name that must match the corresponding MediaOutput setting. It is used to find allocated shared memory by name */
 	FString UniqueName;
 
-	/** Zero latency option to wait for the cross gpu texture rendered on the same frame. May adversely affect fps */
+	/** Reception mode. See ESharedMemoryMediaSourceMode for more details. */
+	ESharedMemoryMediaSourceMode Mode = ESharedMemoryMediaSourceMode::Framelocked;
+
+	/** Zero latency option to wait for the cross gpu texture rendered on the same frame. May adversely affect fps. Only applicable when bUseFrameNumbers is true */
 	bool bZeroLatency = true;
+
+	/** State of each Mode */
+	struct FModeState
+	{
+		/** Freerun mode state */
+		struct FFreerunModeState
+		{
+			/** Keep track of the last source frame that was picked. Used to avoid picking the same frame and detect unexpected source frame changes. */
+			uint32 LastSourceFrameNumberPicked = 0;
+
+			/** Keep track of the last frame considered at each buffer index. Used to manage acks. */
+			uint32 LastSourceFrameNumberConsideredAtIdx[NUMSHAREDMEM]{ 0 };
+
+			/** Reset the state, as if the stream were started again. */
+			void Reset()
+			{
+				LastSourceFrameNumberPicked = 0;
+
+				FMemory::Memset(LastSourceFrameNumberConsideredAtIdx, 0, sizeof(LastSourceFrameNumberConsideredAtIdx));
+			}
+		} Freerun;
+
+		/** Genlock mode state */
+		struct FGenlockModeState
+		{
+			/** Keep track of last picked frame. Used to aim for consecutiveness. */
+			uint32 LastSourceFrameNumberPicked = 0;
+
+			/** Reset the state, as if the stream were started again. */
+			void Reset()
+			{
+				LastSourceFrameNumberPicked = 0;
+			}
+		} Genlock;
+
+		/** Resets the state for all modes */
+		void Reset()
+		{
+			Freerun.Reset();
+			Genlock.Reset();
+		}
+
+	} ModeState;
+
 };
diff --git a/Engine/Plugins/Runtime/nDisplay/Source/DisplayClusterMedia/Private/Input/SharedMemoryMediaSource.cpp b/Engine/Plugins/Runtime/nDisplay/Source/DisplayClusterMedia/Private/Input/SharedMemoryMediaSource.cpp
--- a/Engine/Plugins/Runtime/nDisplay/Source/DisplayClusterMedia/Private/Input/SharedMemoryMediaSource.cpp
+++ b/Engine/Plugins/Runtime/nDisplay/Source/DisplayClusterMedia/Private/Input/SharedMemoryMediaSource.cpp
@@ -18,6 +18,11 @@ bool USharedMemoryMediaSource::GetMediaOption(const FName& Key, bool DefaultValu
 
 int64 USharedMemoryMediaSource::GetMediaOption(const FName& Key, int64 DefaultValue) const
 {
+	if (Key == SharedMemoryMediaOption::Mode)
+	{
+		return int64(Mode);
+	}
+
 	return Super::GetMediaOption(Key, DefaultValue);
 }
 
@@ -43,6 +48,11 @@ bool USharedMemoryMediaSource::HasMediaOption(const FName& Key) const
 		return true;
 	}
 
+	if (Key == SharedMemoryMediaOption::Mode)
+	{
+		return true;
+	}
+
 	return Super::HasMediaOption(Key);
 }
 
diff --git a/Engine/Plugins/Runtime/nDisplay/Source/DisplayClusterMedia/Private/Input/SharedMemoryMediaSourceOptions.h b/Engine/Plugins/Runtime/nDisplay/Source/DisplayClusterMedia/Private/Input/SharedMemoryMediaSourceOptions.h
--- a/Engine/Plugins/Runtime/nDisplay/Source/DisplayClusterMedia/Private/Input/SharedMemoryMediaSourceOptions.h
+++ b/Engine/Plugins/Runtime/nDisplay/Source/DisplayClusterMedia/Private/Input/SharedMemoryMediaSourceOptions.h
@@ -8,5 +8,6 @@
 namespace SharedMemoryMediaOption
 {
 	static const FName UniqueName("UniqueName");
+	static const FName Mode("Mode");
 	static const FName ZeroLatency("ZeroLatency");
 }
diff --git a/Engine/Plugins/Runtime/nDisplay/Source/DisplayClusterMedia/Public/SharedMemoryMediaSource.h b/Engine/Plugins/Runtime/nDisplay/Source/DisplayClusterMedia/Public/SharedMemoryMediaSource.h
--- a/Engine/Plugins/Runtime/nDisplay/Source/DisplayClusterMedia/Public/SharedMemoryMediaSource.h
+++ b/Engine/Plugins/Runtime/nDisplay/Source/DisplayClusterMedia/Public/SharedMemoryMediaSource.h
@@ -7,6 +7,20 @@
 
 #include "SharedMemoryMediaSource.generated.h"
 
+/** 
+ * Mode of operation when receiving data.
+ * Framelocked - Matches source and local frame numbers. Always use this mode in nDisplay.
+ * Genlocked - It doesn't match frame numbers, but it also doesn't skip frames, so will hold back the sender if it is faster than the receiver.
+ * Freerun - It always grabs the latest frame. It may skip frames if they arrive too fast.
+ */
+UENUM(BlueprintType)
+enum class ESharedMemoryMediaSourceMode : uint8
+{
+	Framelocked = 0, // Matches source and local frame numbers. Always use this mode in nDisplay.
+	Genlocked,       // It doesn't match frame numbers, but it also doesn't skip frames, so will hold back the sender if it is faster than the receiver.
+	Freerun,         // It always grabs the latest frame. It may skip frames if they arrive too fast.
+};
+
 /**
  * Media source for SharedMemory streams.
  */
@@ -21,8 +35,17 @@ public:
 	UPROPERTY(EditAnywhere, BlueprintReadWrite, Category = "Input")
 	FString UniqueName = TEXT("UniqueName");
 
-	/** Zero latency option to wait for the cross gpu texture rendered on the same frame. May adversely affect fps */
+	/**
+	 * Mode of operation when receiving data.
+	 * Framelocked - Matches source and local frame numbers. Always use this mode in nDisplay.
+	 * Genlocked - It doesn't match frame numbers, but it also doesn't skip frames, so will hold back the sender if it is faster than the receiver.
+	 * Freerun - It always grabs the latest frame. It may skip frames if they arrive too fast.
+	 */
 	UPROPERTY(EditAnywhere, BlueprintReadWrite, Category = "Input")
+	ESharedMemoryMediaSourceMode Mode;
+
+	/** Zero latency option to wait for the cross gpu texture rendered on the same frame. May adversely affect fps. Only applicable in Framelocked mode. */
+	UPROPERTY(EditAnywhere, BlueprintReadWrite, Category = "Input", meta = (EditCondition = "bUseFrameNumbers"))
 	bool bZeroLatency = true;
 
 public:
