nDisplay[SharedMemoryMedia]: Fixed typeless pixel format.

This was preventing the cross gpu texture from being created in dx12 for some back buffer pixel formats.

Also fixed a possible resource release while in use by an async task in the player.

diff --git a/Engine/Plugins/Runtime/nDisplay/Source/DisplayClusterMedia/Private/Capture/SharedMemoryMediaCapture.cpp b/Engine/Plugins/Runtime/nDisplay/Source/DisplayClusterMedia/Private/Capture/SharedMemoryMediaCapture.cpp
--- a/Engine/Plugins/Runtime/nDisplay/Source/DisplayClusterMedia/Private/Capture/SharedMemoryMediaCapture.cpp
+++ b/Engine/Plugins/Runtime/nDisplay/Source/DisplayClusterMedia/Private/Capture/SharedMemoryMediaCapture.cpp
@@ -116,7 +116,7 @@ void USharedMemoryMediaCapture::StopCaptureImpl(bool bAllowPendingFrameToBeProce
 
 	// Since FlushRenderingCommands was already called by StopCapture, we can safely release all the resources
 
-	// ... if we don't have any running tasks.
+	// Wait for any pending tasks to finish, which could be trying to use the resources as well.
 	while (RunningTasksCount > 0)
 	{
 		FPlatformProcess::SleepNoStats(SpinWaitTimeSeconds);
@@ -198,16 +198,23 @@ void USharedMemoryMediaCapture::OnCustomCapture_RenderingThread(
 		if (!SharedCrossGpuTextures[Idx].IsValid())
 		{
 			const FGuid Guid = FGuid::NewGuid();
-
+						
 			SharedCrossGpuTextures[Idx] = PlatformData->CreateSharedCrossGpuTexture(
-				InSourceTexture->Desc.Format, 
+				InSourceTexture->Desc.Format,
+				EnumHasAnyFlags(InSourceTexture->Desc.Flags, TexCreate_SRGB),
 				CopyInfo.Size.X, 
 				CopyInfo.Size.Y, 
 				Guid, 
 				Idx
 			);
 
-			check(SharedCrossGpuTextures[Idx].IsValid());
+			if (!SharedCrossGpuTextures[Idx].IsValid())
+			{
+				UE_LOG(LogDisplayClusterMedia, Error, TEXT("Unable to create cross GPU texture of the requested type."));
+
+				SetState(EMediaCaptureState::Error);
+				return;
+			}
 
 			SharedCrossGpuTextureGuids[Idx] = Guid;
 			UE_LOG(LogDisplayClusterMedia, Verbose, TEXT("Created SharedGpuTextureGuid[%d] = %s"), Idx, *SharedCrossGpuTextureGuids[Idx].ToString());
diff --git a/Engine/Plugins/Runtime/nDisplay/Source/DisplayClusterMedia/Private/Capture/SharedMemoryMediaOutput.cpp b/Engine/Plugins/Runtime/nDisplay/Source/DisplayClusterMedia/Private/Capture/SharedMemoryMediaOutput.cpp
--- a/Engine/Plugins/Runtime/nDisplay/Source/DisplayClusterMedia/Private/Capture/SharedMemoryMediaOutput.cpp
+++ b/Engine/Plugins/Runtime/nDisplay/Source/DisplayClusterMedia/Private/Capture/SharedMemoryMediaOutput.cpp
@@ -3,6 +3,7 @@
 #include "SharedMemoryMediaOutput.h"
 
 #include "MediaOutput.h"
+#include "Engine/RendererSettings.h"
 #include "SharedMemoryMediaCapture.h"
 #include "SharedMemoryMediaTypes.h"
 
@@ -25,7 +26,9 @@ EPixelFormat USharedMemoryMediaOutput::GetRequestedPixelFormat() const
 {
 	// We are copying to a shared cross gpu texture allocated outside the MF, so this PF needs to be valid but is not relevant.
 	// @todo is there a way to not have the MF render target allocated since we're not using it ?
-	return EPixelFormat::PF_A8;
+
+	static const auto CVarDefaultBackBufferPixelFormat = IConsoleManager::Get().FindTConsoleVariableDataInt(TEXT("r.DefaultBackBufferPixelFormat"));
+	return EDefaultBackBufferPixelFormat::Convert2PixelFormat(EDefaultBackBufferPixelFormat::FromInt(CVarDefaultBackBufferPixelFormat->GetValueOnGameThread()));
 }
 
 EMediaCaptureConversionOperation USharedMemoryMediaOutput::GetConversionOperation(EMediaCaptureSourceType InSourceType) const
diff --git a/Engine/Plugins/Runtime/nDisplay/Source/DisplayClusterMedia/Private/Input/SharedMemoryMediaPlayer.cpp b/Engine/Plugins/Runtime/nDisplay/Source/DisplayClusterMedia/Private/Input/SharedMemoryMediaPlayer.cpp
--- a/Engine/Plugins/Runtime/nDisplay/Source/DisplayClusterMedia/Private/Input/SharedMemoryMediaPlayer.cpp
+++ b/Engine/Plugins/Runtime/nDisplay/Source/DisplayClusterMedia/Private/Input/SharedMemoryMediaPlayer.cpp
@@ -117,6 +117,12 @@ void FSharedMemoryMediaPlayer::Close()
 	// We flush rendering commands so that we can safely release the resources
 	FlushRenderingCommands();
 
+	// Wait for any pending tasks to finish, which could be trying to use the resources as well.
+	while (RunningTasksCount > 0)
+	{
+		FPlatformProcess::SleepNoStats(SpinWaitTimeSeconds);
+	}
+
 	// Release resources related to the indexed textures
 	for (int32 BufferIdx = 0; BufferIdx < NUMSHAREDMEM; ++BufferIdx)
 	{
@@ -790,11 +796,20 @@ void FSharedMemoryMediaPlayer::JustInTimeSampleRender()
 	}
 
 	// spawn async task that will wait on FrameAckFence, ack the frame to the Media Capture, and flag that it is clear for re-use
+
+	RunningTasksCount++;
+
 	UE::Tasks::Launch(UE_SOURCE_LOCATION,
 		[SharedMemoryData, ExpectedFrameNumber, SharedMemoryIdx, this]()
 		{
 			TRACE_CPUPROFILER_EVENT_SCOPE(SharedMemoryMediaFrameAckTask);
 
+			// Decrement RunningTasksCount when the task exits
+			ON_SCOPE_EXIT
+			{
+				RunningTasksCount--;
+			};
+
 			do
 			{
 				const bool bFenceReached = !FrameAckFences[SharedMemoryIdx] || FrameAckFences[SharedMemoryIdx]->Poll();
diff --git a/Engine/Plugins/Runtime/nDisplay/Source/DisplayClusterMedia/Private/Input/SharedMemoryMediaPlayer.h b/Engine/Plugins/Runtime/nDisplay/Source/DisplayClusterMedia/Private/Input/SharedMemoryMediaPlayer.h
--- a/Engine/Plugins/Runtime/nDisplay/Source/DisplayClusterMedia/Private/Input/SharedMemoryMediaPlayer.h
+++ b/Engine/Plugins/Runtime/nDisplay/Source/DisplayClusterMedia/Private/Input/SharedMemoryMediaPlayer.h
@@ -276,6 +276,9 @@ protected:
 	/** Zero latency option to wait for the cross gpu texture rendered on the same frame. May adversely affect fps. Only applicable when bUseFrameNumbers is true */
 	bool bZeroLatency = true;
 
+	/** Counter of running tasks used to detect when to release resources */
+	std::atomic<int32> RunningTasksCount{ 0 };
+
 	/** State of each Mode */
 	struct FModeState
 	{
diff --git a/Engine/Plugins/Runtime/nDisplay/Source/DisplayClusterMedia/Private/Platforms/Windows/SharedMemoryMediaPlatformWindowsD3D12.cpp b/Engine/Plugins/Runtime/nDisplay/Source/DisplayClusterMedia/Private/Platforms/Windows/SharedMemoryMediaPlatformWindowsD3D12.cpp
--- a/Engine/Plugins/Runtime/nDisplay/Source/DisplayClusterMedia/Private/Platforms/Windows/SharedMemoryMediaPlatformWindowsD3D12.cpp
+++ b/Engine/Plugins/Runtime/nDisplay/Source/DisplayClusterMedia/Private/Platforms/Windows/SharedMemoryMediaPlatformWindowsD3D12.cpp
@@ -5,12 +5,68 @@
 #include "ID3D12DynamicRHI.h"
 
 #include "DisplayClusterMediaLog.h"
-
+#include "Misc/EnumClassFlags.h"
 
 // This should run when the module starts and register the creation function for this rhi platform.
 bool FSharedMemoryMediaPlatformWindowsD3D12::bRegistered = FSharedMemoryMediaPlatformFactory::Get()->RegisterPlatformForRhi(
 	ERHIInterfaceType::D3D12, &FSharedMemoryMediaPlatformWindowsD3D12::CreateInstance);
 
+namespace UE::FSharedMemoryMediaPlatformWindowsD3D12
+{
+	DXGI_FORMAT FindSharedResourceDXGIFormat(DXGI_FORMAT InFormat, bool bSRGB)
+	{
+		if (bSRGB)
+		{
+			switch (InFormat)
+			{
+			case DXGI_FORMAT_B8G8R8X8_TYPELESS:    return DXGI_FORMAT_B8G8R8X8_UNORM_SRGB;
+			case DXGI_FORMAT_B8G8R8A8_TYPELESS:    return DXGI_FORMAT_B8G8R8A8_UNORM_SRGB;
+			case DXGI_FORMAT_R8G8B8A8_TYPELESS:    return DXGI_FORMAT_R8G8B8A8_UNORM_SRGB;
+			case DXGI_FORMAT_BC1_TYPELESS:         return DXGI_FORMAT_BC1_UNORM_SRGB;
+			case DXGI_FORMAT_BC2_TYPELESS:         return DXGI_FORMAT_BC2_UNORM_SRGB;
+			case DXGI_FORMAT_BC3_TYPELESS:         return DXGI_FORMAT_BC3_UNORM_SRGB;
+			case DXGI_FORMAT_BC7_TYPELESS:         return DXGI_FORMAT_BC7_UNORM_SRGB;
+			};
+		}
+		else
+		{
+			switch (InFormat)
+			{
+			case DXGI_FORMAT_B8G8R8X8_TYPELESS:    return DXGI_FORMAT_B8G8R8X8_UNORM;
+			case DXGI_FORMAT_B8G8R8A8_TYPELESS: return DXGI_FORMAT_B8G8R8A8_UNORM;
+			case DXGI_FORMAT_R8G8B8A8_TYPELESS: return DXGI_FORMAT_R8G8B8A8_UNORM;
+			case DXGI_FORMAT_BC1_TYPELESS:      return DXGI_FORMAT_BC1_UNORM;
+			case DXGI_FORMAT_BC2_TYPELESS:      return DXGI_FORMAT_BC2_UNORM;
+			case DXGI_FORMAT_BC3_TYPELESS:      return DXGI_FORMAT_BC3_UNORM;
+			case DXGI_FORMAT_BC7_TYPELESS:      return DXGI_FORMAT_BC7_UNORM;
+			};
+		}
+		switch (InFormat)
+		{
+		case DXGI_FORMAT_R32G32B32A32_TYPELESS: return DXGI_FORMAT_R32G32B32A32_UINT;
+		case DXGI_FORMAT_R32G32B32_TYPELESS:    return DXGI_FORMAT_R32G32B32_UINT;
+		case DXGI_FORMAT_R16G16B16A16_TYPELESS: return DXGI_FORMAT_R16G16B16A16_UNORM;
+		case DXGI_FORMAT_R32G32_TYPELESS:       return DXGI_FORMAT_R32G32_UINT;
+		case DXGI_FORMAT_R10G10B10A2_TYPELESS:  return DXGI_FORMAT_R10G10B10A2_UNORM;
+		case DXGI_FORMAT_R16G16_TYPELESS:       return DXGI_FORMAT_R16G16_UNORM;
+		case DXGI_FORMAT_R8G8_TYPELESS:         return DXGI_FORMAT_R8G8_UNORM;
+		case DXGI_FORMAT_R8_TYPELESS:           return DXGI_FORMAT_R8_UNORM;
+
+		case DXGI_FORMAT_BC4_TYPELESS:         return DXGI_FORMAT_BC4_UNORM;
+		case DXGI_FORMAT_BC5_TYPELESS:         return DXGI_FORMAT_BC5_UNORM;
+
+
+
+		case DXGI_FORMAT_R24G8_TYPELESS: return DXGI_FORMAT_R24_UNORM_X8_TYPELESS;
+		case DXGI_FORMAT_R32_TYPELESS: return DXGI_FORMAT_R32_FLOAT;
+		case DXGI_FORMAT_R16_TYPELESS: return DXGI_FORMAT_R16_UNORM;
+			// Changing Depth Buffers to 32 bit on Dingo as D24S8 is actually implemented as a 32 bit buffer in the hardware
+		case DXGI_FORMAT_R32G8X24_TYPELESS: return DXGI_FORMAT_R32_FLOAT_X8X24_TYPELESS;
+		}
+		return InFormat;
+	}
+}
+
 TSharedPtr<FSharedMemoryMediaPlatform, ESPMode::ThreadSafe> FSharedMemoryMediaPlatformWindowsD3D12::CreateInstance()
 {
 	return MakeShared<FSharedMemoryMediaPlatformWindowsD3D12, ESPMode::ThreadSafe>();
@@ -38,16 +94,20 @@ FSharedMemoryMediaPlatformWindowsD3D12::~FSharedMemoryMediaPlatformWindowsD3D12(
 	}
 }
 
-FTextureRHIRef FSharedMemoryMediaPlatformWindowsD3D12::CreateSharedCrossGpuTexture(EPixelFormat Format, int32 Width, int32 Height, const FGuid& Guid, uint32 BufferIdx)
+FTextureRHIRef FSharedMemoryMediaPlatformWindowsD3D12::CreateSharedCrossGpuTexture(EPixelFormat Format, bool bSrgb, int32 Width, int32 Height, const FGuid& Guid, uint32 BufferIdx)
 {
+	using namespace UE::FSharedMemoryMediaPlatformWindowsD3D12;
+
 	check(BufferIdx < UE::SharedMemoryMedia::SenderNumBuffers);
 	check(SharedHandle[BufferIdx] == INVALID_HANDLE_VALUE);
 
 	ID3D12Device* D3D12Device = static_cast<ID3D12Device*>(GDynamicRHI->RHIGetNativeDevice());
 	check(D3D12Device);
 
+	const DXGI_FORMAT DxgiFormat = FindSharedResourceDXGIFormat(DXGI_FORMAT(GPixelFormats[Format].PlatformFormat), bSrgb);
+
 	D3D12_RESOURCE_DESC SharedCrossGpuTextureDesc = CD3DX12_RESOURCE_DESC::Tex2D(
-		DXGI_FORMAT(GPixelFormats[Format].PlatformFormat),
+		DxgiFormat,
 		Width,
 		Height,
 		1, // arraySize
@@ -71,6 +131,16 @@ FTextureRHIRef FSharedMemoryMediaPlatformWindowsD3D12::CreateSharedCrossGpuTextu
 		IID_PPV_ARGS(&SharedGpuResource)
 	);
 
+	if (FAILED(HResult))
+	{
+		UE_LOG(LogDisplayClusterMedia, Error, TEXT(
+			"D3D12Device->CreateCommittedResource failed when creating a cross GPU texture:\n"
+			"0x%X - %s. Texture was %dx%d, EPixelFormat %d and DXGI type %d"),
+			HResult, *GetD3D12ComErrorDescription(HResult), Width, Height, Format, DxgiFormat);
+
+		return nullptr;
+	}
+
 	check(HResult == S_OK);
 	check(SharedGpuResource);
 
@@ -107,6 +177,8 @@ void FSharedMemoryMediaPlatformWindowsD3D12::ReleaseSharedCrossGpuTexture(uint32
 
 FTextureRHIRef FSharedMemoryMediaPlatformWindowsD3D12::OpenSharedCrossGpuTextureByGuid(const FGuid& Guid, FSharedMemoryMediaTextureDescription& OutTextureDescription)
 {
+	using namespace UE::FSharedMemoryMediaPlatformWindowsD3D12;
+
 	ID3D12Device* D3D12Device = static_cast<ID3D12Device*>(GDynamicRHI->RHIGetNativeDevice());
 	check(D3D12Device);
 
@@ -147,15 +219,28 @@ FTextureRHIRef FSharedMemoryMediaPlatformWindowsD3D12::OpenSharedCrossGpuTexture
 
 		OutTextureDescription.Height = SharedGpuTextureDesc.Height;
 		OutTextureDescription.Width = SharedGpuTextureDesc.Width;
+		OutTextureDescription.Format = EPixelFormat::PF_Unknown;
 
 		// Find EPixelFormat from platform format
 		for (int32 FormatIdx = 0; FormatIdx < PF_MAX; FormatIdx++)
 		{
-			if (GPixelFormats[FormatIdx].PlatformFormat == SharedGpuTextureDesc.Format)
+			constexpr bool SrgbOptions[2] = { false, true };
+
+			for (const bool bSrgb : SrgbOptions)
+			{
+				const DXGI_FORMAT ResultingPlatformFormat = FindSharedResourceDXGIFormat(DXGI_FORMAT(GPixelFormats[FormatIdx].PlatformFormat), bSrgb);
+
+				if (ResultingPlatformFormat == SharedGpuTextureDesc.Format)
+				{
+					OutTextureDescription.Format = EPixelFormat(FormatIdx);
+					OutTextureDescription.BytesPerPixel = GPixelFormats[FormatIdx].BlockBytes;
+					OutTextureDescription.Stride = OutTextureDescription.Width * OutTextureDescription.BytesPerPixel;
+					break;
+				}
+			}
+
+			if (OutTextureDescription.Format != EPixelFormat::PF_Unknown)
 			{
-				OutTextureDescription.Format = EPixelFormat(FormatIdx);
-				OutTextureDescription.BytesPerPixel = GPixelFormats[FormatIdx].BlockBytes;
-				OutTextureDescription.Stride = OutTextureDescription.Width * OutTextureDescription.BytesPerPixel;
 				break;
 			}
 		}
diff --git a/Engine/Plugins/Runtime/nDisplay/Source/DisplayClusterMedia/Private/Platforms/Windows/SharedMemoryMediaPlatformWindowsD3D12.h b/Engine/Plugins/Runtime/nDisplay/Source/DisplayClusterMedia/Private/Platforms/Windows/SharedMemoryMediaPlatformWindowsD3D12.h
--- a/Engine/Plugins/Runtime/nDisplay/Source/DisplayClusterMedia/Private/Platforms/Windows/SharedMemoryMediaPlatformWindowsD3D12.h
+++ b/Engine/Plugins/Runtime/nDisplay/Source/DisplayClusterMedia/Private/Platforms/Windows/SharedMemoryMediaPlatformWindowsD3D12.h
@@ -32,7 +32,7 @@ private:
 public:
 
 	//~ Begin FSharedMemoryMediaPlatform interface
-	virtual FTextureRHIRef CreateSharedCrossGpuTexture(EPixelFormat Format, int32 Width, int32 Height, const FGuid& Guid, uint32 BufferIdx) override;
+	virtual FTextureRHIRef CreateSharedCrossGpuTexture(EPixelFormat Format, bool bSrgb, int32 Width, int32 Height, const FGuid& Guid, uint32 BufferIdx) override;
 	virtual void ReleaseSharedCrossGpuTexture(uint32 BufferIdx) override;
 	virtual FTextureRHIRef OpenSharedCrossGpuTextureByGuid(const FGuid& Guid, FSharedMemoryMediaTextureDescription& OutTextureDescription) override;
 	//~ End FSharedMemoryMediaPlatform interface
diff --git a/Engine/Plugins/Runtime/nDisplay/Source/DisplayClusterMedia/Private/SharedMemoryMediaPlatform.h b/Engine/Plugins/Runtime/nDisplay/Source/DisplayClusterMedia/Private/SharedMemoryMediaPlatform.h
--- a/Engine/Plugins/Runtime/nDisplay/Source/DisplayClusterMedia/Private/SharedMemoryMediaPlatform.h
+++ b/Engine/Plugins/Runtime/nDisplay/Source/DisplayClusterMedia/Private/SharedMemoryMediaPlatform.h
@@ -16,7 +16,7 @@ class FSharedMemoryMediaPlatform
 public:
 
 	/** Creates a cross gpu texture */
-	virtual FTextureRHIRef CreateSharedCrossGpuTexture(EPixelFormat Format, int32 Width, int32 Height, const FGuid& Guid, uint32 BufferIdx)
+	virtual FTextureRHIRef CreateSharedCrossGpuTexture(EPixelFormat Format, bool bSrgb, int32 Width, int32 Height, const FGuid& Guid, uint32 BufferIdx)
 	{
 		return nullptr;
 	}
