DMX - Keep light cards flush to wall when controlled via DMX and bAlwaysFlushToWall is set

diff --git a/Engine/Plugins/Runtime/nDisplay/Source/DisplayCluster/Private/Game/EngineClasses/Scene/DisplayClusterLightCardActor.cpp b/Engine/Plugins/Runtime/nDisplay/Source/DisplayCluster/Private/Game/EngineClasses/Scene/DisplayClusterLightCardActor.cpp
--- a/Engine/Plugins/Runtime/nDisplay/Source/DisplayCluster/Private/Game/EngineClasses/Scene/DisplayClusterLightCardActor.cpp
+++ b/Engine/Plugins/Runtime/nDisplay/Source/DisplayCluster/Private/Game/EngineClasses/Scene/DisplayClusterLightCardActor.cpp
@@ -159,13 +159,14 @@ void ADisplayClusterLightCardActor::PostEditChangeProperty(FPropertyChangedEvent
 		PropertyChangedEvent.Property->GetFName() == GET_MEMBER_NAME_CHECKED(ADisplayClusterLightCardActor, Yaw) ||
 		PropertyChangedEvent.Property->GetFName() == GET_MEMBER_NAME_CHECKED(ADisplayClusterLightCardActor, RadialOffset) ||
 		PropertyChangedEvent.Property->GetFName() == GET_MEMBER_NAME_CHECKED(ADisplayClusterLightCardActor, Scale) ||
-		PropertyChangedEvent.Property->GetFName() == GET_MEMBER_NAME_CHECKED(ADisplayClusterLightCardActor, bIsUVLightCard)))
+		PropertyChangedEvent.Property->GetFName() == GET_MEMBER_NAME_CHECKED(ADisplayClusterLightCardActor, bIsUVLightCard) ||
+		PropertyChangedEvent.Property->GetFName() == GET_MEMBER_NAME_CHECKED(ADisplayClusterLightCardActor, bAlwaysFlushToWall)))
 	{
 		UpdateStageActorTransform();
 
-		if (bAlwaysFlushToWall && RootActorOwner.IsValid())
+		if (bAlwaysFlushToWall)
 		{
-			RootActorOwner->MakeStageActorFlushToWall(this);
+			MakeFlushToWall();
 		}
 	}
 
@@ -445,6 +446,14 @@ void ADisplayClusterLightCardActor::UpdateUVIndicator()
 #endif // WITH_EDITOR
 }
 
+void ADisplayClusterLightCardActor::MakeFlushToWall()
+{
+	if (RootActorOwner.IsValid())
+	{
+		RootActorOwner->MakeStageActorFlushToWall(this);
+	}
+}
+
 void ADisplayClusterLightCardActor::SetIsLightCardFlag(bool bNewFlagValue)
 {
 	bIsLightCardFlag = bNewFlagValue;
diff --git a/Engine/Plugins/Runtime/nDisplay/Source/DisplayCluster/Public/DisplayClusterLightCardActor.h b/Engine/Plugins/Runtime/nDisplay/Source/DisplayCluster/Public/DisplayClusterLightCardActor.h
--- a/Engine/Plugins/Runtime/nDisplay/Source/DisplayCluster/Public/DisplayClusterLightCardActor.h
+++ b/Engine/Plugins/Runtime/nDisplay/Source/DisplayCluster/Public/DisplayClusterLightCardActor.h
@@ -96,6 +96,9 @@ public:
 	/** Updates the UV Indicator */
 	void UpdateUVIndicator();
 
+	/** Makes the light card flush to the wall */
+	void MakeFlushToWall();
+
 	/** Configures this light card as a flag */
 	void SetIsLightCardFlag(bool bNewFlagValue);
 
diff --git a/Engine/Plugins/VirtualProduction/DMX/DMXDisplayCluster/Source/DMXDisplayClusterLightCard/Private/DMXDisplayClusterLightCardComponent.cpp b/Engine/Plugins/VirtualProduction/DMX/DMXDisplayCluster/Source/DMXDisplayClusterLightCard/Private/DMXDisplayClusterLightCardComponent.cpp
--- a/Engine/Plugins/VirtualProduction/DMX/DMXDisplayCluster/Source/DMXDisplayClusterLightCard/Private/DMXDisplayClusterLightCardComponent.cpp
+++ b/Engine/Plugins/VirtualProduction/DMX/DMXDisplayCluster/Source/DMXDisplayClusterLightCard/Private/DMXDisplayClusterLightCardComponent.cpp
@@ -220,6 +220,12 @@ struct FDMXDisplayClusterLightCardActorData
 		if (bNeedsUpdateLightCardTransform)
 		{
 			Actor->UpdateStageActorTransform();
+
+			// When the transform properties change and if bAlwaysFlushToWall is set, remain flush with the owning stage actor.
+			if (Actor->bAlwaysFlushToWall)
+			{
+				Actor->MakeFlushToWall();
+			}
 		}
 
 		if (bNeedsUpdateLightCardMaterialInstance)
