nDisplay[SharedMemoryMedia]: Genlock mode now also forces receiver to wait of sender is slow.

This will truly genlock the 2 instances, where one cannot tick faster than the other. Timeouts still apply.

diff --git a/Engine/Plugins/Runtime/nDisplay/Source/DisplayClusterMedia/Private/Input/SharedMemoryMediaPlayer.cpp b/Engine/Plugins/Runtime/nDisplay/Source/DisplayClusterMedia/Private/Input/SharedMemoryMediaPlayer.cpp
--- a/Engine/Plugins/Runtime/nDisplay/Source/DisplayClusterMedia/Private/Input/SharedMemoryMediaPlayer.cpp
+++ b/Engine/Plugins/Runtime/nDisplay/Source/DisplayClusterMedia/Private/Input/SharedMemoryMediaPlayer.cpp
@@ -458,14 +458,12 @@ bool FSharedMemoryMediaPlayer::DetermineNextSourceFrameGenlockMode(uint64 FrameN
 {
 	// Look for the closest thing to the next consecutive frame.
 
+	FModeState::FGenlockModeState& State = ModeState.Genlock;
+
 	FSharedMemoryMediaFrameMetadata* SharedMemoryData[NUMSHAREDMEM] = { 0 };
 	FSharedMemoryMediaFrameMetadata::FSender SenderMetadata[NUMSHAREDMEM];
 
-	uint32 NextFrameNumber = ~0;
-	int32 NextFrameNumberIdx = -1;
-
-	bool bAtLeastOneFrameValid = false;
-	bool bAlreadyPickedFrameFound = false;
+	TArray<uint32, TInlineAllocator<NUMSHAREDMEM>> ValidFrameNumbers;
 
 	for (int32 MemIdx = 0; MemIdx < NUMSHAREDMEM; ++MemIdx)
 	{
@@ -480,58 +478,40 @@ bool FSharedMemoryMediaPlayer::DetermineNextSourceFrameGenlockMode(uint64 FrameN
 			continue;
 		}
 
-		bAtLeastOneFrameValid = true;
-
-		if (SenderMetadata[MemIdx].FrameNumber == ModeState.Genlock.LastSourceFrameNumberPicked)
-		{
-			bAlreadyPickedFrameFound = true;
-		}
-
-		// Detect if this one is closer to being consecutive.
-		if (SenderMetadata[MemIdx].FrameNumber > ModeState.Genlock.LastSourceFrameNumberPicked)
-		{
-			const uint32 ThisDistance = SenderMetadata[MemIdx].FrameNumber - ModeState.Genlock.LastSourceFrameNumberPicked;
-			const uint32 NextDistance = NextFrameNumber - ModeState.Genlock.LastSourceFrameNumberPicked;
-
-			if (ThisDistance < NextDistance)
-			{
-				NextFrameNumber = SenderMetadata[MemIdx].FrameNumber;
-				NextFrameNumberIdx = MemIdx;
-			}
-		}
+		ValidFrameNumbers.Add(SenderMetadata[MemIdx].FrameNumber);
 	}
 
 	// If none were valid, return. We will not wait any time for them to arrive at this point.
-	if (!bAtLeastOneFrameValid)
+	if (!ValidFrameNumbers.Num())
 	{
 		return false;
 	}
 
-	// If none were equal greater than our last frame then no need to render. 
-	if (NextFrameNumberIdx < 0)
-	{
-		// If the previous frame wasn't found though, we should reset our state, as the source may have just(re)connected.
-		if (!bAlreadyPickedFrameFound)
-		{
-			UE_LOG(LogDisplayClusterMedia, Warning, TEXT("SharedMemoryMedia: In genlock mode, the sender's frame count seems to have reset to earlier than %u"),
-				ModeState.Genlock.LastSourceFrameNumberPicked);
+	uint32 MinDistanceFromLastPicked = INT_MAX;
 
-			ModeState.Genlock.Reset();
-		}
+	for (const uint32& ValidFrameNumber : ValidFrameNumbers)
+	{
+		const uint32 DistanceFromPicked = 
+			State.LastSourceFrameNumberPicked > ValidFrameNumber ? 
+			State.LastSourceFrameNumberPicked - ValidFrameNumber : ValidFrameNumber - State.LastSourceFrameNumberPicked;
 
-		return false;
+		MinDistanceFromLastPicked = FMath::Min(DistanceFromPicked, MinDistanceFromLastPicked);
 	}
 
-	// No point in rendering the same frame
-	if (ModeState.Freerun.LastSourceFrameNumberPicked == NextFrameNumber)
+	// If our frame number seems far from what is in the buffers, reset our state
+	if (MinDistanceFromLastPicked > 2)
 	{
-		return false;
-	}
+		UE_LOG(LogDisplayClusterMedia, Warning, TEXT("SharedMemoryMedia: In genlock mode, the sender's frame count seems to have reset to earlier than %u"),
+			State.LastSourceFrameNumberPicked);
 
-	OutExpectedFrameNumber = NextFrameNumber;
-	OutSharedMemoryIdx = NextFrameNumberIdx;
+		State.Reset();
 
-	ModeState.Genlock.LastSourceFrameNumberPicked = OutExpectedFrameNumber;
+		ValidFrameNumbers.Sort();
+		State.LastSourceFrameNumberPicked = ValidFrameNumbers.Last() - 1;
+	}
+
+	OutExpectedFrameNumber = ++State.LastSourceFrameNumberPicked;
+	OutSharedMemoryIdx = OutExpectedFrameNumber % NUMSHAREDMEM;
 
 	return true;
 }
@@ -540,6 +520,8 @@ bool FSharedMemoryMediaPlayer::DetermineNextSourceFrameFreerunMode(uint64 FrameN
 {
 	// Use the latest frame, ack the others (if they haven't already been acked).
 
+	FModeState::FFreerunModeState& State = ModeState.Freerun;
+
 	// See what is on all the reception buffers
 
 	FSharedMemoryMediaFrameMetadata* SharedMemoryData[NUMSHAREDMEM];
@@ -586,35 +568,35 @@ bool FSharedMemoryMediaPlayer::DetermineNextSourceFrameFreerunMode(uint64 FrameN
 	check(LowestFrameNumberIdx >= 0);
 
 	// No point in rendering the same frame
-	if (ModeState.Freerun.LastSourceFrameNumberPicked == HighestFrameNumber)
+	if (State.LastSourceFrameNumberPicked == HighestFrameNumber)
 	{
 		return false;
 	}
-	else if (ModeState.Freerun.LastSourceFrameNumberPicked > HighestFrameNumber)
+	else if (State.LastSourceFrameNumberPicked > HighestFrameNumber)
 	{
 		// We have picked in the past a frame number higher than the highest in the buffers.
 		// This is unexpected, something may have happened to our source.
 		UE_LOG(LogDisplayClusterMedia, Warning, TEXT("SharedMemoryMedia: In freerun mode, the sender's frame count seems to have reset from %u to %u."),
-			ModeState.Freerun.LastSourceFrameNumberPicked, HighestFrameNumber);
+			State.LastSourceFrameNumberPicked, HighestFrameNumber);
 
 		OutExpectedFrameNumber = LowestFrameNumber;
 		OutSharedMemoryIdx = LowestFrameNumberIdx;
 
-		ModeState.Freerun.Reset();
-		ModeState.Freerun.LastSourceFrameNumberConsideredAtIdx[OutSharedMemoryIdx] = OutExpectedFrameNumber;
-		ModeState.Freerun.LastSourceFrameNumberPicked = OutExpectedFrameNumber;
+		State.Reset();
+		State.LastSourceFrameNumberConsideredAtIdx[OutSharedMemoryIdx] = OutExpectedFrameNumber;
+		State.LastSourceFrameNumberPicked = OutExpectedFrameNumber;
 
 		return true;
 	}
 
-	check(ModeState.Freerun.LastSourceFrameNumberPicked < HighestFrameNumber);
+	check(State.LastSourceFrameNumberPicked < HighestFrameNumber);
 
 	// ack the ones we're skipping in favor of the latest one
 	for (int32 MemIdx = 0; MemIdx < NUMSHAREDMEM; ++MemIdx)
 	{
-		const uint32 LastSourceFrameNumberConsideredAtThisIdx = ModeState.Freerun.LastSourceFrameNumberConsideredAtIdx[MemIdx];
+		const uint32 LastSourceFrameNumberConsideredAtThisIdx = State.LastSourceFrameNumberConsideredAtIdx[MemIdx];
 
-		ModeState.Freerun.LastSourceFrameNumberConsideredAtIdx[MemIdx] = SenderMetadata[MemIdx].FrameNumber;
+		State.LastSourceFrameNumberConsideredAtIdx[MemIdx] = SenderMetadata[MemIdx].FrameNumber;
 
 		// Don't ack the latest one just yet (we'll ack it once we're done with it)
 		if (HighestFrameNumberIdx == MemIdx)
@@ -632,7 +614,7 @@ bool FSharedMemoryMediaPlayer::DetermineNextSourceFrameFreerunMode(uint64 FrameN
 	OutExpectedFrameNumber = HighestFrameNumber;
 	OutSharedMemoryIdx = HighestFrameNumberIdx;
 
-	ModeState.Freerun.LastSourceFrameNumberPicked = OutExpectedFrameNumber;
+	State.LastSourceFrameNumberPicked = OutExpectedFrameNumber;
 
 	return true;
 
