nDisplay: Support DataLayers.
Fix lightcards potentially being dirtied by just opening the ICVFX editor.
Fix spawning new stage actors not using the same label calculations as the rest of the engine.
Refactor adding light cards so stage app and ICVFX editor use the same routine.

diff --git a/Engine/Plugins/Runtime/nDisplay/Source/DisplayCluster/DisplayCluster.Build.cs b/Engine/Plugins/Runtime/nDisplay/Source/DisplayCluster/DisplayCluster.Build.cs
--- a/Engine/Plugins/Runtime/nDisplay/Source/DisplayCluster/DisplayCluster.Build.cs
+++ b/Engine/Plugins/Runtime/nDisplay/Source/DisplayCluster/DisplayCluster.Build.cs
@@ -59,7 +59,8 @@ public class DisplayCluster : ModuleRules
 			PrivateDependencyModuleNames.AddRange(
 				new string[]
 				{
-					"ConcertSyncClient"
+					"ConcertSyncClient",
+					"DataLayerEditor"
 				});
 		}
 
diff --git a/Engine/Plugins/Runtime/nDisplay/Source/DisplayCluster/Private/Game/EngineClasses/Scene/DisplayClusterChromakeyCardActor.cpp b/Engine/Plugins/Runtime/nDisplay/Source/DisplayCluster/Private/Game/EngineClasses/Scene/DisplayClusterChromakeyCardActor.cpp
--- a/Engine/Plugins/Runtime/nDisplay/Source/DisplayCluster/Private/Game/EngineClasses/Scene/DisplayClusterChromakeyCardActor.cpp
+++ b/Engine/Plugins/Runtime/nDisplay/Source/DisplayCluster/Private/Game/EngineClasses/Scene/DisplayClusterChromakeyCardActor.cpp
@@ -4,9 +4,11 @@
 
 #include "DisplayClusterRootActor.h"
 #include "Components/DisplayClusterICVFXCameraComponent.h"
+#include "Misc/DisplayClusterLog.h"
 
 #if WITH_EDITOR
 #include "Layers/LayersSubsystem.h"
+#include "DataLayer/DataLayerEditorSubsystem.h"
 #endif
 
 ADisplayClusterChromakeyCardActor::ADisplayClusterChromakeyCardActor(const FObjectInitializer& ObjectInitializer)
@@ -67,11 +69,38 @@ void ADisplayClusterChromakeyCardActor::AddToChromakeyLayer(ADisplayClusterRootA
 		check(!ChromakeyLayerName.IsNone());
 		bool bLayerAdded = false;
 #if WITH_EDITOR
-		ULayersSubsystem* LayersSubsystem = GEditor ? GEditor->GetEditorSubsystem<ULayersSubsystem>() : nullptr;
-		if (LayersSubsystem)
+		//TArray<AActor*> ActorsInLayer; // @todo uv chromakey cards - this will be used to determine translucent sort order
+	
+		if (SupportsLayers())
+		{
+			if (ULayersSubsystem* LayersSubsystem = GEditor ? GEditor->GetEditorSubsystem<ULayersSubsystem>() : nullptr)
+			{
+				LayersSubsystem->AddActorsToLayer(TArray<AActor*>{ this }, ChromakeyLayerName);
+				bLayerAdded = true;
+			}
+		}
+		else if (SupportsDataLayer() && GetLevel())
 		{
-			LayersSubsystem->AddActorsToLayer(TArray<AActor*>{ this }, ChromakeyLayerName);
-			bLayerAdded = true;
+			if (UDataLayerEditorSubsystem* DataLayersSubsystem = GEditor ? GEditor->GetEditorSubsystem<UDataLayerEditorSubsystem>() : nullptr)
+			{
+				if (UDataLayerAsset* DataLayerAsset = InRootActor->GetOrCreateChromakeyCardDataLayerAsset(ChromakeyLayerName))
+				{
+					UDataLayerInstance* DataLayerInstance = DataLayersSubsystem->GetDataLayerInstance(DataLayerAsset);
+					if (DataLayerInstance == nullptr)
+					{
+						FDataLayerCreationParameters CreationParams;
+						CreationParams.WorldDataLayers = GetLevel()->GetWorldDataLayers();
+						CreationParams.DataLayerAsset = DataLayerAsset;
+						DataLayerInstance = DataLayersSubsystem->CreateDataLayerInstance(MoveTemp(CreationParams));
+					}
+			
+					AddDataLayer(DataLayerInstance);
+				}
+				else
+				{
+					UE_LOG(LogDisplayClusterGame, Error, TEXT("Chromakey Card Actor could not find or create the data layer asset for layer '%s'."),  *ChromakeyLayerName.ToString());
+				}
+			}
 		}
 #endif
 		if (!bLayerAdded)
diff --git a/Engine/Plugins/Runtime/nDisplay/Source/DisplayCluster/Private/Game/EngineClasses/Scene/DisplayClusterLightCardActor.cpp b/Engine/Plugins/Runtime/nDisplay/Source/DisplayCluster/Private/Game/EngineClasses/Scene/DisplayClusterLightCardActor.cpp
--- a/Engine/Plugins/Runtime/nDisplay/Source/DisplayCluster/Private/Game/EngineClasses/Scene/DisplayClusterLightCardActor.cpp
+++ b/Engine/Plugins/Runtime/nDisplay/Source/DisplayCluster/Private/Game/EngineClasses/Scene/DisplayClusterLightCardActor.cpp
@@ -5,6 +5,7 @@
 #include "DisplayClusterConfigurationTypes.h"
 #include "DisplayClusterRootActor.h"
 #include "Features/IModularFeatures.h"
+#include "Misc/DisplayClusterLog.h"
 #include "IDisplayClusterLightCardActorExtender.h"
 
 #include "Components/DisplayClusterLabelComponent.h"
@@ -18,6 +19,7 @@
 
 #include "Editor.h"
 #include "Layers/LayersSubsystem.h"
+#include "DataLayer/DataLayerEditorSubsystem.h"
 
 #endif
 
@@ -136,6 +138,14 @@ void ADisplayClusterLightCardActor::Tick(float DeltaSeconds)
 {
 	Super::Tick(DeltaSeconds);
 
+	if (!MainSpringArmComponent || !LightCardTransformerComponent || !LightCardComponent)
+	{
+		// @todo Using WP with MU can trigger this in -game if an actor was added prior to converting to an MU session, and that actor was deleted.
+		// It looks like the actor isn't being destroyed properly on the clients.
+		UE_LOG(LogDisplayClusterGame, Error, TEXT("Light Card Actor is missing components."));
+		return;
+	}
+	
 	ClampLatitudeAndLongitude(Latitude, Longitude);
 
 	UpdateStageActorTransform();
@@ -541,41 +551,72 @@ void ADisplayClusterLightCardActor::AddToLightCardLayer(ADisplayClusterRootActor
 
 	bool bLayerAdded = false;
 #if WITH_EDITOR
-	ULayersSubsystem* LayersSubsystem = GEditor ? GEditor->GetEditorSubsystem<ULayersSubsystem>() : nullptr;
-	if (LayersSubsystem)
+	TArray<AActor*> ActorsInLayer;
+	
+	if (SupportsLayers())
 	{
-		// Adjust transparent sort order for UV light cards
-		if (IsUVActor())
+		if (ULayersSubsystem* LayersSubsystem = GEditor ? GEditor->GetEditorSubsystem<ULayersSubsystem>() : nullptr)
 		{
-			int32 HighestPriority = MIN_int32;
-			bool bExistingPrioritiesFound = false;
-			
-			const TArray<AActor*> ActorsInLayer = LayersSubsystem->GetActorsFromLayer(LightCardLayerName);
-			for (AActor* Actor : ActorsInLayer)
+			ActorsInLayer = LayersSubsystem->GetActorsFromLayer(LightCardLayerName);
+
+			LayersSubsystem->AddActorsToLayer(TArray<AActor*>{ this }, LightCardLayerName);
+			bLayerAdded = true;
+		}
+	}
+	else if (SupportsDataLayer() && GetLevel())
+	{
+		if (UDataLayerEditorSubsystem* DataLayersSubsystem = GEditor ? GEditor->GetEditorSubsystem<UDataLayerEditorSubsystem>() : nullptr)
+		{
+			if (UDataLayerAsset* DataLayerAsset = InRootActor->GetOrCreateLightCardDataLayerAsset(LightCardLayerName))
 			{
-				if (const ADisplayClusterLightCardActor* LightCard = Cast<ADisplayClusterLightCardActor>(Actor))
+				UDataLayerInstance* DataLayerInstance = DataLayersSubsystem->GetDataLayerInstance(DataLayerAsset);
+				if (DataLayerInstance == nullptr)
 				{
-					if (LightCard->IsUVActor())
-					{
-						if (LightCard->LightCardComponent->TranslucencySortPriority > HighestPriority)
-						{
-							HighestPriority = LightCard->LightCardComponent->TranslucencySortPriority;
-						}
-						bExistingPrioritiesFound = true;
-					}
+					FDataLayerCreationParameters CreationParams;
+					CreationParams.WorldDataLayers = GetLevel()->GetWorldDataLayers();
+					CreationParams.DataLayerAsset = DataLayerAsset;
+					DataLayerInstance = DataLayersSubsystem->CreateDataLayerInstance(MoveTemp(CreationParams));
 				}
+
+				ActorsInLayer = DataLayersSubsystem->GetActorsFromDataLayer(DataLayerInstance);
+			
+				AddDataLayer(DataLayerInstance);
 			}
+			else
+			{
+				UE_LOG(LogDisplayClusterGame, Error, TEXT("Light Card Actor could not find or create the data layer asset for layer '%s'."),  *LightCardLayerName.ToString());
+			}
+		}
+	}
+	
+	// Adjust transparent sort order for UV light cards
+	if (IsUVActor() && ActorsInLayer.Num() > 0)
+	{
+		int32 HighestPriority = MIN_int32;
+		bool bExistingPrioritiesFound = false;
 
-			if (bExistingPrioritiesFound)
+		for (AActor* Actor : ActorsInLayer)
+		{
+			if (const ADisplayClusterLightCardActor* LightCard = Cast<ADisplayClusterLightCardActor>(Actor))
 			{
-				// Newly added cards should be on the top-most visibility layer
-				LightCardComponent->TranslucencySortPriority = HighestPriority < MAX_int32 ? HighestPriority + 1 : HighestPriority;
+				if (LightCard->IsUVActor())
+				{
+					if (LightCard->LightCardComponent->TranslucencySortPriority > HighestPriority)
+					{
+						HighestPriority = LightCard->LightCardComponent->TranslucencySortPriority;
+					}
+					bExistingPrioritiesFound = true;
+				}
 			}
 		}
 
-		LayersSubsystem->AddActorsToLayer(TArray<AActor*>{ this }, LightCardLayerName);
-		bLayerAdded = true;
+		if (bExistingPrioritiesFound)
+		{
+			// Newly added cards should be on the top-most visibility layer
+			LightCardComponent->TranslucencySortPriority = HighestPriority < MAX_int32 ? HighestPriority + 1 : HighestPriority;
+		}
 	}
+	
 #endif
 	if (!bLayerAdded)
 	{
diff --git a/Engine/Plugins/Runtime/nDisplay/Source/DisplayCluster/Private/Game/EngineClasses/Scene/DisplayClusterRootActor.cpp b/Engine/Plugins/Runtime/nDisplay/Source/DisplayCluster/Private/Game/EngineClasses/Scene/DisplayClusterRootActor.cpp
--- a/Engine/Plugins/Runtime/nDisplay/Source/DisplayCluster/Private/Game/EngineClasses/Scene/DisplayClusterRootActor.cpp
+++ b/Engine/Plugins/Runtime/nDisplay/Source/DisplayCluster/Private/Game/EngineClasses/Scene/DisplayClusterRootActor.cpp
@@ -13,6 +13,7 @@
 #include "Components/DisplayClusterICVFXCameraComponent.h"
 #include "Components/DisplayClusterSceneComponentSyncThis.h"
 #include "CineCameraComponent.h"
+#include "DisplayClusterChromakeyCardActor.h"
 #include "ProceduralMeshComponent.h"
 
 #include "DisplayClusterLightCardActor.h"
@@ -45,12 +46,16 @@
 #include "Algo/MaxElement.h"
 #include "Components/DisplayClusterStageGeometryComponent.h"
 #include "UObject/Package.h"
+#include "WorldPartition/DataLayer/DataLayerAsset.h"
 
 
 #if WITH_EDITOR
 #include "IConcertSyncClientModule.h"
 #include "IConcertClientWorkspace.h"
 #include "IConcertSyncClient.h"
+
+#include "AssetToolsModule.h"
+#include "DataLayer/DataLayerFactory.h"
 #endif
 
 
@@ -134,6 +139,53 @@ const FDisplayClusterConfigurationRenderFrame& ADisplayClusterRootActor::GetRend
 	return CurrentConfigData->RenderFrameSettings;
 }
 
+UDataLayerAsset* ADisplayClusterRootActor::GetOrCreateLightCardDataLayerAsset(const FName& InName)
+{
+#if WITH_EDITOR
+	if (LightCardDataLayerAsset == nullptr)
+	{
+		Modify();
+		LightCardDataLayerAsset = CreateDataLayerAsset(InName);
+	}
+#endif
+
+	return LightCardDataLayerAsset;
+}
+
+UDataLayerAsset* ADisplayClusterRootActor::GetOrCreateChromakeyCardDataLayerAsset(const FName& InName)
+{
+#if WITH_EDITOR
+	if (ChromakeyCardDataLayerAsset == nullptr)
+	{
+		Modify();
+		ChromakeyCardDataLayerAsset = CreateDataLayerAsset(InName);
+	}
+#endif
+
+	return ChromakeyCardDataLayerAsset;
+}
+
+#if WITH_EDITOR
+UDataLayerAsset* ADisplayClusterRootActor::CreateDataLayerAsset(const FName& InName)
+{
+	IAssetTools& AssetTools = FModuleManager::GetModuleChecked<FAssetToolsModule>("AssetTools").Get();
+
+	const FString PackageName = GetPackage()->GetName();
+	FString AssetName;
+	FString TempPackageName;
+
+	UDataLayerFactory* Factory = NewObject<UDataLayerFactory>();
+
+	AssetTools.CreateUniqueAssetName(PackageName / InName.ToString(), FString(),
+		TempPackageName, AssetName);
+
+	UClass* AssetClass = Factory->GetSupportedClass();
+	
+	return Cast<UDataLayerAsset>(AssetTools.CreateAsset(AssetName,
+		PackageName, AssetClass, Factory));
+}
+#endif
+
 void ADisplayClusterRootActor::InitializeFromConfig(UDisplayClusterConfigurationData* ConfigData)
 {
 	if (ConfigData)
@@ -758,6 +810,46 @@ void ADisplayClusterRootActor::UpdateLightCardPositions()
 	}
 }
 
+void ADisplayClusterRootActor::SetChromakeyCardsOwner()
+{
+	TArray<UDisplayClusterICVFXCameraComponent*> ICVFXComponents;
+	GetComponents(ICVFXComponents);
+
+	for (UDisplayClusterICVFXCameraComponent* Camera : ICVFXComponents)
+	{
+		FDisplayClusterConfigurationICVFX_VisibilityList& ChromakeyCards = Camera->CameraSettings.Chromakey.ChromakeyRenderTexture.ShowOnlyList;
+
+		for (const TSoftObjectPtr<AActor>& Actor : ChromakeyCards.Actors)
+		{
+			if (ADisplayClusterChromakeyCardActor* ChromakeyCardActor = Actor.IsValid() ? Cast<ADisplayClusterChromakeyCardActor>(Actor.Get()) : nullptr)
+			{
+				// Set the owner so the light card actor can look the root actor up in certain situations, like adjusting labels
+				// when running as -game
+				ChromakeyCardActor->SetRootActorOwner(this);
+			}
+		}
+
+		if (ChromakeyCards.ActorLayers.Num())
+		{
+			if (const UWorld* World = GetWorld())
+			{
+				for (TActorIterator<ADisplayClusterChromakeyCardActor> Iter(World); Iter; ++Iter)
+				{
+					ADisplayClusterChromakeyCardActor* ChromakeyCardActor = *Iter;
+					for (const FActorLayer& ActorLayer : ChromakeyCards.ActorLayers)
+					{
+						if (ChromakeyCardActor->Layers.Contains(ActorLayer.Name))
+						{
+							ChromakeyCardActor->SetRootActorOwner(this);
+							break;
+						}
+					}
+				}
+			}
+		}
+	}
+}
+
 bool ADisplayClusterRootActor::GetFlushPositionAndNormal(const FVector& WorldPosition, FVector& OutPosition, FVector& OutNormal)
 {
 	const FVector StagePosition = GetCommonViewPoint()->GetComponentLocation();
diff --git a/Engine/Plugins/Runtime/nDisplay/Source/DisplayCluster/Private/Game/EngineClasses/Scene/DisplayClusterRootActorEditor.cpp b/Engine/Plugins/Runtime/nDisplay/Source/DisplayCluster/Private/Game/EngineClasses/Scene/DisplayClusterRootActorEditor.cpp
--- a/Engine/Plugins/Runtime/nDisplay/Source/DisplayCluster/Private/Game/EngineClasses/Scene/DisplayClusterRootActorEditor.cpp
+++ b/Engine/Plugins/Runtime/nDisplay/Source/DisplayCluster/Private/Game/EngineClasses/Scene/DisplayClusterRootActorEditor.cpp
@@ -153,6 +153,8 @@ void ADisplayClusterRootActor::Tick_Editor(float DeltaSeconds)
 	{
 		ResetPreviewInternals_Editor();
 	}
+	
+	SetChromakeyCardsOwner();
 }
 
 void ADisplayClusterRootActor::PostActorCreated_Editor()
diff --git a/Engine/Plugins/Runtime/nDisplay/Source/DisplayCluster/Public/DisplayClusterChromakeyCardActor.h b/Engine/Plugins/Runtime/nDisplay/Source/DisplayCluster/Public/DisplayClusterChromakeyCardActor.h
--- a/Engine/Plugins/Runtime/nDisplay/Source/DisplayCluster/Public/DisplayClusterChromakeyCardActor.h
+++ b/Engine/Plugins/Runtime/nDisplay/Source/DisplayCluster/Public/DisplayClusterChromakeyCardActor.h
@@ -6,6 +6,7 @@
 
 #include "DisplayClusterChromakeyCardActor.generated.h"
 
+class UDataLayerAsset;
 class UDisplayClusterICVFXCameraComponent;
 
 UCLASS(Blueprintable, NotPlaceable, DisplayName = "Chromakey Card", HideCategories = (Tick, Physics, Collision, Replication, Cooking, Input, Actor))
diff --git a/Engine/Plugins/Runtime/nDisplay/Source/DisplayCluster/Public/DisplayClusterLightCardActor.h b/Engine/Plugins/Runtime/nDisplay/Source/DisplayCluster/Public/DisplayClusterLightCardActor.h
--- a/Engine/Plugins/Runtime/nDisplay/Source/DisplayCluster/Public/DisplayClusterLightCardActor.h
+++ b/Engine/Plugins/Runtime/nDisplay/Source/DisplayCluster/Public/DisplayClusterLightCardActor.h
@@ -10,6 +10,7 @@
 
 class ADisplayClusterRootActor;
 class UActorComponent;
+class UDataLayerAsset;
 class UDisplayClusterLabelComponent;
 class UMeshComponent;
 class USceneComponent;
diff --git a/Engine/Plugins/Runtime/nDisplay/Source/DisplayCluster/Public/DisplayClusterRootActor.h b/Engine/Plugins/Runtime/nDisplay/Source/DisplayCluster/Public/DisplayClusterRootActor.h
--- a/Engine/Plugins/Runtime/nDisplay/Source/DisplayCluster/Public/DisplayClusterRootActor.h
+++ b/Engine/Plugins/Runtime/nDisplay/Source/DisplayCluster/Public/DisplayClusterRootActor.h
@@ -99,6 +99,13 @@ public:
 	const FDisplayClusterConfigurationICVFX_StageSettings& GetStageSettings() const;
 	const FDisplayClusterConfigurationRenderFrame& GetRenderFrameSettings() const;
 
+	UDataLayerAsset* GetOrCreateLightCardDataLayerAsset(const FName& InName);
+	UDataLayerAsset* GetOrCreateChromakeyCardDataLayerAsset(const FName& InName);
+
+private:
+#if WITH_EDITOR
+	UDataLayerAsset* CreateDataLayerAsset(const FName& InName);
+#endif
 protected:
 	//////////////////////////////////////////////////////////////////////////////////////////////
 	// AActor
@@ -121,6 +128,9 @@ protected:
 	/** Updates the world position and rotation of each light card referenced by this root actor's light card list to match the default view origin */
 	void UpdateLightCardPositions();
 
+	/** Sets self as the owner for any chromakey cards */
+	void SetChromakeyCardsOwner();
+
 public:
 	UFUNCTION(BlueprintCallable, Category = "NDisplay")
 	bool GetFlushPositionAndNormal(const FVector& WorldPosition, FVector& OutPosition, FVector& OutNormal);
@@ -281,6 +291,13 @@ private:
 	UPROPERTY()
 	TObjectPtr<UDisplayClusterStageGeometryComponent> StageGeometryComponent;
 
+	/** Data layer used for light card layers */
+	UPROPERTY()
+	TObjectPtr<UDataLayerAsset> LightCardDataLayerAsset;
+
+	/** Data layer used for chromakey card layers */
+	UPROPERTY()
+	TObjectPtr<UDataLayerAsset> ChromakeyCardDataLayerAsset;
 private:
 	// Current operation mode
 	EDisplayClusterOperationMode OperationMode;
diff --git a/Engine/Plugins/Runtime/nDisplay/Source/DisplayClusterLightCardEditor/Private/DisplayClusterLightCardEditor.cpp b/Engine/Plugins/Runtime/nDisplay/Source/DisplayClusterLightCardEditor/Private/DisplayClusterLightCardEditor.cpp
--- a/Engine/Plugins/Runtime/nDisplay/Source/DisplayClusterLightCardEditor/Private/DisplayClusterLightCardEditor.cpp
+++ b/Engine/Plugins/Runtime/nDisplay/Source/DisplayClusterLightCardEditor/Private/DisplayClusterLightCardEditor.cpp
@@ -6,7 +6,6 @@
 #include "DisplayClusterLightCardEditorCommands.h"
 #include "DisplayClusterLightCardEditorUtils.h"
 #include "DisplayClusterLightCardEditorStyle.h"
-#include "Engine/Blueprint.h"
 #include "IDisplayClusterLightCardEditor.h"
 #include "LightCardTemplates/DisplayClusterLightCardTemplate.h"
 #include "LightCardTemplates/DisplayClusterLightCardTemplateHelpers.h"
@@ -39,11 +38,13 @@
 #include "AssetRegistry/AssetRegistryModule.h"
 #include "Editor/Transactor.h"
 #include "Editor/UnrealEdEngine.h"
+#include "Engine/Blueprint.h"
 #include "Framework/Commands/GenericCommands.h"
 #include "Framework/Docking/LayoutExtender.h"
 #include "Framework/Docking/TabManager.h"
 #include "Framework/MultiBox/MultiBoxBuilder.h"
 #include "Framework/MultiBox/MultiBoxExtender.h"
+#include "Framework/Notifications/NotificationManager.h"
 #include "Layers/LayersSubsystem.h"
 #include "Misc/FileHelper.h"
 #include "Misc/ITransaction.h"
@@ -53,6 +54,7 @@
 #include "Widgets/Docking/SDockTab.h"
 #include "Widgets/Input/SComboBox.h"
 #include "Widgets/Input/SNumericEntryBox.h"
+#include "Widgets/Notifications/SNotificationList.h"
 #include "Widgets/Workflow/SWizard.h"
 
 #define LOCTEXT_NAMESPACE "FDisplayClusterLightCardEditor"
@@ -212,6 +214,27 @@ AActor* FDisplayClusterLightCardEditor::SpawnActor(TSubclassOf<AActor> InActorCl
 	{
 		return nullptr;
 	}
+
+	const UClass* ClassToUse = InTemplate ? InTemplate->GetClass() : InActorClass.Get(); 
+
+	// Partitioned worlds that are not saved will fail to create the data layer asset. Detect this and display
+	// an error message for the user.
+	if (ActiveRootActor->GetWorld()->IsPartitionedWorld() &&
+		(ClassToUse->IsChildOf(ADisplayClusterLightCardActor::StaticClass()) ||
+		ClassToUse->IsChildOf(UDisplayClusterLightCardTemplate::StaticClass())))
+	{
+		const FString PathName = ActiveRootActor->GetWorld()->GetPathName();
+		if (PathName.StartsWith(TEXT("/Temp")))
+		{
+			FNotificationInfo NotificationInfo(LOCTEXT("CannotSpawnActorInUnsavedWorldPartition",
+	"Cannot spawn an actor of this type in a world partition map that is not saved. Please save the map."));
+			NotificationInfo.bUseThrobber = false;
+			NotificationInfo.ExpireDuration = 4.0f;
+			FSlateNotificationManager::Get().AddNotification(NotificationInfo)->SetCompletionState(SNotificationItem::CS_Fail);
+
+			return nullptr;
+		}
+	}
 	
 	FScopedTransaction Transaction(LOCTEXT("SpawnActorTransactionMessage", "Spawn Actor"));
 
@@ -219,7 +242,6 @@ AActor* FDisplayClusterLightCardEditor::SpawnActor(TSubclassOf<AActor> InActorCl
 	
 	FDisplayClusterLightCardEditorHelper::FSpawnActorArgs SpawnArgs;
 	{
-		const UClass* ClassToUse = InTemplate ? InTemplate->GetClass() : InActorClass.Get();
 		FString ActorNameStr = InActorName.IsNone() ? ClassToUse->GetDisplayNameText().ToString() : InActorName.ToString();
 		ActorNameStr.RemoveSpacesInline();
 		
@@ -452,34 +474,15 @@ void FDisplayClusterLightCardEditor::AddLightCardsToActor(const TArray<ADisplayC
 {
 	if (ActiveRootActor.IsValid())
 	{
-		UDisplayClusterConfigurationData* ConfigData = ActiveRootActor->GetConfigData();
-		ConfigData->Modify();
-		FDisplayClusterConfigurationICVFX_VisibilityList& RootActorLightCards = ConfigData->StageSettings.Lightcard.ShowOnlyList;
-
+		FDisplayClusterLightCardEditorHelper::FAddLightCardArgs AddLightCardArgs;
+		AddLightCardArgs.bShowLabels = ShouldShowLightCardLabels();
+		AddLightCardArgs.LabelScale = *GetLightCardLabelScale();
+		
+		FDisplayClusterLightCardEditorHelper::AddLightCardsToRootActor(LightCards, ActiveRootActor.Get(), MoveTemp(AddLightCardArgs));
+		
 		for (ADisplayClusterLightCardActor* LightCard : LightCards)
 		{
 			check(LightCard);
-
-			if (!RootActorLightCards.Actors.ContainsByPredicate([&](const TSoftObjectPtr<AActor>& Actor)
-				{
-					// Don't add if a loaded actor is already present.
-					return Actor.Get() == LightCard;
-				}))
-			{
-				LightCard->ShowLightCardLabel(ShouldShowLightCardLabels(), *GetLightCardLabelScale(), ActiveRootActor.Get());
-				
-				const TSoftObjectPtr<AActor> LightCardSoftObject(LightCard);
-
-				// Remove any exact paths to this actor. It's possible invalid actors are present if a light card
-				// was force deleted from a level.
-				RootActorLightCards.Actors.RemoveAll([&](const TSoftObjectPtr<AActor>& Actor)
-					{
-						return Actor == LightCardSoftObject;
-					});
-
-				LightCard->AddToLightCardLayer(ActiveRootActor.Get());
-			}
-
 			RefreshPreviewStageActor(LightCard);
 		}
 	}
diff --git a/Engine/Plugins/Runtime/nDisplay/Source/DisplayClusterLightCardEditor/Private/Viewport/DisplayClusterLightCardEditorViewportClient.cpp b/Engine/Plugins/Runtime/nDisplay/Source/DisplayClusterLightCardEditor/Private/Viewport/DisplayClusterLightCardEditorViewportClient.cpp
--- a/Engine/Plugins/Runtime/nDisplay/Source/DisplayClusterLightCardEditor/Private/Viewport/DisplayClusterLightCardEditorViewportClient.cpp
+++ b/Engine/Plugins/Runtime/nDisplay/Source/DisplayClusterLightCardEditor/Private/Viewport/DisplayClusterLightCardEditorViewportClient.cpp
@@ -1559,7 +1559,8 @@ void FDisplayClusterLightCardEditorViewportClient::UpdatePreviewActor(ADisplayCl
 
 		const TWeakObjectPtr<ADisplayClusterRootActor> RootActorPtr (RootActor);
 		const TWeakPtr<FDisplayClusterLightCardEditorViewportClient> WeakPtrThis = SharedThis(this);
-		
+
+		TWeakObjectPtr<AActor> StageActorWeakPtr(StageActor);
 		// Schedule for the next tick so CDO changes get propagated first in the event of config editor skeleton
 		// regeneration & compiles. nDisplay's custom propagation may have issues if the archetype isn't correct.
 		PreviewWorld->GetTimerManager().SetTimerForNextTick([=]()
@@ -1571,10 +1572,10 @@ void FDisplayClusterLightCardEditorViewportClient::UpdatePreviewActor(ADisplayCl
 				return;
 			}
 
-			if (StageActor)
+			if (StageActorWeakPtr.IsValid())
 			{
-				ActorsRefreshing.Remove(StageActor);
-				DestroyProxy(StageActor);
+				ActorsRefreshing.Remove(StageActorWeakPtr.Get());
+				DestroyProxy(StageActorWeakPtr.Get());
 			}
 			else
 			{
@@ -1653,9 +1654,9 @@ void FDisplayClusterLightCardEditorViewportClient::UpdatePreviewActor(ADisplayCl
 
 				if (LightCardEditorPtr.IsValid())
 				{
-					if (StageActor)
+					if (StageActorWeakPtr.IsValid())
 					{
-						if (AActor* ActorProxy = CreateStageActorProxy(StageActor))
+						if (AActor* ActorProxy = CreateStageActorProxy(StageActorWeakPtr.Get()))
 						{
 							ActorProxiesCreated.Add(ActorProxy);
 						}
diff --git a/Engine/Plugins/Runtime/nDisplay/Source/DisplayClusterMultiUser/DisplayClusterMultiUser.Build.cs b/Engine/Plugins/Runtime/nDisplay/Source/DisplayClusterMultiUser/DisplayClusterMultiUser.Build.cs
--- a/Engine/Plugins/Runtime/nDisplay/Source/DisplayClusterMultiUser/DisplayClusterMultiUser.Build.cs
+++ b/Engine/Plugins/Runtime/nDisplay/Source/DisplayClusterMultiUser/DisplayClusterMultiUser.Build.cs
@@ -14,6 +14,7 @@ public class DisplayClusterMultiUser : ModuleRules
 				"CoreUObject",
 				"ConcertSyncClient",
 				"DisplayClusterConfiguration",
+				"Engine"
 			});
 	}
 }
diff --git a/Engine/Plugins/Runtime/nDisplay/Source/DisplayClusterMultiUser/Private/DisplayClusterMultiUserManager.cpp b/Engine/Plugins/Runtime/nDisplay/Source/DisplayClusterMultiUser/Private/DisplayClusterMultiUserManager.cpp
--- a/Engine/Plugins/Runtime/nDisplay/Source/DisplayClusterMultiUser/Private/DisplayClusterMultiUserManager.cpp
+++ b/Engine/Plugins/Runtime/nDisplay/Source/DisplayClusterMultiUser/Private/DisplayClusterMultiUserManager.cpp
@@ -11,6 +11,8 @@
 #include "IConcertSyncClient.h"
 #include "IConcertSyncClientModule.h"
 
+#include "WorldPartition/DataLayer/DataLayerInstance.h"
+
 #define NDISPLAY_MULTIUSER_TRANSACTION_FILTER TEXT("DisplayClusterMultiUser")
 
 FDisplayClusterMultiUserManager::FDisplayClusterMultiUserManager()
@@ -55,12 +57,18 @@ void FDisplayClusterMultiUserManager::OnApplyRemoteTransaction(ETransactionNotif
 
 ETransactionFilterResult FDisplayClusterMultiUserManager::ShouldObjectBeTransacted(UObject* InObject, UPackage* InPackage)
 {
-	if (InObject && ((InObject->IsA<UDisplayClusterConfigurationData_Base>() &&
-		!InObject->IsTemplate() && !InObject->HasAnyFlags(RF_Transient) && InPackage != GetTransientPackage()) ||
-		(InObject->GetClass()->HasMetaData(TEXT("DisplayClusterMultiUserInclude")))))
+	if (InObject)
 	{
-		UE_LOG(LogDisplayClusterMultiUser, Log, TEXT("FDisplayClusterMultiUser transaction for object: %s"), *InObject->GetName());
-		return ETransactionFilterResult::IncludeObject;
+		const bool bIsValidObjectType = InObject->IsA<UDisplayClusterConfigurationData_Base>()
+		|| InObject->IsA<UDataLayerInstance>();
+
+		if (((bIsValidObjectType && !InObject->IsTemplate() && !InObject->HasAnyFlags(RF_Transient)
+			&& InPackage != GetTransientPackage())
+			|| InObject->GetClass()->HasMetaData(TEXT("DisplayClusterMultiUserInclude"))))
+		{
+			UE_LOG(LogDisplayClusterMultiUser, Log, TEXT("FDisplayClusterMultiUser transaction for object: %s"), *InObject->GetName());
+			return ETransactionFilterResult::IncludeObject;
+		}
 	}
 
 	return ETransactionFilterResult::UseDefault;
diff --git a/Engine/Plugins/Runtime/nDisplay/Source/DisplayClusterScenePreview/Private/DisplayClusterLightCardEditorHelper.cpp b/Engine/Plugins/Runtime/nDisplay/Source/DisplayClusterScenePreview/Private/DisplayClusterLightCardEditorHelper.cpp
--- a/Engine/Plugins/Runtime/nDisplay/Source/DisplayClusterScenePreview/Private/DisplayClusterLightCardEditorHelper.cpp
+++ b/Engine/Plugins/Runtime/nDisplay/Source/DisplayClusterScenePreview/Private/DisplayClusterLightCardEditorHelper.cpp
@@ -2,7 +2,6 @@
 
 #include "DisplayClusterLightCardEditorHelper.h"
 
-#include "Engine/Blueprint.h"
 #include "IDisplayClusterScenePreview.h"
 
 #include "DisplayClusterConfigurationTypes.h"
@@ -20,6 +19,7 @@
 #include "ProceduralMeshComponent.h"
 #include "Components/DisplayClusterStageGeometryComponent.h"
 #include "Containers/ArrayView.h"
+#include "Engine/Blueprint.h"
 #include "Engine/StaticMesh.h"
 #include "Engine/Texture2D.h"
 #include "UObject/Package.h"
@@ -438,21 +438,33 @@ void FDisplayClusterLightCardEditorHelper::VerifyAndFixActorOrigin(const FDispla
 	// Set location and rotation to match the root actor
 	const FVector& NewActorLocation = OriginComponent ? OriginComponent->GetComponentLocation() : FVector::ZeroVector;
 	const FRotator& NewActorRotation = OwningRootActor ? OwningRootActor->GetActorRotation() : FRotator::ZeroRotator;
-	
+
+	const FTransform OldOrigin = Actor->GetOrigin();
 	Actor->SetOrigin({ NewActorRotation, NewActorLocation, FVector::One() });
+
+	const FTransform NewOrigin = Actor->GetOrigin();
+	bool bActorMoved = !OldOrigin.Equals(NewOrigin, 0.f);
 	
 	if (Actor.AsActorChecked()->IsA<ADisplayClusterLightCardActor>())
 	{
+		const FDisplayClusterPositionalParams OldPositionalParams = Actor->GetPositionalParams();
+	
 		// Update the light card spherical coordinates to match its current world coordinates
 		const FVector LightCardEndEffectorLocation = Actor->GetStageActorTransform(true).GetLocation();
 		const FVector ActorRelativeLocation = NewActorRotation.UnrotateVector(LightCardEndEffectorLocation);
 
 		const FSphericalCoordinates SphericalCoords(ActorRelativeLocation);
 		SetActorCoordinates(Actor, SphericalCoords);
+
+		const FDisplayClusterPositionalParams NewPositionalParams = Actor->GetPositionalParams();
+		bActorMoved = bActorMoved || NewPositionalParams != OldPositionalParams;
 	}
 
 #if WITH_EDITOR
-	PostEditChangePropertiesForMovedActor(Actor);
+	if (bActorMoved)
+	{
+		PostEditChangePropertiesForMovedActor(Actor);
+	}
 #endif
 }
 
@@ -847,7 +859,7 @@ AActor* FDisplayClusterLightCardEditorHelper::SpawnStageActor(const FSpawnActorA
 #if WITH_EDITOR
 	if (!bIsPreview)
 	{
-		NewActor->SetActorLabel(NewActor->GetName());
+		FActorLabelUtilities::SetActorLabelUnique(NewActor, ActorName.ToString());
 	}
 #endif
 	
diff --git a/Engine/Plugins/Runtime/nDisplayModularFeatures/Source/DisplayClusterLightCardExtender/Public/StageActor/DisplayClusterPositionalParams.h b/Engine/Plugins/Runtime/nDisplayModularFeatures/Source/DisplayClusterLightCardExtender/Public/StageActor/DisplayClusterPositionalParams.h
--- a/Engine/Plugins/Runtime/nDisplayModularFeatures/Source/DisplayClusterLightCardExtender/Public/StageActor/DisplayClusterPositionalParams.h
+++ b/Engine/Plugins/Runtime/nDisplayModularFeatures/Source/DisplayClusterLightCardExtender/Public/StageActor/DisplayClusterPositionalParams.h
@@ -39,4 +39,21 @@ struct FDisplayClusterPositionalParams
 
 	UPROPERTY(EditAnywhere, BlueprintReadWrite, Category = Orientation)
 	FVector2D Scale = FVector2D(1.0f);
+	
+	FORCEINLINE friend bool operator==(const FDisplayClusterPositionalParams& A, const FDisplayClusterPositionalParams& B)
+	{
+		return A.DistanceFromCenter == B.DistanceFromCenter &&
+			A.Longitude == B.Longitude &&
+			A.Latitude == B.Latitude &&
+			A.Spin == B.Spin &&
+			A.Pitch == B.Pitch &&
+			A.Yaw == B.Yaw &&
+			A.RadialOffset == B.RadialOffset &&
+			A.Scale == B.Scale;
+	}
+
+	FORCEINLINE friend bool operator!=(const FDisplayClusterPositionalParams& A, const FDisplayClusterPositionalParams& B)
+	{
+		return !(A == B);
+	}
 };
