nDisplay: removed "Preview in PIE" checkbox to simplify UX

diff --git a/Engine/Plugins/Runtime/nDisplay/Source/DisplayCluster/Private/Game/EngineClasses/Scene/DisplayClusterRootActorEditor.cpp b/Engine/Plugins/Runtime/nDisplay/Source/DisplayCluster/Private/Game/EngineClasses/Scene/DisplayClusterRootActorEditor.cpp
--- a/Engine/Plugins/Runtime/nDisplay/Source/DisplayCluster/Private/Game/EngineClasses/Scene/DisplayClusterRootActorEditor.cpp
+++ b/Engine/Plugins/Runtime/nDisplay/Source/DisplayCluster/Private/Game/EngineClasses/Scene/DisplayClusterRootActorEditor.cpp
@@ -212,7 +212,7 @@ bool ADisplayClusterRootActor::IsPreviewEnabled() const
 	if (IsRunningGameOrPIE())
 	{
 		// Only PIE is currently supported
-		return bPIEPreviewEnable && GIsPlayInEditorWorld && (PreviewNodeId == DisplayClusterConfigurationStrings::gui::preview::PreviewNodeNone);
+		return bPreviewEnable && GIsPlayInEditorWorld && (PreviewNodeId == DisplayClusterConfigurationStrings::gui::preview::PreviewNodeNone);
 	}
 
 	// Editor case
diff --git a/Engine/Plugins/Runtime/nDisplay/Source/DisplayCluster/Public/DisplayClusterRootActor.h b/Engine/Plugins/Runtime/nDisplay/Source/DisplayCluster/Public/DisplayClusterRootActor.h
--- a/Engine/Plugins/Runtime/nDisplay/Source/DisplayCluster/Public/DisplayClusterRootActor.h
+++ b/Engine/Plugins/Runtime/nDisplay/Source/DisplayCluster/Public/DisplayClusterRootActor.h
@@ -347,10 +347,6 @@ public:
 	/** Render the scene and display it as a preview on the nDisplay root actor in the editor.  This will impact editor performance. */
 	UPROPERTY(EditAnywhere, BlueprintReadWrite, Category = "Editor Preview", meta = (DisplayName = "Enable Editor Preview"))
 	bool bPreviewEnable = true;
-
-	/** Render the scene and display it in the PIE as a preview of the nDisplay root actor. This will impact PIE performance. */
-	UPROPERTY(EditAnywhere, BlueprintReadWrite, Category = "Editor Preview", meta = (DisplayName = "Enabled in PIE"))
-	bool bPIEPreviewEnable = true;
 	
 	/** Adjust resolution scaling for the editor preview. */
 	UPROPERTY(EditAnywhere, BlueprintReadWrite, Category = "Editor Preview", meta = (DisplayName = "Preview Screen Percentage", ClampMin = "0.05", UIMin = "0.05", ClampMax = "1", UIMax = "1", EditCondition = "bPreviewEnable"))
