nDisplay: Fix loading a light card template not using the correct name. Minor regression from actor label changes in 23820376.
Saving a template will remove spaces from the initial name by default.

diff --git a/Engine/Plugins/Runtime/nDisplay/Source/DisplayClusterLightCardEditor/Private/DisplayClusterLightCardEditor.cpp b/Engine/Plugins/Runtime/nDisplay/Source/DisplayClusterLightCardEditor/Private/DisplayClusterLightCardEditor.cpp
--- a/Engine/Plugins/Runtime/nDisplay/Source/DisplayClusterLightCardEditor/Private/DisplayClusterLightCardEditor.cpp
+++ b/Engine/Plugins/Runtime/nDisplay/Source/DisplayClusterLightCardEditor/Private/DisplayClusterLightCardEditor.cpp
@@ -842,6 +842,7 @@ void FDisplayClusterLightCardEditor::CreateLightCardTemplate()
 	}
 	
 	FString DefaultAssetName = LightCardActor->GetActorLabel();
+	DefaultAssetName.RemoveSpacesInline();
 	if (!DefaultAssetName.EndsWith(TEXT("Template")))
 	{
 		DefaultAssetName += TEXT("Template");
diff --git a/Engine/Plugins/Runtime/nDisplay/Source/DisplayClusterScenePreview/Private/DisplayClusterLightCardEditorHelper.cpp b/Engine/Plugins/Runtime/nDisplay/Source/DisplayClusterScenePreview/Private/DisplayClusterLightCardEditorHelper.cpp
--- a/Engine/Plugins/Runtime/nDisplay/Source/DisplayClusterScenePreview/Private/DisplayClusterLightCardEditorHelper.cpp
+++ b/Engine/Plugins/Runtime/nDisplay/Source/DisplayClusterScenePreview/Private/DisplayClusterLightCardEditorHelper.cpp
@@ -804,7 +804,7 @@ AActor* FDisplayClusterLightCardEditorHelper::SpawnStageActor(const FSpawnActorA
 	check(RootActor)
 	check(ActorClass || Template);
 	
-	const FName ActorName = InSpawnArgs.ActorName;
+	FName ActorName = InSpawnArgs.ActorName;
 	const EDisplayClusterMeshProjectionType ProjectionMode = InSpawnArgs.ProjectionMode;
 	ULevel* Level = InSpawnArgs.Level ? InSpawnArgs.Level : RootActor->GetWorld()->GetCurrentLevel();
 	const bool bIsPreview = InSpawnArgs.bIsPreview;
@@ -829,6 +829,7 @@ AActor* FDisplayClusterLightCardEditorHelper::SpawnStageActor(const FSpawnActorA
 		// specifically the static mesh override textures. They will be parented to the template, not the level instance
 		// and prevent the map from saving.
 		ADisplayClusterLightCardActor* NewLightCard = CastChecked<ADisplayClusterLightCardActor>(StaticDuplicateObject(TemplateActor, Level, UniqueName));
+		ActorName = UniqueName;
 
 #if WITH_EDITOR
 		Level->AddLoadedActor(NewLightCard);
