nDisplay[SharedMemoryMedia]: Fixed texture name conflict in media samples.

This was causing the samples of one player to end up on the other player of the same type.

diff --git a/Engine/Plugins/Runtime/nDisplay/Source/DisplayClusterMedia/Private/Input/SharedMemoryMediaPlayer.cpp b/Engine/Plugins/Runtime/nDisplay/Source/DisplayClusterMedia/Private/Input/SharedMemoryMediaPlayer.cpp
--- a/Engine/Plugins/Runtime/nDisplay/Source/DisplayClusterMedia/Private/Input/SharedMemoryMediaPlayer.cpp
+++ b/Engine/Plugins/Runtime/nDisplay/Source/DisplayClusterMedia/Private/Input/SharedMemoryMediaPlayer.cpp
@@ -354,7 +354,7 @@ void FSharedMemoryMediaPlayer::TickFetch(FTimespan DeltaTime, FTimespan Timecode
 			SharedCrossGpuTextureDescriptions[0].Width,
 			SharedCrossGpuTextureDescriptions[0].Height,
 			SharedCrossGpuTextureDescriptions[0].Format,
-			FName(TEXT("SampleCommonTexture"))
+			MakeUniqueObjectName(nullptr, UTexture2D::StaticClass(), FName(*FString::Printf(TEXT("SampleCommonTexture_%s"), *UniqueName)))
 		));
 
 		SampleCommonTexture->UpdateResource();
@@ -648,9 +648,9 @@ void FSharedMemoryMediaPlayer::JustInTimeSampleRender()
 	// We only allow this function to run once per frame.
 	if (LastFrameNumberThatUpdatedJustInTime == GFrameCounterRenderThread)
 	{
-		UE_LOG(LogDisplayClusterMedia, Warning, 
-			TEXT("FSharedMemoryMediaPlayer::JustInTimeSampleRender called more than once in GFrameCounterRenderThread %llu"), 
-			GFrameCounterRenderThread
+		UE_LOG(LogDisplayClusterMedia, Verbose, 
+			TEXT("FSharedMemoryMediaPlayer '%s' JustInTimeSampleRender called more than once in GFrameCounterRenderThread %llu"),
+			*UniqueName, GFrameCounterRenderThread
 		);
 
 		return;
@@ -658,8 +658,9 @@ void FSharedMemoryMediaPlayer::JustInTimeSampleRender()
 
 	LastFrameNumberThatUpdatedJustInTime = GFrameCounterRenderThread;
 
-	UE_LOG(LogDisplayClusterMedia, VeryVerbose, TEXT("FSharedMemoryMediaPlayer::JustInTimeSampleRender called for GFrameCounterRenderThread %llu"), 
-		GFrameCounterRenderThread);
+	UE_LOG(LogDisplayClusterMedia, VeryVerbose, TEXT("FSharedMemoryMediaPlayer '%s' JustInTimeSampleRender called for GFrameCounterRenderThread %llu"), 
+		*UniqueName, GFrameCounterRenderThread
+	);
 
 	for (int32 MemIdx = 0; MemIdx < NUMSHAREDMEM; ++MemIdx)
 	{
