nDisplay[SharedMemoryMedia]: Ensure that only one player with a given UniqueName can play at a time.

Multiple of this kind of players with the same unique name playing at the same time is not supported.

Also fixed edit condition that referenced a property that was removed.

diff --git a/Engine/Plugins/Runtime/nDisplay/Source/DisplayClusterMedia/Private/Input/SharedMemoryMediaPlayer.cpp b/Engine/Plugins/Runtime/nDisplay/Source/DisplayClusterMedia/Private/Input/SharedMemoryMediaPlayer.cpp
--- a/Engine/Plugins/Runtime/nDisplay/Source/DisplayClusterMedia/Private/Input/SharedMemoryMediaPlayer.cpp
+++ b/Engine/Plugins/Runtime/nDisplay/Source/DisplayClusterMedia/Private/Input/SharedMemoryMediaPlayer.cpp
@@ -36,6 +36,8 @@ static FAutoConsoleVariableRef CVarDisplayClusterSharedMemoryLatency(
 );
 
 
+TSet<FString>FSharedMemoryMediaPlayer::UniqueNamesRegistered;
+
 FSharedMemoryMediaPlayer::FSharedMemoryMediaPlayer()
 	: Super()
 {
@@ -48,6 +50,20 @@ FSharedMemoryMediaPlayer::~FSharedMemoryMediaPlayer()
 	delete Samples;
 }
 
+bool FSharedMemoryMediaPlayer::RegisterUniqueName(FString& InUniqueName)
+{
+	bool bIsAlreadyInSet = false;
+
+	UniqueNamesRegistered.Add(InUniqueName, &bIsAlreadyInSet);
+
+	return !bIsAlreadyInSet;
+}
+
+void FSharedMemoryMediaPlayer::UnregisterUniqueName(FString& InUniqueName)
+{
+	UniqueNamesRegistered.Remove(InUniqueName);
+}
+
 bool FSharedMemoryMediaPlayer::Open(const FString& Url, const IMediaOptions* Options)
 {
 	// Get an RHI type specific implementation.
@@ -101,6 +117,12 @@ bool FSharedMemoryMediaPlayer::Open(const FString& Url, const IMediaOptions* Opt
 		}
 	}
 
+	if (!RegisterUniqueName(UniqueName))
+	{
+		UE_LOG(LogDisplayClusterMedia, Error, TEXT("Only one SharedMemoryMediaPlayer the UniqueName '%s' can play at a time"), *UniqueName);
+		return false;
+	}
+
 	return true;
 }
 
@@ -164,6 +186,9 @@ void FSharedMemoryMediaPlayer::Close()
 
 	// Reset our state variables
 	ModeState.Reset();
+
+	// Allow this or other players with the given UniqueName to play
+	UnregisterUniqueName(UniqueName);
 }
 
 FGuid FSharedMemoryMediaPlayer::GetPlayerPluginGUID() const
diff --git a/Engine/Plugins/Runtime/nDisplay/Source/DisplayClusterMedia/Private/Input/SharedMemoryMediaPlayer.h b/Engine/Plugins/Runtime/nDisplay/Source/DisplayClusterMedia/Private/Input/SharedMemoryMediaPlayer.h
--- a/Engine/Plugins/Runtime/nDisplay/Source/DisplayClusterMedia/Private/Input/SharedMemoryMediaPlayer.h
+++ b/Engine/Plugins/Runtime/nDisplay/Source/DisplayClusterMedia/Private/Input/SharedMemoryMediaPlayer.h
@@ -232,6 +232,19 @@ protected:
 	/** DetermineNextSourceFrame implementration for Freerun Mode */
 	bool DetermineNextSourceFrameFreerunMode(uint64 FrameNumber, uint32& OutExpectedFrameNumber, uint32& OutSharedMemoryIdx);
 
+protected:
+
+	/** Used to make sure that players with the same UniqueName are not allowed to play at the same time, which is not supported.
+	 *  Only access from game thread.
+	 */
+	static TSet<FString> UniqueNamesRegistered;
+
+	/** Used by the player when it opens. The player should not play if the return value is false */
+	static bool RegisterUniqueName(FString& InUniqueName);
+
+	/** Used by the player when it closes. */
+	static void UnregisterUniqueName(FString& InUniqueName);
+
 protected:
 
 	/** Number of resources used for pipelining the data flow */
diff --git a/Engine/Plugins/Runtime/nDisplay/Source/DisplayClusterMedia/Public/SharedMemoryMediaSource.h b/Engine/Plugins/Runtime/nDisplay/Source/DisplayClusterMedia/Public/SharedMemoryMediaSource.h
--- a/Engine/Plugins/Runtime/nDisplay/Source/DisplayClusterMedia/Public/SharedMemoryMediaSource.h
+++ b/Engine/Plugins/Runtime/nDisplay/Source/DisplayClusterMedia/Public/SharedMemoryMediaSource.h
@@ -45,7 +45,7 @@ public:
 	ESharedMemoryMediaSourceMode Mode;
 
 	/** Zero latency option to wait for the cross gpu texture rendered on the same frame. May adversely affect fps. Only applicable in Framelocked mode. */
-	UPROPERTY(EditAnywhere, BlueprintReadWrite, Category = "Input", meta = (EditCondition = "bUseFrameNumbers"))
+	UPROPERTY(EditAnywhere, BlueprintReadWrite, Category = "Input")
 	bool bZeroLatency = true;
 
 public:
