[EpicStageApp] Fix LCs created from stage app not appearing in ICVFX outliner

diff --git a/Engine/Plugins/Runtime/nDisplay/Source/DisplayClusterScenePreview/Private/DisplayClusterLightCardEditorHelper.cpp b/Engine/Plugins/Runtime/nDisplay/Source/DisplayClusterScenePreview/Private/DisplayClusterLightCardEditorHelper.cpp
--- a/Engine/Plugins/Runtime/nDisplay/Source/DisplayClusterScenePreview/Private/DisplayClusterLightCardEditorHelper.cpp
+++ b/Engine/Plugins/Runtime/nDisplay/Source/DisplayClusterScenePreview/Private/DisplayClusterLightCardEditorHelper.cpp
@@ -942,6 +942,11 @@ void FDisplayClusterLightCardEditorHelper::AddLightCardsToRootActor(
 
 			LightCard->AddToLightCardLayer(RootActor);
 		}
+
+#if WITH_EDITOR
+		// Fire this event so that e.g. ICVFX panel can update proxies based on the new root/layer
+		LightCard->PostEditChange();
+#endif
 	}
 }
 
