SharedMemoryMediaCapture - Enable AnyThread capture

diff --git a/Engine/Plugins/Runtime/nDisplay/Source/DisplayClusterMedia/Public/SharedMemoryMediaCapture.h b/Engine/Plugins/Runtime/nDisplay/Source/DisplayClusterMedia/Public/SharedMemoryMediaCapture.h
--- a/Engine/Plugins/Runtime/nDisplay/Source/DisplayClusterMedia/Public/SharedMemoryMediaCapture.h
+++ b/Engine/Plugins/Runtime/nDisplay/Source/DisplayClusterMedia/Public/SharedMemoryMediaCapture.h
@@ -34,6 +34,10 @@ protected:
 	virtual bool ShouldCaptureRHIResource() const override;
 	virtual bool InitializeCapture() override;
 	virtual void StopCaptureImpl(bool bAllowPendingFrameToBeProcess) override;
+	virtual bool SupportsAnyThreadCapture() const override
+	{
+		return true;
+	}
 
 	/** For custom conversion, methods that need to be overridden */
 	virtual FIntPoint GetCustomOutputSize(const FIntPoint& InSize) const override;
