nDisplay: ICVFX OCIO
DCRA OCIO Clenaup and UI for new Viewport/ICVFX/LC/LCUV/Chromakey OCIO implemention.
Added alpha channel preservation for LC+Chromakey.
UV LightCard OCIO

diff --git a/Engine/Plugins/Runtime/nDisplay/Shaders/Private/MPCDIOverlayShaders.usf b/Engine/Plugins/Runtime/nDisplay/Shaders/Private/MPCDIOverlayShaders.usf
--- a/Engine/Plugins/Runtime/nDisplay/Shaders/Private/MPCDIOverlayShaders.usf
+++ b/Engine/Plugins/Runtime/nDisplay/Shaders/Private/MPCDIOverlayShaders.usf
@@ -32,26 +32,33 @@ OutputPS Passthrough_PS(OutputVS IN)
 	return OUT;
 }
 
-float4 GetIcvfxOverlayColor(float2 OverlayUV)
+////////////////////////////////////////////////////////////////////////////////
+// LightCard
+////////////////////////////////////////////////////////////////////////////////
+float4 GetLightCardTextureColor(float2 InUV)
 {
-	float4 OverlayColor = OverlayTexture.Sample(OverlaySampler, OverlayUV).rgba;
+	return LightCardTexture.Sample(LightCardSampler, InUV);
+}
 
-#if OVERLAY_ALPHA
-	OverlayColor.a = OverlayAlphaTexture.Sample(OverlayAlphaSampler, OverlayUV).a;
-#endif
+float4 GetUVLightCardTextureColor(float2 InUV)
+{
+	float4 OutColor =UVLightCardTexture.Sample(UVLightCardSampler, InUV);
 
-	return OverlayColor;
+	// UVLightCard alpha not inverted
+	return float4(OutColor.xyz, 1 - OutColor.w);
 }
 
-float3 ComposeIcvfxOverlay(float3 OutColor, float2 OverlayUV)
+float3 ComposeOverlayColor(float4 InColor, float4 InOverlayColor)
 {
+	return InOverlayColor.xyz + InColor.xyz * InOverlayColor.w;
+}
 
-	float4 OverlayColor = GetIcvfxOverlayColor(OverlayUV);
-	float3 FinalColor = OverlayColor.xyz + OutColor * OverlayColor.w;
-
-	return FinalColor;
+float4 AddOverlayColor(float4 InColor, float4 InOverlayColor)
+{
+	return float4(ComposeOverlayColor(InColor, InOverlayColor), max(InColor.w, InOverlayColor.w));
 }
 
+////////////////////////////////////////////////////////////////////////////////
 float2 TileUV(float2 InUV, float InScale)
 {
     return frac(InUV * InScale);
@@ -174,13 +181,6 @@ float4 GetInnerCameraColor(float4 WarpedUV, OutputVS IN)
 	return CameraColor;
 }
 
-float3 GetUVLightCard(float3 SourceColor, float2 UV)
-{
-	// LightCardColor will have a pre-multiplied alpha, so use a pre-multiplied alpha compositing formula
-	float4 LightCardColor = UVLightCardMapTexture.Sample(UVLightCardMapSampler, UV);
-	return SourceColor.rgb * (1 - LightCardColor.a) + LightCardColor.rgb;
-}
-
 OutputPS IcvfxWarpPS(OutputVS IN)
 {
 	OutputPS OUT;
@@ -211,12 +211,11 @@ OutputPS IcvfxWarpPS(OutputVS IN)
 		float4 OutColor = float4(InputTexture.Sample(InputSampler, ViewportUV).rgb, 1.0f);
 	#endif
 
-	#if OVERLAY_UNDER
-		#if UV_LIGHT_CARDS
-			OutColor.xyz = GetUVLightCard(OutColor.xyz, IN.UV_Chromakey.xy);
-		#endif
-	
-		OutColor.xyz = ComposeIcvfxOverlay(OutColor.xyz, OverlayUV);
+	#if LIGHTCARD_UNDER
+		OutColor.xyz = ComposeOverlayColor(OutColor, GetLightCardTextureColor(OverlayUV));
+	#endif
+	#if UVLIGHTCARD_UNDER
+		OutColor.xyz = ComposeOverlayColor(OutColor, GetUVLightCardTextureColor(IN.UV_Chromakey.xy));
 	#endif
 	
 	#if INNER_CAMERA
@@ -224,12 +223,11 @@ OutputPS IcvfxWarpPS(OutputVS IN)
 		OutColor.xyz = lerp(OutColor.xyz, CameraColor.xyz, CameraColor.w);
 	#endif
 
-	#if OVERLAY_OVER
-		#if UV_LIGHT_CARDS
-			OutColor.xyz = GetUVLightCard(OutColor.xyz, IN.UV_Chromakey.xy);
-		#endif
-
-		OutColor.xyz = ComposeIcvfxOverlay(OutColor.xyz, OverlayUV);
+	#if LIGHTCARD_OVER
+		OutColor.xyz = ComposeOverlayColor(OutColor, GetLightCardTextureColor(OverlayUV));
+	#endif
+	#if UVLIGHTCARD_OVER
+		OutColor.xyz = ComposeOverlayColor(OutColor, GetUVLightCardTextureColor(IN.UV_Chromakey.xy));
 	#endif
 
 #else 
@@ -237,11 +235,15 @@ OutputPS IcvfxWarpPS(OutputVS IN)
 	#if INNER_CAMERA
 		float4 OutColor = GetInnerCameraColor(WarpedUV, IN);
 	#else
-		#if OVERLAY_OVER
-			float4 OutColor = GetIcvfxOverlayColor(OverlayUV);
-		#else
-			float4 OutColor = float4(0, 0, 0, 0);
+		float4 OutColor = float4(0, 0, 0, 0);
+
+		#if LIGHTCARD_OVER
+			OutColor = AddOverlayColor(OutColor, GetLightCardTextureColor(OverlayUV));
+		#endif
+		#if UVLIGHTCARD_OVER
+			OutColor = AddOverlayColor(OutColor, GetUVLightCardTextureColor(IN.UV_Chromakey.xy));
 		#endif
+
 	#endif
 #endif
 
@@ -251,4 +253,3 @@ OutputPS IcvfxWarpPS(OutputVS IN)
 	OUT.Color = ApplyBlending(OutColor.xyz, IN, OutColor.w);
 	return OUT;
 }
-
diff --git a/Engine/Plugins/Runtime/nDisplay/Shaders/Private/MPCDIUtils.ush b/Engine/Plugins/Runtime/nDisplay/Shaders/Private/MPCDIUtils.ush
--- a/Engine/Plugins/Runtime/nDisplay/Shaders/Private/MPCDIUtils.ush
+++ b/Engine/Plugins/Runtime/nDisplay/Shaders/Private/MPCDIUtils.ush
@@ -14,23 +14,25 @@ Texture2D InputTexture;
 Texture2D WarpMapTexture;
 Texture2D AlphaMapTexture;
 Texture2D BetaMapTexture;
-Texture2D OverlayTexture;
-Texture2D OverlayAlphaTexture;
+
 Texture2D InnerCameraTexture;
 Texture2D ChromakeyCameraTexture;
 Texture2D ChromakeyMarkerTexture;
-Texture2D UVLightCardMapTexture;
+
+Texture2D LightCardTexture;
+Texture2D UVLightCardTexture;
 
 SamplerState InputSampler;
 SamplerState WarpMapSampler;
 SamplerState AlphaMapSampler;
 SamplerState BetaMapSampler;
-SamplerState OverlaySampler;
-SamplerState OverlayAlphaSampler;
+
 SamplerState InnerCameraSampler;
 SamplerState ChromakeyCameraSampler;
 SamplerState ChromakeyMarkerSampler;
-SamplerState UVLightCardMapSampler;
+
+SamplerState LightCardSampler;
+SamplerState UVLightCardSampler;
 
 float4x4 MeshToStageProjectionMatrix;
 float4x4 ViewportTextureProjectionMatrix;
diff --git a/Engine/Plugins/Runtime/nDisplay/Source/DisplayCluster/DisplayCluster.Build.cs b/Engine/Plugins/Runtime/nDisplay/Source/DisplayCluster/DisplayCluster.Build.cs
--- a/Engine/Plugins/Runtime/nDisplay/Source/DisplayCluster/DisplayCluster.Build.cs
+++ b/Engine/Plugins/Runtime/nDisplay/Source/DisplayCluster/DisplayCluster.Build.cs
@@ -9,10 +9,10 @@ public class DisplayCluster : ModuleRules
 	{
 		PublicDefinitions.Add("WITH_OCIO=0");
 
-		// [temporary] We need this to be able to use some private data types. This should
-		// be removed once we move the nD rendering pipeline to RDG.
-		string EngineDir = Path.GetFullPath(Target.RelativeEnginePath);
-		PrivateIncludePaths.Add(Path.Combine(EngineDir, "Source", "Runtime", "Renderer", "Private"));
+		PrivateIncludePaths.AddRange(
+			new string[] {
+				System.IO.Path.Combine(GetModuleDirectory("Renderer"), "Private"), //required for FPostProcessMaterialInputs
+			});
 
 		PublicDependencyModuleNames.AddRange(
 			new string[] {
diff --git a/Engine/Plugins/Runtime/nDisplay/Source/DisplayCluster/Private/Game/EngineClasses/Basics/DisplayClusterViewportClient.cpp b/Engine/Plugins/Runtime/nDisplay/Source/DisplayCluster/Private/Game/EngineClasses/Basics/DisplayClusterViewportClient.cpp
--- a/Engine/Plugins/Runtime/nDisplay/Source/DisplayCluster/Private/Game/EngineClasses/Basics/DisplayClusterViewportClient.cpp
+++ b/Engine/Plugins/Runtime/nDisplay/Source/DisplayCluster/Private/Game/EngineClasses/Basics/DisplayClusterViewportClient.cpp
@@ -394,14 +394,6 @@ void UDisplayClusterViewportClient::Draw(FViewport* InViewport, FCanvas* SceneCa
 	//Get world for render
 	UWorld* const MyWorld = GetWorld();
 
-	// Initialize new render frame resources
-	FDisplayClusterRenderFrame RenderFrame;
-	if (!DCRenderDevice->BeginNewFrame(InViewport, MyWorld, RenderFrame))
-	{
-		// skip rendering: Can't build render frame
-		return;
-	}
-
 	////////////////////////////////
 	// Otherwise we use our own version of the UGameViewportClient::Draw which is basically
 	// a simpler version of the original one but with multiple ViewFamilies support
@@ -464,6 +456,17 @@ void UDisplayClusterViewportClient::Draw(FViewport* InViewport, FCanvas* SceneCa
 	// Gather all view families first
 	TArray<FSceneViewFamilyContext*> ViewFamilies;
 
+	// Initialize new render frame resources
+	FDisplayClusterRenderFrame RenderFrame;
+	if (!DCRenderDevice->BeginNewFrame(InViewport, MyWorld, RenderFrame))
+	{
+		// skip rendering: Can't build render frame
+		return;
+	}
+
+	// Handle special viewports game-thread logic at frame begin
+	DCRenderDevice->InitializeNewFrame();
+
 	for (FDisplayClusterRenderFrame::FFrameRenderTarget& DCRenderTarget : RenderFrame.RenderTargets)
 	{
 		for (FDisplayClusterRenderFrame::FFrameViewFamily& DCViewFamily : DCRenderTarget.ViewFamilies)
diff --git a/Engine/Plugins/Runtime/nDisplay/Source/DisplayCluster/Private/Render/Device/DisplayClusterDeviceBase.cpp b/Engine/Plugins/Runtime/nDisplay/Source/DisplayCluster/Private/Render/Device/DisplayClusterDeviceBase.cpp
--- a/Engine/Plugins/Runtime/nDisplay/Source/DisplayCluster/Private/Render/Device/DisplayClusterDeviceBase.cpp
+++ b/Engine/Plugins/Runtime/nDisplay/Source/DisplayCluster/Private/Render/Device/DisplayClusterDeviceBase.cpp
@@ -170,10 +170,10 @@ EStereoscopicPass FDisplayClusterDeviceBase::GetViewPassForIndex(bool bStereoReq
 	{
 		if (IsInRenderingThread())
 		{
-			if (ViewportManagerProxyPtr)
+			if (ViewportManagerProxy.IsValid())
 			{
 				uint32 ViewportContextNum = 0;
-				IDisplayClusterViewportProxy* ViewportProxy = ViewportManagerProxyPtr->FindViewport_RenderThread(ViewIndex, &ViewportContextNum);
+				IDisplayClusterViewportProxy* ViewportProxy = ViewportManagerProxy->FindViewport_RenderThread(ViewIndex, &ViewportContextNum);
 				if (ViewportProxy)
 				{
 					const FDisplayClusterViewport_Context& Context = ViewportProxy->GetContexts_RenderThread()[ViewportContextNum];
@@ -183,10 +183,10 @@ EStereoscopicPass FDisplayClusterDeviceBase::GetViewPassForIndex(bool bStereoReq
 		}
 		else
 		{
-			if (ViewportManagerPtr)
+			if (ViewportManager)
 			{
 				uint32 ViewportContextNum = 0;
-				IDisplayClusterViewport* ViewportPtr = ViewportManagerPtr->FindViewport(ViewIndex, &ViewportContextNum);
+				IDisplayClusterViewport* ViewportPtr = ViewportManager->FindViewport(ViewIndex, &ViewportContextNum);
 				if (ViewportPtr)
 				{
 					const FDisplayClusterViewport_Context& Context = ViewportPtr->GetContexts()[ViewportContextNum];
@@ -203,7 +203,7 @@ void FDisplayClusterDeviceBase::AdjustViewRect(int32 ViewIndex, int32& X, int32&
 {
 	check(IsInGameThread());
 
-	if (ViewportManagerPtr == nullptr || ViewportManagerPtr->IsSceneOpened() == false)
+	if (ViewportManager == nullptr || ViewportManager->IsSceneOpened() == false)
 	{
 		return;
 	}
@@ -216,7 +216,7 @@ void FDisplayClusterDeviceBase::AdjustViewRect(int32 ViewIndex, int32& X, int32&
 	}
 
 	uint32 ViewportContextNum = 0;
-	IDisplayClusterViewport* ViewportPtr = ViewportManagerPtr->FindViewport(ViewIndex, &ViewportContextNum);
+	IDisplayClusterViewport* ViewportPtr = ViewportManager->FindViewport(ViewIndex, &ViewportContextNum);
 	if (ViewportPtr == nullptr)
 	{
 		UE_LOG(LogDisplayClusterRender, Warning, TEXT("Viewport StereoViewIndex='%i' not found"), ViewIndex);
@@ -239,7 +239,7 @@ void FDisplayClusterDeviceBase::CalculateStereoViewOffset(const int32 ViewIndex,
 	check(IsInGameThread());
 	check(WorldToMeters > 0.f);
 
-	if (ViewportManagerPtr == nullptr || ViewportManagerPtr->IsSceneOpened() == false)
+	if (ViewportManager == nullptr || ViewportManager->IsSceneOpened() == false)
 	{
 		return;
 	}
@@ -252,7 +252,7 @@ void FDisplayClusterDeviceBase::CalculateStereoViewOffset(const int32 ViewIndex,
 	}
 
 	uint32 ViewportContextNum = 0;
-	IDisplayClusterViewport* ViewportPtr = ViewportManagerPtr->FindViewport(ViewIndex, &ViewportContextNum);
+	IDisplayClusterViewport* ViewportPtr = ViewportManager->FindViewport(ViewIndex, &ViewportContextNum);
 	if (ViewportPtr == nullptr)
 	{
 		UE_LOG(LogDisplayClusterRender, Warning, TEXT("Viewport StereoViewIndex='%i' not found"), ViewIndex);
@@ -380,10 +380,10 @@ FMatrix FDisplayClusterDeviceBase::GetStereoProjectionMatrix(const int32 ViewInd
 	// ViewIndex == eSSE_MONOSCOPIC(-1) is a special case called for ISR culling math.
 	// Since nDisplay is not ISR compatible, we ignore this request. This won't be neccessary once
 	// we stop using nDisplay as a stereoscopic rendering device (IStereoRendering).
-	if (ViewportManagerPtr && ViewportManagerPtr->IsSceneOpened() && ViewIndex >= 0)
+	if (ViewportManager && ViewportManager->IsSceneOpened() && ViewIndex >= 0)
 	{
 		uint32 ViewportContextNum = 0;
-		IDisplayClusterViewport* ViewportPtr = ViewportManagerPtr->FindViewport(ViewIndex, &ViewportContextNum);
+		IDisplayClusterViewport* ViewportPtr = ViewportManager->FindViewport(ViewIndex, &ViewportContextNum);
 		if (ViewportPtr == nullptr)
 		{
 			UE_LOG(LogDisplayClusterRender, Warning, TEXT("Viewport StereoViewIndex='%i' not found"), ViewIndex);
@@ -403,59 +403,71 @@ bool FDisplayClusterDeviceBase::BeginNewFrame(FViewport* InViewport, UWorld* InW
 	check(IsInGameThread());
 	check(InViewport);
 
-	IDisplayClusterViewportManagerProxy* NewViewportManagerProxy = nullptr;
-
 	IDisplayCluster& DisplayCluster = IDisplayCluster::Get();
-	ADisplayClusterRootActor* RootActor = DisplayCluster.GetGameMgr()->GetRootActor();
-	if (RootActor)
+	if (ADisplayClusterRootActor* RootActor = DisplayCluster.GetGameMgr()->GetRootActor())
 	{
-		IDisplayClusterViewportManager* ViewportManager = RootActor->GetViewportManager();
-		if (ViewportManager)
+		if (IDisplayClusterViewportManager* ViewportManagerPtr = RootActor->GetViewportManager())
 		{
 			const FString LocalNodeId = DisplayCluster.GetConfigMgr()->GetLocalNodeId();
+
 			// Update local node viewports (update\create\delete) and build new render frame
-			if (ViewportManager->UpdateConfiguration(RenderFrameMode, LocalNodeId, RootActor))
+			if (ViewportManagerPtr->UpdateConfiguration(RenderFrameMode, LocalNodeId, RootActor))
 			{
-				if (ViewportManager->BeginNewFrame(InViewport, InWorld, OutRenderFrame))
+				if (ViewportManagerPtr->BeginNewFrame(InViewport, InWorld, OutRenderFrame))
 				{
-					// Begin use viewport manager for current frame
-					ViewportManagerPtr = ViewportManager;
-
-					// Send viewport manager proxy on render thread
-					NewViewportManagerProxy = ViewportManager->GetProxy();
-
 					// update total number of views for this frame (in multiple families)
 					DesiredNumberOfViews = OutRenderFrame.DesiredNumberOfViews;
+
+					return true;
 				}
 			}
 		}
 	}
 
-	// Update render thread viewport manager proxy
-	ENQUEUE_RENDER_COMMAND(DisplayClusterDevice_SetViewportManagerPtr)(
-		[DCRenderDevice = this, ViewportManagerProxy = NewViewportManagerProxy](FRHICommandListImmediate& RHICmdList)
-	{
-		DCRenderDevice->ViewportManagerProxyPtr = ViewportManagerProxy;
-	});
-
-	return NewViewportManagerProxy != nullptr;
+	return false;
 }
 
-void FDisplayClusterDeviceBase::FinalizeNewFrame()
+void FDisplayClusterDeviceBase::InitializeNewFrame()
 {
 	IDisplayCluster& DisplayCluster = IDisplayCluster::Get();
-	ADisplayClusterRootActor* RootActor = DisplayCluster.GetGameMgr()->GetRootActor();
-	if (RootActor)
+	if (ADisplayClusterRootActor* RootActor = DisplayCluster.GetGameMgr()->GetRootActor())
 	{
-		IDisplayClusterViewportManager* ViewportManager = RootActor->GetViewportManager();
-		if (ViewportManager)
+		if (IDisplayClusterViewportManager* ViewportManagerPtr = RootActor->GetViewportManager())
 		{
-			ViewportManager->FinalizeNewFrame();
+			// Begin use viewport manager for current frame
+			ViewportManager = ViewportManagerPtr;
+
+			// Initialize frame for render
+			ViewportManager->InitializeNewFrame();
+
+			FDisplayClusterViewportManagerProxy* ViewportManagerProxyPtr = static_cast<FDisplayClusterViewportManagerProxy*>(ViewportManager->GetProxy());
+
+			// Send viewport manager proxy on render thread
+			ENQUEUE_RENDER_COMMAND(DisplayClusterDevice_SetViewportManagerProxy)(
+				[DCRenderDevice = this, NewViewportManagerProxy = ViewportManagerProxyPtr->AsShared()](FRHICommandListImmediate& RHICmdList)
+				{
+					DCRenderDevice->ViewportManagerProxy = NewViewportManagerProxy;
+				});
 		}
 	}
+}
+
+void FDisplayClusterDeviceBase::FinalizeNewFrame()
+{
+	if (ViewportManager)
+	{
+		ViewportManager->FinalizeNewFrame();
+	}
 
 	// reset viewport manager ptr on game thread
-	ViewportManagerPtr = nullptr;
+	ViewportManager = nullptr;
+
+	// Reset viewport manager proxy on render thread
+	ENQUEUE_RENDER_COMMAND(DisplayClusterDevice_ResetViewportManagerProxy)(
+		[DCRenderDevice = this](FRHICommandListImmediate& RHICmdList)
+		{
+			DCRenderDevice->ViewportManagerProxy.Reset();
+		});
 }
 
 DECLARE_GPU_STAT_NAMED(nDisplay_Device_RenderTexture, TEXT("nDisplay RenderDevice::RenderTexture"));
@@ -483,10 +495,10 @@ void FDisplayClusterDeviceBase::RenderTexture_RenderThread(FRHICommandListImmedi
 			TransitionAndCopyTexture(RHICmdList, SrcTexture, BackBuffer, CopyInfo);
 		}
 
-		if (RenderFrameMode == EDisplayClusterRenderFrameMode::Stereo && ViewportManagerProxyPtr)
+		if (RenderFrameMode == EDisplayClusterRenderFrameMode::Stereo && ViewportManagerProxy.IsValid())
 		{
 			// QuadBufStereo: Copy RIGHT_EYE to backbuffer
-			ViewportManagerProxyPtr->ResolveFrameTargetToBackBuffer_RenderThread(RHICmdList, 1, 1, BackBuffer, WindowSize);
+			ViewportManagerProxy->ResolveFrameTargetToBackBuffer_RenderThread(RHICmdList, 1, 1, BackBuffer, WindowSize);
 		}
 
 		const bool bClearTextureEnabled = CVarClearTextureEnabled.GetValueOnRenderThread() != 0;
@@ -552,9 +564,9 @@ void FDisplayClusterDeviceBase::StartFinalPostprocessSettings(struct FPostProces
 	check(IsInGameThread());
 
 	// eSSP_FULL pass reserved for UE internal render
-	if (StereoPassType != EStereoscopicPass::eSSP_FULL && ViewportManagerPtr)
+	if (StereoPassType != EStereoscopicPass::eSSP_FULL && ViewportManager)
 	{
-		IDisplayClusterViewport* ViewportPtr = ViewportManagerPtr->FindViewport(StereoViewIndex);
+		IDisplayClusterViewport* ViewportPtr = ViewportManager->FindViewport(StereoViewIndex);
 		if (ViewportPtr)
 		{
 			ViewportPtr->GetViewport_CustomPostProcessSettings().DoPostProcess(IDisplayClusterViewport_CustomPostProcessSettings::ERenderPass::Start, StartPostProcessingSettings);
@@ -567,9 +579,9 @@ bool FDisplayClusterDeviceBase::OverrideFinalPostprocessSettings(struct FPostPro
 	check(IsInGameThread());
 
 	// eSSP_FULL pass reserved for UE internal render
-	if (StereoPassType != EStereoscopicPass::eSSP_FULL && ViewportManagerPtr)
+	if (StereoPassType != EStereoscopicPass::eSSP_FULL && ViewportManager)
 	{
-		IDisplayClusterViewport* ViewportPtr = ViewportManagerPtr->FindViewport(StereoViewIndex);
+		IDisplayClusterViewport* ViewportPtr = ViewportManager->FindViewport(StereoViewIndex);
 		if (ViewportPtr)
 		{
 			return ViewportPtr->GetViewport_CustomPostProcessSettings().DoPostProcess(IDisplayClusterViewport_CustomPostProcessSettings::ERenderPass::Override, OverridePostProcessingSettings, &BlendWeight);
@@ -584,9 +596,9 @@ void FDisplayClusterDeviceBase::EndFinalPostprocessSettings(struct FPostProcessS
 	check(IsInGameThread());
 
 	// eSSP_FULL pass reserved for UE internal render
-	if (StereoPassType != EStereoscopicPass::eSSP_FULL && ViewportManagerPtr && FinalPostProcessingSettings != nullptr)
+	if (StereoPassType != EStereoscopicPass::eSSP_FULL && ViewportManager && FinalPostProcessingSettings != nullptr)
 	{
-		IDisplayClusterViewport* ViewportPtr = ViewportManagerPtr->FindViewport(StereoViewIndex);
+		IDisplayClusterViewport* ViewportPtr = ViewportManager->FindViewport(StereoViewIndex);
 		if (ViewportPtr)
 		{
 			ViewportPtr->GetViewport_CustomPostProcessSettings().DoPostProcess(IDisplayClusterViewport_CustomPostProcessSettings::ERenderPass::Final, FinalPostProcessingSettings);
diff --git a/Engine/Plugins/Runtime/nDisplay/Source/DisplayCluster/Private/Render/Device/DisplayClusterDeviceBase.h b/Engine/Plugins/Runtime/nDisplay/Source/DisplayCluster/Private/Render/Device/DisplayClusterDeviceBase.h
--- a/Engine/Plugins/Runtime/nDisplay/Source/DisplayCluster/Private/Render/Device/DisplayClusterDeviceBase.h
+++ b/Engine/Plugins/Runtime/nDisplay/Source/DisplayCluster/Private/Render/Device/DisplayClusterDeviceBase.h
@@ -17,10 +17,9 @@
 class IDisplayClusterPostProcess;
 class FDisplayClusterPresentationBase;
 class IDisplayClusterViewportManager;
-class IDisplayClusterViewportManagerProxy;
+class FDisplayClusterViewportManagerProxy;
 class FSceneView;
 
-
 /**
  * Abstract render device
  */
@@ -45,6 +44,8 @@ public:
 	virtual IDisplayClusterPresentation* GetPresentation() const override;
 
 	virtual bool BeginNewFrame(FViewport* InViewport, UWorld* InWorld, FDisplayClusterRenderFrame& OutRenderFrame) override;
+
+	virtual void InitializeNewFrame() override;
 	virtual void FinalizeNewFrame() override;
 
 public:
@@ -116,8 +117,8 @@ protected:
 
 private:
 	// Runtime viewport manager api for game and render threads. Internal usage only
-	IDisplayClusterViewportManager*      ViewportManagerPtr = nullptr;
-	IDisplayClusterViewportManagerProxy* ViewportManagerProxyPtr = nullptr;
+	IDisplayClusterViewportManager* ViewportManager = nullptr;
+	TSharedPtr<FDisplayClusterViewportManagerProxy, ESPMode::ThreadSafe> ViewportManagerProxy;
 
 	EDisplayClusterRenderFrameMode RenderFrameMode = EDisplayClusterRenderFrameMode::Mono;
 	int32 DesiredNumberOfViews = 0;
diff --git a/Engine/Plugins/Runtime/nDisplay/Source/DisplayCluster/Private/Render/Viewport/Configuration/DisplayClusterViewportConfiguration.cpp b/Engine/Plugins/Runtime/nDisplay/Source/DisplayCluster/Private/Render/Viewport/Configuration/DisplayClusterViewportConfiguration.cpp
--- a/Engine/Plugins/Runtime/nDisplay/Source/DisplayCluster/Private/Render/Viewport/Configuration/DisplayClusterViewportConfiguration.cpp
+++ b/Engine/Plugins/Runtime/nDisplay/Source/DisplayCluster/Private/Render/Viewport/Configuration/DisplayClusterViewportConfiguration.cpp
@@ -22,6 +22,9 @@
 #include "Misc/DisplayClusterGlobals.h"
 #include "Misc/DisplayClusterLog.h"
 
+#include "Engine/Classes/Engine/RendererSettings.h"
+
+
 ///////////////////////////////////////////////////////////////////
 int32 GDisplayClusterCrossGPUTransferEnable = 0;
 static FAutoConsoleVariableRef CDisplayClusterCrossGPUTransferEnable(
@@ -51,9 +54,106 @@ static FAutoConsoleVariableRef CVarDisplayClusterCrossGPUTransferPullData(
 	ECVF_RenderThreadSafe
 );
 
+/**
+* This enum is for CVar values only and is used to process the logic that converts the values to the runtime enum in GetAlphaChannelCaptureMode().
+*/
+enum class ECVarDisplayClusterAlphaChannelCaptureMode : uint8
+{
+	/** [Auto]
+	 * Use one of the [ThroughTonemapper] or [FXAA] modes available for the current project settings.
+	 */
+	Auto = 0,
+
+	/** [Disabled]
+	 * Disable alpha channel saving.
+	 */
+	 Disabled,
+
+	/** [ThroughTonemapper]
+	 * When rendering with the PropagateAlpha experimental mode turned on, the alpha channel is forwarded to post-processes.
+	 * In this case, the alpha channel is anti-aliased along with the color.
+	 * Since some post-processing may change the alpha, it is copied at the beginning of the PP and restored after all post-processing is completed.
+	 */
+	ThroughTonemapper,
+
+	/** [FXAA]
+	 * Otherwise, if the PropagateAlpha mode is disabled in the project settings, we need to save the alpha before it becomes invalid.
+	 * The alpha is valid until the scene color is resolved (on the ResolvedSceneColor callback). Therefore, it is copied to a temporary resource on this cb.
+	 * Since we need to remove AA jittering (because alpha copied before AA), anti-aliasing is turned off for this viewport.
+	 * And finally the FXAA is used for smoothing.
+	 */
+	FXAA,
+
+	/** [Copy]
+	 * Disable AA and TAA, Copy alpha from scenecolor texture to final.
+	 * These experimental (temporary) modes for the performance tests.
+	 */
+	 Copy,
+
+	/** [CopyAA]
+	 * Use AA, disable TAA, Copy alpha from scenecolor texture to final.
+	 * These experimental (temporary) modes for the performance tests.
+	 */
+	 CopyAA,
+
+	COUNT
+};
+
+/**
+ *Choose method to preserve alpha channel
+ */
+int32 GDisplayClusterAlphaChannelCaptureMode = 0;
+static FAutoConsoleVariableRef CVarDisplayClusterAlphaChannelCaptureMode(
+	TEXT("nDisplay.render.AlphaChannelCaptureMode"),
+	GDisplayClusterAlphaChannelCaptureMode,
+	TEXT("Alpha channel capture mode\n")
+	TEXT("0 - Auto (default) ThroughTonemapper or FXAA depend from prj settings\n")
+	TEXT("1 - Disabled\n")
+	TEXT("2 - ThroughTonemapper\n")
+	TEXT("3 - FXAA\n")
+	TEXT("4 - Copy [experimental]\n")
+	TEXT("5 - CopyAA [experimental]\n"),
+	ECVF_RenderThreadSafe
+);
+
 ///////////////////////////////////////////////////////////////////
 // FDisplayClusterViewportConfiguration
 ///////////////////////////////////////////////////////////////////
+EDisplayClusterRenderFrameAlphaChannelCaptureMode FDisplayClusterViewportConfiguration::GetAlphaChannelCaptureMode() const
+{
+	ECVarDisplayClusterAlphaChannelCaptureMode AlphaChannelCaptureMode = (ECVarDisplayClusterAlphaChannelCaptureMode)FMath::Clamp(GDisplayClusterAlphaChannelCaptureMode, 0, (int32)ECVarDisplayClusterAlphaChannelCaptureMode::COUNT - 1);
+
+	static const auto CVarPropagateAlpha = IConsoleManager::Get().FindTConsoleVariableDataInt(TEXT("r.PostProcessing.PropagateAlpha"));
+	const EAlphaChannelMode::Type PropagateAlpha = EAlphaChannelMode::FromInt(CVarPropagateAlpha->GetValueOnGameThread());
+	const bool bAllowThroughTonemapper = PropagateAlpha == EAlphaChannelMode::AllowThroughTonemapper;
+
+	switch (AlphaChannelCaptureMode)
+	{
+	case ECVarDisplayClusterAlphaChannelCaptureMode::Auto:
+		// Fit best possible AA
+		return bAllowThroughTonemapper ? EDisplayClusterRenderFrameAlphaChannelCaptureMode::ThroughTonemapper : EDisplayClusterRenderFrameAlphaChannelCaptureMode::FXAA;
+
+	case ECVarDisplayClusterAlphaChannelCaptureMode::ThroughTonemapper:
+		// Disable alpha capture if PropagateAlpha not valid
+		return bAllowThroughTonemapper ? EDisplayClusterRenderFrameAlphaChannelCaptureMode::ThroughTonemapper : EDisplayClusterRenderFrameAlphaChannelCaptureMode::None;
+
+	case ECVarDisplayClusterAlphaChannelCaptureMode::FXAA:
+		return EDisplayClusterRenderFrameAlphaChannelCaptureMode::FXAA;
+
+	case ECVarDisplayClusterAlphaChannelCaptureMode::Copy:
+		return EDisplayClusterRenderFrameAlphaChannelCaptureMode::Copy;
+
+	case ECVarDisplayClusterAlphaChannelCaptureMode::CopyAA:
+		return EDisplayClusterRenderFrameAlphaChannelCaptureMode::CopyAA;
+
+	case ECVarDisplayClusterAlphaChannelCaptureMode::Disabled:
+	default:
+		break;
+	}
+
+	return EDisplayClusterRenderFrameAlphaChannelCaptureMode::None;
+}
+
 bool FDisplayClusterViewportConfiguration::SetRootActor(ADisplayClusterRootActor* InRootActorPtr)
 {
 	check(IsInGameThread());
@@ -105,6 +205,9 @@ bool FDisplayClusterViewportConfiguration::ImplUpdateConfiguration(EDisplayClust
 			RenderFrameSettings.RenderMode = InRenderMode;
 			RenderFrameSettings.ClusterNodeId = InClusterNodeId;
 
+			// Support alpha channel capture
+			RenderFrameSettings.AlphaChannelCaptureMode = GetAlphaChannelCaptureMode();
+
 			if (InPreviewSettings != nullptr)
 			{
 				// Downscale resources with PreviewDownscaleRatio
@@ -220,7 +323,7 @@ bool FDisplayClusterViewportConfiguration::UpdatePreviewConfiguration(EDisplayCl
 }
 #endif
 
-void FDisplayClusterViewportConfiguration::ImplUpdateConfigurationVisibility(ADisplayClusterRootActor& RootActor, const UDisplayClusterConfigurationData& ConfigurationData)
+void FDisplayClusterViewportConfiguration::ImplUpdateConfigurationVisibility(ADisplayClusterRootActor& RootActor, const UDisplayClusterConfigurationData& ConfigurationData) const
 {
 	// Hide root actor components for all viewports
 	TSet<FPrimitiveComponentId> RootActorHidePrimitivesList;
diff --git a/Engine/Plugins/Runtime/nDisplay/Source/DisplayCluster/Private/Render/Viewport/Configuration/DisplayClusterViewportConfiguration.h b/Engine/Plugins/Runtime/nDisplay/Source/DisplayCluster/Private/Render/Viewport/Configuration/DisplayClusterViewportConfiguration.h
--- a/Engine/Plugins/Runtime/nDisplay/Source/DisplayCluster/Private/Render/Viewport/Configuration/DisplayClusterViewportConfiguration.h
+++ b/Engine/Plugins/Runtime/nDisplay/Source/DisplayCluster/Private/Render/Viewport/Configuration/DisplayClusterViewportConfiguration.h
@@ -48,7 +48,12 @@ private:
 
 	void ImplUpdateRenderFrameConfiguration(const FDisplayClusterConfigurationRenderFrame& InRenderFrameConfiguration);
 	void ImplPostUpdateRenderFrameConfiguration();
-	void ImplUpdateConfigurationVisibility(ADisplayClusterRootActor& InRootActor, const UDisplayClusterConfigurationData& InConfigurationData);
+
+	/** Hide DCRA components for nDisplay rendering.*/
+	void ImplUpdateConfigurationVisibility(ADisplayClusterRootActor& InRootActor, const UDisplayClusterConfigurationData& InConfigurationData) const;
+
+	/** Get alpha channel capture mode (for LightCard, ChromaKey).*/
+	EDisplayClusterRenderFrameAlphaChannelCaptureMode GetAlphaChannelCaptureMode() const;
 
 private:
 	FDisplayClusterViewportManager&    ViewportManager;
diff --git a/Engine/Plugins/Runtime/nDisplay/Source/DisplayCluster/Private/Render/Viewport/Configuration/DisplayClusterViewportConfigurationBase.cpp b/Engine/Plugins/Runtime/nDisplay/Source/DisplayCluster/Private/Render/Viewport/Configuration/DisplayClusterViewportConfigurationBase.cpp
--- a/Engine/Plugins/Runtime/nDisplay/Source/DisplayCluster/Private/Render/Viewport/Configuration/DisplayClusterViewportConfigurationBase.cpp
+++ b/Engine/Plugins/Runtime/nDisplay/Source/DisplayCluster/Private/Render/Viewport/Configuration/DisplayClusterViewportConfigurationBase.cpp
@@ -46,7 +46,7 @@ void FDisplayClusterViewportConfigurationBase::Update(const FString& ClusterNode
 		for (FDisplayClusterViewport* It : ViewportManager.ImplGetViewports())
 		{
 			// ignore ICVFX internal resources
-			if ((It->GetRenderSettingsICVFX().RuntimeFlags & ViewportRuntime_InternalResource) == 0)
+			if (!EnumHasAllFlags(It->GetRenderSettingsICVFX().RuntimeFlags, EDisplayClusterViewportRuntimeICVFXFlags::InternalResource))
 			{
 				// Only viewports from cluster node in render
 				if (ClusterNodeId.IsEmpty() || It->GetClusterNodeId() == ClusterNodeId)
@@ -97,7 +97,7 @@ void FDisplayClusterViewportConfigurationBase::Update(const TArray<FString>& InV
 		for (FDisplayClusterViewport* It : ViewportManager.ImplGetViewports())
 		{
 			// ignore ICVFX internal resources
-			if ((It->GetRenderSettingsICVFX().RuntimeFlags & ViewportRuntime_InternalResource) == 0)
+			if (!EnumHasAllFlags(It->GetRenderSettingsICVFX().RuntimeFlags, EDisplayClusterViewportRuntimeICVFXFlags::InternalResource))
 			{
 				if (InViewportNames.Find(It->GetId()) == INDEX_NONE)
 				{
diff --git a/Engine/Plugins/Runtime/nDisplay/Source/DisplayCluster/Private/Render/Viewport/Configuration/DisplayClusterViewportConfigurationHelpers.cpp b/Engine/Plugins/Runtime/nDisplay/Source/DisplayCluster/Private/Render/Viewport/Configuration/DisplayClusterViewportConfigurationHelpers.cpp
--- a/Engine/Plugins/Runtime/nDisplay/Source/DisplayCluster/Private/Render/Viewport/Configuration/DisplayClusterViewportConfigurationHelpers.cpp
+++ b/Engine/Plugins/Runtime/nDisplay/Source/DisplayCluster/Private/Render/Viewport/Configuration/DisplayClusterViewportConfigurationHelpers.cpp
@@ -185,7 +185,7 @@ void FDisplayClusterViewportConfigurationHelpers::UpdateBaseViewportSetting(FDis
 
 	const FDisplayClusterConfigurationViewport_RenderSettings& InRenderSettings = InConfigurationViewport.RenderSettings;
 
-	// OCIO
+	// Update OCIO for Viewport
 	FDisplayClusterViewportConfigurationHelpers_OpenColorIO::UpdateBaseViewport(DstViewport, RootActor, InConfigurationViewport);
 
 	// Additional per-viewport PostProcess
@@ -229,7 +229,7 @@ void FDisplayClusterViewportConfigurationHelpers::UpdateBaseViewportSetting(FDis
 
 		if (InConfigurationViewport.ICVFX.bAllowICVFX)
 		{
-			TargetFlags |= ViewportICVFX_Enable;
+			EnumAddFlags(TargetFlags, EDisplayClusterViewportICVFXFlags::Enable);
 
 			EDisplayClusterConfigurationICVFX_OverrideCameraRenderMode CameraRenderMode = InConfigurationViewport.ICVFX.CameraRenderMode;
 
@@ -243,20 +243,22 @@ void FDisplayClusterViewportConfigurationHelpers::UpdateBaseViewportSetting(FDis
 			{
 				// Disable camera frame render for this viewport
 			case EDisplayClusterConfigurationICVFX_OverrideCameraRenderMode::Disabled:
-				TargetFlags |= ViewportICVFX_DisableCamera | ViewportICVFX_DisableChromakey | ViewportICVFX_DisableChromakeyMarkers;
+				EnumAddFlags(TargetFlags, EDisplayClusterViewportICVFXFlags::DisableCamera | EDisplayClusterViewportICVFXFlags::DisableChromakey | EDisplayClusterViewportICVFXFlags::DisableChromakeyMarkers);
 				break;
 
 				// Disable chromakey render for this viewport
 			case EDisplayClusterConfigurationICVFX_OverrideCameraRenderMode::DisableChromakey:
-				TargetFlags |= ViewportICVFX_DisableChromakey | ViewportICVFX_DisableChromakeyMarkers;
+				EnumAddFlags(TargetFlags, EDisplayClusterViewportICVFXFlags::DisableChromakey | EDisplayClusterViewportICVFXFlags::DisableChromakeyMarkers);
 				break;
 
 				// Disable chromakey markers render for this viewport
 			case EDisplayClusterConfigurationICVFX_OverrideCameraRenderMode::DisableChromakeyMarkers:
-				TargetFlags |= ViewportICVFX_DisableChromakeyMarkers;
+
+				EnumAddFlags(TargetFlags, EDisplayClusterViewportICVFXFlags::DisableChromakeyMarkers);
 				break;
-				// Use default rendering rules
+
 			default:
+				// Use default rendering rules
 				break;
 			}
 
@@ -264,18 +266,21 @@ void FDisplayClusterViewportConfigurationHelpers::UpdateBaseViewportSetting(FDis
 			{
 				// Render incamera frame over lightcard for this viewport
 			case EDisplayClusterConfigurationICVFX_OverrideLightcardRenderMode::Over:
-				TargetFlags |= ViewportICVFX_OverrideLightcardMode;
+
+				EnumAddFlags(TargetFlags, EDisplayClusterViewportICVFXFlags::OverrideLightcardMode);
 				DstViewport.RenderSettingsICVFX.ICVFX.LightcardMode = EDisplayClusterShaderParametersICVFX_LightcardRenderMode::Over;
 				break;
 
 				// Over lightcard over incamera frame  for this viewport
 			case EDisplayClusterConfigurationICVFX_OverrideLightcardRenderMode::Under:
-				TargetFlags |= ViewportICVFX_OverrideLightcardMode;
+
+				EnumAddFlags(TargetFlags, EDisplayClusterViewportICVFXFlags::OverrideLightcardMode);
 				DstViewport.RenderSettingsICVFX.ICVFX.LightcardMode = EDisplayClusterShaderParametersICVFX_LightcardRenderMode::Under;
 				break;
 
 			case EDisplayClusterConfigurationICVFX_OverrideLightcardRenderMode::Disabled:
-				TargetFlags |= ViewportICVFX_DisableLightcard;
+
+				EnumAddFlags(TargetFlags, EDisplayClusterViewportICVFXFlags::DisableLightcard);
 				break;
 
 			default:
diff --git a/Engine/Plugins/Runtime/nDisplay/Source/DisplayCluster/Private/Render/Viewport/Configuration/DisplayClusterViewportConfigurationHelpers_ICVFX.cpp b/Engine/Plugins/Runtime/nDisplay/Source/DisplayCluster/Private/Render/Viewport/Configuration/DisplayClusterViewportConfigurationHelpers_ICVFX.cpp
--- a/Engine/Plugins/Runtime/nDisplay/Source/DisplayCluster/Private/Render/Viewport/Configuration/DisplayClusterViewportConfigurationHelpers_ICVFX.cpp
+++ b/Engine/Plugins/Runtime/nDisplay/Source/DisplayCluster/Private/Render/Viewport/Configuration/DisplayClusterViewportConfigurationHelpers_ICVFX.cpp
@@ -32,6 +32,7 @@
 #include "Render/Viewport/Containers/ImplDisplayClusterViewport_CustomFrustum.h"
 
 #include "Render/Viewport/RenderFrame/DisplayClusterRenderFrameSettings.h"
+#include "Render/Viewport/LightCard/DisplayClusterViewportLightCardManager.h"
 
 #include "Containers/DisplayClusterProjectionCameraPolicySettings.h"
 #include "DisplayClusterProjectionStrings.h"
@@ -165,7 +166,7 @@ FDisplayClusterViewport* FDisplayClusterViewportConfigurationHelpers_ICVFX::Impl
 		if (NewViewport != nullptr)
 		{
 			// Mark as internal resource
-			NewViewport->RenderSettingsICVFX.RuntimeFlags |= ViewportRuntime_InternalResource;
+			EnumAddFlags(NewViewport->RenderSettingsICVFX.RuntimeFlags, EDisplayClusterViewportRuntimeICVFXFlags::InternalResource);
 
 			// Dont show ICVFX composing viewports on frame target
 			NewViewport->RenderSettings.bVisible = false;
@@ -218,14 +219,38 @@ FDisplayClusterViewport* FDisplayClusterViewportConfigurationHelpers_ICVFX::GetO
 	}
 
 	// Mark viewport as used
-	CameraViewport->RenderSettingsICVFX.RuntimeFlags &= ~(ViewportRuntime_Unused);
+	EnumRemoveFlags(CameraViewport->RenderSettingsICVFX.RuntimeFlags, EDisplayClusterViewportRuntimeICVFXFlags::Unused);
 
 	// Add viewport ICVFX usage as Incamera
-	CameraViewport->RenderSettingsICVFX.RuntimeFlags |= ViewportRuntime_ICVFXIncamera;
+	EnumAddFlags(CameraViewport->RenderSettingsICVFX.RuntimeFlags, EDisplayClusterViewportRuntimeICVFXFlags::InCamera);
 
 	return CameraViewport;
 }
 
+void FDisplayClusterViewportConfigurationHelpers_ICVFX::ReuseUVLightCardViewportWithinClusterNode(FDisplayClusterViewport& InUVLightCardViewport)
+{
+	const TArray<FDisplayClusterViewport*>& Viewports = 
+#if WITH_EDITOR
+		InUVLightCardViewport.ImplGetOwner().ImplGetWholeClusterViewports_Editor();
+#else
+		InUVLightCardViewport.ImplGetOwner().ImplGetViewports();
+#endif
+
+	for (FDisplayClusterViewport* ViewportIt : Viewports)
+	{
+		if(ViewportIt && ViewportIt != &InUVLightCardViewport && !ViewportIt->RenderSettings.IsViewportOverrided()
+		&& EnumHasAnyFlags(ViewportIt->RenderSettingsICVFX.RuntimeFlags, EDisplayClusterViewportRuntimeICVFXFlags::UVLightcard))
+		{
+			if (ViewportIt->IsOpenColorIOEquals(InUVLightCardViewport))
+			{
+				// Reuse exist viewport:
+				InUVLightCardViewport.RenderSettings.SetViewportOverride(ViewportIt->GetId(), EDisplayClusterViewportOverrideMode::All);
+				break;
+			}
+		}
+	}
+}
+
 #if WITH_EDITOR
 TArray<FDisplayClusterViewport*> FDisplayClusterViewportConfigurationHelpers_ICVFX::PreviewGetRenderedInCameraViewports(ADisplayClusterRootActor& InRootActor, UDisplayClusterICVFXCameraComponent& InCameraComponent, bool bGetChromakey)
 {
@@ -236,15 +261,15 @@ TArray<FDisplayClusterViewport*> FDisplayClusterViewportConfigurationHelpers_ICV
 	if (ViewportManager != nullptr)
 	{
 		const FString ICVFXCameraId = InCameraComponent.GetCameraUniqueId();
-		const EDisplayClusterViewportRuntimeICVFXFlags InMask = bGetChromakey ? ViewportRuntime_ICVFXChromakey : ViewportRuntime_ICVFXIncamera;
+		const EDisplayClusterViewportRuntimeICVFXFlags InRuntimeFlagsMask = bGetChromakey ? EDisplayClusterViewportRuntimeICVFXFlags::Chromakey : EDisplayClusterViewportRuntimeICVFXFlags::InCamera;
 		const FString ViewportTypeId = bGetChromakey ? DisplayClusterViewportStrings::icvfx::chromakey : DisplayClusterViewportStrings::icvfx::camera;
 
 		for (FDisplayClusterViewport* ViewportIt : ViewportManager->ImplGetWholeClusterViewports_Editor())
 		{
 			if (ViewportIt != nullptr
-				&& (ViewportIt->RenderSettingsICVFX.RuntimeFlags & InMask) != 0
+				&& EnumHasAnyFlags(ViewportIt->RenderSettingsICVFX.RuntimeFlags, InRuntimeFlagsMask)
 				&& (ViewportIt->InputShaderResources.Num() > 0 && ViewportIt->InputShaderResources[0] != nullptr && ViewportIt->Contexts.Num() > 0)
-				&& (ViewportIt->RenderSettings.OverrideViewportId.IsEmpty()))
+				&& (!ViewportIt->RenderSettings.IsViewportOverrided()))
 			{
 				// this is incamera viewport. Check by name
 				const FString RequiredViewportId = ImplGetNameICVFX(ViewportIt->GetClusterNodeId(), ICVFXCameraId, ViewportTypeId);
@@ -278,12 +303,14 @@ void FDisplayClusterViewportConfigurationHelpers_ICVFX::PreviewReuseInnerFrustum
 
 	for (FDisplayClusterViewport* ViewportIt : PreviewGetRenderedInCameraViewports(InRootActor, InCameraComponent, false))
 	{
-		if (ViewportIt && ViewportIt->GetClusterNodeId() != InCameraViewport.GetClusterNodeId()
-			&& FDisplayClusterViewportConfigurationHelpers_OpenColorIO::IsInnerFrustumViewportSettingsEqual_Editor(*ViewportIt, InCameraViewport, InCameraComponent)
+		if (ViewportIt && ViewportIt != &InCameraViewport && ViewportIt->GetClusterNodeId() != InCameraViewport.GetClusterNodeId()
 			&& FDisplayClusterViewportConfigurationHelpers_Postprocess::IsInnerFrustumViewportSettingsEqual_Editor(*ViewportIt, InCameraViewport, InCameraComponent))
 		{
+			EDisplayClusterViewportOverrideMode ViewportOverrideMode = ViewportIt->IsOpenColorIOEquals(InCameraViewport) ?
+				EDisplayClusterViewportOverrideMode::All : EDisplayClusterViewportOverrideMode::InernalRTT;
+
 			// Reuse exist viewport:
-			InCameraViewport.RenderSettings.OverrideViewportId = ViewportIt->GetId();
+			InCameraViewport.RenderSettings.SetViewportOverride(ViewportIt->GetId(), ViewportOverrideMode);
 			return;
 		}
 	}
@@ -308,10 +335,14 @@ void FDisplayClusterViewportConfigurationHelpers_ICVFX::PreviewReuseChromakeyVie
 
 	for (FDisplayClusterViewport* ViewportIt : PreviewGetRenderedInCameraViewports(InRootActor, InCameraComponent, true))
 	{
-		if (ViewportIt && ViewportIt->GetClusterNodeId() != InChromakeyViewport.GetClusterNodeId())
+		if (ViewportIt && ViewportIt!= &InChromakeyViewport  && ViewportIt->GetClusterNodeId() != InChromakeyViewport.GetClusterNodeId())
 		{
+			// Chromakey support OCIO
+			EDisplayClusterViewportOverrideMode ViewportOverrideMode = ViewportIt->IsOpenColorIOEquals(InChromakeyViewport) ?
+				EDisplayClusterViewportOverrideMode::All : EDisplayClusterViewportOverrideMode::InernalRTT;
+
 			// Reuse exist viewport from other node
-			InChromakeyViewport.RenderSettings.OverrideViewportId = ViewportIt->GetId();
+			InChromakeyViewport.RenderSettings.SetViewportOverride(ViewportIt->GetId(), ViewportOverrideMode);
 			return;
 		}
 	}
@@ -341,10 +372,10 @@ FDisplayClusterViewport* FDisplayClusterViewportConfigurationHelpers_ICVFX::GetO
 	}
 
 	// Mark viewport as used
-	ChromakeyViewport->RenderSettingsICVFX.RuntimeFlags &= ~(ViewportRuntime_Unused);
+	EnumRemoveFlags(ChromakeyViewport->RenderSettingsICVFX.RuntimeFlags, EDisplayClusterViewportRuntimeICVFXFlags::Unused);
 
 	// Add viewport ICVFX usage as Chromakey
-	ChromakeyViewport->RenderSettingsICVFX.RuntimeFlags |= ViewportRuntime_ICVFXChromakey;
+	EnumAddFlags(ChromakeyViewport->RenderSettingsICVFX.RuntimeFlags, EDisplayClusterViewportRuntimeICVFXFlags::Chromakey);
 
 	return ChromakeyViewport;
 }
@@ -371,14 +402,44 @@ FDisplayClusterViewport* FDisplayClusterViewportConfigurationHelpers_ICVFX::GetO
 	}
 
 	// Mark viewport as used
-	LightcardViewport->RenderSettingsICVFX.RuntimeFlags &= ~(ViewportRuntime_Unused);
+	EnumRemoveFlags(LightcardViewport->RenderSettingsICVFX.RuntimeFlags, EDisplayClusterViewportRuntimeICVFXFlags::Unused);
 
 	// Add viewport ICVFX usage as Lightcard
-	LightcardViewport->RenderSettingsICVFX.RuntimeFlags |= ViewportRuntime_ICVFXLightcard;
+	EnumAddFlags(LightcardViewport->RenderSettingsICVFX.RuntimeFlags, EDisplayClusterViewportRuntimeICVFXFlags::Lightcard);
 
 	return LightcardViewport;
 }
 
+FDisplayClusterViewport* FDisplayClusterViewportConfigurationHelpers_ICVFX::GetOrCreateUVLightcardViewport(FDisplayClusterViewport& BaseViewport, ADisplayClusterRootActor& RootActor)
+{
+	// Create new lightcard viewport
+	const FString ResourceId = DisplayClusterViewportStrings::icvfx::uv_lightcard;
+
+	FDisplayClusterViewport* UVLightcardViewport = ImplFindViewport(RootActor, BaseViewport.GetId(), ResourceId);
+	if (UVLightcardViewport == nullptr)
+	{
+		TSharedPtr<IDisplayClusterProjectionPolicy, ESPMode::ThreadSafe> UVLightcardProjectionPolicy;
+		if (!ImplCreateProjectionPolicy(RootActor, BaseViewport.GetId(), ResourceId, false, UVLightcardProjectionPolicy))
+		{
+			return nullptr;
+		}
+
+		UVLightcardViewport = ImplCreateViewport(RootActor, BaseViewport.GetId(), ResourceId, UVLightcardProjectionPolicy);
+		if (UVLightcardViewport == nullptr)
+		{
+			return nullptr;
+		}
+	}
+
+	// Mark viewport as used
+	EnumRemoveFlags(UVLightcardViewport->RenderSettingsICVFX.RuntimeFlags, EDisplayClusterViewportRuntimeICVFXFlags::Unused);
+
+	// Add viewport ICVFX usage as Lightcard
+	EnumAddFlags(UVLightcardViewport->RenderSettingsICVFX.RuntimeFlags, EDisplayClusterViewportRuntimeICVFXFlags::UVLightcard);
+
+	return UVLightcardViewport;
+}
+
 bool FDisplayClusterViewportConfigurationHelpers_ICVFX::IsCameraUsed(UDisplayClusterICVFXCameraComponent& InCameraComponent)
 {
 	const FDisplayClusterConfigurationICVFX_CameraSettings& CameraSettings = InCameraComponent.GetCameraSettingsICVFX();
@@ -514,8 +575,8 @@ void FDisplayClusterViewportConfigurationHelpers_ICVFX::UpdateCameraViewportSett
 	// Update camera viewport projection policy settings
 	ImplUpdateCameraProjectionSettings(DstViewport.ProjectionPolicy, CameraSettings, InCameraComponent.GetCameraComponent());
 
-	// OCIO
-	FDisplayClusterViewportConfigurationHelpers_OpenColorIO::UpdateICVFXCameraViewport(DstViewport, RootActor, InCameraComponent);
+	// Update OCIO for Camera Viewport
+	FDisplayClusterViewportConfigurationHelpers_OpenColorIO::UpdateCameraViewport(DstViewport, RootActor, InCameraComponent);
 
 	// Motion blur:
 	const FDisplayClusterViewport_CameraMotionBlur CameraMotionBlurParameters = InCameraComponent.GetMotionBlurParameters();
@@ -615,6 +676,10 @@ void FDisplayClusterViewportConfigurationHelpers_ICVFX::UpdateChromakeyViewportS
 
 	FDisplayClusterViewportConfigurationHelpers::UpdateViewportSetting_OverlayRenderSettings(DstViewport, InRenderSettings.AdvancedRenderSettings);
 
+	// Update OCIO for Chromakey Viewport
+	// Note: Chromakey OCIO is temporarily disabled
+	// FDisplayClusterViewportConfigurationHelpers_OpenColorIO::UpdateChromakeyViewport(DstViewport, RootActor, InCameraComponent);
+
 	// Support inner camera custom frustum
 	FDisplayClusterViewportConfigurationHelpers_ICVFX::UpdateCameraCustomFrustum(DstViewport, CameraSettings.CustomFrustum);
 
@@ -634,7 +699,7 @@ void FDisplayClusterViewportConfigurationHelpers_ICVFX::UpdateChromakeyViewportS
 	// Debug: override the texture of the target viewport from this chromakeyRTT
 	if (ChromakeySettings.ChromakeyRenderTexture.bReplaceCameraViewport)
 	{
-		InCameraViewport.RenderSettings.OverrideViewportId = DstViewport.GetId();
+		InCameraViewport.RenderSettings.SetViewportOverride(DstViewport.GetId());
 	}
 }
 
@@ -713,20 +778,42 @@ bool FDisplayClusterViewportConfigurationHelpers_ICVFX::IsShouldUseLightcard(con
 	return FDisplayClusterViewportConfigurationHelpers_Visibility::IsValid(InLightcardSettings.ShowOnlyList);
 }
 
+bool FDisplayClusterViewportConfigurationHelpers_ICVFX::IsShouldUseUVLightcard(FDisplayClusterViewportManager& InViewportManager, const FDisplayClusterConfigurationICVFX_LightcardSettings& InLightcardSettings)
+{
+	if (InLightcardSettings.bEnable)
+	{
+		TSharedPtr<FDisplayClusterViewportLightCardManager, ESPMode::ThreadSafe> LightCardManager = InViewportManager.GetLightCardManager();;
+		if (LightCardManager.IsValid() && LightCardManager->IsUVLightCardEnabled())
+		{
+			return true;
+		}
+	}
+
+	// dont use lightcard if disabled
+	return false;
+}
+
 void FDisplayClusterViewportConfigurationHelpers_ICVFX::UpdateLightcardViewportSetting(FDisplayClusterViewport& DstViewport, FDisplayClusterViewport& BaseViewport, ADisplayClusterRootActor& RootActor)
 {
 	const FDisplayClusterConfigurationICVFX_StageSettings& StageSettings = RootActor.GetStageSettings();
 	const FDisplayClusterConfigurationICVFX_LightcardSettings& LightcardSettings = StageSettings.Lightcard;
 
-	check(LightcardSettings.bEnable);
-
 	// Reset runtime flags from prev frame:
 	DstViewport.ResetRuntimeParameters();
 
 	// LIghtcard texture used as overlay
 	DstViewport.RenderSettings.bVisible = false;
 
+	if (!LightcardSettings.bEnable)
+	{
+		// Disable this viewport
+		DstViewport.RenderSettings.bEnable = false;
+		return;
+	}
+
 	FDisplayClusterViewportConfigurationHelpers_Postprocess::UpdateLightcardPostProcessSettings(DstViewport, BaseViewport, RootActor);
+
+	// Update OCIO for Lightcard Viewport
 	FDisplayClusterViewportConfigurationHelpers_OpenColorIO::UpdateLightcardViewport(DstViewport, BaseViewport, RootActor);
 
 	DstViewport.RenderSettings.CaptureMode = EDisplayClusterViewportCaptureMode::Lightcard;
@@ -752,7 +839,7 @@ void FDisplayClusterViewportConfigurationHelpers_ICVFX::UpdateLightcardViewportS
 	DstViewport.RenderSettings.AssignParentViewport(BaseViewport.GetId(), BaseViewport.RenderSettings);
 
 	// Global lighcard rendering mode
-	if ((BaseViewport.RenderSettingsICVFX.Flags & ViewportICVFX_OverrideLightcardMode) == 0)
+	if (!EnumHasAllFlags(BaseViewport.RenderSettingsICVFX.Flags, EDisplayClusterViewportICVFXFlags::OverrideLightcardMode))
 	{
 		// Use global lightcard blending mode
 		switch (LightcardSettings.Blendingmode)
@@ -766,14 +853,23 @@ void FDisplayClusterViewportConfigurationHelpers_ICVFX::UpdateLightcardViewportS
 			break;
 		};
 	}
-	// Debug: override the texture of the target viewport from this lightcard RTT
-	if (InRenderSettings.bReplaceViewport)
+
+	if (EnumHasAnyFlags(DstViewport.RenderSettingsICVFX.RuntimeFlags, EDisplayClusterViewportRuntimeICVFXFlags::Lightcard))
 	{
-		BaseViewport.RenderSettings.OverrideViewportId = DstViewport.GetId();
+		// Debug: override the texture of the target viewport from this lightcard RTT
+		if (InRenderSettings.bReplaceViewport)
+		{
+			BaseViewport.RenderSettings.SetViewportOverride(DstViewport.GetId());
+		}
+		else
+		{
+			BaseViewport.RenderSettingsICVFX.ICVFX.Lightcard.ViewportId = DstViewport.GetId();
+		}
 	}
-	else
+
+	if (EnumHasAnyFlags(DstViewport.RenderSettingsICVFX.RuntimeFlags, EDisplayClusterViewportRuntimeICVFXFlags::UVLightcard))
 	{
-		BaseViewport.RenderSettingsICVFX.ICVFX.Lightcard.ViewportId = DstViewport.GetId();
+		BaseViewport.RenderSettingsICVFX.ICVFX.UVLightcard.ViewportId = DstViewport.GetId();
 	}
 }
 
diff --git a/Engine/Plugins/Runtime/nDisplay/Source/DisplayCluster/Private/Render/Viewport/Configuration/DisplayClusterViewportConfigurationHelpers_ICVFX.h b/Engine/Plugins/Runtime/nDisplay/Source/DisplayCluster/Private/Render/Viewport/Configuration/DisplayClusterViewportConfigurationHelpers_ICVFX.h
--- a/Engine/Plugins/Runtime/nDisplay/Source/DisplayCluster/Private/Render/Viewport/Configuration/DisplayClusterViewportConfigurationHelpers_ICVFX.h
+++ b/Engine/Plugins/Runtime/nDisplay/Source/DisplayCluster/Private/Render/Viewport/Configuration/DisplayClusterViewportConfigurationHelpers_ICVFX.h
@@ -18,7 +18,6 @@ struct FDisplayClusterConfigurationICVFX_ChromakeyMarkers;
 struct FDisplayClusterConfigurationICVFX_LightcardSettings;
 struct FDisplayClusterConfigurationICVFX_CameraCustomFrustum;
 struct FDisplayClusterConfigurationICVFX_CameraSettings;
-struct FOpenColorIODisplayConfiguration;
 
 struct FCameraContext_ICVFX
 {
@@ -31,14 +30,40 @@ struct FCameraContext_ICVFX
 class FDisplayClusterViewportConfigurationHelpers_ICVFX
 {
 public:
+	/* Get the ViewportManager object from DCRA. */
 	static FDisplayClusterViewportManager* GetViewportManager(ADisplayClusterRootActor& RootActor);
 
+	/* Find existing InnerFrustum viewport. */
 	static FDisplayClusterViewport* FindCameraViewport(ADisplayClusterRootActor& RootActor, UDisplayClusterICVFXCameraComponent& InCameraComponent);
 
+	/* Return the new or existing InnerFrustum viewport. */
 	static FDisplayClusterViewport* GetOrCreateCameraViewport(ADisplayClusterRootActor& RootActor, UDisplayClusterICVFXCameraComponent& InCameraComponent);
+
+	/* Return the new or existing Chromakey viewport. */
 	static FDisplayClusterViewport* GetOrCreateChromakeyViewport(ADisplayClusterRootActor& RootActor, UDisplayClusterICVFXCameraComponent& InCameraComponent);
+	
+	/* Return the new or existing LightCard viewport. */
 	static FDisplayClusterViewport* GetOrCreateLightcardViewport(FDisplayClusterViewport& BaseViewport, ADisplayClusterRootActor& RootActor);
 
+	/* Return the new or existing UVLightCard viewport. */
+	static FDisplayClusterViewport* GetOrCreateUVLightcardViewport(FDisplayClusterViewport& BaseViewport, ADisplayClusterRootActor& RootActor);
+
+	/* Because UVLightCard viewports share the same outer texture (from LightCardManager), they can clone each other.
+	 * Except when the OCIO settings are not equal.
+	 * This function sets the InUVLightCardViewport rendering settings to minimize rendering costs.
+	 */
+	static void ReuseUVLightCardViewportWithinClusterNode(FDisplayClusterViewport& InUVLightCardViewport);
+
+	/* Returns true if the use of the LightCard viewport is allowed in InLightcardSettings.
+	 * These viewports also require visible LightCard actors to render.
+	 */
+	static bool IsShouldUseLightcard(const FDisplayClusterConfigurationICVFX_LightcardSettings& InLightcardSettings);
+
+	/* Returns true if the use of the UVLightCard viewport is allowed in InLightcardSettings.
+	 * These viewports also require visible UVLightCard actors to render.
+	 */
+	static bool IsShouldUseUVLightcard(FDisplayClusterViewportManager& InViewportManager, const FDisplayClusterConfigurationICVFX_LightcardSettings& InLightcardSettings);
+
 #if WITH_EDITOR
 	static TArray<FDisplayClusterViewport*> PreviewGetRenderedInCameraViewports(ADisplayClusterRootActor& RootActor, UDisplayClusterICVFXCameraComponent& InCameraComponent, bool bGetChromakey = false);
 	static void PreviewReuseInnerFrustumViewportWithinClusterNodes_Editor(FDisplayClusterViewport& InCameraViewport, ADisplayClusterRootActor& RootActor, UDisplayClusterICVFXCameraComponent& InCameraComponent);
@@ -57,8 +82,6 @@ public:
 	static void UpdateCameraSettings_Chromakey(FDisplayClusterShaderParameters_ICVFX::FCameraSettings& InOutCameraSettings, const FDisplayClusterConfigurationICVFX_ChromakeySettings& InChromakeySettings, FDisplayClusterViewport* InChromakeyViewport);
 	static void UpdateCameraSettings_ChromakeyMarkers(FDisplayClusterShaderParameters_ICVFX::FCameraSettings& InOutCameraSettings, const FDisplayClusterConfigurationICVFX_ChromakeyMarkers& InChromakeyMarkers);
 
-	static bool IsShouldUseLightcard(const FDisplayClusterConfigurationICVFX_LightcardSettings& InLightcardSettings);
-
 	static void UpdateCameraCustomFrustum(FDisplayClusterViewport& DstViewport, const FDisplayClusterConfigurationICVFX_CameraCustomFrustum& InCameraCustomFrustum);
 	static void UpdateCameraViewportBufferRatio(FDisplayClusterViewport& DstViewport, const FDisplayClusterConfigurationICVFX_CameraSettings& CameraSettings);
 
diff --git a/Engine/Plugins/Runtime/nDisplay/Source/DisplayCluster/Private/Render/Viewport/Configuration/DisplayClusterViewportConfigurationHelpers_OpenColorIO.cpp b/Engine/Plugins/Runtime/nDisplay/Source/DisplayCluster/Private/Render/Viewport/Configuration/DisplayClusterViewportConfigurationHelpers_OpenColorIO.cpp
--- a/Engine/Plugins/Runtime/nDisplay/Source/DisplayCluster/Private/Render/Viewport/Configuration/DisplayClusterViewportConfigurationHelpers_OpenColorIO.cpp
+++ b/Engine/Plugins/Runtime/nDisplay/Source/DisplayCluster/Private/Render/Viewport/Configuration/DisplayClusterViewportConfigurationHelpers_OpenColorIO.cpp
@@ -13,183 +13,114 @@
 #include "Components/DisplayClusterICVFXCameraComponent.h"
 #include "OpenColorIODisplayExtension.h"
 
-bool FDisplayClusterViewportConfigurationHelpers_OpenColorIO::ImplUpdateOuterViewportOCIO(FDisplayClusterViewport& DstViewport, ADisplayClusterRootActor& RootActor)
+///////////////////////////////////////////////////////////////////////////////////////
+// FDisplayClusterViewportConfigurationHelpers_OpenColorIO
+///////////////////////////////////////////////////////////////////////////////////////
+bool FDisplayClusterViewportConfigurationHelpers_OpenColorIO::UpdateBaseViewport(FDisplayClusterViewport& DstViewport, ADisplayClusterRootActor& RootActor, const UDisplayClusterConfigurationViewport& InViewportConfiguration)
 {
-	const FDisplayClusterConfigurationICVFX_StageSettings& StageSettings = RootActor.GetStageSettings();
-	if (StageSettings.bUseOverallClusterOCIOConfiguration)
+	if (!EnumHasAnyFlags(DstViewport.RenderSettingsICVFX.RuntimeFlags, EDisplayClusterViewportRuntimeICVFXFlags::InCamera | EDisplayClusterViewportRuntimeICVFXFlags::Chromakey | EDisplayClusterViewportRuntimeICVFXFlags::Lightcard | EDisplayClusterViewportRuntimeICVFXFlags::UVLightcard))
 	{
-		for (const FDisplayClusterConfigurationOCIOProfile& OCIOProfileIt : StageSettings.PerViewportOCIOProfiles)
+		if (const FOpenColorIOColorConversionSettings* OCIOConfiguration = RootActor.GetStageSettings().FindViewportOCIOConfiguration(DstViewport.GetId()))
 		{
-			if (OCIOProfileIt.bIsEnabled)
-			{
-				for (const FString& ViewportNameIt : OCIOProfileIt.ApplyOCIOToObjects)
-				{
-					if (DstViewport.GetId().Compare(ViewportNameIt, ESearchCase::IgnoreCase) == 0)
-					{
-						// Use per-viewport OCIO overrides
-						if (FDisplayClusterViewportConfigurationHelpers_OpenColorIO::ImplUpdate(DstViewport, OCIOProfileIt.OCIOConfiguration))
-						{
-							return true;
-						}
-
-						break;
-					}
-				}
-			}
-		}
+			ApplyOCIOConfiguration(DstViewport, *OCIOConfiguration);
 
-		// use all viewports OCIO RootActor
-		if (StageSettings.AllViewportsOCIOConfiguration.bIsEnabled)
-		{
-			if (ImplUpdate(DstViewport, StageSettings.AllViewportsOCIOConfiguration.OCIOConfiguration))
-			{
-				return true;
-			}
+			return true;
 		}
-	}
 
-	// No valid OCIO for ICVX, disable
-	ImplDisable(DstViewport);
+		DisableOCIOConfiguration(DstViewport);
+	}
 
 	return false;
 }
 
-bool FDisplayClusterViewportConfigurationHelpers_OpenColorIO::UpdateBaseViewport(FDisplayClusterViewport& DstViewport, ADisplayClusterRootActor& RootActor, const UDisplayClusterConfigurationViewport& InViewportConfiguration)
-{
-	// ICVFX logic for OCIO
-	return ImplUpdateOuterViewportOCIO(DstViewport, RootActor);
-}
-
 bool FDisplayClusterViewportConfigurationHelpers_OpenColorIO::UpdateLightcardViewport(FDisplayClusterViewport& DstViewport, FDisplayClusterViewport& BaseViewport, ADisplayClusterRootActor& RootActor)
 {
-	const FDisplayClusterConfigurationICVFX_StageSettings& StageSettings = RootActor.GetStageSettings();
-	const FDisplayClusterConfigurationICVFX_LightcardSettings& LightcardSettings = StageSettings.Lightcard;
-
-	// First try use global OCIO from stage settings
-	if (LightcardSettings.bEnableOuterViewportOCIO)
+	if (EnumHasAnyFlags(DstViewport.RenderSettingsICVFX.RuntimeFlags, EDisplayClusterViewportRuntimeICVFXFlags::Lightcard | EDisplayClusterViewportRuntimeICVFXFlags::UVLightcard))
 	{
-		return  ImplUpdateOuterViewportOCIO(DstViewport, RootActor);
+		if (const FOpenColorIOColorConversionSettings* OCIOConfiguration = RootActor.GetStageSettings().FindLightcardOCIOConfiguration(BaseViewport.GetId()))
+		{
+			ApplyOCIOConfiguration(DstViewport, *OCIOConfiguration);
+
+			return true;
+		}
+
+		// No OCIO defined for this lightcard viewport, disabled
+		DisableOCIOConfiguration(DstViewport);
 	}
 
-	// No OCIO defined for this lightcard viewport, disabled
-	ImplDisable(DstViewport);
-	
 	return false;
 }
 
-#if WITH_EDITOR
-bool FDisplayClusterViewportConfigurationHelpers_OpenColorIO::IsInnerFrustumViewportSettingsEqual_Editor(const FDisplayClusterViewport& InViewport1, const FDisplayClusterViewport& InViewport2, UDisplayClusterICVFXCameraComponent& InCameraComponent)
+bool FDisplayClusterViewportConfigurationHelpers_OpenColorIO::UpdateCameraViewport(FDisplayClusterViewport& DstViewport, ADisplayClusterRootActor& RootActor, UDisplayClusterICVFXCameraComponent& InCameraComponent)
 {
-	const FDisplayClusterConfigurationICVFX_CameraSettings& CameraSettings = InCameraComponent.GetCameraSettingsICVFX();
-	if (CameraSettings.AllNodesOCIOConfiguration.bIsEnabled)
+	if (EnumHasAllFlags(DstViewport.RenderSettingsICVFX.RuntimeFlags, EDisplayClusterViewportRuntimeICVFXFlags::InCamera))
 	{
-		for (const FDisplayClusterConfigurationOCIOProfile& OCIOProfileIt : CameraSettings.PerNodeOCIOProfiles)
+		if (const FOpenColorIOColorConversionSettings* OCIOConfiguration = InCameraComponent.GetCameraSettingsICVFX().FindInnerFrustumOCIOConfiguration(DstViewport.GetClusterNodeId()))
 		{
-			if (OCIOProfileIt.bIsEnabled)
-			{
-				const FString* CustomNode1 = OCIOProfileIt.ApplyOCIOToObjects.FindByPredicate([ClusterNodeId = InViewport1.GetClusterNodeId()](const FString& InClusterNodeId)
-				{
-					return ClusterNodeId.Equals(InClusterNodeId, ESearchCase::IgnoreCase);
-				});
-
-				const FString* CustomNode2 = OCIOProfileIt.ApplyOCIOToObjects.FindByPredicate([ClusterNodeId = InViewport2.GetClusterNodeId()](const FString& InClusterNodeId)
-				{
-					return ClusterNodeId.Equals(InClusterNodeId, ESearchCase::IgnoreCase);
-				});
-
-				if (CustomNode1 && CustomNode2)
-				{
-					// equal custom settings
-					return true;
-				}
-
-				if (CustomNode1 || CustomNode2)
-				{
-					// one of node has custom settings
-					return false;
-				}
-			}
+			ApplyOCIOConfiguration(DstViewport, *OCIOConfiguration);
+
+			return true;
 		}
+
+		DisableOCIOConfiguration(DstViewport);
 	}
 
-	return true;
+	return false;
 }
-#endif
 
-bool FDisplayClusterViewportConfigurationHelpers_OpenColorIO::UpdateICVFXCameraViewport(FDisplayClusterViewport& DstViewport, ADisplayClusterRootActor& RootActor, UDisplayClusterICVFXCameraComponent& InCameraComponent)
+bool FDisplayClusterViewportConfigurationHelpers_OpenColorIO::UpdateChromakeyViewport(FDisplayClusterViewport& DstViewport, ADisplayClusterRootActor& RootActor, UDisplayClusterICVFXCameraComponent& InCameraComponent)
 {
-	const FDisplayClusterConfigurationICVFX_CameraSettings& CameraSettings = InCameraComponent.GetCameraSettingsICVFX();
-	if (CameraSettings.AllNodesOCIOConfiguration.bIsEnabled)
+	if (EnumHasAllFlags(DstViewport.RenderSettingsICVFX.RuntimeFlags, EDisplayClusterViewportRuntimeICVFXFlags::Chromakey))
 	{
-		const FDisplayClusterRenderFrameSettings& RenderFrameSettings = DstViewport.Owner.GetRenderFrameSettings();
-
-		const FString& ClusterNodeId = DstViewport.GetClusterNodeId();
-		check(!ClusterNodeId.IsEmpty());
-
-		for (const FDisplayClusterConfigurationOCIOProfile& OCIOProfileIt : CameraSettings.PerNodeOCIOProfiles)
+		if (const FOpenColorIOColorConversionSettings* OCIOConfiguration = InCameraComponent.GetCameraSettingsICVFX().FindChromakeyOCIOConfiguration(DstViewport.GetClusterNodeId()))
 		{
-			if (OCIOProfileIt.bIsEnabled)
-			{
-				for (const FString& ClusterNodeIt : OCIOProfileIt.ApplyOCIOToObjects)
-				{
-					if (ClusterNodeId.Compare(ClusterNodeIt, ESearchCase::IgnoreCase) == 0)
-					{
-						// Use cluster node OCIO
-						if (FDisplayClusterViewportConfigurationHelpers_OpenColorIO::ImplUpdate(DstViewport, OCIOProfileIt.OCIOConfiguration))
-						{
-							return true;
-						}
-
-						break;
-					}
-				}
-			}
-		}
+			ApplyOCIOConfiguration(DstViewport, *OCIOConfiguration);
 
-		// cluster node OCIO override not found, use all nodes configuration
-		if (ImplUpdate(DstViewport, CameraSettings.AllNodesOCIOConfiguration.OCIOConfiguration))
-		{
 			return true;
 		}
-	}
 
-	// No OCIO defined for this camera, disabled
-	ImplDisable(DstViewport);
+		DisableOCIOConfiguration(DstViewport);
+	}
 
 	return false;
 }
 
-bool FDisplayClusterViewportConfigurationHelpers_OpenColorIO::ImplUpdate(FDisplayClusterViewport& DstViewport, const FOpenColorIODisplayConfiguration& InConfiguration)
+void FDisplayClusterViewportConfigurationHelpers_OpenColorIO::ApplyOCIOConfiguration(FDisplayClusterViewport& DstViewport, const FOpenColorIOColorConversionSettings& InConfiguration)
 {
-	if (InConfiguration.bIsEnabled && InConfiguration.ColorConfiguration.IsValid())
+	// Check OCIO before apply
+	if(!InConfiguration.IsValid()
+	|| !InConfiguration.SourceColorSpace.IsValid()
+	|| !InConfiguration.DestinationColorSpace.IsValid())
 	{
-		// Create/Update OCIO:
-		if (DstViewport.OpenColorIODisplayExtension.IsValid() == false)
-		{
-			DstViewport.OpenColorIODisplayExtension = FSceneViewExtensions::NewExtension<FOpenColorIODisplayExtension>(nullptr);
-			if (DstViewport.OpenColorIODisplayExtension.IsValid())
-			{
-				// assign active func
-				DstViewport.OpenColorIODisplayExtension->IsActiveThisFrameFunctions.Reset(1);
-				DstViewport.OpenColorIODisplayExtension->IsActiveThisFrameFunctions.Add(DstViewport.GetSceneViewExtensionIsActiveFunctor());
-			}
-		}
+		DisableOCIOConfiguration(DstViewport);
 
-		if (DstViewport.OpenColorIODisplayExtension.IsValid())
-		{
-			// Update configuration
-			DstViewport.OpenColorIODisplayExtension->SetDisplayConfiguration(InConfiguration);
-		}
+		return;
+	}
 
-		return true;
+	if (DstViewport.OpenColorIO.IsValid() && DstViewport.OpenColorIO->IsDisplayConfigurationEquals(InConfiguration))
+	{
+		// Already assigned
+		return;
 	}
 
-	return false;
+	DstViewport.OpenColorIO.Reset();
+	DstViewport.OpenColorIO = MakeShared<FDisplayClusterViewport_OpenColorIO>(InConfiguration);
 }
 
-void FDisplayClusterViewportConfigurationHelpers_OpenColorIO::ImplDisable(FDisplayClusterViewport& DstViewport)
+void FDisplayClusterViewportConfigurationHelpers_OpenColorIO::DisableOCIOConfiguration(FDisplayClusterViewport& DstViewport)
 {
 	// Remove OICO ref
-	DstViewport.OpenColorIODisplayExtension.Reset();
+	DstViewport.OpenColorIO.Reset();
 }
 
+#if WITH_EDITOR
+bool FDisplayClusterViewportConfigurationHelpers_OpenColorIO::IsInnerFrustumViewportSettingsEqual_Editor(const FDisplayClusterViewport& InViewport1, const FDisplayClusterViewport& InViewport2, UDisplayClusterICVFXCameraComponent& InCameraComponent)
+{
+	return InCameraComponent.GetCameraSettingsICVFX().IsInnerFrustumViewportSettingsEqual_Editor(InViewport1.GetClusterNodeId(), InViewport2.GetClusterNodeId());
+}
+
+bool FDisplayClusterViewportConfigurationHelpers_OpenColorIO::IsChromakeyViewportSettingsEqual_Editor(const FDisplayClusterViewport& InViewport1, const FDisplayClusterViewport& InViewport2, UDisplayClusterICVFXCameraComponent& InCameraComponent)
+{
+	return InCameraComponent.GetCameraSettingsICVFX().IsChromakeyViewportSettingsEqual_Editor(InViewport1.GetClusterNodeId(), InViewport2.GetClusterNodeId());
+}
+#endif
diff --git a/Engine/Plugins/Runtime/nDisplay/Source/DisplayCluster/Private/Render/Viewport/Configuration/DisplayClusterViewportConfigurationHelpers_OpenColorIO.h b/Engine/Plugins/Runtime/nDisplay/Source/DisplayCluster/Private/Render/Viewport/Configuration/DisplayClusterViewportConfigurationHelpers_OpenColorIO.h
--- a/Engine/Plugins/Runtime/nDisplay/Source/DisplayCluster/Private/Render/Viewport/Configuration/DisplayClusterViewportConfigurationHelpers_OpenColorIO.h
+++ b/Engine/Plugins/Runtime/nDisplay/Source/DisplayCluster/Private/Render/Viewport/Configuration/DisplayClusterViewportConfigurationHelpers_OpenColorIO.h
@@ -10,24 +10,88 @@ class ADisplayClusterRootActor;
 class UDisplayClusterConfigurationViewport;
 class UDisplayClusterICVFXCameraComponent;
 
-struct FOpenColorIODisplayConfiguration;
 struct FOpenColorIOColorConversionSettings;
 
 class FDisplayClusterViewportConfigurationHelpers_OpenColorIO
 {
 public:
+	/** Update OCIO for base viewport (Outer).
+	 *
+	 * @param DstViewport             - Dest viewport object.
+	 * @param RootActor               - DCRA object.
+	 * @param InViewportConfiguration - Actual viewport configuration.
+	 *
+	 * @return - true if success.
+	 */
 	static bool UpdateBaseViewport(FDisplayClusterViewport& DstViewport, ADisplayClusterRootActor& RootActor, const UDisplayClusterConfigurationViewport& InViewportConfiguration);
-	static bool UpdateICVFXCameraViewport(FDisplayClusterViewport& DstViewport, ADisplayClusterRootActor& RootActor, UDisplayClusterICVFXCameraComponent& InCameraComponent);
+
+	/** Update OCIO for LightCard viewport.
+	 *
+	 * @param DstViewport  - Dest LightCard viewport object.
+	 * @param BaseViewport - Base (Outer) viewport object.
+	 * @param RootActor    - DCRA object.
+	 *
+	 * @return - true if success.
+	 */
 	static bool UpdateLightcardViewport(FDisplayClusterViewport& DstViewport, FDisplayClusterViewport& BaseViewport, ADisplayClusterRootActor& RootActor);
 
+	/** Update OCIO for InnerFrustum(InCamera) viewport.
+	 *
+	 * @param DstViewport       - Dest InCamera viewport object
+	 * @param RootActor         - DCRA object
+	 * @param InCameraComponent - InCamera component with configuration
+	 *
+	 * @return - true if success
+	 */
+	static bool UpdateCameraViewport(FDisplayClusterViewport& DstViewport, ADisplayClusterRootActor& RootActor, UDisplayClusterICVFXCameraComponent& InCameraComponent);
+
+	/** Update OCIO for Chromakey viewport.
+	 *
+	 * @param DstViewport       - Dest InCamera viewport object
+	 * @param RootActor         - DCRA object
+	 * @param InCameraComponent - InCamera component with configuration
+	 *
+	 * @return - true if success
+	 */
+	static bool UpdateChromakeyViewport(FDisplayClusterViewport& DstViewport, ADisplayClusterRootActor& RootActor, UDisplayClusterICVFXCameraComponent& InCameraComponent);
+
 #if WITH_EDITOR
-	// return true when same settings used for both viewports
+	/** Compare InnerFrustum OCIO configuration for input viewports
+	 *
+	 * @param InViewport1       - Viewport #1
+	 * @param InViewport2       - Viewport #2
+	 * @param InCameraComponent - InCamera component with configuration
+	 *
+	 * @return - true if OCIO is equals
+	 */
 	static bool IsInnerFrustumViewportSettingsEqual_Editor(const FDisplayClusterViewport& InViewport1, const FDisplayClusterViewport& InViewport2, UDisplayClusterICVFXCameraComponent& InCameraComponent);
+
+	/** Compare Chromakey OCIO configuration for input viewports
+	 *
+	 * @param InViewport1       - Viewport #1
+	 * @param InViewport2       - Viewport #2
+	 * @param InCameraComponent - InCamera component with configuration
+	 *
+	 * @return - true if OCIO is equals
+	 */
+	static bool IsChromakeyViewportSettingsEqual_Editor(const FDisplayClusterViewport& InViewport1, const FDisplayClusterViewport& InViewport2, UDisplayClusterICVFXCameraComponent& InCameraComponent);
 #endif
 
 private:
-	static bool ImplUpdate(FDisplayClusterViewport& DstViewport, const FOpenColorIODisplayConfiguration& InConfiguration);
-	static void ImplDisable(FDisplayClusterViewport& DstViewport);
+	/** Apply OCIO to viewport.
+	 *
+	 * @param DstViewport       - Dest viewport object
+	 * @param InConfiguration   - OCIO configuration
+	 *
+	 * @return - none
+	 */
+	static void ApplyOCIOConfiguration(FDisplayClusterViewport& DstViewport, const FOpenColorIOColorConversionSettings& InConfiguration);
 
-	static bool ImplUpdateOuterViewportOCIO(FDisplayClusterViewport& DstViewport, ADisplayClusterRootActor& RootActor);
+	/** Reset OCIO for viewport.
+	 *
+	 * @param DstViewport       - Dest viewport object
+	 *
+	 * @return - none
+	 */
+	static void DisableOCIOConfiguration(FDisplayClusterViewport& DstViewport);
 };
diff --git a/Engine/Plugins/Runtime/nDisplay/Source/DisplayCluster/Private/Render/Viewport/Configuration/DisplayClusterViewportConfigurationHelpers_Visibility.cpp b/Engine/Plugins/Runtime/nDisplay/Source/DisplayCluster/Private/Render/Viewport/Configuration/DisplayClusterViewportConfigurationHelpers_Visibility.cpp
--- a/Engine/Plugins/Runtime/nDisplay/Source/DisplayCluster/Private/Render/Viewport/Configuration/DisplayClusterViewportConfigurationHelpers_Visibility.cpp
+++ b/Engine/Plugins/Runtime/nDisplay/Source/DisplayCluster/Private/Render/Viewport/Configuration/DisplayClusterViewportConfigurationHelpers_Visibility.cpp
@@ -131,7 +131,7 @@ void FDisplayClusterViewportConfigurationHelpers_Visibility::UpdateHideList_ICVF
 				ViewportIt->VisibilitySettings.UpdateConfiguration(EDisplayClusterViewport_VisibilityMode::Hide, ActorLayerNames, AdditionalComponentsList);
 
 				// Support additional hide list for outer viewports
-				if ((ViewportIt->GetRenderSettingsICVFX().RuntimeFlags & ViewportRuntime_ICVFXTarget) != 0)
+				if (EnumHasAllFlags(ViewportIt->GetRenderSettingsICVFX().RuntimeFlags, EDisplayClusterViewportRuntimeICVFXFlags::Target))
 				{
 					ViewportIt->VisibilitySettings.AppendHideList(OuterActorLayerNames, OuterAdditionalComponentsList);
 				}
diff --git a/Engine/Plugins/Runtime/nDisplay/Source/DisplayCluster/Private/Render/Viewport/Configuration/DisplayClusterViewportConfigurationICVFX.cpp b/Engine/Plugins/Runtime/nDisplay/Source/DisplayCluster/Private/Render/Viewport/Configuration/DisplayClusterViewportConfigurationICVFX.cpp
--- a/Engine/Plugins/Runtime/nDisplay/Source/DisplayCluster/Private/Render/Viewport/Configuration/DisplayClusterViewportConfigurationICVFX.cpp
+++ b/Engine/Plugins/Runtime/nDisplay/Source/DisplayCluster/Private/Render/Viewport/Configuration/DisplayClusterViewportConfigurationICVFX.cpp
@@ -100,7 +100,7 @@ public:
 
 	void AssignChromakeyToTargetViewport(FDisplayClusterViewport& DstViewport, UDisplayClusterICVFXCameraComponent& InCameraComponent)
 	{
-		if ((DstViewport.RenderSettingsICVFX.Flags & ViewportICVFX_DisableChromakey) != 0)
+		if (EnumHasAnyFlags(DstViewport.RenderSettingsICVFX.Flags, EDisplayClusterViewportICVFXFlags::DisableChromakey))
 		{
 			// chromakey disabled for this viewport
 			return;
@@ -111,7 +111,7 @@ public:
 
 		FDisplayClusterViewportConfigurationHelpers_ICVFX::UpdateCameraSettings_Chromakey(DstCameraData, CameraSettings.Chromakey, ChromakeyViewport);
 
-		if ((DstViewport.RenderSettingsICVFX.Flags & ViewportICVFX_DisableChromakeyMarkers) == 0)
+		if (!EnumHasAnyFlags(DstViewport.RenderSettingsICVFX.Flags, EDisplayClusterViewportICVFXFlags::DisableChromakeyMarkers))
 		{
 			FDisplayClusterViewportConfigurationHelpers_ICVFX::UpdateCameraSettings_ChromakeyMarkers(DstCameraData, CameraSettings.Chromakey.ChromakeyMarkers);
 		}
@@ -182,7 +182,10 @@ public:
 					ViewportIt->RenderSettingsICVFX.ICVFX.Cameras.Add(ShaderParametersCameraSettings);
 
 					// At lest one target must accept chromakey
-					bAllowChromakey |= (ViewportIt->RenderSettingsICVFX.Flags & ViewportICVFX_DisableChromakey) == 0;
+					if (!EnumHasAnyFlags(ViewportIt->RenderSettingsICVFX.Flags, EDisplayClusterViewportICVFXFlags::DisableChromakey))
+					{
+						bAllowChromakey = true;
+					}
 				}
 			}
 
@@ -296,6 +299,26 @@ bool FDisplayClusterViewportConfigurationICVFX::CreateLightcardViewport(FDisplay
 	return false;
 }
 
+bool FDisplayClusterViewportConfigurationICVFX::CreateUVLightcardViewport(FDisplayClusterViewport& BaseViewport)
+{
+	FDisplayClusterViewport* UVLightCardViewport = FDisplayClusterViewportConfigurationHelpers_ICVFX::GetOrCreateUVLightcardViewport(BaseViewport, RootActor);
+	if (UVLightCardViewport)
+	{
+		// Update UV LightCard viewport settings
+		FDisplayClusterViewportConfigurationHelpers_ICVFX::UpdateLightcardViewportSetting(*UVLightCardViewport, BaseViewport, RootActor);
+
+		// Support projection policy update
+		FDisplayClusterViewportConfigurationHelpers::UpdateProjectionPolicy(*UVLightCardViewport);
+
+		// Optimize: re-use UVLightCard viewports with equals OCIO
+		FDisplayClusterViewportConfigurationHelpers_ICVFX::ReuseUVLightCardViewportWithinClusterNode(*UVLightCardViewport);
+
+		return true;
+	}
+
+	return false;
+}
+
 void FDisplayClusterViewportConfigurationICVFX::Update()
 {
 	const FDisplayClusterConfigurationICVFX_StageSettings& StageSettings = RootActor.GetStageSettings();
@@ -325,7 +348,7 @@ void FDisplayClusterViewportConfigurationICVFX::Update()
 
 			// If not all viewports disable camera:
 			// Collect all ICVFX cameras from stage
-			if ((TargetViewportsFlags & ViewportICVFX_DisableCamera) == 0)
+			if (!EnumHasAnyFlags(TargetViewportsFlags, EDisplayClusterViewportICVFXFlags::DisableCamera))
 			{
 				ImplGetCameras();
 			}
@@ -337,7 +360,7 @@ void FDisplayClusterViewportConfigurationICVFX::Update()
 				for (FDisplayClusterViewport* TargetIt : TargetViewports)
 				{
 					// Target viewpot must support camera render:
-					if ((TargetIt->RenderSettingsICVFX.Flags & ViewportICVFX_DisableCamera) == 0)
+					if (!EnumHasAnyFlags(TargetIt->RenderSettingsICVFX.Flags, EDisplayClusterViewportICVFXFlags::DisableCamera))
 					{
 						// Add this target to all cameras visible on it
 						for (FDisplayClusterViewportConfigurationCameraICVFX& CameraIt : StageCameras)
@@ -362,17 +385,28 @@ void FDisplayClusterViewportConfigurationICVFX::Update()
 			}
 
 			// If not all viewports disable lightcard
-			if ((TargetViewportsFlags & ViewportICVFX_DisableLightcard) == 0)
+			if (!EnumHasAnyFlags(TargetViewportsFlags, EDisplayClusterViewportICVFXFlags::DisableLightcard))
 			{
 				// Allocate and assign lightcard resources
-				if (FDisplayClusterViewportConfigurationHelpers_ICVFX::IsShouldUseLightcard(StageSettings.Lightcard))
 				{
+
+					const bool bUseLightCard = FDisplayClusterViewportConfigurationHelpers_ICVFX::IsShouldUseLightcard(StageSettings.Lightcard);
+					const bool bUseUVLightCard = FDisplayClusterViewportConfigurationHelpers_ICVFX::IsShouldUseUVLightcard(*ViewportManager, StageSettings.Lightcard);
+
 					for (FDisplayClusterViewport* TargetIt : TargetViewports)
 					{
 						// only for support targets
-						if (TargetIt && (TargetIt->RenderSettingsICVFX.Flags & ViewportICVFX_DisableLightcard) == 0)
+						if (TargetIt && !EnumHasAnyFlags(TargetIt->RenderSettingsICVFX.Flags, EDisplayClusterViewportICVFXFlags::DisableLightcard))
 						{
-							CreateLightcardViewport(*TargetIt);
+							if (bUseLightCard)
+							{
+								CreateLightcardViewport(*TargetIt);
+							}
+
+							if (bUseUVLightCard)
+							{
+								CreateUVLightcardViewport(*TargetIt);
+							}
 						}
 					}
 				}
@@ -398,11 +432,9 @@ void FDisplayClusterViewportConfigurationICVFX::Update()
 				// Freeze render for lightcards when outer viewports freezed
 				if (StageSettings.Lightcard.bIgnoreOuterViewportsFreezingForLightcards == false)
 				{
-					const EDisplayClusterViewportRuntimeICVFXFlags LightcardViewportMask = ViewportRuntime_ICVFXLightcard;
-
 					for (FDisplayClusterViewport* ViewportIt : ViewportManager->ImplGetViewports())
 					{
-						if (ViewportIt && (ViewportIt->RenderSettingsICVFX.RuntimeFlags & LightcardViewportMask) != 0)
+						if (ViewportIt && EnumHasAnyFlags(ViewportIt->RenderSettingsICVFX.RuntimeFlags, EDisplayClusterViewportRuntimeICVFXFlags::Lightcard | EDisplayClusterViewportRuntimeICVFXFlags::UVLightcard))
 						{
 							ViewportIt->RenderSettings.bFreezeRendering = true;
 						}
@@ -476,9 +508,9 @@ void FDisplayClusterViewportConfigurationICVFX::UpdateHideList(FDisplayClusterVi
 	// Collect viewports, that use ICVFX hide list
 	for (FDisplayClusterViewport* ViewportIt : ViewportManager.ImplGetViewports())
 	{
-		const bool bInternalResource = (ViewportIt->RenderSettingsICVFX.RuntimeFlags & ViewportRuntime_InternalResource) != 0;
-		const bool bIsInCamera = (ViewportIt->RenderSettingsICVFX.RuntimeFlags & ViewportRuntime_ICVFXIncamera) != 0;
-		const bool bICVFX_Enable = (ViewportIt->RenderSettingsICVFX.Flags & ViewportICVFX_Enable) != 0;
+		const bool bInternalResource = EnumHasAnyFlags(ViewportIt->RenderSettingsICVFX.RuntimeFlags, EDisplayClusterViewportRuntimeICVFXFlags::InternalResource);
+		const bool bIsInCamera = EnumHasAnyFlags(ViewportIt->RenderSettingsICVFX.RuntimeFlags, EDisplayClusterViewportRuntimeICVFXFlags::InCamera);
+		const bool bICVFX_Enable = EnumHasAnyFlags(ViewportIt->RenderSettingsICVFX.Flags, EDisplayClusterViewportICVFXFlags::Enable);
 
 		if ((bICVFX_Enable && !bInternalResource)  || (bInternalResource && bIsInCamera))
 		{
@@ -495,10 +527,10 @@ void FDisplayClusterViewportConfigurationICVFX::ImplBeginReallocateViewports(FDi
 	for (FDisplayClusterViewport* ViewportIt : ViewportManager.ImplGetViewports())
 	{
 		// Runtime icvfx viewport support reallocate feature:
-		if ((ViewportIt->RenderSettingsICVFX.RuntimeFlags & ViewportRuntime_InternalResource) != 0)
+		if (EnumHasAllFlags(ViewportIt->RenderSettingsICVFX.RuntimeFlags, EDisplayClusterViewportRuntimeICVFXFlags::InternalResource))
 		{
 			// Mark all dynamic ICVFX viewports for delete
-			ViewportIt->RenderSettingsICVFX.RuntimeFlags |= ViewportRuntime_Unused;
+			EnumAddFlags(ViewportIt->RenderSettingsICVFX.RuntimeFlags, EDisplayClusterViewportRuntimeICVFXFlags::Unused);
 		}
 	}
 }
@@ -510,7 +542,7 @@ void FDisplayClusterViewportConfigurationICVFX::ImplFinishReallocateViewports(FD
 	// Collect all unused viewports for remove
 	for (FDisplayClusterViewport* ViewportIt : ViewportManager.ImplGetViewports())
 	{
-		if ((ViewportIt->RenderSettingsICVFX.RuntimeFlags & ViewportRuntime_Unused) != 0)
+		if (EnumHasAllFlags(ViewportIt->RenderSettingsICVFX.RuntimeFlags, EDisplayClusterViewportRuntimeICVFXFlags::Unused))
 		{
 			UnusedViewports.Add(ViewportIt);
 		}
@@ -551,16 +583,16 @@ void FDisplayClusterViewportConfigurationICVFX::ImplGetCameras()
 
 EDisplayClusterViewportICVFXFlags FDisplayClusterViewportConfigurationICVFX::ImplGetTargetViewports(FDisplayClusterViewportManager& ViewportManager, TArray<class FDisplayClusterViewport*>& OutTargets)
 {
-	EDisplayClusterViewportICVFXFlags InvFlags = ViewportICVFX_None;
+	EDisplayClusterViewportICVFXFlags InvFlags = EDisplayClusterViewportICVFXFlags::None;
 
 	// Collect invertet disable flags from all target viewports
 	for (FDisplayClusterViewport* ViewportIt : ViewportManager.ImplGetViewports())
 	{
 		// Process only external viewports:
-		if ((ViewportIt->RenderSettingsICVFX.RuntimeFlags & ViewportRuntime_InternalResource) == 0)
+		if (!EnumHasAnyFlags(ViewportIt->RenderSettingsICVFX.RuntimeFlags, EDisplayClusterViewportRuntimeICVFXFlags::InternalResource))
 		{
 			//Raise new projection target if possible
-			if (ViewportIt->RenderSettings.bEnable && (ViewportIt->RenderSettingsICVFX.Flags & ViewportICVFX_Enable) != 0)
+			if (ViewportIt->RenderSettings.bEnable && EnumHasAnyFlags(ViewportIt->RenderSettingsICVFX.Flags, EDisplayClusterViewportICVFXFlags::Enable))
 			{
 				if (ViewportIt->ProjectionPolicy.IsValid() && ViewportIt->ProjectionPolicy->ShouldSupportICVFX())
 				{
@@ -568,7 +600,7 @@ EDisplayClusterViewportICVFXFlags FDisplayClusterViewportConfigurationICVFX::Imp
 					OutTargets.Add(ViewportIt);
 
 					// proj policy support ICVFX, Use this viewport as icvfx target
-					ViewportIt->RenderSettingsICVFX.RuntimeFlags |= ViewportRuntime_ICVFXTarget;
+					EnumAddFlags(ViewportIt->RenderSettingsICVFX.RuntimeFlags, EDisplayClusterViewportRuntimeICVFXFlags::Target);
 
 					// Update targets use flags:
 					InvFlags |= ~(ViewportIt->RenderSettingsICVFX.Flags);
diff --git a/Engine/Plugins/Runtime/nDisplay/Source/DisplayCluster/Private/Render/Viewport/Configuration/DisplayClusterViewportConfigurationICVFX.h b/Engine/Plugins/Runtime/nDisplay/Source/DisplayCluster/Private/Render/Viewport/Configuration/DisplayClusterViewportConfigurationICVFX.h
--- a/Engine/Plugins/Runtime/nDisplay/Source/DisplayCluster/Private/Render/Viewport/Configuration/DisplayClusterViewportConfigurationICVFX.h
+++ b/Engine/Plugins/Runtime/nDisplay/Source/DisplayCluster/Private/Render/Viewport/Configuration/DisplayClusterViewportConfigurationICVFX.h
@@ -33,8 +33,12 @@ private:
 	void ImplGetCameras();
 	EDisplayClusterViewportICVFXFlags ImplGetTargetViewports(FDisplayClusterViewportManager& ViewportManager, TArray<FDisplayClusterViewport*>& OutTargets);
 
+	/** Create a LightCard viewport for the  BaseViewport. */
 	bool CreateLightcardViewport(FDisplayClusterViewport& BaseViewport);
 
+	/** Create an UVLightCard viewport for the BaseViewport. */
+	bool CreateUVLightcardViewport(FDisplayClusterViewport& BaseViewport);
+
 	void UpdateHideList(FDisplayClusterViewportManager& ViewportManager);
 	void UpdateCameraHideList(FDisplayClusterViewportManager& ViewportManager);
 
diff --git a/Engine/Plugins/Runtime/nDisplay/Source/DisplayCluster/Private/Render/Viewport/Configuration/DisplayClusterViewportConfigurationProjectionPolicy.cpp b/Engine/Plugins/Runtime/nDisplay/Source/DisplayCluster/Private/Render/Viewport/Configuration/DisplayClusterViewportConfigurationProjectionPolicy.cpp
--- a/Engine/Plugins/Runtime/nDisplay/Source/DisplayCluster/Private/Render/Viewport/Configuration/DisplayClusterViewportConfigurationProjectionPolicy.cpp
+++ b/Engine/Plugins/Runtime/nDisplay/Source/DisplayCluster/Private/Render/Viewport/Configuration/DisplayClusterViewportConfigurationProjectionPolicy.cpp
@@ -31,7 +31,7 @@ void FDisplayClusterViewportConfigurationProjectionPolicy::Update()
 	for (FDisplayClusterViewport* ViewportIt : ViewportManager.ImplGetViewports())
 	{
 		// ignore ICVFX internal resources
-		if (ViewportIt && (ViewportIt->GetRenderSettingsICVFX().RuntimeFlags & ViewportRuntime_InternalResource) == 0)
+		if (ViewportIt && !EnumHasAnyFlags(ViewportIt->GetRenderSettingsICVFX().RuntimeFlags, EDisplayClusterViewportRuntimeICVFXFlags::InternalResource))
 		{
 			// Support advanced logic for 'camera' projection policy
 			if (ViewportIt->ProjectionPolicy.IsValid() && ViewportIt->ProjectionPolicy->GetType().Compare(DisplayClusterProjectionStrings::projection::Camera) == 0)
@@ -115,7 +115,9 @@ bool FDisplayClusterViewportConfigurationProjectionPolicy::UpdateCameraPolicy_IC
 	FDisplayClusterViewport* ExistCameraViewport = FDisplayClusterViewportConfigurationHelpers_ICVFX::FindCameraViewport(RootActor, *InICVFXCameraComponent);
 	if (ExistCameraViewport)
 	{
-		DstViewport.RenderSettings.OverrideViewportId = ExistCameraViewport->GetId();
+		// Re-use camera RTT, and support OCIO:
+		const EDisplayClusterViewportOverrideMode ViewportOverrideMode = DstViewport.IsOpenColorIOEquals(*ExistCameraViewport) ? EDisplayClusterViewportOverrideMode::All : EDisplayClusterViewportOverrideMode::InernalRTT;
+		DstViewport.RenderSettings.SetViewportOverride(ExistCameraViewport->GetId(), ViewportOverrideMode);
 	}
 	else
 	{
diff --git a/Engine/Plugins/Runtime/nDisplay/Source/DisplayCluster/Private/Render/Viewport/Containers/DisplayClusterViewportProxyData.cpp b/Engine/Plugins/Runtime/nDisplay/Source/DisplayCluster/Private/Render/Viewport/Containers/DisplayClusterViewportProxyData.cpp
--- a/Engine/Plugins/Runtime/nDisplay/Source/DisplayCluster/Private/Render/Viewport/Containers/DisplayClusterViewportProxyData.cpp
+++ b/Engine/Plugins/Runtime/nDisplay/Source/DisplayCluster/Private/Render/Viewport/Containers/DisplayClusterViewportProxyData.cpp
@@ -15,6 +15,8 @@ FDisplayClusterViewportProxyData::FDisplayClusterViewportProxyData(const FDispla
 
 	DstViewportProxy = SrcViewport->ViewportProxy;
 
+	OpenColorIO = SrcViewport->OpenColorIO;
+
 	RenderSettings = SrcViewport->RenderSettings;
 	RenderSettingsICVFX.SetParameters(SrcViewport->RenderSettingsICVFX);
 	PostRenderSettings.SetParameters(SrcViewport->PostRenderSettings);
@@ -48,6 +50,8 @@ void FDisplayClusterViewportProxyData::UpdateProxy_RenderThread() const
 	check(IsInRenderingThread());
 	check(DstViewportProxy);
 
+	DstViewportProxy->OpenColorIO = OpenColorIO;
+
 	DstViewportProxy->OverscanSettings = OverscanSettings;
 
 	DstViewportProxy->RemapMesh = RemapMesh;
diff --git a/Engine/Plugins/Runtime/nDisplay/Source/DisplayCluster/Private/Render/Viewport/Containers/DisplayClusterViewportProxyData.h b/Engine/Plugins/Runtime/nDisplay/Source/DisplayCluster/Private/Render/Viewport/Containers/DisplayClusterViewportProxyData.h
--- a/Engine/Plugins/Runtime/nDisplay/Source/DisplayCluster/Private/Render/Viewport/Containers/DisplayClusterViewportProxyData.h
+++ b/Engine/Plugins/Runtime/nDisplay/Source/DisplayCluster/Private/Render/Viewport/Containers/DisplayClusterViewportProxyData.h
@@ -14,6 +14,7 @@ class FDisplayClusterViewport;
 class FDisplayClusterViewportProxy;
 class FDisplayClusterViewport_Context;
 class IDisplayClusterRender_MeshComponent;
+class FDisplayClusterViewport_OpenColorIO;
 
 //
 // Container for data exchange game->render threads
@@ -38,6 +39,9 @@ private:
 	FDisplayClusterViewport_RenderSettingsICVFX  RenderSettingsICVFX;
 	FDisplayClusterViewport_PostRenderSettings   PostRenderSettings;
 
+	// OpenColorIO
+	TSharedPtr<FDisplayClusterViewport_OpenColorIO, ESPMode::ThreadSafe> OpenColorIO;
+
 	// Projection policy instance that serves this viewport
 	TSharedPtr<IDisplayClusterProjectionPolicy, ESPMode::ThreadSafe> ProjectionPolicy;
 
diff --git a/Engine/Plugins/Runtime/nDisplay/Source/DisplayCluster/Private/Render/Viewport/DisplayClusterViewport.cpp b/Engine/Plugins/Runtime/nDisplay/Source/DisplayCluster/Private/Render/Viewport/DisplayClusterViewport.cpp
--- a/Engine/Plugins/Runtime/nDisplay/Source/DisplayCluster/Private/Render/Viewport/DisplayClusterViewport.cpp
+++ b/Engine/Plugins/Runtime/nDisplay/Source/DisplayCluster/Private/Render/Viewport/DisplayClusterViewport.cpp
@@ -5,8 +5,11 @@
 #include "Render/Viewport/DisplayClusterViewportManager.h"
 #include "Render/Viewport/DisplayClusterViewportProxy.h"
 #include "Render/Viewport/DisplayClusterViewportManagerProxy.h"
+#include "Render/Viewport/DisplayClusterViewport_OpenColorIO.h"
+#include "Render/Viewport/DisplayClusterViewportManagerViewExtension.h"
 
 #include "Render/Viewport/Configuration/DisplayClusterViewportConfiguration.h"
+#include "Render/Viewport/LightCard/DisplayClusterViewportLightCardManager.h"
 
 #include "Render/Projection/IDisplayClusterProjectionPolicy.h"
 
@@ -69,6 +72,7 @@ FDisplayClusterViewport::~FDisplayClusterViewport()
 	);
 
 	ViewportProxy.Reset();
+	OpenColorIO.Reset();
 
 	// Handle projection policy EndScene event
 	HandleEndScene();
@@ -91,53 +95,90 @@ const FDisplayClusterRenderFrameSettings& FDisplayClusterViewport::GetRenderFram
 	return Owner.GetRenderFrameSettings();
 }
 
-const TArray<FSceneViewExtensionRef> FDisplayClusterViewport::GatherActiveExtensions(FViewport* InViewport) const
+bool FDisplayClusterViewport::IsOpenColorIOEquals(const FDisplayClusterViewport& InViewport) const
 {
-	if (InViewport)
-	{
-		FDisplayClusterSceneViewExtensionContext ViewExtensionContext(InViewport, &Owner, GetId());
-		return GEngine->ViewExtensions->GatherActiveExtensions(ViewExtensionContext);
-	}
-	else
+	bool bEnabledOCIO_1 = OpenColorIO.IsValid();
+	bool bEnabledOCIO_2 = InViewport.OpenColorIO.IsValid();
+
+	if (bEnabledOCIO_1 == bEnabledOCIO_2)
 	{
-		UWorld* CurrentWorld = Owner.GetCurrentWorld();
-		if (CurrentWorld)
+		if (!bEnabledOCIO_1)
 		{
-			FDisplayClusterSceneViewExtensionContext ViewExtensionContext(CurrentWorld->Scene, &Owner, GetId());
-			return GEngine->ViewExtensions->GatherActiveExtensions(ViewExtensionContext);
+			// Both OCIO disabled
+			return true;
+		}
+
+		if (OpenColorIO->IsDisplayConfigurationEquals(InViewport.OpenColorIO->GetDisplayConfiguration()))
+		{
+			return true;
 		}
 	}
 
-	return TArray<FSceneViewExtensionRef>();
+	return false;
 }
 
-FSceneViewExtensionIsActiveFunctor FDisplayClusterViewport::GetSceneViewExtensionIsActiveFunctor() const
+const TArray<FSceneViewExtensionRef> FDisplayClusterViewport::GatherActiveExtensions(FViewport* InViewport) const
 {
-	FSceneViewExtensionIsActiveFunctor IsActiveFunction;
-	IsActiveFunction.IsActiveFunction = [this](const ISceneViewExtension* SceneViewExtension, const FSceneViewExtensionContext& Context)
+	// Use VE from engine for default render and MRQ:
+	switch (RenderSettings.CaptureMode)
 	{
-		if (Context.IsA(FDisplayClusterSceneViewExtensionContext()))
+	case EDisplayClusterViewportCaptureMode::Default:
+	case EDisplayClusterViewportCaptureMode::MoviePipeline:
+		if (InViewport)
 		{
-			const FDisplayClusterSceneViewExtensionContext& DisplayContext = static_cast<const FDisplayClusterSceneViewExtensionContext&>(Context);
-
-			// Find exist viewport by name
-			IDisplayClusterViewport* PublicViewport = DisplayContext.ViewportManager->FindViewport(DisplayContext.ViewportId);
-			if (PublicViewport)
+			FDisplayClusterSceneViewExtensionContext ViewExtensionContext(InViewport, &Owner, GetId());
+			return GEngine->ViewExtensions->GatherActiveExtensions(ViewExtensionContext);
+		}
+		else
+		{
+			UWorld* CurrentWorld = Owner.GetCurrentWorld();
+			if (CurrentWorld)
 			{
-				FDisplayClusterViewport* Viewport = static_cast<FDisplayClusterViewport*>(PublicViewport);
-
-				if (Viewport->OpenColorIODisplayExtension.IsValid() && Viewport->OpenColorIODisplayExtension.Get() == SceneViewExtension)
-				{
-					// This viewport use this OCIO extension
-					return  TOptional<bool>(true);
-				}
+				FDisplayClusterSceneViewExtensionContext ViewExtensionContext(CurrentWorld->Scene, &Owner, GetId());
+				return GEngine->ViewExtensions->GatherActiveExtensions(ViewExtensionContext);
 			}
 		}
 
-		return TOptional<bool>(false);
+		// No extension found.
+		return TArray<FSceneViewExtensionRef>();
+		
+	default:
+		break;
+	}
+
+	// Get custom VE:
+	TArray<FSceneViewExtensionRef> OutCustomExtensions;
+
+	// Initialize custom extensions:
+	switch (RenderSettings.CaptureMode)
+	{
+	case EDisplayClusterViewportCaptureMode::Chromakey:
+	case EDisplayClusterViewportCaptureMode::Lightcard:
+	{
+		// Chromakey and LightCard use only nDisplay VE for callback purposes (preserve alpha channel, etc).
+		TSharedPtr<FDisplayClusterViewportManagerViewExtension, ESPMode::ThreadSafe> ViewportManagerViewExtension = Owner.GetViewportManagerViewExtension();
+		if (ViewportManagerViewExtension.IsValid())
+		{
+			OutCustomExtensions.Add(ViewportManagerViewExtension->AsShared());
+		}
+	}
+	break;
+
+	default:
+		break;
+	}
+
+	// Sort extensions in order of priority (copied from FSceneViewExtensions::GatherActiveExtensions)
+	struct SortPriority
+	{
+		bool operator () (const FSceneViewExtensionRef& A, const FSceneViewExtensionRef& B) const
+		{
+			return A->GetPriority() > B->GetPriority();
+		}
 	};
+	Sort(OutCustomExtensions.GetData(), OutCustomExtensions.Num(), SortPriority());
 
-	return IsActiveFunction;
+	return OutCustomExtensions;
 }
 
 bool FDisplayClusterViewport::HandleStartScene()
@@ -249,6 +290,11 @@ void FDisplayClusterViewport::SetupSceneView(uint32 ContextNum, class UWorld* Wo
 		return;
 	}
 
+	if (OpenColorIO.IsValid())
+	{
+		OpenColorIO->SetupSceneView(InOutViewFamily, InOutView);
+	}
+
 	if(Contexts[ContextNum].GPUIndex >= 0)
 	{
 		// Use custom GPUIndex for render
@@ -268,9 +314,7 @@ void FDisplayClusterViewport::SetupSceneView(uint32 ContextNum, class UWorld* Wo
 	{
 	case EDisplayClusterViewportCaptureMode::Chromakey:
 	case EDisplayClusterViewportCaptureMode::Lightcard:
-
 		InOutView.bAllowRayTracing = false;
-		
 		break;
 
 	default:
@@ -330,11 +374,11 @@ float FDisplayClusterViewport::GetClusterRenderTargetRatioMult(const FDisplayClu
 	float ClusterRenderTargetRatioMult = InFrameSettings.ClusterRenderTargetRatioMult;
 
 	// Support Outer viewport cluster rtt multiplier
-	if ((RenderSettingsICVFX.RuntimeFlags & ViewportRuntime_ICVFXTarget) != 0)
+	if (EnumHasAllFlags(RenderSettingsICVFX.RuntimeFlags, EDisplayClusterViewportRuntimeICVFXFlags::Target))
 	{
 		ClusterRenderTargetRatioMult *= InFrameSettings.ClusterICVFXOuterViewportRenderTargetRatioMult;
 	}
-	else if ((RenderSettingsICVFX.RuntimeFlags & ViewportRuntime_ICVFXIncamera) != 0)
+	else if (EnumHasAllFlags(RenderSettingsICVFX.RuntimeFlags, EDisplayClusterViewportRuntimeICVFXFlags::InCamera))
 	{
 		ClusterRenderTargetRatioMult *= InFrameSettings.ClusterICVFXInnerViewportRenderTargetRatioMult;
 	}
@@ -368,12 +412,12 @@ float FDisplayClusterViewport::GetCustomBufferRatio(const FDisplayClusterRenderF
 	// Global multiplier
 	CustomBufferRatio *= InFrameSettings.ClusterBufferRatioMult;
 
-	if ((RenderSettingsICVFX.RuntimeFlags & ViewportRuntime_ICVFXTarget) != 0)
+	if (EnumHasAllFlags(RenderSettingsICVFX.RuntimeFlags, EDisplayClusterViewportRuntimeICVFXFlags::Target))
 	{
 		// Outer viewport
 		CustomBufferRatio *= InFrameSettings.ClusterICVFXOuterViewportBufferRatioMult;
 	}
-	else if ((RenderSettingsICVFX.RuntimeFlags & ViewportRuntime_ICVFXIncamera) != 0)
+	else if (EnumHasAllFlags(RenderSettingsICVFX.RuntimeFlags, EDisplayClusterViewportRuntimeICVFXFlags::InCamera))
 	{
 		// Inner Frustum
 		CustomBufferRatio *= InFrameSettings.ClusterICVFXInnerFrustumBufferRatioMult;
@@ -382,18 +426,37 @@ float FDisplayClusterViewport::GetCustomBufferRatio(const FDisplayClusterRenderF
 	return CustomBufferRatio;
 }
 
+void FDisplayClusterViewport::ResetFrameContexts()
+{
+
+#if WITH_EDITOR
+	OutputPreviewTargetableResource.SafeRelease();
+#endif
+
+	// Discard resources that are not used in frame composition
+	RenderTargets.Empty();
+	OutputFrameTargetableResources.Empty();
+	AdditionalFrameTargetableResources.Empty();
+
+	// Release old contexts
+	Contexts.Empty();
+
+	// Free internal resources
+	InputShaderResources.Empty();
+	AdditionalTargetableResources.Empty();
+	MipsShaderResources.Empty();
+}
+
 bool FDisplayClusterViewport::UpdateFrameContexts(const uint32 InStereoViewIndex, const FDisplayClusterRenderFrameSettings& InFrameSettings)
 {
 	check(IsInGameThread());
 
-	uint32 FrameTargetsAmount = 2;
+	const uint32 FrameTargetsAmount = Owner.GetViewPerViewportAmount();
+
 	FIntRect     DesiredFrameTargetRect = RenderSettings.Rect;
 	{
 		switch (InFrameSettings.RenderMode)
 		{
-		case EDisplayClusterRenderFrameMode::Mono:
-			FrameTargetsAmount = 1;
-			break;
 		case EDisplayClusterRenderFrameMode::SideBySide:
 			AdjustRect(DesiredFrameTargetRect, 0.5f, 1.f);
 			break;
@@ -409,8 +472,6 @@ bool FDisplayClusterViewport::UpdateFrameContexts(const uint32 InStereoViewIndex
 			// Align each frame to zero
 			DesiredFrameTargetRect = FIntRect(FIntPoint(0, 0), DesiredFrameTargetRect.Size());
 
-			// Mono
-			FrameTargetsAmount = 1;
 			break;
 		}
 		default:
@@ -419,7 +480,7 @@ bool FDisplayClusterViewport::UpdateFrameContexts(const uint32 InStereoViewIndex
 	}
 
 	// Special case mono->stereo
-	uint32 ViewportContextAmount = RenderSettings.bForceMono ? 1 : FrameTargetsAmount;
+	const uint32 ViewportContextAmount = RenderSettings.bForceMono ? 1 : FrameTargetsAmount;
 
 #if WITH_EDITOR
 	OutputPreviewTargetableResource.SafeRelease();
@@ -534,12 +595,62 @@ bool FDisplayClusterViewport::UpdateFrameContexts(const uint32 InStereoViewIndex
 	}
 
 	bool bDisableInternalResources = false;
-	if (RenderSettings.bSkipRendering || RenderSettings.OverrideViewportId.IsEmpty() == false)
+	if (RenderSettings.bSkipRendering)
 	{
 		bDisableInternalResources = true;
 		bDisableRender = true;
 	}
 
+	if (RenderSettings.IsViewportOverrided())
+	{
+		switch (RenderSettings.ViewportOverrideMode)
+		{
+		case EDisplayClusterViewportOverrideMode::InernalRTT:
+			bDisableRender = true;
+			break;
+
+		case EDisplayClusterViewportOverrideMode::All:
+			bDisableRender = true;
+			bDisableInternalResources = true;
+			break;
+
+		default:
+			break;
+		}
+	}
+
+	// UV LightCard viewport use unique whole-cluster texture from LC manager
+	if (EnumHasAllFlags(RenderSettingsICVFX.RuntimeFlags, EDisplayClusterViewportRuntimeICVFXFlags::UVLightcard))
+	{
+		// Use external texture from LightCardManager instead of rendering
+		bDisableRender = true;
+
+		// Use the UVLightCard viewport only when this type of lightcards has been defined
+		bool bUseUVLightCardViewport = false;
+
+		TSharedPtr<FDisplayClusterViewportLightCardManager, ESPMode::ThreadSafe> LightCardManager = Owner.GetLightCardManager();
+		if (LightCardManager.IsValid() && LightCardManager->IsUVLightCardEnabled())
+		{
+			// Custom viewport size from LC Manager
+			ContextSize = LightCardManager->GetUVLightCardResourceSize();
+
+			// Size must be not null
+			if (ContextSize.GetMin() > 1)
+			{
+				FrameTargetRect = RenderTargetRect = FIntRect(FIntPoint(0, 0), ContextSize);
+
+				// Allow to use this viewport
+				bUseUVLightCardViewport = true;
+			}
+		}
+
+		if(!bUseUVLightCardViewport)
+		{
+			// do not use UV LightCard viewport
+			return false;
+		}
+	}
+
 	//Add new contexts
 	for (uint32 ContextIt = 0; ContextIt < ViewportContextAmount; ++ContextIt)
 	{
