nDisplay: ICVFX OCIO

diff --git a/Engine/Plugins/Runtime/nDisplay/Source/DisplayCluster/Private/Render/Viewport/DisplayClusterViewport.h b/Engine/Plugins/Runtime/nDisplay/Source/DisplayCluster/Private/Render/Viewport/DisplayClusterViewport.h
--- a/Engine/Plugins/Runtime/nDisplay/Source/DisplayCluster/Private/Render/Viewport/DisplayClusterViewport.h
+++ b/Engine/Plugins/Runtime/nDisplay/Source/DisplayCluster/Private/Render/Viewport/DisplayClusterViewport.h
@@ -12,9 +12,10 @@
 
 #include "Render/Viewport/DisplayClusterViewport_CustomPostProcessSettings.h"
 #include "Render/Viewport/DisplayClusterViewport_VisibilitySettings.h"
+#include "Render/Viewport/DisplayClusterViewport_OpenColorIO.h"
 
+#include "SceneViewExtension.h"
 #include "SceneViewExtensionContext.h"
-#include "OpenColorIODisplayExtension.h"
 
 class FDisplayClusterViewportRenderTargetResource;
 class FDisplayClusterViewportTextureResource;
@@ -187,8 +188,21 @@ public:
 	// Active view extension for this viewport
 	const TArray<FSceneViewExtensionRef> GatherActiveExtensions(FViewport* InViewport) const;
 
+	/** Initialize viewport contexts and resources for new frame
+	 *
+	 * @param InStereoViewIndex - initial StereoViewIndex for this viewport
+	 * @param InFrameSettings   - New settings for current frame
+	 *
+	 * @return - true, if success
+	 */
 	bool UpdateFrameContexts(const uint32 InStereoViewIndex, const FDisplayClusterRenderFrameSettings& InFrameSettings);
 
+	/** Reset viewport contexts and resources. */
+	void ResetFrameContexts();
+
+	/** Compare OCIO with another viewport, return true if they are equal. */
+	bool IsOpenColorIOEquals(const FDisplayClusterViewport& InViewport) const;
+
 public:
 	FIntRect GetValidRect(const FIntRect& InRect, const TCHAR* DbgSourceName);
 
@@ -208,11 +222,8 @@ private:
 #endif
 
 public:
-	// Support OCIO:
-	FSceneViewExtensionIsActiveFunctor GetSceneViewExtensionIsActiveFunctor() const;
-
-	// OCIO wrapper
-	TSharedPtr<FOpenColorIODisplayExtension, ESPMode::ThreadSafe> OpenColorIODisplayExtension;
+	/** nDisplay OpenColorIO object. */
+	TSharedPtr<FDisplayClusterViewport_OpenColorIO, ESPMode::ThreadSafe> OpenColorIO;
 
 public:
 	// Projection policy instance that serves this viewport
diff --git a/Engine/Plugins/Runtime/nDisplay/Source/DisplayCluster/Private/Render/Viewport/DisplayClusterViewportManager.cpp b/Engine/Plugins/Runtime/nDisplay/Source/DisplayCluster/Private/Render/Viewport/DisplayClusterViewportManager.cpp
--- a/Engine/Plugins/Runtime/nDisplay/Source/DisplayCluster/Private/Render/Viewport/DisplayClusterViewportManager.cpp
+++ b/Engine/Plugins/Runtime/nDisplay/Source/DisplayCluster/Private/Render/Viewport/DisplayClusterViewportManager.cpp
@@ -49,22 +49,6 @@ static FAutoConsoleVariableRef CVarDisplayClusterChromaKeyAllowNanite(
 	ECVF_RenderThreadSafe
 );
 
-int32 GDisplayClusterLightcardsAllowAA = 0;
-static FAutoConsoleVariableRef CVarDisplayClusterLightcardsAllowAA(
-	TEXT("DC.Lightcards.AllowAA"),
-	GDisplayClusterLightcardsAllowAA,
-	TEXT("0 disables AntiAliasing when rendering lightcards.Otherwise uses default showflag."),
-	ECVF_RenderThreadSafe
-);
-
-int32 GDisplayClusterChromaKeyAllowAA = 0;
-static FAutoConsoleVariableRef CVarDisplayClusterChromaKeyAllowAA(
-	TEXT("DC.ChromaKey.AllowAA"),
-	GDisplayClusterChromaKeyAllowAA,
-	TEXT("0 disables AntiAliasing when rendering custom chroma keys. Otherwise uses default showflag."),
-	ECVF_RenderThreadSafe
-);
-
 ///////////////////////////////////////////////////////////////////////////////////////
 //          FDisplayClusterViewportManager
 ///////////////////////////////////////////////////////////////////////////////////////
@@ -423,14 +407,11 @@ void FDisplayClusterViewportManager::ImplUpdateClusterNodeViewports(const EDispl
 			{
 				if (Viewport && (InClusterNodeId.IsEmpty() || Viewport->GetClusterNodeId() == InClusterNodeId))
 				{
-					bool bHasParentViewport = Viewport->RenderSettings.GetParentViewportId().IsEmpty() == false;
-					bool bHasViewportOverride = Viewport->RenderSettings.OverrideViewportId.IsEmpty() == false;
-
-					if (bHasViewportOverride)
+					if (Viewport->RenderSettings.IsViewportOverrided())
 					{
 						OverriddenViewports.Add(Viewport);
 					}
-					else if (bHasParentViewport)
+					else if (Viewport->RenderSettings.IsViewportHasParent())
 					{
 						LinkedViewports.Add(Viewport);
 					}
@@ -452,6 +433,36 @@ void FDisplayClusterViewportManager::ImplUpdateClusterNodeViewports(const EDispl
 	}
 }
 
+int32 FDisplayClusterViewportManager::GetViewPerViewportAmount() const
+{
+	switch (GetRenderFrameSettings().RenderMode)
+	{
+	case EDisplayClusterRenderFrameMode::Mono:
+	case EDisplayClusterRenderFrameMode::PreviewInScene:
+		return 1;
+		break;
+	default:
+		break;
+	}
+
+	return 2;
+}
+
+int32 FDisplayClusterViewportManager::FindFirstViewportStereoViewIndex(FDisplayClusterViewport* InViewport) const
+{
+	const int32 ViewportIndex = Viewports.Find(InViewport);
+	if (ViewportIndex != INDEX_NONE)
+	{
+		const int32 ViewPerViewportAmount = GetViewPerViewportAmount();
+		// Begin from 1, because INDEX_NONE use ViewState[0] in LocalPlayer.cpp:786
+		const int32 FirstViewportStereoViewIndex = (ViewportIndex * ViewPerViewportAmount) + 1;
+
+		return FirstViewportStereoViewIndex;
+	}
+
+	return INDEX_NONE;
+}
+
 bool FDisplayClusterViewportManager::BeginNewFrame(FViewport* InViewport, UWorld* InWorld, FDisplayClusterRenderFrame& OutRenderFrame)
 {
 	check(IsInGameThread());
@@ -492,24 +503,26 @@ bool FDisplayClusterViewportManager::BeginNewFrame(FViewport* InViewport, UWorld
 		PostProcessManager->HandleSetupNewFrame();
 	}
 
-	// generate unique stereo view index for each frame
-	// Begin from 1, because INDEX_NONE use ViewState[0] in LocalPlayer.cpp:786
-	uint32 StereoViewIndex = 1;
-
 	const FDisplayClusterRenderFrameSettings& RenderFrameSettings = GetRenderFrameSettings();
 
 	// Initialize viewports from new render settings, and create new contexts, reset prev frame resources
 	for (FDisplayClusterViewport* Viewport : ClusterNodeViewports)
 	{
-		// Save orig viewport contexts
-		TArray<FDisplayClusterViewport_Context> PrevContexts;
-		PrevContexts.Append(Viewport->GetContexts());
-		if (Viewport->UpdateFrameContexts(StereoViewIndex, RenderFrameSettings))
+		const int32 FirstViewportStereoViewIndex = FindFirstViewportStereoViewIndex(Viewport);
+		if (FirstViewportStereoViewIndex != INDEX_NONE)
 		{
-			StereoViewIndex += Viewport->Contexts.Num();
-		}
+			// Save orig viewport contexts
+			TArray<FDisplayClusterViewport_Context> PrevContexts;
+			PrevContexts.Append(Viewport->GetContexts());
+			Viewport->UpdateFrameContexts(FirstViewportStereoViewIndex, RenderFrameSettings);
 
-		HandleViewportRTTChanges(PrevContexts, Viewport->GetContexts());
+			HandleViewportRTTChanges(PrevContexts, Viewport->GetContexts());
+		}
+		else
+		{
+			// do not render this viewport
+			Viewport->ResetFrameContexts();
+		}
 	}
 
 	// Handle scene RTT resize
@@ -556,7 +569,13 @@ bool FDisplayClusterViewportManager::BeginNewFrame(FViewport* InViewport, UWorld
 
 void FDisplayClusterViewportManager::UpdateDesiredNumberOfViews(FDisplayClusterRenderFrame& InOutRenderFrame)
 {
-	InOutRenderFrame.DesiredNumberOfViews = 0;
+	// Total views ammount
+	const int32 ViewPerViewportAmount = GetViewPerViewportAmount();
+	const int32 DesiredNumberOfViews = (Viewports.Num() * ViewPerViewportAmount) + 1;
+
+	InOutRenderFrame.DesiredNumberOfViews = DesiredNumberOfViews;
+
+	// Total view on current cluster node
 	InOutRenderFrame.ViewportsAmount = 0;
 
 	for (FDisplayClusterRenderFrame::FFrameRenderTarget& RenderTargetIt : InOutRenderFrame.RenderTargets)
@@ -575,13 +594,30 @@ void FDisplayClusterViewportManager::UpdateDesiredNumberOfViews(FDisplayClusterR
 
 					// Get StereoViewIndex for this viewport
 					const int32 StereoViewIndex = ViewIt.Viewport->GetContexts()[ViewIt.ContextNum].StereoViewIndex;
-					InOutRenderFrame.DesiredNumberOfViews++;
 				}
 			}
 		}
 	}
 }
 
+void FDisplayClusterViewportManager::InitializeNewFrame()
+{
+	check(IsInGameThread());
+
+	// Before render
+	for (FDisplayClusterViewport* Viewport : ClusterNodeViewports)
+	{
+		if (Viewport)
+		{
+			// Initialize OCIO resources
+			if (Viewport->OpenColorIO.IsValid())
+			{
+				Viewport->OpenColorIO->UpdateOpenColorIORenderPassResources();
+			}
+		}
+	}
+}
+
 void FDisplayClusterViewportManager::FinalizeNewFrame()
 {
 	check(IsInGameThread());
@@ -635,16 +671,10 @@ FSceneViewFamily::ConstructionValues FDisplayClusterViewportManager::CreateViewF
 
 	bool bResolveScene = true;
 
-	// Control AA for ChromaKey and Lightcards:
+	// Control NaniteMeshes for ChromaKey and Lightcards:
 	switch (InFrameTarget.CaptureMode)
 	{
 	case EDisplayClusterViewportCaptureMode::Chromakey:
-		if (!GDisplayClusterChromaKeyAllowAA)
-		{
-			InEngineShowFlags.SetAntiAliasing(0);
-			InEngineShowFlags.SetTemporalAA(false);
-		}
-		
 		if (!GDisplayClusterChromaKeyAllowNanite)
 		{
 			InEngineShowFlags.SetNaniteMeshes(0);
@@ -652,15 +682,24 @@ FSceneViewFamily::ConstructionValues FDisplayClusterViewportManager::CreateViewF
 
 		break;
 	case EDisplayClusterViewportCaptureMode::Lightcard:
-		if (!GDisplayClusterLightcardsAllowAA)
+		if (!GDisplayClusterLightcardsAllowNanite)
 		{
-			InEngineShowFlags.SetAntiAliasing(0);
-			InEngineShowFlags.SetTemporalAA(false);
+			InEngineShowFlags.SetNaniteMeshes(0);
 		}
+		break;
+	default:
+		break;
+	}
 
-		if (!GDisplayClusterLightcardsAllowNanite)
+	const FDisplayClusterRenderFrameSettings& RenderFrameSettings = GetRenderFrameSettings();
+	switch (RenderFrameSettings.RenderMode)
+	{
+	case EDisplayClusterRenderFrameMode::PreviewInScene:
+
+		if (RenderFrameSettings.bPreviewEnablePostProcess == false)
 		{
-			InEngineShowFlags.SetNaniteMeshes(0);
+			// Disable postprocess for preview
+			InEngineShowFlags.PostProcessing = 0;
 		}
 		break;
 	default:
@@ -671,8 +710,33 @@ FSceneViewFamily::ConstructionValues FDisplayClusterViewportManager::CreateViewF
 	{
 	case EDisplayClusterViewportCaptureMode::Chromakey:
 	case EDisplayClusterViewportCaptureMode::Lightcard:
-		bResolveScene = false;
+		switch (GetRenderFrameSettings().AlphaChannelCaptureMode)
+		{
+			case EDisplayClusterRenderFrameAlphaChannelCaptureMode::Copy:
+				// Disable AA
+				InEngineShowFlags.SetAntiAliasing(0);
+				InEngineShowFlags.SetTemporalAA(0);
+				break;
+
+			case EDisplayClusterRenderFrameAlphaChannelCaptureMode::CopyAA:
+				// Use AA
+				InEngineShowFlags.SetAntiAliasing(1);
+				InEngineShowFlags.SetTemporalAA(0);
+				break;
+
+			case EDisplayClusterRenderFrameAlphaChannelCaptureMode::FXAA:
+				// Alpha captured without AA, own FXAA used
+				InEngineShowFlags.SetAntiAliasing(0);
+				InEngineShowFlags.SetTemporalAA(0);
+				break;
+
+			default:
+				break;
+		}
+
+		// Disable postprocess for LC\CK
 		InEngineShowFlags.SetPostProcessing(0);
+
 		InEngineShowFlags.SetAtmosphere(0);
 		InEngineShowFlags.SetFog(0);
 		InEngineShowFlags.SetVolumetricFog(0);
@@ -696,72 +760,49 @@ FSceneViewFamily::ConstructionValues FDisplayClusterViewportManager::CreateViewF
 		break;
 	}
 
+	return FSceneViewFamily::ConstructionValues(InFrameTarget.RenderTargetPtr, InScene, InEngineShowFlags)
+		.SetResolveScene(bResolveScene)
+		.SetRealtimeUpdate(true)
+		.SetGammaCorrection(1.0f)
+		.SetAdditionalViewFamily(bInAdditionalViewFamily);
+}
+
+bool FDisplayClusterViewportManager::ShouldRenderFinalColor() const
+{
+#if WITH_EDITOR
 	const FDisplayClusterRenderFrameSettings& RenderFrameSettings = GetRenderFrameSettings();
 	switch (RenderFrameSettings.RenderMode)
 	{
 	case EDisplayClusterRenderFrameMode::PreviewInScene:
+		// Disable postprocess for preview
+		return RenderFrameSettings.bPreviewEnablePostProcess;
 
-		if (RenderFrameSettings.bPreviewEnablePostProcess == false)
-		{
-			// Disable postprocess for preview
-			InEngineShowFlags.PostProcessing = 0;
-		}
-		break;
 	default:
 		break;
 	}
+#endif
 
-	return FSceneViewFamily::ConstructionValues(InFrameTarget.RenderTargetPtr, InScene, InEngineShowFlags)
-		.SetResolveScene(bResolveScene)
-		.SetRealtimeUpdate(true)
-		.SetAdditionalViewFamily(bInAdditionalViewFamily);
+	return true;
 }
 
 void FDisplayClusterViewportManager::ConfigureViewFamily(const FDisplayClusterRenderFrame::FFrameRenderTarget& InFrameTarget, const FDisplayClusterRenderFrame::FFrameViewFamily& InFrameViewFamily, FSceneViewFamilyContext& ViewFamily)
 {
-	ViewFamily.SceneCaptureCompositeMode = ESceneCaptureCompositeMode::SCCM_Overwrite;
-
 	// Note: EngineShowFlags should have already been configured in CreateViewFamilyConstructionValues.
-	switch (InFrameTarget.CaptureMode)
+	ViewFamily.SceneCaptureCompositeMode = ESceneCaptureCompositeMode::SCCM_Overwrite;
+	if (!ShouldRenderFinalColor())
 	{
-		case EDisplayClusterViewportCaptureMode::Chromakey:
-		case EDisplayClusterViewportCaptureMode::Lightcard:
-			ViewFamily.SceneCaptureSource = ESceneCaptureSource::SCS_SceneColorHDR;
-
-			//Do not use view extensions for LightCard and Chromakey.
-			ViewFamily.ViewExtensions.Empty();
-			break;
-
-		default:
-
-			// Gather Scene View Extensions
-			ViewFamily.ViewExtensions = InFrameViewFamily.ViewExtensions;
-			break;
+		ViewFamily.SceneCaptureSource = ESceneCaptureSource::SCS_SceneColorHDR;
 	}
 
+	// Gather Scene View Extensions
+	// Note: Chromakey and Light maps use their own view extensions, collected in FDisplayClusterViewport::GatherActiveExtensions().
+	ViewFamily.ViewExtensions = InFrameViewFamily.ViewExtensions;
+
 	// Scene View Extension activation with ViewportId granularity only works if you have one ViewFamily per ViewportId
 	for (FSceneViewExtensionRef& ViewExt : ViewFamily.ViewExtensions)
 	{
 		ViewExt->SetupViewFamily(ViewFamily);
 	}
-
-#if WITH_EDITOR
-	const FDisplayClusterRenderFrameSettings& RenderFrameSettings = GetRenderFrameSettings();
-	switch (RenderFrameSettings.RenderMode)
-	{
-	case EDisplayClusterRenderFrameMode::PreviewInScene:
-		if (RenderFrameSettings.bPreviewEnablePostProcess == false)
-		{
-			// Disable postprocess for preview
-			ViewFamily.SceneCaptureSource = ESceneCaptureSource::SCS_SceneColorHDR;
-			ViewFamily.SceneCaptureCompositeMode = ESceneCaptureCompositeMode::SCCM_Overwrite;
-		}
-		break;
-	default:
-		break;
-	}
-#endif
-
 }
 
 void FDisplayClusterViewportManager::RenderFrame(FViewport* InViewport)
@@ -959,11 +1000,6 @@ TSharedPtr<IDisplayClusterProjectionPolicy, ESPMode::ThreadSafe> FDisplayCluster
 	return nullptr;
 }
 
-TSharedPtr<IDisplayClusterViewportLightCardManager, ESPMode::ThreadSafe> FDisplayClusterViewportManager::GetLightCardManager() const
-{
-	return LightCardManager;
-}
-
 void FDisplayClusterViewportManager::MarkComponentGeometryDirty(const FName InComponentName)
 {
 	check(IsInGameThread());
diff --git a/Engine/Plugins/Runtime/nDisplay/Source/DisplayCluster/Private/Render/Viewport/DisplayClusterViewportManager.h b/Engine/Plugins/Runtime/nDisplay/Source/DisplayCluster/Private/Render/Viewport/DisplayClusterViewportManager.h
--- a/Engine/Plugins/Runtime/nDisplay/Source/DisplayCluster/Private/Render/Viewport/DisplayClusterViewportManager.h
+++ b/Engine/Plugins/Runtime/nDisplay/Source/DisplayCluster/Private/Render/Viewport/DisplayClusterViewportManager.h
@@ -20,7 +20,6 @@ class FDisplayClusterRenderFrameManager;
 class FDisplayClusterRenderFrame; 
 class FDisplayClusterViewportManagerProxy;
 class FDisplayClusterViewportLightCardManager;
-class IDisplayClusterViewportLightCardManager;
 class IDisplayClusterProjectionPolicy;
 
 class  UDisplayClusterConfigurationViewport;
@@ -54,6 +53,8 @@ public:
 	virtual bool UpdateCustomConfiguration(EDisplayClusterRenderFrameMode InRenderMode, const TArray<FString>& InViewportNames, class ADisplayClusterRootActor* InRootActorPtr) override;
 
 	virtual bool BeginNewFrame(FViewport* InViewport, UWorld* InWorld, FDisplayClusterRenderFrame& OutRenderFrame) override;
+
+	virtual void InitializeNewFrame() override;
 	virtual void FinalizeNewFrame() override;
 
 	virtual FSceneViewFamily::ConstructionValues CreateViewFamilyConstructionValues(
@@ -92,8 +93,6 @@ public:
 		return TArrayView<IDisplayClusterViewport*>((IDisplayClusterViewport**)(ClusterNodeViewports.GetData()), ClusterNodeViewports.Num());
 	}
 
-	virtual TSharedPtr<IDisplayClusterViewportLightCardManager, ESPMode::ThreadSafe> GetLightCardManager() const override;
-
 	virtual void MarkComponentGeometryDirty(const FName InComponentName = NAME_None) override;
 
 	virtual void AddReferencedObjects(FReferenceCollector& Collector) override;
@@ -125,6 +124,24 @@ public:
 
 	const FDisplayClusterRenderFrameSettings& GetRenderFrameSettings() const;
 
+	/** Return number of contexts per viewport. */
+	int32 GetViewPerViewportAmount() const;
+
+	/** Return initial StereoViewIndex for the input viewport. */
+	int32 FindFirstViewportStereoViewIndex(FDisplayClusterViewport* InViewport) const;
+
+	/** Returns true if the final color is renderable (to support preview in scene). */
+	bool ShouldRenderFinalColor() const;
+
+	/** Return nDisplay VE object. */
+	TSharedPtr<class FDisplayClusterViewportManagerViewExtension, ESPMode::ThreadSafe> GetViewportManagerViewExtension() const
+	{ return ViewportManagerViewExtension; }
+
+	/** Return LightCardManager object. */
+	TSharedPtr<FDisplayClusterViewportLightCardManager, ESPMode::ThreadSafe> GetLightCardManager() const
+	{ return LightCardManager; }
+
+	/** Return PostProcessManager object. */
 	TSharedPtr<FDisplayClusterViewportPostProcessManager, ESPMode::ThreadSafe> GetPostProcessManager() const
 	{ return PostProcessManager; }
 
diff --git a/Engine/Plugins/Runtime/nDisplay/Source/DisplayCluster/Private/Render/Viewport/DisplayClusterViewportManagerProxy.cpp b/Engine/Plugins/Runtime/nDisplay/Source/DisplayCluster/Private/Render/Viewport/DisplayClusterViewportManagerProxy.cpp
--- a/Engine/Plugins/Runtime/nDisplay/Source/DisplayCluster/Private/Render/Viewport/DisplayClusterViewportManagerProxy.cpp
+++ b/Engine/Plugins/Runtime/nDisplay/Source/DisplayCluster/Private/Render/Viewport/DisplayClusterViewportManagerProxy.cpp
@@ -12,8 +12,8 @@
 #include "Render/Projection/IDisplayClusterProjectionPolicyFactory.h"
 #include "Render/Projection/IDisplayClusterProjectionPolicy.h"
 
-#include "Render/Viewport/IDisplayClusterViewportLightCardManager.h"
 #include "Render/Viewport/LightCard/DisplayClusterViewportLightCardManager.h"
+#include "Render/Viewport/LightCard/DisplayClusterViewportLightCardManagerProxy.h"
 
 #include "Render/Viewport/DisplayClusterViewportProxy.h"
 #include "Render/Viewport/RenderTarget/DisplayClusterRenderTargetManager.h"
@@ -91,17 +91,13 @@ void FDisplayClusterViewportManagerProxy::Release_RenderThread()
 	}
 
 	if (PostProcessManager.IsValid())
-		{
+	{
 		PostProcessManager->Release();
 		PostProcessManager.Reset();
-		}
-
-	if (LightCardManager.IsValid())
-	{
-		LightCardManager->Release();
-		LightCardManager.Reset();
 	}
 
+	LightCardManagerProxy.Reset();
+
 	ViewportManagerViewExtension.Reset();
 }
 
@@ -109,7 +105,7 @@ void FDisplayClusterViewportManagerProxy::Initialize(FDisplayClusterViewportMana
 {
 	RenderTargetManager = InViewportManager.RenderTargetManager;
 	PostProcessManager = InViewportManager.PostProcessManager;
-	LightCardManager = InViewportManager.LightCardManager;
+	LightCardManagerProxy = InViewportManager.LightCardManager->GetLightCardManagerProxy();
 	ViewportManagerViewExtension = InViewportManager.ViewportManagerViewExtension;
 }
 
@@ -251,14 +247,18 @@ void FDisplayClusterViewportManagerProxy::ImplRenderFrame(FViewport* InViewport)
 		IDisplayCluster::Get().GetCallbacks().OnDisplayClusterProcessLatency_RenderThread().Broadcast(RHICmdList, ViewportManagerProxy, InViewport);
 	});
 
+	ENQUEUE_RENDER_COMMAND(DisplayClusterRenderFrame_UpdateDeferredResources)(
+		[ViewportManagerProxy = SharedThis(this), InViewport](FRHICommandListImmediate& RHICmdList)
+	{
+		// Update viewports resources: vp/texture overlay, OCIO, blur, nummips, etc
+		ViewportManagerProxy->UpdateDeferredResources_RenderThread(RHICmdList);
+	});
+
 	ENQUEUE_RENDER_COMMAND(DisplayClusterRenderFrame_WarpBlend)(
 		[InViewportManagerProxy = SharedThis(this), InViewport](FRHICommandListImmediate& RHICmdList)
 	{
 		const FDisplayClusterViewportManagerProxy* ViewportManagerProxy = &InViewportManagerProxy.Get();
 
-		// Update viewports resources: overlay, vp-overla, blur, nummips, etc
-		ViewportManagerProxy->UpdateDeferredResources_RenderThread(RHICmdList);
-
 		if (ViewportManagerProxy->PostProcessManager.IsValid())
 		{
 			ViewportManagerProxy->PostProcessManager->HandleBeginUpdateFrameResources_RenderThread(RHICmdList, ViewportManagerProxy);
@@ -312,7 +312,7 @@ void FDisplayClusterViewportManagerProxy::UpdateDeferredResources_RenderThread(F
 
 	for (const TSharedPtr<FDisplayClusterViewportProxy, ESPMode::ThreadSafe>& ViewportProxy : ClusterNodeViewportProxies)
 	{
-		if (ViewportProxy->RenderSettings.OverrideViewportId.IsEmpty())
+		if (!ViewportProxy->RenderSettings.IsViewportOverrided())
 		{
 			ViewportProxy->UpdateDeferredResources(RHICmdList);
 		}
@@ -358,6 +358,9 @@ void FDisplayClusterViewportManagerProxy::ImplClearFrameTargets_RenderThread(FRH
 	}
 }
 
+/**
+ * WarpBlend render pass. (internal)
+ */
 enum class EWarpPass : uint8
 {
 	Begin = 0,
@@ -620,12 +623,6 @@ bool FDisplayClusterViewportManagerProxy::ResolveFrameTargetToBackBuffer_RenderT
 	return false;
 }
 
-TSharedPtr<IDisplayClusterViewportLightCardManager, ESPMode::ThreadSafe> FDisplayClusterViewportManagerProxy::GetLightCardManager_RenderThread() const
-{
-	return LightCardManager;
-}
-
-
 FDisplayClusterViewportProxy* FDisplayClusterViewportManagerProxy::ImplFindViewport_RenderThread(const FString& ViewportId) const
 {
 	check(IsInRenderingThread());
@@ -639,7 +636,7 @@ FDisplayClusterViewportProxy* FDisplayClusterViewportManagerProxy::ImplFindViewp
 	return (DesiredViewport != nullptr) ? DesiredViewport->Get() : nullptr;
 }
 
-IDisplayClusterViewportProxy* FDisplayClusterViewportManagerProxy::FindViewport_RenderThread(const int32 StereoViewIndex, uint32* OutContextNum) const
+FDisplayClusterViewportProxy* FDisplayClusterViewportManagerProxy::ImplFindViewport_RenderThread(const int32 StereoViewIndex, uint32* OutContextNum) const
 {
 	check(IsInRenderingThread());
 
diff --git a/Engine/Plugins/Runtime/nDisplay/Source/DisplayCluster/Private/Render/Viewport/DisplayClusterViewportManagerProxy.h b/Engine/Plugins/Runtime/nDisplay/Source/DisplayCluster/Private/Render/Viewport/DisplayClusterViewportManagerProxy.h
--- a/Engine/Plugins/Runtime/nDisplay/Source/DisplayCluster/Private/Render/Viewport/DisplayClusterViewportManagerProxy.h
+++ b/Engine/Plugins/Runtime/nDisplay/Source/DisplayCluster/Private/Render/Viewport/DisplayClusterViewportManagerProxy.h
@@ -14,13 +14,14 @@ class FDisplayClusterRenderTargetManager;
 class FDisplayClusterViewportPostProcessManager;
 class FDisplayClusterViewportManager;
 class FDisplayClusterViewportResource;
-class FDisplayClusterViewportLightCardManager;
-class IDisplayClusterViewportLightCardManager;
+class FDisplayClusterViewportLightCardManagerProxy;
 class IDisplayClusterProjectionPolicy;
-
 class FViewport;
 
-
+/**
+ * ViewportManagerProxy implementation.
+ * This is a proxy object for the rendering thread. The data is copied every frame from the ViewportManager on the game thread.
+ */
 class FDisplayClusterViewportManagerProxy
 	: public IDisplayClusterViewportManagerProxy
 	, public TSharedFromThis<FDisplayClusterViewportManagerProxy, ESPMode::ThreadSafe>
@@ -29,20 +30,17 @@ public:
 	FDisplayClusterViewportManagerProxy();
 	virtual ~FDisplayClusterViewportManagerProxy();
 
-	void Release_RenderThread();
-
 public:
-	void DoCrossGPUTransfers_RenderThread(FRHICommandListImmediate& RHICmdList) const;
-	void UpdateDeferredResources_RenderThread(FRHICommandListImmediate& RHICmdList) const;
-	void UpdateFrameResources_RenderThread(FRHICommandListImmediate& RHICmdList, bool bWarpBlendEnabled) const;
-
-	/** Render thread funcs */
+	//~IDisplayClusterViewportManagerProxy
 	virtual IDisplayClusterViewportProxy* FindViewport_RenderThread(const FString& InViewportId) const override
 	{
 		return ImplFindViewport_RenderThread(InViewportId);
 	}
 
-	virtual IDisplayClusterViewportProxy* FindViewport_RenderThread(const int32 StereoViewIndex, uint32* OutContextNum = nullptr) const override;
+	virtual IDisplayClusterViewportProxy* FindViewport_RenderThread(const int32 StereoViewIndex, uint32* OutContextNum = nullptr) const override
+	{
+		return ImplFindViewport_RenderThread(StereoViewIndex, OutContextNum);
+	}
 
 	virtual const TArrayView<TSharedPtr<IDisplayClusterViewportProxy, ESPMode::ThreadSafe>> GetViewports_RenderThread() const override
 	{
@@ -51,51 +49,85 @@ public:
 
 	virtual bool GetFrameTargets_RenderThread(TArray<FRHITexture2D*>& OutFrameResources, TArray<FIntPoint>& OutTargetOffsets, TArray<FRHITexture2D*>* OutAdditionalFrameResources = nullptr) const override;
 	virtual bool ResolveFrameTargetToBackBuffer_RenderThread(FRHICommandListImmediate& RHICmdList, const uint32 InContextNum, const int32 DestArrayIndex, FRHITexture2D* DestTexture, FVector2D WindowSize) const override;
+	//~~IDisplayClusterViewportManagerProxy
+
+	/** Custom cross-GPU implementation for mGPU. Applyed for viewports with bOverrideCrossGPUTransfer=true. */
+	void DoCrossGPUTransfers_RenderThread(FRHICommandListImmediate& RHICmdList) const;
+
+	/** UpdateDeferredResources for all viewports (OCIO, PP, mips, etc). */
+	void UpdateDeferredResources_RenderThread(FRHICommandListImmediate& RHICmdList) const;
+
+	/** Apply WarpBlend and resolve to frame resources. */
+	void UpdateFrameResources_RenderThread(FRHICommandListImmediate& RHICmdList, bool bWarpBlendEnabled) const;
 
-	virtual TSharedPtr<IDisplayClusterViewportLightCardManager, ESPMode::ThreadSafe> GetLightCardManager_RenderThread() const override;
+	/** Release all referenced objects and resources. */
+	void Release_RenderThread();
 
-	// internal use only
+	/** Register new viewport proxy. */
 	void CreateViewport_RenderThread(const TSharedPtr<FDisplayClusterViewportProxy, ESPMode::ThreadSafe>& InViewportProxy);
+
+	/** Unregister exist viewport proxy. */
 	void DeleteViewport_RenderThread(const TSharedPtr<FDisplayClusterViewportProxy, ESPMode::ThreadSafe>& InViewportProxy);
 
+	/** Reset all references to InDeletedResourcePtr for all viewport proxies and delete the resource. */
 	void DeleteResource_RenderThread(FDisplayClusterViewportResource* InDeletedResourcePtr);
+
+	/** Initialize ViewportManagerProxy from ViewportManager. */
 	void Initialize(FDisplayClusterViewportManager& InViewportManager);
 
-	const FDisplayClusterRenderFrameSettings& GetRenderFrameSettings_RenderThread() const
-	{
-		check(IsInRenderingThread());
+	/** Get LightCardManager proxy object. */
+	TSharedPtr<FDisplayClusterViewportLightCardManagerProxy, ESPMode::ThreadSafe> GetLightCardManagerProxy_RenderThread() const
+	{ return LightCardManagerProxy; }
 
-		return RenderFrameSettings;
-	}
+	/** Get RenderFrameSettings on render thread. */
+	const FDisplayClusterRenderFrameSettings& GetRenderFrameSettings_RenderThread() const
+	{ return RenderFrameSettings; }
 
+	/** Get the viewports of the cluster node on the render thread. */
 	const TArray<TSharedPtr<FDisplayClusterViewportProxy, ESPMode::ThreadSafe>>& ImplGetViewportProxies_RenderThread() const
-	{
-		check(IsInRenderingThread());
-
-		return ClusterNodeViewportProxies;
-	}
+	{ return ClusterNodeViewportProxies; }
 
+	/** Find viewport by name. */
 	FDisplayClusterViewportProxy* ImplFindViewport_RenderThread(const FString& InViewportId) const;
 
+	/** Find viewport and context index by StereoViewIndex. */
+	FDisplayClusterViewportProxy* ImplFindViewport_RenderThread(const int32 StereoViewIndex, uint32* OutContextNum = nullptr) const;
+
+	/** Copy RenderFrameSettings from game to render thread. */
 	void ImplUpdateRenderFrameSettings(const FDisplayClusterRenderFrameSettings& InRenderFrameSettings);
+
+	/** Copy Viewports data to Proxies from game to render thread. */
 	void ImplUpdateViewports(const TArray<FDisplayClusterViewport*>& InViewports);
 
+	/** Rendering the final frame of nDisplay. Called after all viewports have been rendered in RTT. */
 	void ImplRenderFrame(FViewport* InViewport);
+
+	/** Rendering black color in all frame textures. */
 	void ImplClearFrameTargets_RenderThread(FRHICommandListImmediate& RHICmdList) const;
 
 private:
+	/** Update ClusterNodeViewportProxies from ViewportProxies for the current cluster node. */
 	void ImplUpdateClusterNodeViewportProxies();
 
 private:
+	/** RenderTarget manager object. Used to remove related objects in order. */
 	TSharedPtr<FDisplayClusterRenderTargetManager, ESPMode::ThreadSafe>        RenderTargetManager;
+
+	/** PostProcess manager object. Used to remove related objects in order. */
 	TSharedPtr<FDisplayClusterViewportPostProcessManager, ESPMode::ThreadSafe> PostProcessManager;
-	TSharedPtr<FDisplayClusterViewportLightCardManager, ESPMode::ThreadSafe>   LightCardManager;
 
+	/** LightCard manager proxy object. */
+	TSharedPtr<FDisplayClusterViewportLightCardManagerProxy, ESPMode::ThreadSafe> LightCardManagerProxy;
+
+	/** Copy of RenderFrameSettings for the render thread. */
 	FDisplayClusterRenderFrameSettings RenderFrameSettings;
 
+	/** Proxies of the entire cluster. */
 	TArray<TSharedPtr<FDisplayClusterViewportProxy, ESPMode::ThreadSafe>> ViewportProxies;
+
+	/** Proxies of the current cluster node. */
 	TArray<TSharedPtr<FDisplayClusterViewportProxy, ESPMode::ThreadSafe>> ClusterNodeViewportProxies;
 
-	// Handle special features
+	/** nDisplay VE object. Used to remove related objects in order. */
 	TSharedPtr<class FDisplayClusterViewportManagerViewExtension, ESPMode::ThreadSafe> ViewportManagerViewExtension;
 };
diff --git a/Engine/Plugins/Runtime/nDisplay/Source/DisplayCluster/Private/Render/Viewport/DisplayClusterViewportManagerViewExtension.cpp b/Engine/Plugins/Runtime/nDisplay/Source/DisplayCluster/Private/Render/Viewport/DisplayClusterViewportManagerViewExtension.cpp
--- a/Engine/Plugins/Runtime/nDisplay/Source/DisplayCluster/Private/Render/Viewport/DisplayClusterViewportManagerViewExtension.cpp
+++ b/Engine/Plugins/Runtime/nDisplay/Source/DisplayCluster/Private/Render/Viewport/DisplayClusterViewportManagerViewExtension.cpp
@@ -2,13 +2,19 @@
 
 #include "DisplayClusterViewportManagerViewExtension.h"
 #include "DisplayClusterSceneViewExtensions.h"
-#include "Render/Viewport/IDisplayClusterViewportManager.h"
+
+#include "Render/Viewport/DisplayClusterViewportManager.h"
 #include "Render/Viewport/DisplayClusterViewportManagerProxy.h"
 #include "SceneRendering.h"
 
 #include "IDisplayCluster.h"
 #include "IDisplayClusterCallbacks.h"
 
+#include "ScreenPass.h"
+#include "CommonRenderResources.h"
+#include "PostProcess/PostProcessing.h"
+// for FPostProcessMaterialInputs
+#include "PostProcess/PostProcessMaterial.h"
 
 namespace DisplayClusterViewportManagerViewExtensionHelpers
 {
@@ -19,10 +25,10 @@ using namespace DisplayClusterViewportManagerViewExtensionHelpers;
 ///////////////////////////////////////////////////////////////////////////////////////
 // FDisplayClusterViewportManagerViewExtension
 ///////////////////////////////////////////////////////////////////////////////////////
-FDisplayClusterViewportManagerViewExtension::FDisplayClusterViewportManagerViewExtension(const FAutoRegister& AutoRegister, const IDisplayClusterViewportManager* InViewportManager)
+FDisplayClusterViewportManagerViewExtension::FDisplayClusterViewportManagerViewExtension(const FAutoRegister& AutoRegister, const FDisplayClusterViewportManager* InViewportManager)
 	: FSceneViewExtensionBase(AutoRegister)
 	, ViewportManager(InViewportManager)
-	, ViewportManagerProxy(InViewportManager->GetProxy())
+	, ViewportManagerProxy(InViewportManager->GetViewportManagerProxy())
 {
 	RegisterCallbacks();
 }
@@ -43,6 +49,82 @@ void FDisplayClusterViewportManagerViewExtension::PreRenderView_RenderThread(FRD
 	}
 }
 
+void FDisplayClusterViewportManagerViewExtension::SubscribeToPostProcessingPass(EPostProcessingPass PassId, FAfterPassCallbackDelegateArray& InOutPassCallbacks, bool bIsPassEnabled)
+{
+	switch (PassId)
+	{
+	case EPostProcessingPass::SSRInput:
+		for (const FViewportProxy& ViewportIt : ViewportProxies)
+		{
+			if (ViewportIt.ViewportProxy.IsValid() && ViewportIt.ViewportProxy->ShouldUsePostProcessPassAfterSSRInput())
+			{
+				InOutPassCallbacks.Add(FAfterPassCallbackDelegate::CreateRaw(this, &FDisplayClusterViewportManagerViewExtension::PostProcessPassAfterSSRInput_RenderThread));
+				break;
+			}
+		}
+		break;
+
+	case EPostProcessingPass::FXAA:
+		for (const FViewportProxy& ViewportIt : ViewportProxies)
+		{
+			if (ViewportIt.ViewportProxy.IsValid() && ViewportIt.ViewportProxy->ShouldUsePostProcessPassAfterFXAA())
+			{
+				InOutPassCallbacks.Add(FAfterPassCallbackDelegate::CreateRaw(this, &FDisplayClusterViewportManagerViewExtension::PostProcessPassAfterFXAA_RenderThread));
+				break;
+			}
+		}
+		break;
+
+	default:
+		break;
+	}
+}
+
+FScreenPassTexture FDisplayClusterViewportManagerViewExtension::PostProcessPassAfterFXAA_RenderThread(FRDGBuilder& GraphBuilder, const FSceneView& View, const FPostProcessMaterialInputs& Inputs)
+{
+	uint32 ContextNum = 0;
+	if (FDisplayClusterViewportProxy* ViewportProxyPtr = ViewportManagerProxy->ImplFindViewport_RenderThread(View.StereoViewIndex, &ContextNum))
+	{
+		if (ViewportProxyPtr->ShouldUsePostProcessPassAfterFXAA())
+		{
+			return ViewportProxyPtr->OnPostProcessPassAfterFXAA_RenderThread(GraphBuilder, View, Inputs, ContextNum);
+		}
+	}
+
+	return ReturnUntouchedSceneColorForPostProcessing(Inputs);
+}
+
+FScreenPassTexture FDisplayClusterViewportManagerViewExtension::PostProcessPassAfterSSRInput_RenderThread(FRDGBuilder& GraphBuilder, const FSceneView& View, const FPostProcessMaterialInputs& Inputs)
+{
+	uint32 ContextNum = 0;
+	if (FDisplayClusterViewportProxy* ViewportProxyPtr = ViewportManagerProxy->ImplFindViewport_RenderThread(View.StereoViewIndex, &ContextNum))
+	{
+		if (ViewportProxyPtr->ShouldUsePostProcessPassAfterSSRInput())
+		{
+			return ViewportProxyPtr->OnPostProcessPassAfterSSRInput_RenderThread(GraphBuilder, View, Inputs, ContextNum);
+		}
+	}
+
+	return ReturnUntouchedSceneColorForPostProcessing(Inputs);
+}
+
+/**
+* A helper function that extracts the right scene color texture, untouched, to be used further in post processing.
+*/
+FScreenPassTexture FDisplayClusterViewportManagerViewExtension::ReturnUntouchedSceneColorForPostProcessing(const FPostProcessMaterialInputs& InOutInputs)
+{
+	if (InOutInputs.OverrideOutput.IsValid())
+	{
+		return InOutInputs.OverrideOutput;
+	}
+	else
+	{
+		/** We don't want to modify scene texture in any way. We just want it to be passed back onto the next stage. */
+		FScreenPassTexture SceneTexture = const_cast<FScreenPassTexture&>(InOutInputs.Textures[(uint32)EPostProcessMaterialInput::SceneColor]);
+		return SceneTexture;
+	}
+}
+
 bool FDisplayClusterViewportManagerViewExtension::IsActiveThisFrame_Internal(const FSceneViewExtensionContext& Context) const
 {
 	static const FDisplayClusterSceneViewExtensionContext DCViewExtensionContext;
@@ -61,7 +143,7 @@ bool FDisplayClusterViewportManagerViewExtension::IsActiveThisFrame_Internal(con
 
 void FDisplayClusterViewportManagerViewExtension::PreRenderViewFamily_RenderThread(FRDGBuilder& GraphBuilder, FSceneViewFamily& InViewFamily)
 {
-	Viewports.Empty();
+	ViewportProxies.Empty();
 
 	// Get all viewport proxies that are uses in this view family
 	for (int32 ViewIndex = 0; ViewIndex < InViewFamily.Views.Num(); ViewIndex++)
@@ -77,7 +159,7 @@ void FDisplayClusterViewportManagerViewExtension::PreRenderViewFamily_RenderThre
 				NewViewportProxy.ViewportProxy = (static_cast<FDisplayClusterViewportProxy*>(ViewportProxyPtr))->AsShared();
 				NewViewportProxy.ViewIndex = ViewIndex;
 
-				Viewports.Add(NewViewportProxy);
+				ViewportProxies.Add(NewViewportProxy);
 			}
 		}
 	}
@@ -85,11 +167,11 @@ void FDisplayClusterViewportManagerViewExtension::PreRenderViewFamily_RenderThre
 
 void FDisplayClusterViewportManagerViewExtension::PostRenderViewFamily_RenderThread(FRDGBuilder& GraphBuilder, FSceneViewFamily& InViewFamily)
 {
-	for (FViewportProxy& ViewportIt : Viewports)
+	for (FViewportProxy& ViewportIt : ViewportProxies)
 	{
 		if (const FSceneView* SceneView = InViewFamily.Views[ViewportIt.ViewIndex])
 		{
-			ViewportIt.ViewportProxy->PostRenderViewFamily_RenderThread(GraphBuilder, InViewFamily, *SceneView, ViewportIt.ViewportProxyContext);
+			ViewportIt.ViewportProxy->OnPostRenderViewFamily_RenderThread(GraphBuilder, InViewFamily, *SceneView, ViewportIt.ViewportProxyContext);
 		}
 	}
 }
@@ -99,7 +181,7 @@ void FDisplayClusterViewportManagerViewExtension::OnResolvedSceneColor_RenderThr
 {
 	check(IsInRenderingThread());
 
-	for (FViewportProxy& ViewportIt : Viewports)
+	for (FViewportProxy& ViewportIt : ViewportProxies)
 	{
 		ViewportIt.ViewportProxy->OnResolvedSceneColor_RenderThread(GraphBuilder, SceneTextures, ViewportIt.ViewportProxyContext);
 	}
diff --git a/Engine/Plugins/Runtime/nDisplay/Source/DisplayCluster/Private/Render/Viewport/DisplayClusterViewportManagerViewExtension.h b/Engine/Plugins/Runtime/nDisplay/Source/DisplayCluster/Private/Render/Viewport/DisplayClusterViewportManagerViewExtension.h
--- a/Engine/Plugins/Runtime/nDisplay/Source/DisplayCluster/Private/Render/Viewport/DisplayClusterViewportManagerViewExtension.h
+++ b/Engine/Plugins/Runtime/nDisplay/Source/DisplayCluster/Private/Render/Viewport/DisplayClusterViewportManagerViewExtension.h
@@ -4,8 +4,11 @@
 #include "SceneViewExtension.h"
 #include "Render/Viewport/Containers/DisplayClusterViewportProxy_Context.h"
 
-class IDisplayClusterViewportManager;
-class IDisplayClusterViewportManagerProxy;
+class FDisplayClusterViewportManager;
+class FDisplayClusterViewportManagerProxy;
+class FDisplayClusterViewportProxy;
+
+#define DISPLAYCLUSTER_SCENE_VIEW_EXTENSION_PRIORITY -1
 
 /**
  * View extension applying an DC Viewport features
@@ -13,14 +16,14 @@ class IDisplayClusterViewportManagerProxy;
 class FDisplayClusterViewportManagerViewExtension : public FSceneViewExtensionBase
 {
 public:
-	FDisplayClusterViewportManagerViewExtension(const FAutoRegister& AutoRegister, const IDisplayClusterViewportManager* InViewportManager);
+	FDisplayClusterViewportManagerViewExtension(const FAutoRegister& AutoRegister, const FDisplayClusterViewportManager* InViewportManager);
 	virtual ~FDisplayClusterViewportManagerViewExtension();
 
 public:
 	//~ Begin ISceneViewExtension interface
 	virtual int32 GetPriority() const override
 	{
-		return -1;
+		return DISPLAYCLUSTER_SCENE_VIEW_EXTENSION_PRIORITY;
 	}
 
 	virtual void SetupViewFamily(FSceneViewFamily& InViewFamily) override
@@ -34,30 +37,69 @@ public:
 
 	virtual void PreRenderView_RenderThread(FRDGBuilder& GraphBuilder, FSceneView& InView) override;
 
+	virtual void SubscribeToPostProcessingPass(EPostProcessingPass PassId, FAfterPassCallbackDelegateArray& InOutPassCallbacks, bool bIsPassEnabled) override;
+
 	virtual bool IsActiveThisFrame_Internal(const FSceneViewExtensionContext& Context) const override;
 
 	virtual void PreRenderViewFamily_RenderThread(FRDGBuilder& GraphBuilder, FSceneViewFamily& InViewFamily) override;
 	virtual void PostRenderViewFamily_RenderThread(FRDGBuilder& GraphBuilder, FSceneViewFamily& InViewFamily) override;
 	//~End ISceneVIewExtension interface
 
+	/** PP Helper: Get output PP texture from inputs (OverrideOutput or SceneColor).
+	 *
+	 * @param InOutInputs - PP material inputs
+	 *
+	 * @return - Screen pass texture
+	 */
+	static FScreenPassTexture ReturnUntouchedSceneColorForPostProcessing(const FPostProcessMaterialInputs& InOutInputs);
+
 private:
-	/** Rendered callback (get scene textures to share) */
+	/** nDisplay VE Callback [subscribed to Renderer:ResolvedSceneColorCallbacks].
+	 *
+	 * @param GraphBuilder   - RDG interface
+	 * @param SceneTextures  - Scene textures (SceneColor, Depth, etc)
+	 *
+	 * @return - none
+	 */
 	void OnResolvedSceneColor_RenderThread(FRDGBuilder& GraphBuilder, const FSceneTextures& SceneTextures);
 
+	/** Callback OnPostProcessPassAfterFXAA.
+	 *
+	 * @param GraphBuilder - RDG interface
+	 * @param View         - Scene View
+	 * @param Inputs       - PP Input resources
+	 *
+	 * @return - Screen pass texture
+	 */
+	FScreenPassTexture PostProcessPassAfterFXAA_RenderThread(FRDGBuilder& GraphBuilder, const FSceneView& View, const FPostProcessMaterialInputs& Inputs);
+
+	/** Callback OnPostProcessPassAfterSSRInput.
+	 *
+	 * @param GraphBuilder - RDG interface
+	 * @param View         - Scene View
+	 * @param Inputs       - PP Input resources
+	 *
+	 * @return - Screen pass texture
+	 */
+	FScreenPassTexture PostProcessPassAfterSSRInput_RenderThread(FRDGBuilder& GraphBuilder, const FSceneView& View, const FPostProcessMaterialInputs& Inputs);
+
+	/** Register callbacks to external modules (Renderer:ResolvedSceneColor). */
 	void RegisterCallbacks();
+
+	/** UnRegister callbacks to external modules (Renderer:ResolvedSceneColor). */
 	void UnregisterCallbacks();
 
 private:
 	FDelegateHandle ResolvedSceneColorCallbackHandle;
 
 private:
-	const IDisplayClusterViewportManager* ViewportManager;
-	const IDisplayClusterViewportManagerProxy* ViewportManagerProxy;
+	const FDisplayClusterViewportManager* ViewportManager;
+	const TSharedPtr<FDisplayClusterViewportManagerProxy, ESPMode::ThreadSafe> ViewportManagerProxy;
 
 	struct FViewportProxy
 	{
 		// The ViewportProxy must exist until this object is removed.
-		TSharedPtr<class IDisplayClusterViewportProxy, ESPMode::ThreadSafe> ViewportProxy;
+		TSharedPtr<FDisplayClusterViewportProxy, ESPMode::ThreadSafe> ViewportProxy;
 
 		// The index of view in viewfamily collection
 		int32 ViewIndex = 0;
@@ -66,6 +108,6 @@ private:
 		FDisplayClusterViewportProxy_Context ViewportProxyContext;
 	};
 
-	TArray<FViewportProxy> Viewports;
-	
+	/** RenderThreadData:. Viewport proxies from rendered view family. */
+	TArray<FViewportProxy> ViewportProxies;
 };
diff --git a/Engine/Plugins/Runtime/nDisplay/Source/DisplayCluster/Private/Render/Viewport/DisplayClusterViewportManager_EditorPreview.cpp b/Engine/Plugins/Runtime/nDisplay/Source/DisplayCluster/Private/Render/Viewport/DisplayClusterViewportManager_EditorPreview.cpp
--- a/Engine/Plugins/Runtime/nDisplay/Source/DisplayCluster/Private/Render/Viewport/DisplayClusterViewportManager_EditorPreview.cpp
+++ b/Engine/Plugins/Runtime/nDisplay/Source/DisplayCluster/Private/Render/Viewport/DisplayClusterViewportManager_EditorPreview.cpp
@@ -120,6 +120,9 @@ bool FDisplayClusterViewportManager::RenderInEditor(class FDisplayClusterRenderF
 	FSceneInterface* PreviewScene = CurrentWorld->Scene;
 	FEngineShowFlags EngineShowFlags = FEngineShowFlags(EShowFlagInitMode::ESFIM_Game);
 
+	// Handle special viewports game-thread logic at frame begin
+	InitializeNewFrame();
+
 	int32 ViewportIndex = 0;
 	bool bViewportsRenderPassDone = false;
 
diff --git a/Engine/Plugins/Runtime/nDisplay/Source/DisplayCluster/Private/Render/Viewport/DisplayClusterViewportProxy.cpp b/Engine/Plugins/Runtime/nDisplay/Source/DisplayCluster/Private/Render/Viewport/DisplayClusterViewportProxy.cpp
--- a/Engine/Plugins/Runtime/nDisplay/Source/DisplayCluster/Private/Render/Viewport/DisplayClusterViewportProxy.cpp
+++ b/Engine/Plugins/Runtime/nDisplay/Source/DisplayCluster/Private/Render/Viewport/DisplayClusterViewportProxy.cpp
@@ -5,11 +5,15 @@
 #include "Render/Viewport/DisplayClusterViewportManagerProxy.h"
 #include "Render/Viewport/DisplayClusterViewport.h"
 #include "Render/Viewport/DisplayClusterViewportManager.h"
+#include "Render/Viewport/DisplayClusterViewportManagerViewExtension.h"
 
 #include "Render/Viewport/Containers/DisplayClusterViewport_PostRenderSettings.h"
+#include "Render/Viewport/Configuration/DisplayClusterViewportConfiguration.h"
+
 #include "Render/Viewport/RenderTarget/DisplayClusterRenderTargetResource.h"
 #include "Render/Viewport/RenderFrame/DisplayClusterRenderFrameSettings.h"
-#include "Render/Viewport/Configuration/DisplayClusterViewportConfiguration.h"
+
+#include "Render/Viewport/LightCard/DisplayClusterViewportLightCardManagerProxy.h"
 
 #include "Render/Projection/IDisplayClusterProjectionPolicy.h"
 
@@ -39,10 +43,16 @@
 
 #include "ScreenRendering.h"
 #include "PostProcess/SceneFilterRendering.h"
+#include "PostProcess/PostProcessAA.h"
 #include "SceneRendering.h"
 #include "RenderGraphUtils.h"
 #include "RenderGraphResources.h"
 
+#include "ScreenPass.h"
+#include "PostProcess/PostProcessing.h"
+// for FPostProcessMaterialInputs
+#include "PostProcess/PostProcessMaterial.h"
+
 static TAutoConsoleVariable<int32> CVarDisplayClusterRenderOverscanResolve(
 	TEXT("nDisplay.render.overscan.resolve"),
 	1,
@@ -51,12 +61,61 @@ static TAutoConsoleVariable<int32> CVarDisplayClusterRenderOverscanResolve(
 	ECVF_RenderThreadSafe
 );
 
+int32 GDisplayClusterShadersICVFXFXAALightCard = 2;
+static FAutoConsoleVariableRef CVarDisplayClusterShadersICVFXFXAALightCard(
+	TEXT("nDisplay.render.icvfx.fxaa.lightcard"),
+	GDisplayClusterShadersICVFXFXAALightCard,
+	TEXT("FXAA quality for lightcard (0 - disable).\n")
+	TEXT("1..6 - FXAA quality from Lowest Quality(Fastest) to Highest Quality(Slowest).\n"),
+	ECVF_RenderThreadSafe
+);
+
+int32 GDisplayClusterShadersICVFXFXAAChromakey = 2;
+static FAutoConsoleVariableRef CVarDisplayClusterShadersICVFXFXAAChromakey(
+	TEXT("nDisplay.render.icvfx.fxaa.chromakey"),
+	GDisplayClusterShadersICVFXFXAAChromakey,
+	TEXT("FXAA quality for chromakey (0 - disable).\n")
+	TEXT("1..6 - FXAA quality from Lowest Quality(Fastest) to Highest Quality(Slowest).\n"),
+	ECVF_RenderThreadSafe
+);
+
 ///////////////////////////////////////////////////////////////////////////////////////
 namespace DisplayClusterViewportProxyHelpers
 {
 	// The viewport override has the maximum depth. This protects against a link cycle
 	static const int32 DisplayClusterViewportProxyResourcesOverrideRecursionDepthMax = 4;
 
+	static bool GetFXAAQuality(const EDisplayClusterViewportCaptureMode InCaptureMode, EFXAAQuality& OutFXAAQuality)
+	{
+		// Get FXAA quality for current viewport
+		EFXAAQuality FXAAQuality = EFXAAQuality::MAX;
+
+		switch (InCaptureMode)
+		{
+		case EDisplayClusterViewportCaptureMode::Chromakey:
+			if (GDisplayClusterShadersICVFXFXAAChromakey > 0)
+			{
+				OutFXAAQuality = (EFXAAQuality)FMath::Clamp(GDisplayClusterShadersICVFXFXAAChromakey - 1, 0, (int32)EFXAAQuality::MAX - 1);
+				return true;
+			}
+			break;
+
+		case EDisplayClusterViewportCaptureMode::Lightcard:
+			if (GDisplayClusterShadersICVFXFXAALightCard > 0)
+			{
+				OutFXAAQuality = (EFXAAQuality)FMath::Clamp(GDisplayClusterShadersICVFXFXAALightCard - 1, 0, (int32)EFXAAQuality::MAX - 1);
+				return true;
+			}
+			break;
+
+		default:
+			break;
+		}
+
+		// No FXAA
+		return false;
+	}
+
 	/**
 	* Get viewport RHI resources
 	*
@@ -193,18 +252,40 @@ void FDisplayClusterViewportProxy::HandleResourceDelete_RenderThread(FDisplayClu
 bool FDisplayClusterViewportProxy::IsShouldOverrideViewportResource(const EDisplayClusterViewportResourceType InResourceType) const
 {
 	// Override resources from other viewport
-	if (RenderSettings.OverrideViewportId.IsEmpty() == false)
+	if (RenderSettings.IsViewportOverrided())
 	{
-		switch (InResourceType)
+		switch (RenderSettings.ViewportOverrideMode)
 		{
-		case EDisplayClusterViewportResourceType::InternalRenderTargetResource:
-		case EDisplayClusterViewportResourceType::InputShaderResource:
-		case EDisplayClusterViewportResourceType::MipsShaderResource:
-		case EDisplayClusterViewportResourceType::AdditionalTargetableResource:
-			return true;
+		case EDisplayClusterViewportOverrideMode::All:
+		{
+			switch (InResourceType)
+			{
+			case EDisplayClusterViewportResourceType::InternalRenderTargetResource:
+			case EDisplayClusterViewportResourceType::InputShaderResource:
+			case EDisplayClusterViewportResourceType::MipsShaderResource:
+			case EDisplayClusterViewportResourceType::AdditionalTargetableResource:
+				return true;
+
+			default:
+				break;
+			}
+		}
+		break;
+
+		case EDisplayClusterViewportOverrideMode::InernalRTT:
+		{
+			switch (InResourceType)
+			{
+			case EDisplayClusterViewportResourceType::InternalRenderTargetResource:
+				return true;
+
+			default:
+				break;
+			}
+		}
+		break;
 
 		default:
-			// By default use all output resources from this viewport
 			break;
 		}
 	}
@@ -218,6 +299,40 @@ bool FDisplayClusterViewportProxy::GetResources_RenderThread(const EDisplayClust
 	return ImplGetResources_RenderThread(InResourceType, OutResources, false);
 }
 
+const FDisplayClusterViewportProxy& FDisplayClusterViewportProxy::GetRenderingViewportProxy() const
+{
+	switch (RenderSettings.ViewportOverrideMode)
+	{
+	case EDisplayClusterViewportOverrideMode::All:
+	case EDisplayClusterViewportOverrideMode::InernalRTT:
+		if (FDisplayClusterViewportProxy const* OverrideViewportProxy = Owner->ImplFindViewport_RenderThread(RenderSettings.ViewportOverrideId))
+		{
+			return *OverrideViewportProxy;
+		}
+
+		break;
+	}
+
+	return *this;
+}
+
+bool FDisplayClusterViewportProxy::IsInputRenderTargetResourceExists() const
+{
+	if (PostRenderSettings.Replace.IsEnabled())
+	{
+		// Use external texture
+		return true;
+	}
+
+	if (EnumHasAnyFlags(RenderSettingsICVFX.RuntimeFlags, EDisplayClusterViewportRuntimeICVFXFlags::UVLightcard))
+	{
+		// Use external UVLightCard Resource
+		return true;
+	}
+
+	return !RenderTargets.IsEmpty();
+}
+
 bool FDisplayClusterViewportProxy::ImplGetResources_RenderThread(const EDisplayClusterViewportResourceType InResourceType, TArray<FRHITexture2D*>& OutResources, const int32 InRecursionDepth) const
 {
 	check(IsInRenderingThread());
@@ -227,10 +342,7 @@ bool FDisplayClusterViewportProxy::ImplGetResources_RenderThread(const EDisplayC
 	{
 		if (InRecursionDepth < DisplayClusterViewportProxyResourcesOverrideRecursionDepthMax)
 		{
-			if (FDisplayClusterViewportProxy const* OverrideViewportProxy = Owner->ImplFindViewport_RenderThread(RenderSettings.OverrideViewportId))
-			{
-				return OverrideViewportProxy->ImplGetResources_RenderThread(InResourceType, OutResources, InRecursionDepth + 1);
-			}
+			return GetRenderingViewportProxy().ImplGetResources_RenderThread(InResourceType, OutResources, InRecursionDepth + 1);
 		}
 
 		return false;
@@ -242,29 +354,63 @@ bool FDisplayClusterViewportProxy::ImplGetResources_RenderThread(const EDisplayC
 	{
 	case EDisplayClusterViewportResourceType::InternalRenderTargetResource:
 	{
+		bool bResult = false;
+
 		if (Contexts.Num() > 0)
 		{
-			for (int32 ContextIt = 0; ContextIt < Contexts.Num(); ContextIt++)
+
+			// 1. Replace RTT from configuration
+			if (!bResult && PostRenderSettings.Replace.IsEnabled())
 			{
+				bResult = true;
+
 				// Support texture replace:
-				if (PostRenderSettings.Replace.IsEnabled())
+				if (FRHITexture2D* ReplaceTextureRHI = PostRenderSettings.Replace.TextureRHI->GetTexture2D())
+				{
+					for (int32 ContextIndex = 0; ContextIndex < Contexts.Num(); ContextIndex++)
+					{
+						OutResources.Add(ReplaceTextureRHI);
+					}
+				}
+			}
+
+			// 2. Replace RTT from UVLightCard:
+			if (!bResult && EnumHasAnyFlags(RenderSettingsICVFX.RuntimeFlags, EDisplayClusterViewportRuntimeICVFXFlags::UVLightcard))
+			{
+				bResult = true;
+				
+				// Get resources from external UV LightCard manager
+				TSharedPtr<FDisplayClusterViewportLightCardManagerProxy, ESPMode::ThreadSafe> LightCardManager = Owner->GetLightCardManagerProxy_RenderThread();
+				if (LightCardManager.IsValid())
 				{
-					OutResources.Add(PostRenderSettings.Replace.TextureRHI->GetTexture2D());
+					if (FRHITexture* UVLightCardRHIResource = LightCardManager->GetUVLightCardRHIResource_RenderThread())
+					{
+						for (int32 ContextIndex = 0; ContextIndex < Contexts.Num(); ContextIndex++)
+						{
+							OutResources.Add(UVLightCardRHIResource);
+						}
+					}
 				}
-				else
+			}
+
+			// 3. Finally Use InternalRTT
+			if (!bResult)
+			{
+				for (int32 ContextIndex = 0; ContextIndex < Contexts.Num(); ContextIndex++)
 				{
-					if (ContextIt < RenderTargets.Num())
+					if (RenderTargets.IsValidIndex(ContextIndex))
 					{
-						if (FDisplayClusterViewportRenderTargetResource* Input = RenderTargets[ContextIt])
+						if (FDisplayClusterViewportRenderTargetResource* Input = RenderTargets[ContextIndex])
 						{
 							OutResources.Add(Input->GetViewportRenderTargetResourceRHI());
+							bResult = true;
 						}
 					}
 				}
 			}
 		}
 
-		if (Contexts.Num() != OutResources.Num())
+		if (!bResult || Contexts.Num() != OutResources.Num())
 		{
 			OutResources.Empty();
 		}
@@ -471,10 +617,7 @@ bool FDisplayClusterViewportProxy::ImplGetResourcesWithRects_RenderThread(const
 	{
 		if (InRecursionDepth < DisplayClusterViewportProxyResourcesOverrideRecursionDepthMax)
 		{
-			if (FDisplayClusterViewportProxy const* OverrideViewportProxy = Owner->ImplFindViewport_RenderThread(RenderSettings.OverrideViewportId))
-			{
-				return OverrideViewportProxy->ImplGetResourcesWithRects_RenderThread(InResourceType, OutResources, OutResourceRects, InRecursionDepth + 1);
-			}
+			return GetRenderingViewportProxy().ImplGetResourcesWithRects_RenderThread(InResourceType, OutResources, OutResourceRects, InRecursionDepth + 1);
 		}
 
 		return false;
@@ -521,25 +664,66 @@ bool FDisplayClusterViewportProxy::ImplGetResourcesWithRects_RenderThread(const
 	return true;
 }
 
-// Apply postprocess, generate mips, etc from settings in FDisplayClusterViewporDeferredUpdateSettings
 void FDisplayClusterViewportProxy::UpdateDeferredResources(FRHICommandListImmediate& RHICmdList) const
 {
 	check(IsInRenderingThread());
 
-	if (RenderTargets.Num() == 0 && PostRenderSettings.Replace.IsEnabled() == false)
+	if (RenderSettings.bFreezeRendering || RenderSettings.bSkipRendering)
 	{
-		// Internal RTT required for deferred update. Except when use 'Replace' as source
+		// Disable deferred update
 		return;
 	}
 
-	if (RenderSettings.bFreezeRendering || RenderSettings.bSkipRendering || RenderSettings.OverrideViewportId.IsEmpty() == false)
+	if (RenderSettings.IsViewportOverrided() && RenderSettings.ViewportOverrideMode == EDisplayClusterViewportOverrideMode::All)
 	{
-		// Disable deferred update
+		// Disable deferred update for clone viewports
 		return;
 	}
 
-	// Pass 0: Resolve from RTT region to separated viewport context resource:
-	ResolveResources_RenderThread(RHICmdList, EDisplayClusterViewportResourceType::InternalRenderTargetResource, EDisplayClusterViewportResourceType::InputShaderResource);
+	const FDisplayClusterViewportProxy& SourceViewportProxy = GetRenderingViewportProxy();
+	if(!SourceViewportProxy.IsInputRenderTargetResourceExists())
+	{
+		// No input RTT resource for deferred update
+		return;
+	}
+
+	bool bPass0Applied = false;
+
+	// Support OCIO on pass 0
+	if (OpenColorIO.IsValid() && OpenColorIO->IsEnabled_RenderThread())
+	{
+		TArray<FRHITexture2D*> Input, Output;
+		TArray<FIntRect> InputRects, OutputRects;
+		if (SourceViewportProxy.GetResourcesWithRects_RenderThread(EDisplayClusterViewportResourceType::InternalRenderTargetResource, Input, InputRects)
+			&& GetResourcesWithRects_RenderThread(EDisplayClusterViewportResourceType::InputShaderResource, Output, OutputRects)
+			&& Input.Num() == Output.Num())
+		{
+			FRDGBuilder GraphBuilder(RHICmdList);
+
+			for (int32 ContextNum = 0; ContextNum < Input.Num(); ContextNum++)
+			{
+				if (OpenColorIO->AddPass_RenderThread(
+					GraphBuilder,
+					SourceViewportProxy.GetContexts_RenderThread()[ContextNum],
+					Input[ContextNum], InputRects[ContextNum],
+					Output[ContextNum], OutputRects[ContextNum]))
+				{
+					bPass0Applied = true;
+				}
+			}
+
+			if (bPass0Applied)
+			{
+				GraphBuilder.Execute();
+			}
+		}
+	}
+	
+	if(!bPass0Applied)
+	{
+		// Pass 0: Resolve from RTT region to separated viewport context resource:
+		ResolveResources_RenderThread(RHICmdList, EDisplayClusterViewportResourceType::InternalRenderTargetResource, EDisplayClusterViewportResourceType::InputShaderResource);
+	}
 
 	// Pass 1: Generate blur postprocess effect for render target texture rect for all contexts
 	if (PostRenderSettings.PostprocessBlur.IsEnabled())
@@ -578,7 +762,7 @@ void FDisplayClusterViewportProxy::UpdateDeferredResources(FRHICommandListImmedi
 }
 
 template<class TScreenPixelShader>
-void ResampleCopyTextureImpl_RenderThread(FRHICommandListImmediate& RHICmdList, FRHITexture* SrcTexture, FRHITexture* DstTexture, const FIntRect& SrcRect, const FIntRect& DstRect, const bool bCopyOnlyAlphaChannel)
+void ResampleCopyTextureImpl_RenderThread(FRHICommandListImmediate& RHICmdList, FRHITexture* SrcTexture, FRHITexture* DstTexture, const FIntRect& SrcRect, const FIntRect& DstRect, const EDisplayClusterTextureCopyMode InCopyMode = EDisplayClusterTextureCopyMode::RGBA)
 {
 	// Texture format mismatch, use a shader to do the copy.
 	// #todo-renderpasses there's no explicit resolve here? Do we need one?
@@ -598,23 +782,27 @@ void ResampleCopyTextureImpl_RenderThread(FRHICommandListImmediate& RHICmdList,
 		FGraphicsPipelineStateInitializer GraphicsPSOInit;
 		RHICmdList.ApplyCachedRenderTargets(GraphicsPSOInit);
 
-		if(bCopyOnlyAlphaChannel)
+		switch (InCopyMode)
 		{
+		case EDisplayClusterTextureCopyMode::Alpha:
 			// Copy alpha channel from source to dest
-			// RT0ColorWriteMask,RT0ColorBlendOp,RT0ColorSrcBlend,RT0ColorDestBlend,RT0AlphaBlendOp,RT0AlphaSrcBlend,RT0AlphaDestBlend,
 			GraphicsPSOInit.BlendState = TStaticBlendState <CW_ALPHA, BO_Add, BF_One, BF_Zero, BO_Add, BF_One, BF_Zero>::GetRHI();
-		}
-		else
-		{
-		GraphicsPSOInit.BlendState = TStaticBlendState<>::GetRHI();
+			break;
+
+		case EDisplayClusterTextureCopyMode::RGB:
+			// Copy only RGB channels from source to dest
+			GraphicsPSOInit.BlendState = TStaticBlendState <CW_RGB, BO_Add, BF_One, BF_Zero, BO_Add, BF_One, BF_Zero>::GetRHI();
+			break;
+
+		case EDisplayClusterTextureCopyMode::RGBA:
+		default:
+			GraphicsPSOInit.BlendState = TStaticBlendState<>::GetRHI();
+			break;
 		}
 
 		GraphicsPSOInit.RasterizerState = TStaticRasterizerState<>::GetRHI();
 		GraphicsPSOInit.DepthStencilState = TStaticDepthStencilState<false, CF_Always>::GetRHI();
 
-
-
-
 		FGlobalShaderMap* ShaderMap = GetGlobalShaderMap(GMaxRHIFeatureLevel);
 		TShaderMapRef<FScreenVS> VertexShader(ShaderMap);
 		TShaderMapRef<TScreenPixelShader> PixelShader(ShaderMap);
@@ -660,7 +848,7 @@ void ImplResolveResource(FRHICommandListImmediate& RHICmdList, FRHITexture2D* In
 	if (bOutputIsPreviewResource)
 	{
 		// Preview require a normal alpha (re-invert)
-		ResampleCopyTextureImpl_RenderThread<FScreenPSInvertAlpha>(RHICmdList, InputResource, OutputResource, InputRect, OutputRect, false);
+		ResampleCopyTextureImpl_RenderThread<FScreenPSInvertAlpha>(RHICmdList, InputResource, OutputResource, InputRect, OutputRect);
 	}
 	else if (InputRect.Size() == OutputRect.Size() && InputResource->GetFormat() == OutputResource->GetFormat())
 	{
@@ -675,7 +863,7 @@ void ImplResolveResource(FRHICommandListImmediate& RHICmdList, FRHITexture2D* In
 	}
 	else
 	{
-		ResampleCopyTextureImpl_RenderThread<FScreenPS>(RHICmdList, InputResource, OutputResource, InputRect, OutputRect, false);
+		ResampleCopyTextureImpl_RenderThread<FScreenPS>(RHICmdList, InputResource, OutputResource, InputRect, OutputRect);
 	}
 }
 
@@ -762,10 +950,177 @@ bool FDisplayClusterViewportProxy::ImplResolveResources_RenderThread(FRHICommand
 	return false;
 }
 
+bool FDisplayClusterViewportProxy::CopyResource_RenderThread(FRDGBuilder& GraphBuilder, const EDisplayClusterTextureCopyMode InCopyMode, const int32 InContextNum, const EDisplayClusterViewportResourceType InSrcResourceType, const EDisplayClusterViewportResourceType InDestResourceType)
+{
+	TArray<FRHITexture2D*> SrcResources;
+	TArray<FIntRect> SrcResourceRects;
+	if (GetResourcesWithRects_RenderThread(InSrcResourceType, SrcResources, SrcResourceRects))
+	{
+		check(SrcResources.IsValidIndex(InContextNum));
+		check(SrcResourceRects.IsValidIndex(InContextNum));
+
+		const FIntRect SrcRect = SrcResourceRects[InContextNum];
+
+		FRDGTextureRef SrcTextureRef = RegisterExternalTexture(GraphBuilder, SrcResources[InContextNum]->GetTexture2D(), TEXT("DCViewportProxyCopySrcResource"));
+		GraphBuilder.SetTextureAccessFinal(SrcTextureRef, ERHIAccess::SRVGraphics);
+
+		return CopyResource_RenderThread(GraphBuilder, InCopyMode, InContextNum, SrcTextureRef, SrcRect, InDestResourceType);
+	}
+
+	return false;
+}
+
+BEGIN_SHADER_PARAMETER_STRUCT(FDisplayClusterCopyTextureParameters, )
+RDG_TEXTURE_ACCESS(Input, ERHIAccess::CopySrc)
+RDG_TEXTURE_ACCESS(Output, ERHIAccess::CopyDest)
+END_SHADER_PARAMETER_STRUCT()
+
+bool FDisplayClusterViewportProxy::CopyResource_RenderThread(FRDGBuilder& GraphBuilder, const EDisplayClusterTextureCopyMode InCopyMode, const int32 InContextNum, FRDGTextureRef InSrcTextureRef, const FIntRect& InSrcRect, const EDisplayClusterViewportResourceType InDestResourceType)
+{
+	TArray<FRHITexture2D*> DestResources;
+	TArray<FIntRect> DestResourceRects;
+	if (HasBeenProduced(InSrcTextureRef) && GetResourcesWithRects_RenderThread(InDestResourceType, DestResources, DestResourceRects))
+	{
+		check(DestResources.IsValidIndex(InContextNum));
+		check(DestResourceRects.IsValidIndex(InContextNum));
+
+		FRDGTextureRef DestTextureRef = RegisterExternalTexture(GraphBuilder, DestResources[InContextNum]->GetTexture2D(), TEXT("DCViewportProxyCopyDestResource"));
+		const FIntRect DestRect = DestResourceRects[InContextNum];
+
+		FDisplayClusterCopyTextureParameters* PassParameters = GraphBuilder.AllocParameters<FDisplayClusterCopyTextureParameters>();
+		PassParameters->Input = InSrcTextureRef;
+		PassParameters->Output = DestTextureRef;
+		GraphBuilder.SetTextureAccessFinal(DestTextureRef, ERHIAccess::RTV);
+
+		GraphBuilder.AddPass(
+			RDG_EVENT_NAME("DisplayClusterViewportProxy_CopyResource(%s)", *GetId()),
+			PassParameters,
+			ERDGPassFlags::Copy | ERDGPassFlags::NeverCull,
+			[InCopyMode, InContextNum, InSrcTextureRef, InSrcRect, DestTextureRef, DestRect](FRHICommandListImmediate& RHICmdList)
+			{
+				ResampleCopyTextureImpl_RenderThread<FScreenPS>(RHICmdList, InSrcTextureRef->GetRHI(), DestTextureRef->GetRHI(), InSrcRect, DestRect, InCopyMode);
+			});
+
+		return true;
+	}
+
+	return false;
+}
+
+bool FDisplayClusterViewportProxy::CopyResource_RenderThread(FRDGBuilder& GraphBuilder, const EDisplayClusterTextureCopyMode InCopyMode, const int32 InContextNum, const EDisplayClusterViewportResourceType InSrcResourceType, FRDGTextureRef InDestTextureRef, const FIntRect& InDestRect)
+{
+	TArray<FRHITexture2D*> SrcResources;
+	TArray<FIntRect> SrcResourceRects;
+	if (HasBeenProduced(InDestTextureRef) && GetResourcesWithRects_RenderThread(InSrcResourceType, SrcResources, SrcResourceRects))
+	{
+		check(SrcResources.IsValidIndex(InContextNum));
+		check(SrcResourceRects.IsValidIndex(InContextNum));
+
+		FRDGTextureRef SrcTextureRef = RegisterExternalTexture(GraphBuilder, SrcResources[InContextNum]->GetTexture2D(), TEXT("DCViewportProxyCopySrcResource"));
+		const FIntRect SrcRect = SrcResourceRects[InContextNum];
+
+		FDisplayClusterCopyTextureParameters* PassParameters = GraphBuilder.AllocParameters<FDisplayClusterCopyTextureParameters>();
+		PassParameters->Input = SrcTextureRef;
+		PassParameters->Output = InDestTextureRef;
+		GraphBuilder.SetTextureAccessFinal(SrcTextureRef, ERHIAccess::RTV);
+
+		GraphBuilder.AddPass(
+			RDG_EVENT_NAME("DisplayClusterViewportProxy_CopyResource(%s)", *GetId()),
+			PassParameters,
+			ERDGPassFlags::Copy | ERDGPassFlags::NeverCull,
+			[InCopyMode, InContextNum, SrcTextureRef, SrcRect, InDestTextureRef, InDestRect](FRHICommandListImmediate& RHICmdList)
+			{
+				ResampleCopyTextureImpl_RenderThread<FScreenPS>(RHICmdList, SrcTextureRef->GetRHI(), InDestTextureRef->GetRHI(), SrcRect, InDestRect, InCopyMode);
+			});
+
+		return true;
+	}
+
+	return false;
+}
+
+bool FDisplayClusterViewportProxy::ShouldUseAlphaChannel_RenderThread() const
+{
+	// Chromakey and Light Cards use alpha channel
+	if (EnumHasAnyFlags(RenderSettingsICVFX.RuntimeFlags, EDisplayClusterViewportRuntimeICVFXFlags::Lightcard | EDisplayClusterViewportRuntimeICVFXFlags::Chromakey))
+	{
+		return true;
+	}
+
+	return false;
+}
+
+bool FDisplayClusterViewportProxy::ShouldUsePostProcessPassAfterSSRInput() const
+{	
+	if (ShouldUseAlphaChannel_RenderThread())
+	{
+		return Owner->GetRenderFrameSettings_RenderThread().AlphaChannelCaptureMode == EDisplayClusterRenderFrameAlphaChannelCaptureMode::ThroughTonemapper;
+	}
+
+	return false;
+}
+
+bool FDisplayClusterViewportProxy::ShouldUsePostProcessPassAfterFXAA() const
+{
+	if (ShouldUseAlphaChannel_RenderThread())
+	{
+		return Owner->GetRenderFrameSettings_RenderThread().AlphaChannelCaptureMode == EDisplayClusterRenderFrameAlphaChannelCaptureMode::ThroughTonemapper;
+	}
+
+	return false;
+}
+
 void FDisplayClusterViewportProxy::OnResolvedSceneColor_RenderThread(FRDGBuilder& GraphBuilder, const FSceneTextures& SceneTextures, const FDisplayClusterViewportProxy_Context& InProxyContext)
-{ }
+{
+	const uint32 InContextNum = InProxyContext.ContextNum;
+	if (ShouldUseAlphaChannel_RenderThread())
+	{
+		switch (Owner->GetRenderFrameSettings_RenderThread().AlphaChannelCaptureMode)
+		{
+		case EDisplayClusterRenderFrameAlphaChannelCaptureMode::FXAA:
+		case EDisplayClusterRenderFrameAlphaChannelCaptureMode::Copy:
+		case EDisplayClusterRenderFrameAlphaChannelCaptureMode::CopyAA:
+		{
+			const FIntRect SrcRect = GetFinalContextRect(EDisplayClusterViewportResourceType::InternalRenderTargetResource, Contexts[InContextNum].RenderTargetRect);
+			// Copy alpha channel from 'SceneTextures.Color.Resolve' to 'InputShaderResource'
+			CopyResource_RenderThread(GraphBuilder, EDisplayClusterTextureCopyMode::Alpha, InContextNum, SceneTextures.Color.Resolve, SrcRect, EDisplayClusterViewportResourceType::InputShaderResource);
+		}
+		break;
+
+		default:
+			break;
+		}
+	}
+}
 
-void FDisplayClusterViewportProxy::PostRenderViewFamily_RenderThread(FRDGBuilder& GraphBuilder, FSceneViewFamily& InViewFamily, const FSceneView& InSceneView, const FDisplayClusterViewportProxy_Context& InProxyContext)
+FScreenPassTexture FDisplayClusterViewportProxy::OnPostProcessPassAfterSSRInput_RenderThread(FRDGBuilder& GraphBuilder, const FSceneView& View, const FPostProcessMaterialInputs& Inputs, const uint32 ContextNum)
+{
+	FScreenPassTexture OutScreenPassTexture = FDisplayClusterViewportManagerViewExtension::ReturnUntouchedSceneColorForPostProcessing(Inputs);
+	if (OutScreenPassTexture.IsValid())
+	{
+		// Copy alpha channel to 'InputShaderResource'
+		const FIntRect SrcRect = GetFinalContextRect(EDisplayClusterViewportResourceType::InternalRenderTargetResource, Contexts[ContextNum].RenderTargetRect);
+		CopyResource_RenderThread(GraphBuilder, EDisplayClusterTextureCopyMode::Alpha, ContextNum, OutScreenPassTexture.Texture, SrcRect, EDisplayClusterViewportResourceType::InputShaderResource);
+	}
+
+	return OutScreenPassTexture;
+}
+
+FScreenPassTexture FDisplayClusterViewportProxy::OnPostProcessPassAfterFXAA_RenderThread(FRDGBuilder& GraphBuilder, const FSceneView& View, const FPostProcessMaterialInputs& Inputs, const uint32 ContextNum)
+{
+	FScreenPassTexture OutScreenPassTexture = FDisplayClusterViewportManagerViewExtension::ReturnUntouchedSceneColorForPostProcessing(Inputs);
+	if (OutScreenPassTexture.IsValid())
+	{
+		// Restore alpha channel after OCIO
+		// Copy alpha channel from 'InputShaderResource'
+		const FIntRect DestRect = GetFinalContextRect(EDisplayClusterViewportResourceType::InternalRenderTargetResource, Contexts[ContextNum].RenderTargetRect);
+		CopyResource_RenderThread(GraphBuilder, EDisplayClusterTextureCopyMode::Alpha, ContextNum, EDisplayClusterViewportResourceType::InputShaderResource, OutScreenPassTexture.Texture, DestRect);
+	}
+
+	return OutScreenPassTexture;
+}
+
+void FDisplayClusterViewportProxy::OnPostRenderViewFamily_RenderThread(FRDGBuilder& GraphBuilder, FSceneViewFamily& InViewFamily, const FSceneView& InSceneView, const FDisplayClusterViewportProxy_Context& InProxyContext)
 {
 	const uint32 InContextNum = InProxyContext.ContextNum;
 
@@ -781,6 +1136,12 @@ void FDisplayClusterViewportProxy::PostRenderViewFamily_RenderThread(FRDGBuilder
 	}
 #endif
 
+	if (Contexts.IsValidIndex(InContextNum))
+	{
+		Contexts[InContextNum].RenderThreadData.EngineDisplayGamma = InSceneView.Family->RenderTarget->GetDisplayGamma();
+		Contexts[InContextNum].RenderThreadData.EngineShowFlags = InSceneView.Family->EngineShowFlags;
+	}
+
 	if (!InProxyContext.ViewFamilyProfileDescription.IsEmpty())
 	{
 		static IDisplayClusterCallbacks& DCCallbacksAPI = IDisplayCluster::Get().GetCallbacks();
@@ -790,4 +1151,66 @@ void FDisplayClusterViewportProxy::PostRenderViewFamily_RenderThread(FRDGBuilder
 			DCCallbacksAPI.OnDisplayClusterPostRenderViewFamily_RenderThread().Broadcast(GraphBuilder, InViewFamily, this);
 		}
 	}
+
+	if (ShouldUseAlphaChannel_RenderThread())
+	{
+		switch (Owner->GetRenderFrameSettings_RenderThread().AlphaChannelCaptureMode)
+		{
+		case EDisplayClusterRenderFrameAlphaChannelCaptureMode::FXAA:
+		{
+			// Restore alpha channed
+			EFXAAQuality FXAAQuality = EFXAAQuality::Q0;
+			if (GetFXAAQuality(RenderSettings.CaptureMode, FXAAQuality))
+			{
+				// Apply FXAA for RGB only
+				// Note: Add AA for alpha channel
+
+				// Copy Alpha channels back from'InputShaderResource' to 'InternalRenderTargetResource'
+				CopyResource_RenderThread(GraphBuilder, EDisplayClusterTextureCopyMode::Alpha, InContextNum, EDisplayClusterViewportResourceType::InputShaderResource, EDisplayClusterViewportResourceType::InternalRenderTargetResource);
+
+				checkSlow(InSceneView.bIsViewInfo);
+				const FViewInfo& ViewInfo = static_cast<const FViewInfo&>(InSceneView);
+
+				// 1. Copy RGB channels from 'InternalRenderTargetResource' to 'InputShaderResource'
+				CopyResource_RenderThread(GraphBuilder, EDisplayClusterTextureCopyMode::RGB, InContextNum, EDisplayClusterViewportResourceType::InternalRenderTargetResource, EDisplayClusterViewportResourceType::InputShaderResource);
+
+				check(InputShaderResources.IsValidIndex(InContextNum));
+				if (FDisplayClusterViewportTextureResource* InputResource = InputShaderResources[InContextNum])
+				{
+					// 2. FXAA render pass to 'InputShaderResource'
+					// Input is 'InputShaderResource'
+					FRDGTextureRef InputTexture = RegisterExternalTexture(GraphBuilder, InputResource->GetTextureRHI(), TEXT("DCViewportProxyFXAAResource"));
+					GraphBuilder.SetTextureAccessFinal(InputTexture, ERHIAccess::RTV);
+
+					FFXAAInputs PassInputs;
+					PassInputs.SceneColor = FScreenPassTexture(InputTexture);
+					PassInputs.Quality = FXAAQuality;
+
+					// 2.1. Do FXAA
+					FScreenPassTexture OutputColorTexture = AddFXAAPass(GraphBuilder, ViewInfo, PassInputs);
+
+					// 2.2. Copy FXAA result from 'OutputTexture' to the 'InternalRenderTargetResource'
+					if (OutputColorTexture.Texture)
+					{
+						const FIntVector OutputTextureSize = OutputColorTexture.Texture->Desc.GetSize();
+						const FIntRect SrcRect(FIntPoint(0, 0), FIntPoint(OutputTextureSize.X, OutputTextureSize.Y));
+
+						CopyResource_RenderThread(GraphBuilder, EDisplayClusterTextureCopyMode::RGB, InContextNum, OutputColorTexture.Texture, SrcRect, EDisplayClusterViewportResourceType::InternalRenderTargetResource);
+					}
+				}
+				break;
+			}
+			// don't break (FXAA not used, just copy alpha)
+		}
+
+		case EDisplayClusterRenderFrameAlphaChannelCaptureMode::Copy:
+		case EDisplayClusterRenderFrameAlphaChannelCaptureMode::CopyAA:
+			// Copy Alpha channels back from'InputShaderResource' to 'InternalRenderTargetResource'
+			CopyResource_RenderThread(GraphBuilder, EDisplayClusterTextureCopyMode::Alpha, InContextNum, EDisplayClusterViewportResourceType::InputShaderResource, EDisplayClusterViewportResourceType::InternalRenderTargetResource);
+			break;
+
+		default:
+			break;
+		}
+	}
 }
diff --git a/Engine/Plugins/Runtime/nDisplay/Source/DisplayCluster/Private/Render/Viewport/DisplayClusterViewportProxy.h b/Engine/Plugins/Runtime/nDisplay/Source/DisplayCluster/Private/Render/Viewport/DisplayClusterViewportProxy.h
--- a/Engine/Plugins/Runtime/nDisplay/Source/DisplayCluster/Private/Render/Viewport/DisplayClusterViewportProxy.h
+++ b/Engine/Plugins/Runtime/nDisplay/Source/DisplayCluster/Private/Render/Viewport/DisplayClusterViewportProxy.h
@@ -22,7 +22,27 @@ class IDisplayClusterShaders;
 class IDisplayClusterProjectionPolicy;
 class IDisplayClusterRender_MeshComponent;
 class FDisplayClusterViewportReadPixelsData;
+class FDisplayClusterViewport_OpenColorIO;
+
+class FRDGBuilder;
+class FSceneView;
+struct FPostProcessMaterialInputs;
+struct FScreenPassTexture;
+struct FSceneTextures;
+
+/**
+ * Copy mode for texture channels.
+ */
+enum class EDisplayClusterTextureCopyMode : uint8
+{
+	RGBA = 0,
+	RGB,
+	Alpha,
+};
 
+/**
+ * nDisplay viewport proxy implementation.
+ */
 class FDisplayClusterViewportProxy
 	: public IDisplayClusterViewportProxy
 	, public TSharedFromThis<FDisplayClusterViewportProxy, ESPMode::ThreadSafe>
@@ -90,9 +110,6 @@ public:
 		Contexts.Append(InContexts);
 	}
 
-	// Apply postprocess, generate mips, etc from settings in FDisplayClusterViewporDeferredUpdateSettings
-	void UpdateDeferredResources(FRHICommandListImmediate& RHICmdList) const;
-
 	//  Return viewport scene proxy resources by type
 	virtual bool GetResources_RenderThread(const EDisplayClusterViewportResourceType InResourceType, TArray<FRHITexture2D*>& OutResources) const override;
 	virtual bool GetResourcesWithRects_RenderThread(const EDisplayClusterViewportResourceType InResourceType, TArray<FRHITexture2D*>& OutResources, TArray<FIntRect>& OutRects) const override;
@@ -103,17 +120,79 @@ public:
 	virtual EDisplayClusterViewportResourceType GetOutputResourceType_RenderThread() const override;
 
 	virtual const IDisplayClusterViewportManagerProxy& GetOwner_RenderThread() const override;
-
-	virtual void OnResolvedSceneColor_RenderThread(FRDGBuilder& GraphBuilder, const FSceneTextures& SceneTextures, const FDisplayClusterViewportProxy_Context& InProxyContext) override;
-	virtual void PostRenderViewFamily_RenderThread(FRDGBuilder& InGraphBuilder, class FSceneViewFamily& InViewFamily, const class FSceneView& InSceneView, const FDisplayClusterViewportProxy_Context& InProxyContext) override;
-
 	///////////////////////////////
 	// ~IDisplayClusterViewportProxy
 	///////////////////////////////
 
-	// Release internal resource refs
+	/** Release internal resource refs. */
 	void HandleResourceDelete_RenderThread(class FDisplayClusterViewportResource* InDeletedResourcePtr);
 
+	/** Resolve viewport RTT: render OCIO, PP, generate MIPS, etc.
+	 *
+	 * @param RHICmdList - RHI interface
+	 *
+	 * @return - none
+	 */
+	void UpdateDeferredResources(FRHICommandListImmediate& RHICmdList) const;
+
+	/** nDisplay VE Callback [subscribed to Renderer:ResolvedSceneColorCallbacks].
+	 *
+	 * @param GraphBuilder   - RDG interface
+	 * @param SceneTextures  - Scene textures (SceneColor, Depth, etc)
+	 * @param InProxyContext - Saved from game thread context for this proxy
+	 *
+	 * @return - none
+	 */
+	void OnResolvedSceneColor_RenderThread(FRDGBuilder& GraphBuilder, const FSceneTextures& SceneTextures, const FDisplayClusterViewportProxy_Context& InProxyContext);
+
+	/** nDisplay VE Callback [PostRenderViewFamily_RenderThread()].
+	 *
+	 * @param GraphBuilder   - RDG interface
+	 * @param InViewFamily   - Scene View Family
+	 * @param InSceneView    - Scene View
+	 * @param InProxyContext - Saved from game thread context for this proxy
+	 *
+	 * @return - none
+	 */
+	void OnPostRenderViewFamily_RenderThread(FRDGBuilder& InGraphBuilder, FSceneViewFamily& InViewFamily, const FSceneView& InSceneView, const FDisplayClusterViewportProxy_Context& InProxyContext);
+
+	/** Allow callback OnPostProcessPassAfterFXAA.  */
+	bool ShouldUsePostProcessPassAfterFXAA() const;
+
+	/** Callback OnPostProcessPassAfterFXAA.
+	 *
+	 * @param GraphBuilder - RDG interface
+	 * @param View         - Scene View
+	 * @param Inputs       - PP Input resources
+	 * @param ContextNum   - viewport context index
+	 *
+	 * @return - Screen pass texture
+	 */
+	FScreenPassTexture OnPostProcessPassAfterFXAA_RenderThread(FRDGBuilder& GraphBuilder, const FSceneView& View, const FPostProcessMaterialInputs& Inputs, const uint32 ContextNum);
+
+	/** Allow callback OnPostProcessPassAfterSSRInput. */
+	bool ShouldUsePostProcessPassAfterSSRInput() const;
+
+	/** Callback OnPostProcessPassAfterSSRInput.
+	 *
+	 * @param GraphBuilder - RDG interface
+	 * @param View         - Scene View
+	 * @param Inputs       - PP Input resources
+	 * @param ContextNum   - viewport context index
+	 *
+	 * @return - Screen pass texture
+	 */
+	FScreenPassTexture OnPostProcessPassAfterSSRInput_RenderThread(FRDGBuilder& GraphBuilder, const FSceneView& View, const FPostProcessMaterialInputs& Inputs, const uint32 ContextNum);
+
+	/** Enable alpha channel for this viewport (useful for overlays with alpha channel: ChromaKey, LightCard). */
+	bool ShouldUseAlphaChannel_RenderThread() const;
+
+	/** AfterWarp resolve viewport resources: ViewportRemap, etc
+	 *
+	 * @param RHICmdList - RHI interface
+	 *
+	 * @return - none
+	 */
 	void PostResolveViewport_RenderThread(FRHICommandListImmediate& RHICmdList) const;
 
 #if WITH_EDITOR
@@ -152,15 +231,65 @@ private:
 	bool ImplResolveResources_RenderThread(FRHICommandListImmediate& RHICmdList, FDisplayClusterViewportProxy const* SourceProxy, const EDisplayClusterViewportResourceType InputResourceType, const EDisplayClusterViewportResourceType OutputResourceType, const int32 InContextNum) const;
 	bool IsShouldOverrideViewportResource(const EDisplayClusterViewportResourceType InResourceType) const;
 
+	/** Copy pixel channels between resources.
+	 *
+	 * @param GraphBuilder       - RDG interface
+	 * @param InCopyMode         - channels copy mode
+	 * @param InContextNum       - viewport context
+	 * @param InSrcTextureRef    - Input resource
+	 * @param InSrcRect          - Input resource rect
+	 * @param InDestResourceType - Dest viewport resource type
+	 *
+	 * @return - true, if success
+	 */
+	bool CopyResource_RenderThread(FRDGBuilder& GraphBuilder, const EDisplayClusterTextureCopyMode InCopyMode, const int32 InContextNum, FRDGTextureRef InSrcTextureRef, const FIntRect& InSrcRect, const EDisplayClusterViewportResourceType InDestResourceType);
+
+	/** Copy pixel channels between resources.
+	 *
+	 * @param GraphBuilder      - RDG interface
+	 * @param InCopyMode        - channels copy mode
+	 * @param InContextNum      - viewport context
+	 * @param InDestTextureRef  - Dest resource
+	 * @param InDestRect        - Dest resource rect
+	 *
+	 * @return - true, if success
+	 */
+	bool CopyResource_RenderThread(FRDGBuilder& GraphBuilder, const EDisplayClusterTextureCopyMode InCopyMode, const int32 InContextNum, const EDisplayClusterViewportResourceType InSrcResourceType, FRDGTextureRef InDestTextureRef, const FIntRect& InDestRect);
+
+	/** Copy pixel channels between resources.
+	 *
+	 * @param GraphBuilder       - RDG interface
+	 * @param InCopyMode         - channels copy mode
+	 * @param InContextNum       - viewport context
+	 * @param InSrcResourceType  - Input viewport resource type
+	 * @param InDestResourceType - Dest viewport resource type
+	 *
+	 * @return - true, if success
+	 */
+	bool CopyResource_RenderThread(FRDGBuilder& GraphBuilder, const EDisplayClusterTextureCopyMode InCopyMode, const int32 InContextNum, const EDisplayClusterViewportResourceType InSrcResourceType, const EDisplayClusterViewportResourceType InDestResourceType);
+
+	/** Return viewport used to render RTT (support ViewportOverride)
+	 *
+	 * @return - viewport proxy with valid render resources
+	 */
+	const FDisplayClusterViewportProxy& GetRenderingViewportProxy() const;
+
+	/** Check if there is an RTT source (internal or external) in this viewport proxy. */
+	bool IsInputRenderTargetResourceExists() const;
+
 protected:
 	friend FDisplayClusterViewportProxyData;
 	friend FDisplayClusterViewportManagerProxy;
 
-	// Unique viewport name
+	/** Unique viewport name. */
 	const FString ViewportId;
 
+	/** Cluster node name. */
 	const FString ClusterNodeId;
 
+	/** OpenColorIO nDisplay interface ref. */
+	TSharedPtr<FDisplayClusterViewport_OpenColorIO, ESPMode::ThreadSafe> OpenColorIO;
+
 	// Viewport render params
 	mutable FDisplayClusterViewport_RenderSettings       RenderSettings;
 	FDisplayClusterViewport_RenderSettingsICVFX  RenderSettingsICVFX;
diff --git a/Engine/Plugins/Runtime/nDisplay/Source/DisplayCluster/Private/Render/Viewport/DisplayClusterViewportStrings.h b/Engine/Plugins/Runtime/nDisplay/Source/DisplayCluster/Private/Render/Viewport/DisplayClusterViewportStrings.h
--- a/Engine/Plugins/Runtime/nDisplay/Source/DisplayCluster/Private/Render/Viewport/DisplayClusterViewportStrings.h
+++ b/Engine/Plugins/Runtime/nDisplay/Source/DisplayCluster/Private/Render/Viewport/DisplayClusterViewportStrings.h
@@ -14,10 +14,11 @@ namespace DisplayClusterViewportStrings
 	// ICVFXstrings
 	namespace icvfx
 	{
-		static constexpr const TCHAR* prefix         = TEXT("icvfx");
-		static constexpr const TCHAR* camera         = TEXT("incamera");
-		static constexpr const TCHAR* chromakey      = TEXT("chromakey");
-		static constexpr const TCHAR* lightcard      = TEXT("lightcard");
-		static constexpr const TCHAR* lightcard_OCIO = TEXT("lightcard_ocio");
+		static constexpr const TCHAR* prefix       = TEXT("icvfx");
+		static constexpr const TCHAR* camera       = TEXT("incamera");
+		static constexpr const TCHAR* chromakey    = TEXT("chromakey");
+		static constexpr const TCHAR* lightcard    = TEXT("lightcard");
+		static constexpr const TCHAR* uv_lightcard = TEXT("uv_lightcard");
+		
 	}
 };
\ No newline at end of file
diff --git a/Engine/Plugins/Runtime/nDisplay/Source/DisplayCluster/Private/Render/Viewport/DisplayClusterViewport_EditorPreview.cpp b/Engine/Plugins/Runtime/nDisplay/Source/DisplayCluster/Private/Render/Viewport/DisplayClusterViewport_EditorPreview.cpp
--- a/Engine/Plugins/Runtime/nDisplay/Source/DisplayCluster/Private/Render/Viewport/DisplayClusterViewport_EditorPreview.cpp
+++ b/Engine/Plugins/Runtime/nDisplay/Source/DisplayCluster/Private/Render/Viewport/DisplayClusterViewport_EditorPreview.cpp
@@ -119,7 +119,6 @@ FSceneView* FDisplayClusterViewport::ImplCalcScenePreview(FSceneViewFamilyContex
 			ViewInitOptions.OverlayColor = FLinearColor::Black;
 		}
 
-		ViewInitOptions.bIsSceneCapture = true;
 		ViewInitOptions.bSceneCaptureUsesRayTracing = false;
 		ViewInitOptions.bIsPlanarReflection = false;
 
diff --git a/Engine/Plugins/Runtime/nDisplay/Source/DisplayCluster/Private/Render/Viewport/DisplayClusterViewport_OpenColorIO.cpp b/Engine/Plugins/Runtime/nDisplay/Source/DisplayCluster/Private/Render/Viewport/DisplayClusterViewport_OpenColorIO.cpp
new file mode 100644
--- /dev/null
+++ b/Engine/Plugins/Runtime/nDisplay/Source/DisplayCluster/Private/Render/Viewport/DisplayClusterViewport_OpenColorIO.cpp
@@ -0,0 +1,175 @@
+// Copyright Epic Games, Inc. All Rights Reserved.
+
+#include "Render/Viewport/DisplayClusterViewport_OpenColorIO.h"
+#include "Render/Viewport/DisplayClusterViewport.h"
+
+#include "Render/Viewport/Containers/DisplayClusterViewport_Context.h"
+
+#include "OpenColorIOConfiguration.h"
+#include "OpenColorIOColorTransform.h"
+#include "OpenColorIODisplayManager.h"
+#include "OpenColorIODisplayExtension.h"
+#include "OpenColorIOShader.h"
+
+#include "Shader.h"
+#include "SceneView.h"
+#include "ScreenPass.h"
+#include "CommonRenderResources.h"
+
+namespace {
+	/** Create dummy view for OCIO render pass
+	 *
+	 * @param InViewRect - input texture view rect
+	 *
+	 * @return - new view info
+	 */
+	FViewInfo CreateDummyViewInfo(const FIntRect& InViewRect)
+	{
+		FSceneViewFamily ViewFamily(FSceneViewFamily::ConstructionValues(nullptr, nullptr, FEngineShowFlags(ESFIM_Game))
+			.SetTime(FGameTime())
+			.SetGammaCorrection(1.0f));
+		FSceneViewInitOptions ViewInitOptions;
+		ViewInitOptions.ViewFamily = &ViewFamily;
+		ViewInitOptions.SetViewRectangle(InViewRect);
+		ViewInitOptions.ViewOrigin = FVector::ZeroVector;
+		ViewInitOptions.ViewRotationMatrix = FMatrix::Identity;
+		ViewInitOptions.ProjectionMatrix = FMatrix::Identity;
+
+		return FViewInfo(ViewInitOptions);
+	}
+}
+
+//////////////////////////////////////////////////////////////////////////
+// FDisplayClusterViewport_OpenColorIO
+//////////////////////////////////////////////////////////////////////////
+FDisplayClusterViewport_OpenColorIO::FDisplayClusterViewport_OpenColorIO(const FOpenColorIOColorConversionSettings& InDisplayConfiguration)
+	: DisplayConfiguration(InDisplayConfiguration)
+{ }
+
+FDisplayClusterViewport_OpenColorIO::~FDisplayClusterViewport_OpenColorIO()
+{ }
+
+void FDisplayClusterViewport_OpenColorIO::AddReferencedObjects(FReferenceCollector& Collector)
+{
+	if (DisplayConfiguration.ConfigurationSource)
+	{
+		Collector.AddReferencedObject(DisplayConfiguration.ConfigurationSource);
+	}
+}
+
+void FDisplayClusterViewport_OpenColorIO::SetupSceneView(FSceneViewFamily& InOutViewFamily, FSceneView& InOutView) const
+{
+	if (bShaderResourceValid)
+	{
+		// Code copied from FOpenColorIODisplayExtension::SetupView():
+
+		//Force ToneCurve to be off while we'are alive to make sure the input color space is the working space : srgb linear
+		InOutViewFamily.EngineShowFlags.SetToneCurve(false);
+		// This flags sets tonampper to output to ETonemapperOutputDevice::LinearNoToneCurve
+		InOutViewFamily.SceneCaptureSource = SCS_FinalColorHDR;
+
+		InOutView.FinalPostProcessSettings.bOverride_ToneCurveAmount = 1;
+		InOutView.FinalPostProcessSettings.ToneCurveAmount = 0.0;
+	}
+}
+
+void FDisplayClusterViewport_OpenColorIO::UpdateOpenColorIORenderPassResources()
+{
+	bShaderResourceValid = false;
+	if (!bConfigurationDataValid)
+	{
+		// The data in this configuration was not correctly initialized in the previous frame. Ignore this OCIO
+		return;
+	}
+
+	FOpenColorIOTransformResource* ShaderResource = nullptr;
+	TSortedMap<int32, FTextureResource*> TransformTextureResources;
+
+	if (DisplayConfiguration.ConfigurationSource != nullptr)
+	{
+		const ERHIFeatureLevel::Type FeatureLevel = GMaxRHIFeatureLevel;
+
+		const bool bFoundTransform = DisplayConfiguration.ConfigurationSource->GetRenderResources(
+			FeatureLevel
+			, DisplayConfiguration
+			, ShaderResource
+			, TransformTextureResources);
+
+		if (bFoundTransform)
+		{
+			// Transform was found, so shader must be there but doesn't mean the actual shader is available
+			check(ShaderResource);
+			if (ShaderResource->GetShaderGameThread<FOpenColorIOPixelShader>().IsNull())
+			{
+				ensureMsgf(false, TEXT("Can't apply display look - Shader was invalid for Resource %s"), *ShaderResource->GetFriendlyName());
+
+				//Invalidate shader resource
+				ShaderResource = nullptr;
+
+				// Data in this configuration is corrupted. Ignore this configuration for next frames
+				bConfigurationDataValid = false;
+			}
+			else
+			{
+				bShaderResourceValid = true;
+			}
+		}
+	}
+
+	if (!bShaderResourceValid)
+	{
+		ShaderResource = nullptr;
+		TransformTextureResources.Empty();
+	}
+
+	ENQUEUE_RENDER_COMMAND(ProcessColorSpaceTransform)(
+		[This = SharedThis(this), ShaderResource, TextureResources = MoveTemp(TransformTextureResources)](FRHICommandListImmediate& RHICmdList)
+		{
+			//Caches render thread resource to be used when applying configuration in PostRenderViewFamily_RenderThread
+			This->CachedResourcesRenderThread.ShaderResource = ShaderResource;
+			This->CachedResourcesRenderThread.TextureResources = TextureResources;
+		}
+	);
+}
+
+bool FDisplayClusterViewport_OpenColorIO::IsEnabled_RenderThread() const
+{
+	return CachedResourcesRenderThread.ShaderResource != nullptr;
+}
+
+bool FDisplayClusterViewport_OpenColorIO::IsDisplayConfigurationEquals(const FOpenColorIOColorConversionSettings& InDisplayConfiguration) const
+{
+	return DisplayConfiguration.ToString() == InDisplayConfiguration.ToString();
+}
+
+bool FDisplayClusterViewport_OpenColorIO::AddPass_RenderThread(FRDGBuilder& GraphBuilder, const FDisplayClusterViewport_Context& InViewportContext,
+	FRHITexture2D* InputTextureRHI, const FIntRect& InputRect, FRHITexture2D* OutputTextureRHI, const FIntRect& OutputRect) const
+{
+	if (IsEnabled_RenderThread())
+	{
+		FRDGTextureRef InputTexture = GraphBuilder.RegisterExternalTexture(CreateRenderTarget(InputTextureRHI, TEXT("nDisplayOCIOInputTexture")));
+		FRDGTextureRef OutputTexture = GraphBuilder.RegisterExternalTexture(CreateRenderTarget(OutputTextureRHI, TEXT("nDisplayOCIORenderTargetTexture")));
+
+		FScreenPassRenderTarget Output = FScreenPassRenderTarget(OutputTexture, OutputRect, ERenderTargetLoadAction::EClear);
+
+		const float DefaultDisplayGamma = FOpenColorIODisplayExtension::DefaultDisplayGamma;
+		const float EngineDisplayGamma = InViewportContext.RenderThreadData.EngineDisplayGamma;
+
+		// There is a special case where post processing and tonemapper are disabled. In this case tonemapper applies a static display Inverse of Gamma which defaults to 2.2.
+		// In the case when Both PostProcessing and ToneMapper are disabled we apply gamma manually. In every other case we apply inverse gamma before applying OCIO.
+		const float DisplayGamma = (InViewportContext.RenderThreadData.EngineShowFlags.Tonemapper == 0) || (InViewportContext.RenderThreadData.EngineShowFlags.PostProcessing == 0) ? DefaultDisplayGamma : DefaultDisplayGamma / EngineDisplayGamma;
+
+		FOpenColorIORendering::AddPass_RenderThread(
+			GraphBuilder,
+			CreateDummyViewInfo(Output.ViewRect),
+			FScreenPassTexture(InputTexture),
+			Output,
+			CachedResourcesRenderThread,
+			DisplayGamma
+		);
+
+		return true;
+	}
+
+	return false;
+}
diff --git a/Engine/Plugins/Runtime/nDisplay/Source/DisplayCluster/Private/Render/Viewport/DisplayClusterViewport_OpenColorIO.h b/Engine/Plugins/Runtime/nDisplay/Source/DisplayCluster/Private/Render/Viewport/DisplayClusterViewport_OpenColorIO.h
new file mode 100644
--- /dev/null
+++ b/Engine/Plugins/Runtime/nDisplay/Source/DisplayCluster/Private/Render/Viewport/DisplayClusterViewport_OpenColorIO.h
@@ -0,0 +1,93 @@
+// Copyright Epic Games, Inc. All Rights Reserved.
+
+#pragma once
+
+#include "CoreMinimal.h"
+
+#include "OpenColorIOColorSpace.h"
+#include "OpenColorIORendering.h"
+
+#include "RHI.h"
+#include "RHICommandList.h"
+#include "RHIResources.h"
+#include "Templates/SharedPointer.h"
+
+#include "UObject/GCObject.h"
+
+/**
+ * nDisplay OCIO implementation.
+ * 
+ */
+class FDisplayClusterViewport_OpenColorIO
+	: public TSharedFromThis<FDisplayClusterViewport_OpenColorIO, ESPMode::ThreadSafe>
+	, public FGCObject
+{
+public:
+	FDisplayClusterViewport_OpenColorIO(const FOpenColorIOColorConversionSettings& InDisplayConfiguration);
+	virtual ~FDisplayClusterViewport_OpenColorIO();
+
+	//~Begin FGCObject interface
+	virtual void AddReferencedObjects(FReferenceCollector& Collector) override;
+	virtual FString GetReferencerName() const override
+	{
+		return TEXT("FDisplayClusterViewport_OpenColorIO");
+	}
+	//~End FGCObject interface
+
+public:
+	/** Update render thread resources. */
+	void UpdateOpenColorIORenderPassResources();
+
+	/** Setup view for OCIO.
+	 *
+	 * @param InOutViewFamily - [In,Out] ViewFamily.
+	 * @param InOutView       - [In,Out] View.
+	 *
+	 * @return - none.
+	 */
+	void SetupSceneView(class FSceneViewFamily& InOutViewFamily, class FSceneView& InOutView) const;
+
+	/** Add OCIO render pass.
+	 *
+	 * @param GraphBuilder      - RDG interface.
+	 * @param InViewportContext - DC viewport context.
+	 * @param InputTextureRHI   - Source texture.
+	 * @param InputRect         - Source rec.
+	 * @param OutputTextureRHI  - Destination texture.
+	 * @param OutputRect        - Destination rec.
+	 *
+	 * @return - true if success.
+	 */
+	bool AddPass_RenderThread(class FRDGBuilder& GraphBuilder, const class FDisplayClusterViewport_Context& InViewportContext,
+		FRHITexture2D* InputTextureRHI, const FIntRect& InputRect, FRHITexture2D* OutputTextureRHI, const FIntRect& OutputRect) const;
+
+	/** Is OCIO enabled for render thread */
+	bool IsEnabled_RenderThread() const;
+
+	/** Compare two OCIO configurations.
+	 *
+	 * @param InDisplayConfiguration - configuration to compare with.
+	 *
+	 * @return - true if equal.
+	 */
+	bool IsDisplayConfigurationEquals(const FOpenColorIOColorConversionSettings& InDisplayConfiguration) const;
+
+	/** Get current OCIO configuration. */
+	const FOpenColorIOColorConversionSettings& GetDisplayConfiguration() const
+	{
+		return DisplayConfiguration;
+	}
+
+private:
+	/** Cached pass resources required to apply conversion for render thread. */
+	FOpenColorIORenderPassResources CachedResourcesRenderThread;
+
+	/** Configuration to apply during post render callback. */
+	FOpenColorIOColorConversionSettings DisplayConfiguration;
+
+	/** Shader resources state. */
+	bool bShaderResourceValid = true;
+
+	/** Configuration data state. */
+	bool bConfigurationDataValid = true;
+};
