nDisplay: ICVFX OCIO

diff --git a/Engine/Plugins/Runtime/nDisplay/Source/DisplayCluster/Private/Render/Viewport/LightCard/DisplayClusterViewportLightCardManager.cpp b/Engine/Plugins/Runtime/nDisplay/Source/DisplayCluster/Private/Render/Viewport/LightCard/DisplayClusterViewportLightCardManager.cpp
--- a/Engine/Plugins/Runtime/nDisplay/Source/DisplayCluster/Private/Render/Viewport/LightCard/DisplayClusterViewportLightCardManager.cpp
+++ b/Engine/Plugins/Runtime/nDisplay/Source/DisplayCluster/Private/Render/Viewport/LightCard/DisplayClusterViewportLightCardManager.cpp
@@ -1,6 +1,8 @@
 // Copyright Epic Games, Inc. All Rights Reserved.
 
 #include "DisplayClusterViewportLightCardManager.h"
+#include "DisplayClusterViewportLightCardManagerProxy.h"
+#include "DisplayClusterViewportLightCardResource.h"
 
 #include "DisplayClusterLightCardActor.h"
 #include "Blueprints/DisplayClusterBlueprintLib.h"
@@ -14,44 +16,28 @@
 #include "SceneInterface.h"
 #include "UObject/Package.h"
 
+///////////////////////////////////////////////////////////////////////////////////////////////
 /** Console variable used to control the size of the UV light card map texture */
 static TAutoConsoleVariable<int32> CVarUVLightCardTextureSize(
 	TEXT("nDisplay.render.uvlightcards.UVTextureSize"),
 	512,
-	TEXT("The size of the texture UV light cards are rendered to."));
+	TEXT("The size of the texture UV light cards are rendered to.")
+);
 
-//-----------------------------------------------------------------------------------------------------------------
-// FDisplayClusterLightCardMap
-//-----------------------------------------------------------------------------------------------------------------
-void FDisplayClusterLightCardMap::InitDynamicRHI()
-{
-	ETextureCreateFlags CreateFlags = TexCreate_Dynamic;
-	CreateFlags |= TexCreate_MultiGPUGraphIgnore;
-
-	const FRHITextureCreateDesc Desc = FRHITextureCreateDesc::Create2D(TEXT("DisplayClusterLightCardMap"))
-		.SetExtent(GetSizeX(), GetSizeY())
-		.SetFormat(PF_FloatRGBA)
-		.SetFlags(CreateFlags | ETextureCreateFlags::RenderTargetable | ETextureCreateFlags::ShaderResource)
-		.SetInitialState(ERHIAccess::SRVMask)
-		.SetClearValue(FClearValueBinding::Transparent);
-
-	RenderTargetTextureRHI = TextureRHI = RHICreateTexture(Desc);
-}
-
-//-----------------------------------------------------------------------------------------------------------------
+///////////////////////////////////////////////////////////////////////////////////////////////
 // FDisplayClusterViewportLightCardManager
-//-----------------------------------------------------------------------------------------------------------------
+///////////////////////////////////////////////////////////////////////////////////////////////
 FDisplayClusterViewportLightCardManager::FDisplayClusterViewportLightCardManager(FDisplayClusterViewportManager& InViewportManager)
 	: ViewportManager(InViewportManager)
 {
-	ProxyData = MakeShared<FProxyData, ESPMode::ThreadSafe>();
+	LightCardManagerProxy = MakeShared<FDisplayClusterViewportLightCardManagerProxy, ESPMode::ThreadSafe>();
 }
 
 FDisplayClusterViewportLightCardManager::~FDisplayClusterViewportLightCardManager()
 {
 	Release();
 
-	ProxyData.Reset();
+	LightCardManagerProxy.Reset();
 }
 
 void FDisplayClusterViewportLightCardManager::Release()
@@ -60,28 +46,61 @@ void FDisplayClusterViewportLightCardManager::Release()
 	const bool bIsInRenderingThread = IsInRenderingThread();
 	check(!bIsInRenderingThread || (bIsInRenderingThread && PreviewWorld == nullptr));
 
+	// Release UVLightCard
+	ReleaseUVLightCardData();
+	ReleaseUVLightCardResource();
+
 	// Deleting PreviewScene is only called from the game thread
 	DestroyPreviewWorld();
+}
 
-	UVLightCards.Empty();
-	ReleaseUVLightCardMap();
+///////////////////////////////////////////////////////////////////////////////////////////////
+void FDisplayClusterViewportLightCardManager::UpdateConfiguration()
+{
+	UpdateUVLightCardData();
 }
 
-void FDisplayClusterViewportLightCardManager::AddReferencedObjects(FReferenceCollector& Collector)
+void FDisplayClusterViewportLightCardManager::HandleStartScene()
 {
-	Collector.AddReferencedObject(PreviewWorld);
+	InitializePreviewWorld();
+
+	UpdateUVLightCardData();
 }
 
-FRHITexture* FDisplayClusterViewportLightCardManager::GetUVLightCardMap_RenderThread() const
-{ 
-	check(IsInRenderingThread());
+void FDisplayClusterViewportLightCardManager::HandleEndScene()
+{
+	DestroyPreviewWorld();
 
-	return ProxyData.IsValid() ? ProxyData->GetUVLightCardMap_RenderThread() : nullptr;
+	ReleaseUVLightCardData();
 }
 
-void FDisplayClusterViewportLightCardManager::UpdateConfiguration()
+void FDisplayClusterViewportLightCardManager::RenderFrame()
+{
+	RenderUVLightCard();
+}
+
+///////////////////////////////////////////////////////////////////////////////////////////////
+FIntPoint FDisplayClusterViewportLightCardManager::GetUVLightCardResourceSize() const
+{
+	return UVLightCardResource.IsValid() ? UVLightCardResource->GetSizeXY() : FIntPoint(0, 0);
+}
+
+bool FDisplayClusterViewportLightCardManager::IsUVLightCardEnabled() const
+{
+	return !UVLightCardPrimitiveComponents.IsEmpty();
+}
+
+void FDisplayClusterViewportLightCardManager::ReleaseUVLightCardData()
+{
+	UVLightCardPrimitiveComponents.Empty();
+}
+
+void FDisplayClusterViewportLightCardManager::UpdateUVLightCardData()
 {
-	UVLightCards.Empty();
+	ReleaseUVLightCardData();
+
+	/** The list of UV light card actors that are referenced by the root actor */
+	TArray<ADisplayClusterLightCardActor*> UVLightCardActors;
 
 	if (ADisplayClusterRootActor* RootActorPtr = ViewportManager.GetRootActor())
 	{
@@ -92,72 +111,103 @@ void FDisplayClusterViewportLightCardManager::UpdateConfiguration()
 		{
 			if (LightCard->bIsUVLightCard)
 			{
-				UVLightCards.Add(LightCard);
+				UVLightCardActors.Add(LightCard);
+			}
+		}
+	}
+
+	TArray<UMeshComponent*> LightCardMeshComponents;
+	for (ADisplayClusterLightCardActor* LightCard : UVLightCardActors)
+	{
+		if (LightCard->IsHidden() || LightCard->IsActorBeingDestroyed() || LightCard->GetWorld() == nullptr)
+		{
+			continue;
+		}
+
+		LightCardMeshComponents.Empty(LightCardMeshComponents.Num());
+		LightCard->GetLightCardMeshComponents(LightCardMeshComponents);
+
+		for (UMeshComponent* LightCardMeshComp : LightCardMeshComponents)
+		{
+			if (LightCardMeshComp && LightCardMeshComp->SceneProxy == nullptr)
+			{
+				UVLightCardPrimitiveComponents.Add(LightCardMeshComp);
 			}
 		}
 	}
 }
 
-void FDisplayClusterViewportLightCardManager::HandleStartScene()
+void FDisplayClusterViewportLightCardManager::CreateUVLightCardResource(const FIntPoint& InResourceSize)
 {
-	InitializePreviewWorld();
+	UVLightCardResource = MakeShared<FDisplayClusterViewportLightCardResource>(InResourceSize);
+	LightCardManagerProxy->UpdateUVLightCardResource(UVLightCardResource);
 }
 
-void FDisplayClusterViewportLightCardManager::HandleEndScene()
+void FDisplayClusterViewportLightCardManager::ReleaseUVLightCardResource()
 {
-	DestroyPreviewWorld();
+	if (UVLightCardResource.IsValid())
+	{
+		LightCardManagerProxy->ReleaseUVLightCardResource();
+	}
+
+	UVLightCardResource.Reset();
 }
 
-void FDisplayClusterViewportLightCardManager::RenderFrame()
+void FDisplayClusterViewportLightCardManager::UpdateUVLightCardResource()
 {
-	if (PreviewWorld && ProxyData.IsValid())
+	const uint32 UVLightCardTextureSize = CVarUVLightCardTextureSize.GetValueOnGameThread();
+	const FIntPoint UVLightCardResourceSize = FIntPoint(UVLightCardTextureSize, UVLightCardTextureSize);
+
+	if (UVLightCardResource.IsValid())
 	{
-		if (FSceneInterface* SceneInterface = PreviewWorld->Scene)
+		if (UVLightCardResource->GetSizeXY() != UVLightCardResourceSize)
 		{
-			InitializeUVLightCardMap();
+			ReleaseUVLightCardResource();
+		}
+	}
 
-			/** A list of primitive components that have been added to the preview scene for rendering in the current frame */
-			TArray<UPrimitiveComponent*> LoadedPrimitiveComponents;
+	if (!UVLightCardResource.IsValid())
+	{
+		CreateUVLightCardResource(UVLightCardResourceSize);
+	}
+}
 
-			TArray<UMeshComponent*> LightCardMeshComponents;
+void FDisplayClusterViewportLightCardManager::RenderUVLightCard()
+{
+	// Render UV LightCard:
+	if (IsUVLightCardEnabled() && PreviewWorld && LightCardManagerProxy.IsValid())
+	{
+		UpdateUVLightCardResource();
 
-			bool bLoadedPrimitives = false;
-			for (ADisplayClusterLightCardActor* LightCard : UVLightCards)
+		if (UVLightCardResource.IsValid())
+		{
+			for (UPrimitiveComponent* LoadedComponent : UVLightCardPrimitiveComponents)
 			{
-				if (LightCard->IsHidden() || LightCard->IsActorBeingDestroyed() || LightCard->GetWorld() == nullptr)
-				{
-					continue;
-				}
-
-				LightCardMeshComponents.Empty(LightCardMeshComponents.Num());
-				LightCard->GetLightCardMeshComponents(LightCardMeshComponents);
-
-				for (UMeshComponent* LightCardMeshComp : LightCardMeshComponents)
-				{
-					if (LightCardMeshComp && LightCardMeshComp->SceneProxy == nullptr)
-					{
-						SceneInterface->AddPrimitive(LightCardMeshComp);
-						LoadedPrimitiveComponents.Add(LightCardMeshComp);
-
-						bLoadedPrimitives = true;
-					}
-				}
+				PreviewWorld->Scene->AddPrimitive(LoadedComponent);
 			}
 
-			ENQUEUE_RENDER_COMMAND(DisplayClusterViewportLightCardManager_RenderFrame)(
-				[InProxyData = ProxyData, bLoadedPrimitives, SceneInterface](FRHICommandListImmediate& RHICmdList)
-				{
-					InProxyData->RenderLightCardMap_RenderThread(RHICmdList, bLoadedPrimitives, SceneInterface);
-				});
+			LightCardManagerProxy->RenderUVLightCard(PreviewWorld->Scene, ADisplayClusterLightCardActor::UVPlaneDefaultSize, ViewportManager.ShouldRenderFinalColor());
 
-			for (UPrimitiveComponent* LoadedComponent : LoadedPrimitiveComponents)
+			for (UPrimitiveComponent* LoadedComponent : UVLightCardPrimitiveComponents)
 			{
-				SceneInterface->RemovePrimitive(LoadedComponent);
+				PreviewWorld->Scene->RemovePrimitive(LoadedComponent);
 			}
 		}
 	}
+	else
+	{
+		ReleaseUVLightCardResource();
+	}
+
+}
+
+///////////////////////////////////////////////////////////////////////////////////////////////
+void FDisplayClusterViewportLightCardManager::AddReferencedObjects(FReferenceCollector& Collector)
+{
+	Collector.AddReferencedObject(PreviewWorld);
 }
 
+///////////////////////////////////////////////////////////////////////////////////////////////
 void FDisplayClusterViewportLightCardManager::InitializePreviewWorld()
 {
 	if (!PreviewWorld)
@@ -190,86 +240,9 @@ void FDisplayClusterViewportLightCardManager::DestroyPreviewWorld()
 		PreviewWorld->WorldType = EWorldType::Inactive;
 
 		GEngine->DestroyWorldContext(PreviewWorld);
+
 		PreviewWorld->DestroyWorld(false);
 		PreviewWorld->MarkObjectsPendingKill();
 		PreviewWorld = nullptr;
 	}
 }
-
-void FDisplayClusterViewportLightCardManager::InitializeUVLightCardMap()
-{
-	const uint32 LightCardTextureSize = CVarUVLightCardTextureSize.GetValueOnAnyThread();
-	if (UVLightCardMap && UVLightCardMap->GetSizeX() != LightCardTextureSize)
-	{
-		ReleaseUVLightCardMap();
-	}
-
-	if (UVLightCardMap == nullptr && ProxyData.IsValid())
-	{
-		UVLightCardMap = new FDisplayClusterLightCardMap(LightCardTextureSize);
-
-		ENQUEUE_RENDER_COMMAND(DisplayClusterViewportLightCardManager_InitializeUVLightCardMap)(
-			[InProxyData = ProxyData, InUVLightCardMap = UVLightCardMap](FRHICommandListImmediate& RHICmdList)
-			{
-				InProxyData->InitializeUVLightCardMap_RenderThread(InUVLightCardMap);
-			});
-	}
-}
-
-void FDisplayClusterViewportLightCardManager::ReleaseUVLightCardMap()
-{
-	UVLightCardMap = nullptr;
-
-	if (ProxyData.IsValid())
-	{
-		ENQUEUE_RENDER_COMMAND(DisplayClusterViewportLightCardManager_ReleaseUVLightCardMap)(
-			[InProxyData = ProxyData](FRHICommandListImmediate& RHICmdList)
-			{
-				InProxyData->ReleaseUVLightCardMap_RenderThread();
-			});
-	}
-}
-
-//-----------------------------------------------------------------------------------------------------------------
-// FDisplayClusterViewportLightCardManager::FProxyData
-//-----------------------------------------------------------------------------------------------------------------
-FDisplayClusterViewportLightCardManager::FProxyData::~FProxyData()
-{
-	ReleaseUVLightCardMap_RenderThread();
-}
-
-void FDisplayClusterViewportLightCardManager::FProxyData::InitializeUVLightCardMap_RenderThread(FDisplayClusterLightCardMap* InUVLightCardMap)
-{
-	if (UVLightCardMap == nullptr)
-	{
-		// Store a copy of the texture's pointer on the render thread and initialize the texture's resources
-		UVLightCardMap = InUVLightCardMap;
-		UVLightCardMap->InitResource();
-	}
-}
-
-void FDisplayClusterViewportLightCardManager::FProxyData::ReleaseUVLightCardMap_RenderThread()
-{
-	// Release the texture's resources and delete the texture object from the rendering thread
-	if (UVLightCardMap)
-	{
-		UVLightCardMap->ReleaseResource();
-
-		delete UVLightCardMap;
-		UVLightCardMap = nullptr;
-	}
-}
-
-void FDisplayClusterViewportLightCardManager::FProxyData::RenderLightCardMap_RenderThread(FRHICommandListImmediate& RHICmdList, const bool bLoadedPrimitives, FSceneInterface* InSceneInterface)
-{
-	bHasUVLightCards = bLoadedPrimitives;
-
-	IDisplayClusterShaders& ShadersAPI = IDisplayClusterShaders::Get();
-	ShadersAPI.RenderPreprocess_UVLightCards(RHICmdList, InSceneInterface, UVLightCardMap, ADisplayClusterLightCardActor::UVPlaneDefaultSize);
-}
-
-FRHITexture* FDisplayClusterViewportLightCardManager::FProxyData::GetUVLightCardMap_RenderThread() const
-{
-	return (bHasUVLightCards && UVLightCardMap != nullptr) ? UVLightCardMap->GetTextureRHI() : nullptr;
-}
-
diff --git a/Engine/Plugins/Runtime/nDisplay/Source/DisplayCluster/Private/Render/Viewport/LightCard/DisplayClusterViewportLightCardManager.h b/Engine/Plugins/Runtime/nDisplay/Source/DisplayCluster/Private/Render/Viewport/LightCard/DisplayClusterViewportLightCardManager.h
--- a/Engine/Plugins/Runtime/nDisplay/Source/DisplayCluster/Private/Render/Viewport/LightCard/DisplayClusterViewportLightCardManager.h
+++ b/Engine/Plugins/Runtime/nDisplay/Source/DisplayCluster/Private/Render/Viewport/LightCard/DisplayClusterViewportLightCardManager.h
@@ -3,42 +3,21 @@
 #pragma once
 
 #include "CoreMinimal.h"
-#include "Render/Viewport/IDisplayClusterViewportLightCardManager.h"
-#include "UObject/GCObject.h"
 
-#include "RenderResource.h"
-#include "UnrealClient.h"
+#include "DisplayClusterViewportLightCardManagerProxy.h"
+#include "DisplayClusterViewportLightCardResource.h"
+
+#include "UObject/GCObject.h"
 
 class ADisplayClusterLightCardActor;
 class FDisplayClusterViewportManager;
-class FPreviewScene;
-class FSceneInterface;
 class UWorld;
 
-/** A render targetable texture resource used to render the UV light cards to */
-class FDisplayClusterLightCardMap : public FTexture, public FRenderTarget
-{
-public:
-	FDisplayClusterLightCardMap(uint32 InSize)
-		: Size(InSize)
-	{ }
-
-	virtual ~FDisplayClusterLightCardMap() = default;
-
-	virtual uint32 GetSizeX() const override { return Size; }
-	virtual uint32 GetSizeY() const override { return Size; }
-	virtual FIntPoint GetSizeXY() const override { return FIntPoint(Size, Size); }
-
-	virtual void InitDynamicRHI() override;
-
-	virtual FString GetFriendlyName() const override { return TEXT("DisplayClusterLightCardMap"); }
-
-private:
-	uint32 Size;
-};
-
-/** Manages the rendering of UV light cards for the viewport manager */
-class FDisplayClusterViewportLightCardManager : public IDisplayClusterViewportLightCardManager, public FGCObject
+/**
+ * Manages the rendering of UV light cards for the viewport manager (Game Thread object)
+ */
+class FDisplayClusterViewportLightCardManager
+	: public FGCObject
 {
 public:
 	FDisplayClusterViewportLightCardManager(FDisplayClusterViewportManager& InViewportManager);
@@ -52,16 +31,27 @@ public:
 	//~ End FGCObject interface
 
 public:
-	//~ Begin IDisplayClusterViewportLightCardManager interface
-	virtual FRHITexture* GetUVLightCardMap_RenderThread() const override;
-	//~ End IDisplayClusterViewportLightCardManager interface
+	/** Return true if UV LightCard is used in this frame. */
+	bool IsUVLightCardEnabled() const;
+
+	/** Get UV LightCard texture size. */
+	FIntPoint GetUVLightCardResourceSize() const;
+
+	/** Get proxy object. */
+	TSharedPtr<FDisplayClusterViewportLightCardManagerProxy, ESPMode::ThreadSafe> GetLightCardManagerProxy() const
+	{ return LightCardManagerProxy; }
 
 public:
+	/** Update internals for the new frame. */
 	void UpdateConfiguration();
 
+	/** Handle StartScene event: created and update internal resources. */
 	void HandleStartScene();
+
+	/** Handle EndScene event: release internal resources. */
 	void HandleEndScene();
 
+	/** Render internal resoures for current frame. */
 	void RenderFrame();
 
 private:
@@ -71,11 +61,24 @@ private:
 	/** Destroys the preview world used to render the UV light cards */
 	void DestroyPreviewWorld();
 
-	/** Initializes the UV light card map texture */
-	void InitializeUVLightCardMap();
+private:
+	/** Render UVLightCard */
+	void RenderUVLightCard();
+
+	/** Update the UV light card map texture */
+	void UpdateUVLightCardResource();
 
 	/** Releases the UV light card map texture */
-	void ReleaseUVLightCardMap();
+	void ReleaseUVLightCardResource();
+
+	/** Create the UV light card map texture */
+	void CreateUVLightCardResource(const FIntPoint& InResourceSize);
+
+	/** Update UVLightCard data game thread*/
+	void UpdateUVLightCardData();
+
+	/** Release UVLightCard data game thread*/
+	void ReleaseUVLightCardData();
 
 private:
 	/** A reference to the owning viewport manager */
@@ -84,30 +87,14 @@ private:
 	/** The preview world the UV light card proxies live in for rendering */
 	UWorld* PreviewWorld = nullptr;
 
-	/** The list of UV light card actors that are referenced by the root actor */
-	TArray<ADisplayClusterLightCardActor*> UVLightCards;
+private:
+	/** A list of primitive components that have been added to the preview scene for rendering in the current frame */
+	TArray<UPrimitiveComponent*> UVLightCardPrimitiveComponents;
 
 	/** The render target to which the UV light card map is rendered */
-	FDisplayClusterLightCardMap* UVLightCardMap = nullptr;
+	TSharedPtr<FDisplayClusterViewportLightCardResource, ESPMode::ThreadSafe> UVLightCardResource;
 
-	struct FProxyData
-	{
-		~FProxyData();
-
-		void InitializeUVLightCardMap_RenderThread(FDisplayClusterLightCardMap* InUVLightCardMap);
-		void ReleaseUVLightCardMap_RenderThread();
-
-		void RenderLightCardMap_RenderThread(FRHICommandListImmediate& RHICmdList, const bool bLoadedPrimitives, FSceneInterface* InSceneInterface);
-
-		FRHITexture* GetUVLightCardMap_RenderThread() const;
-
-	private:
-		/** The render thread copy of the pointer to the UV ligth card map */
-		FDisplayClusterLightCardMap* UVLightCardMap = nullptr;
-
-		/** A render thread flag that indicates the light card manager has UV light cards to render */
-		bool bHasUVLightCards = false;
-	};
-
-	TSharedPtr<FProxyData, ESPMode::ThreadSafe> ProxyData;
+private:
+	/** RenderThread Proxy object*/
+	TSharedPtr<FDisplayClusterViewportLightCardManagerProxy, ESPMode::ThreadSafe> LightCardManagerProxy;
 };
diff --git a/Engine/Plugins/Runtime/nDisplay/Source/DisplayCluster/Private/Render/Viewport/LightCard/DisplayClusterViewportLightCardManagerProxy.cpp b/Engine/Plugins/Runtime/nDisplay/Source/DisplayCluster/Private/Render/Viewport/LightCard/DisplayClusterViewportLightCardManagerProxy.cpp
new file mode 100644
--- /dev/null
+++ b/Engine/Plugins/Runtime/nDisplay/Source/DisplayCluster/Private/Render/Viewport/LightCard/DisplayClusterViewportLightCardManagerProxy.cpp
@@ -0,0 +1,84 @@
+// Copyright Epic Games, Inc. All Rights Reserved.
+
+#include "DisplayClusterViewportLightCardManagerProxy.h"
+
+#include "IDisplayClusterShaders.h"
+
+#include "SceneInterface.h"
+#include "RenderingThread.h"
+
+///////////////////////////////////////////////////////////////////////////////////////////////
+// FDisplayClusterViewportLightCardManagerProxy
+///////////////////////////////////////////////////////////////////////////////////////////////
+FDisplayClusterViewportLightCardManagerProxy::~FDisplayClusterViewportLightCardManagerProxy()
+{
+	ImplReleaseUVLightCardResource_RenderThread();
+}
+
+///////////////////////////////////////////////////////////////////////////////////////////////
+FRHITexture* FDisplayClusterViewportLightCardManagerProxy::GetUVLightCardRHIResource_RenderThread() const
+{
+	return UVLightCardMapResource.IsValid() ? UVLightCardMapResource->GetTextureRHI().GetReference() : nullptr;
+}
+
+void FDisplayClusterViewportLightCardManagerProxy::UpdateUVLightCardResource(const TSharedPtr<FDisplayClusterViewportLightCardResource, ESPMode::ThreadSafe>& InUVLightCardMapResource)
+{
+	ENQUEUE_RENDER_COMMAND(DisplayClusterViewportLightCardManagerProxy_UpdateUVLightCardResource)(
+		[InProxyData = SharedThis(this), UVLightCardMapResource = InUVLightCardMapResource](FRHICommandListImmediate& RHICmdList)
+		{
+			InProxyData->ImplUpdateUVLightCardResource_RenderThread(UVLightCardMapResource);
+		});
+}
+
+void FDisplayClusterViewportLightCardManagerProxy::ReleaseUVLightCardResource()
+{
+	ENQUEUE_RENDER_COMMAND(DisplayClusterViewportLightCardManagerProxy_ReleaseUVLightCardResource)(
+		[InProxyData = SharedThis(this)](FRHICommandListImmediate& RHICmdList)
+		{
+			InProxyData->ImplReleaseUVLightCardResource_RenderThread();
+		});
+}
+
+void FDisplayClusterViewportLightCardManagerProxy::RenderUVLightCard(FSceneInterface* InSceneInterface, const float UVPlaneDefaultSize, const bool bRenderFinalColor) const
+{
+	ENQUEUE_RENDER_COMMAND(DisplayClusterViewportLightCardManagerProxy_RenderUVLightCard)(
+		[InProxyData = SharedThis(this), SceneInterface = InSceneInterface, UVPlaneDefaultSize, bRenderFinalColor](FRHICommandListImmediate& RHICmdList)
+		{
+			InProxyData->ImplRenderUVLightCard_RenderThread(RHICmdList, SceneInterface, UVPlaneDefaultSize, bRenderFinalColor);
+		});
+}
+
+void FDisplayClusterViewportLightCardManagerProxy::ImplUpdateUVLightCardResource_RenderThread(const TSharedPtr<FDisplayClusterViewportLightCardResource, ESPMode::ThreadSafe>& InUVLightCardMapResource)
+{
+	if (UVLightCardMapResource != InUVLightCardMapResource)
+	{
+		ImplReleaseUVLightCardResource_RenderThread();
+
+		// Update resource ptr
+		UVLightCardMapResource = InUVLightCardMapResource;
+		if (UVLightCardMapResource.IsValid())
+		{
+			UVLightCardMapResource->InitResource();
+		}
+	}
+}
+
+void FDisplayClusterViewportLightCardManagerProxy::ImplReleaseUVLightCardResource_RenderThread()
+{
+	// Release the texture's resources and delete the texture object from the rendering thread
+	if (UVLightCardMapResource.IsValid())
+	{
+		UVLightCardMapResource->ReleaseResource();
+
+		UVLightCardMapResource.Reset();
+	}
+}
+
+void FDisplayClusterViewportLightCardManagerProxy::ImplRenderUVLightCard_RenderThread(FRHICommandListImmediate& RHICmdList, FSceneInterface* InSceneInterface, const float InUVPlaneDefaultSize, const bool bRenderFinalColor) const
+{
+	if(InSceneInterface && UVLightCardMapResource.IsValid())
+	{
+		IDisplayClusterShaders& ShadersAPI = IDisplayClusterShaders::Get();
+		ShadersAPI.RenderPreprocess_UVLightCards(RHICmdList, InSceneInterface, UVLightCardMapResource.Get(), InUVPlaneDefaultSize, bRenderFinalColor);
+	}
+}
diff --git a/Engine/Plugins/Runtime/nDisplay/Source/DisplayCluster/Private/Render/Viewport/LightCard/DisplayClusterViewportLightCardManagerProxy.h b/Engine/Plugins/Runtime/nDisplay/Source/DisplayCluster/Private/Render/Viewport/LightCard/DisplayClusterViewportLightCardManagerProxy.h
new file mode 100644
--- /dev/null
+++ b/Engine/Plugins/Runtime/nDisplay/Source/DisplayCluster/Private/Render/Viewport/LightCard/DisplayClusterViewportLightCardManagerProxy.h
@@ -0,0 +1,46 @@
+// Copyright Epic Games, Inc. All Rights Reserved.
+
+#pragma once
+
+#include "CoreMinimal.h"
+
+#include "DisplayClusterViewportLightCardResource.h"
+
+#include "RHI.h"
+#include "RHICommandList.h"
+#include "RHIResources.h"
+#include "Templates/SharedPointer.h"
+
+class FSceneInterface;
+
+/*
+ * Manages the rendering of UV light cards for the viewport manager (Render Thread proxy object)
+ */
+class FDisplayClusterViewportLightCardManagerProxy
+	: public TSharedFromThis<FDisplayClusterViewportLightCardManagerProxy, ESPMode::ThreadSafe>
+{
+public:
+	virtual ~FDisplayClusterViewportLightCardManagerProxy();
+
+	/** Update UVLightCard resource. */
+	void UpdateUVLightCardResource(const TSharedPtr<FDisplayClusterViewportLightCardResource, ESPMode::ThreadSafe>& InUVLightCardMapResource);
+
+	/** Release UVLightCard resource. */
+	void ReleaseUVLightCardResource();
+
+	/** Render UVLightCard. */
+	void RenderUVLightCard(FSceneInterface* InSceneInterface, const float InUVPlaneDefaultSize, const bool bRenderFinalColor) const;
+
+	/** Get current UVLightCard RHI resource on rendering thread. */
+	FRHITexture* GetUVLightCardRHIResource_RenderThread() const;
+
+protected:
+	void ImplUpdateUVLightCardResource_RenderThread(const TSharedPtr<FDisplayClusterViewportLightCardResource, ESPMode::ThreadSafe>& InUVLightCardMapResource);
+	void ImplReleaseUVLightCardResource_RenderThread();
+
+	void ImplRenderUVLightCard_RenderThread(FRHICommandListImmediate& RHICmdList, FSceneInterface* InSceneInterface, const float InUVPlaneDefaultSize, const bool bRenderFinalColor) const;
+
+private:
+	/** The render thread copy of the pointer to the UV ligth card map */
+	TSharedPtr<FDisplayClusterViewportLightCardResource, ESPMode::ThreadSafe> UVLightCardMapResource;
+};
diff --git a/Engine/Plugins/Runtime/nDisplay/Source/DisplayCluster/Private/Render/Viewport/RenderFrame/DisplayClusterRenderFrameSettings.h b/Engine/Plugins/Runtime/nDisplay/Source/DisplayCluster/Private/Render/Viewport/RenderFrame/DisplayClusterRenderFrameSettings.h
--- a/Engine/Plugins/Runtime/nDisplay/Source/DisplayCluster/Private/Render/Viewport/RenderFrame/DisplayClusterRenderFrameSettings.h
+++ b/Engine/Plugins/Runtime/nDisplay/Source/DisplayCluster/Private/Render/Viewport/RenderFrame/DisplayClusterRenderFrameSettings.h
@@ -5,6 +5,7 @@
 #include "CoreMinimal.h"
 
 #include "Render/Viewport/RenderFrame/DisplayClusterRenderFrameEnums.h"
+#include "Render/Viewport/RenderFrame/DisplayClusterRenderFrameSettingsEnums.h"
 
 // Settings for render frame builder
 struct FDisplayClusterRenderFrameSettings
@@ -30,6 +31,9 @@ struct FDisplayClusterRenderFrameSettings
 
 	} CrossGPUTransfer;
 
+	// Alpha channel capture mode for viewports (Lightcard, chromakey)
+	EDisplayClusterRenderFrameAlphaChannelCaptureMode AlphaChannelCaptureMode;
+
 	// Some frame postprocess require additional render targetable resources
 	bool bShouldUseAdditionalFrameTargetableResource = false;
 
diff --git a/Engine/Plugins/Runtime/nDisplay/Source/DisplayCluster/Private/Render/Viewport/RenderTarget/DisplayClusterRenderTargetManager.cpp b/Engine/Plugins/Runtime/nDisplay/Source/DisplayCluster/Private/Render/Viewport/RenderTarget/DisplayClusterRenderTargetManager.cpp
--- a/Engine/Plugins/Runtime/nDisplay/Source/DisplayCluster/Private/Render/Viewport/RenderTarget/DisplayClusterRenderTargetManager.cpp
+++ b/Engine/Plugins/Runtime/nDisplay/Source/DisplayCluster/Private/Render/Viewport/RenderTarget/DisplayClusterRenderTargetManager.cpp
@@ -21,9 +21,8 @@ namespace DisplayClusterRenderTargetManager
 		switch (CaptureMode)
 		{
 		case EDisplayClusterViewportCaptureMode::Chromakey:
-			return EPixelFormat::PF_B8G8R8A8;
-
 		case EDisplayClusterViewportCaptureMode::Lightcard:
+			// The Chromakey and LightCard always uses PF_FloatRGBA for OCIO support.
 			return EPixelFormat::PF_FloatRGBA;
 
 		case EDisplayClusterViewportCaptureMode::MoviePipeline:
diff --git a/Engine/Plugins/Runtime/nDisplay/Source/DisplayCluster/Public/Components/DisplayClusterICVFXCameraComponent.h b/Engine/Plugins/Runtime/nDisplay/Source/DisplayCluster/Public/Components/DisplayClusterICVFXCameraComponent.h
--- a/Engine/Plugins/Runtime/nDisplay/Source/DisplayCluster/Public/Components/DisplayClusterICVFXCameraComponent.h
+++ b/Engine/Plugins/Runtime/nDisplay/Source/DisplayCluster/Public/Components/DisplayClusterICVFXCameraComponent.h
@@ -133,13 +133,13 @@ private:
 	UPROPERTY(EditAnywhere, Transient, Category = Chromakey, meta = (PropertyPath = "CameraSettings.Chromakey.ChromakeyRenderTexture", EditConditionPath = "CameraSettings.bEnable && CameraSettings.Chromakey.bEnable"))
 	FDisplayClusterEditorPropertyReference ChromakeyRenderTextureRef;
 
-	UPROPERTY(EditAnywhere, Transient, Category = OCIO, meta = (PropertyPath = "CameraSettings.AllNodesOCIOConfiguration.bIsEnabled", DisplayName = "Enable Inner Frustum OCIO", ToolTip = "Enable the application of an OpenColorIO configuration to all nodes."))
-	FDisplayClusterEditorPropertyReference OCIOConfigurationEnabledRef;
+	UPROPERTY(EditAnywhere, Transient, Category = OCIO, meta = (PropertyPath = "CameraSettings.CameraOCIO.AllNodesOCIOConfiguration.bIsEnabled", DisplayName = "Enable Inner Frustum OCIO"))
+	FDisplayClusterEditorPropertyReference EnableInnerFrustuOCIORef;
 
-	UPROPERTY(EditAnywhere, Transient, Category = OCIO, meta = (PropertyPath = "CameraSettings.AllNodesOCIOConfiguration.OCIOConfiguration.ColorConfiguration", DisplayName = "All Nodes Color Configuration", ToolTip = "Apply this OpenColorIO configuration to all nodes.", EditConditionPath = "CameraSettings.AllNodesOCIOConfiguration.bIsEnabled"))
-	FDisplayClusterEditorPropertyReference OCIOColorConfiguratonRef;
+	UPROPERTY(EditAnywhere, Transient, Category = OCIO, meta = (PropertyPath = "CameraSettings.CameraOCIO.AllNodesOCIOConfiguration.ColorConfiguration", DisplayName = "All Nodes Color Configuration"))
+	FDisplayClusterEditorPropertyReference AllNodesOCIOConfigurationRef;
 
-	UPROPERTY(EditAnywhere, Transient, Category = OCIO, meta = (PropertyPath = "CameraSettings.PerNodeOCIOProfiles", EditConditionPath = "CameraSettings.AllNodesOCIOConfiguration.bIsEnabled"))
+	UPROPERTY(EditAnywhere, Transient, Category = OCIO, meta = (PropertyPath = "CameraSettings.CameraOCIO.PerNodeOCIOProfiles", DisplayName = "Per-Node OCIO Overrides"))
 	FDisplayClusterEditorPropertyReference PerNodeOCIOProfilesRef;
 
 	UPROPERTY(EditAnywhere, Transient, Category = "Inner Frustum Color Grading", meta = (PropertyPath = "CameraSettings.EnableInnerFrustumColorGrading", DisplayName = "Enable Inner Frustum Color Grading"))
diff --git a/Engine/Plugins/Runtime/nDisplay/Source/DisplayCluster/Public/DisplayClusterRootActor.h b/Engine/Plugins/Runtime/nDisplay/Source/DisplayCluster/Public/DisplayClusterRootActor.h
--- a/Engine/Plugins/Runtime/nDisplay/Source/DisplayCluster/Public/DisplayClusterRootActor.h
+++ b/Engine/Plugins/Runtime/nDisplay/Source/DisplayCluster/Public/DisplayClusterRootActor.h
@@ -236,13 +236,13 @@ private:
 	UPROPERTY(EditAnywhere, Transient, Category = "Color Grading", meta = (PropertyPath = "CurrentConfigData.StageSettings.PerViewportColorGrading"))
 	FDisplayClusterEditorPropertyReference PerViewportColorGradingRef;
 
-	UPROPERTY(EditAnywhere, Transient, Category = OCIO, meta = (PropertyPath = "CurrentConfigData.StageSettings.bUseOverallClusterOCIOConfiguration"))
-	FDisplayClusterEditorPropertyReference EnableClusterOCIORef;
+	UPROPERTY(EditAnywhere, Transient, Category = OCIO, meta = (PropertyPath = "CurrentConfigData.StageSettings.ViewportOCIO.AllViewportsOCIOConfiguration.bIsEnabled", DisplayName = "Enable Viewpport OCIO"))
+	FDisplayClusterEditorPropertyReference EnableViewportOCIORef;
 
-	UPROPERTY(EditAnywhere, Transient, Category = OCIO, meta = (DisplayName = "All Viewports Color Configuration", PropertyPath = "CurrentConfigData.StageSettings.AllViewportsOCIOConfiguration.OCIOConfiguration.ColorConfiguration", ToolTip = "Apply this OpenColorIO configuration to all viewports.", EditConditionPath = "CurrentConfigData.StageSettings.bUseOverallClusterOCIOConfiguration"))
-	FDisplayClusterEditorPropertyReference ClusterOCIOColorConfigurationRef;
+	UPROPERTY(EditAnywhere, Transient, Category = OCIO, meta = (PropertyPath = "CurrentConfigData.StageSettings.ViewportOCIO.AllViewportsOCIOConfiguration.ColorConfiguration", DisplayName = "All Viewports Color Configuration"))
+	FDisplayClusterEditorPropertyReference AllViewportColorConfigurationRef;
 
-	UPROPERTY(EditAnywhere, Transient, Category = OCIO, meta = (PropertyPath = "CurrentConfigData.StageSettings.PerViewportOCIOProfiles"))
+	UPROPERTY(EditAnywhere, Transient, Category = OCIO, meta = (PropertyPath = "CurrentConfigData.StageSettings.ViewportOCIO.PerViewportOCIOProfiles", DisplayName = "Per-Viewport OCIO Overrides"))
 	FDisplayClusterEditorPropertyReference PerViewportOCIOProfilesRef;
 
 	UPROPERTY(EditAnywhere, Transient, Category = "Light Cards", meta = (PropertyPath = "CurrentConfigData.StageSettings.Lightcard.bEnable"))
@@ -254,6 +254,16 @@ private:
 	UPROPERTY(EditAnywhere, Transient, Category = "Light Cards", meta = (PropertyPath = "CurrentConfigData.StageSettings.Lightcard.ShowOnlyList", EditConditionPath = "CurrentConfigData.StageSettings.Lightcard.bEnable"))
 	FDisplayClusterEditorPropertyReference LightCardContentRef;
 
+	UPROPERTY(EditAnywhere, Transient, Category = "Light Cards", meta = (PropertyPath = "CurrentConfigData.StageSettings.Lightcard.LightcardOCIO.LightcardOCIOMode", DisplayName = "Light Cards OCIO"))
+	FDisplayClusterEditorPropertyReference LightcardOCIOModeRef;
+
+	UPROPERTY(EditAnywhere, Transient, Category = "Light Cards", meta = (PropertyPath = "CurrentConfigData.StageSettings.Lightcard.LightcardOCIO.CustomOCIO.AllViewportsOCIOConfiguration.ColorConfiguration", DisplayName = "All Viewports Color Configuration"))
+	FDisplayClusterEditorPropertyReference LightcardAllViewportColorConfigurationRef;
+
+	UPROPERTY(EditAnywhere, Transient, Category = "Light Cards", meta = (PropertyPath = "CurrentConfigData.StageSettings.Lightcard.LightcardOCIO.CustomOCIO.PerViewportOCIOProfiles", DisplayName = "Per-Viewport OCIO Overrides"))
+	FDisplayClusterEditorPropertyReference LightcardPerViewportOCIOProfilesRef;
+
+	// Media
 	UPROPERTY(EditAnywhere, Transient, Category = "Media", meta = (PropertyPath = "CurrentConfigData.MediaSettings"))
 	FDisplayClusterEditorPropertyReference MediaSettingsRef;
 
diff --git a/Engine/Plugins/Runtime/nDisplay/Source/DisplayCluster/Public/Render/Device/IDisplayClusterRenderDevice.h b/Engine/Plugins/Runtime/nDisplay/Source/DisplayCluster/Public/Render/Device/IDisplayClusterRenderDevice.h
--- a/Engine/Plugins/Runtime/nDisplay/Source/DisplayCluster/Public/Render/Device/IDisplayClusterRenderDevice.h
+++ b/Engine/Plugins/Runtime/nDisplay/Source/DisplayCluster/Public/Render/Device/IDisplayClusterRenderDevice.h
@@ -51,8 +51,27 @@ public:
 	virtual void EndScene()
 	{ }
 
-	// update settings from root actor config data, and build new frame structure
+	/**
+	* Update settings from root actor config data, and build new frame structure
+	* This function also updates all viewport settings from config.
+	* Must be called before the InitializeNewFrame() function.
+	*/
 	virtual bool BeginNewFrame(FViewport* InViewport, UWorld* InWorld, FDisplayClusterRenderFrame& OutRenderFrame) = 0;
+
+	/**
+	* Initialize internal data to render a new frame
+	* Initialize references to DCRA objects.
+	* Must be called before viewports start rendering
+	* [ This pair of functions InitializeNewFrame\FinalizeNewFrame should always be called together. ]
+	*/
+	virtual void InitializeNewFrame() = 0;
+
+	/**
+	* Finish rendering the frame and compose the viewports for the final RTT (Warp&Blend, OCIO, ICVFX, etc.)
+	* Release references to DCRA objects.
+	* Must be called after all viewports have been rendered.
+	* [ This pair of functions InitializeNewFrame\FinalizeNewFrame should always be called together. ]
+	*/
 	virtual void FinalizeNewFrame() = 0;
 
 	/**
diff --git a/Engine/Plugins/Runtime/nDisplay/Source/DisplayCluster/Public/Render/Viewport/Containers/DisplayClusterViewport_Context.h b/Engine/Plugins/Runtime/nDisplay/Source/DisplayCluster/Public/Render/Viewport/Containers/DisplayClusterViewport_Context.h
--- a/Engine/Plugins/Runtime/nDisplay/Source/DisplayCluster/Public/Render/Viewport/Containers/DisplayClusterViewport_Context.h
+++ b/Engine/Plugins/Runtime/nDisplay/Source/DisplayCluster/Public/Render/Viewport/Containers/DisplayClusterViewport_Context.h
@@ -4,6 +4,7 @@
 
 #include "CoreMinimal.h"
 #include "StereoRendering.h"
+#include "ShowFlags.h"
 
 class FDisplayClusterViewport_Context
 {
@@ -67,8 +68,18 @@ public:
 
 	struct FRenderThreadData
 	{
+		FRenderThreadData()
+			: EngineShowFlags(ESFIM_All0)
+		{ }
+		
 		// GPUIndex used to render this context.
 		int32 GPUIndex = INDEX_NONE;
+
+		// Display gamma used to render this context
+		float EngineDisplayGamma = 2.2f;
+
+		// Engine flags used to render this context
+		FEngineShowFlags EngineShowFlags;
 	};
 
 	// This data updated only on rendering thread
diff --git a/Engine/Plugins/Runtime/nDisplay/Source/DisplayCluster/Public/Render/Viewport/Containers/DisplayClusterViewport_Enums.h b/Engine/Plugins/Runtime/nDisplay/Source/DisplayCluster/Public/Render/Viewport/Containers/DisplayClusterViewport_Enums.h
--- a/Engine/Plugins/Runtime/nDisplay/Source/DisplayCluster/Public/Render/Viewport/Containers/DisplayClusterViewport_Enums.h
+++ b/Engine/Plugins/Runtime/nDisplay/Source/DisplayCluster/Public/Render/Viewport/Containers/DisplayClusterViewport_Enums.h
@@ -44,3 +44,31 @@ enum class EDisplayClusterViewportResourceType : uint8
 	OutputPreviewTargetableResource,
 #endif
 };
+
+enum class EDisplayClusterViewportCaptureMode : uint8
+{
+	// Use current scene format, no alpha
+	Default = 0,
+
+	// use small BGRA 8bit texture with alpha for masking
+	Chromakey,
+
+	// use hi-res float texture with alpha for compisiting
+	Lightcard,
+
+	// Special hi-res mode for movie pipeline
+	MoviePipeline,
+};
+
+enum class EDisplayClusterViewportOverrideMode : uint8
+{
+	// Do not override this viewport from the other one (Render viewport; create all resources)
+	None = 0,
+
+	// Override internalRTT from the other viewport (Don't render this viewport; Don't create RTT resource)
+	// Useful for custom PP on the same InRTT. (OCIO per-viewport\node)
+	InernalRTT,
+
+	// Override all - clone viewport (Dont render; Don't create resources;)
+	All
+};
diff --git a/Engine/Plugins/Runtime/nDisplay/Source/DisplayCluster/Public/Render/Viewport/Containers/DisplayClusterViewport_RenderSettings.h b/Engine/Plugins/Runtime/nDisplay/Source/DisplayCluster/Public/Render/Viewport/Containers/DisplayClusterViewport_RenderSettings.h
--- a/Engine/Plugins/Runtime/nDisplay/Source/DisplayCluster/Public/Render/Viewport/Containers/DisplayClusterViewport_RenderSettings.h
+++ b/Engine/Plugins/Runtime/nDisplay/Source/DisplayCluster/Public/Render/Viewport/Containers/DisplayClusterViewport_RenderSettings.h
@@ -4,21 +4,12 @@
 
 #include "CoreMinimal.h"
 
-enum class EDisplayClusterViewportCaptureMode : uint8
-{
-	// Use current scene format, no alpha
-	Default = 0,
-
-	// use small BGRA 8bit texture with alpha for masking
-	Chromakey,
-
-	// use hi-res float texture with alpha for compisiting
-	Lightcard,
-
-	// Special hi-res mode for movie pipeline
-	MoviePipeline,
-};
+#include "DisplayClusterViewport_Enums.h"
 
+/**
+ * nDisplay viewport render settings.
+ * These are runtime settings, updated every frame from the cluster configuration.
+ */
 class FDisplayClusterViewport_RenderSettings
 {
 public:
@@ -97,8 +88,12 @@ public:
 	// Special capture modes (chromakey, lightcard) change RTT format and render flags
 	EDisplayClusterViewportCaptureMode CaptureMode = EDisplayClusterViewportCaptureMode::Default;
 
-	// Override image from this viewport
-	FString OverrideViewportId;
+	// Override resources from this viewport
+	EDisplayClusterViewportOverrideMode ViewportOverrideMode = EDisplayClusterViewportOverrideMode::None;
+	FString ViewportOverrideId;
+
+	// Override all viewport resources
+	bool bOverrideWholeViewport= true;
 
 public:
 	// Reset runtime values from prev frame
@@ -112,7 +107,8 @@ public:
 
 		CaptureMode = EDisplayClusterViewportCaptureMode::Default;
 
-		OverrideViewportId.Empty();
+		ViewportOverrideMode = EDisplayClusterViewportOverrideMode::None;
+		ViewportOverrideId.Empty();
 	}
 
 	inline void FinishUpdateSettings()
@@ -125,6 +121,22 @@ public:
 		return ParentViewportId;
 	}
 
+	inline bool IsViewportHasParent() const
+	{
+		return !ParentViewportId.IsEmpty();
+	}
+
+	inline void SetViewportOverride(const FString& InViewportOverrideId, const EDisplayClusterViewportOverrideMode InViewportOverrideMode = EDisplayClusterViewportOverrideMode::All)
+	{
+		ViewportOverrideMode = InViewportOverrideMode;
+		ViewportOverrideId = InViewportOverrideId;
+	}
+
+	inline bool IsViewportOverrided() const
+	{
+		return ViewportOverrideMode != EDisplayClusterViewportOverrideMode::None && !ViewportOverrideId.IsEmpty();
+	}
+
 	// Call this after UpdateSettings()
 	inline void AssignParentViewport(const FString& InParentViewportId, const FDisplayClusterViewport_RenderSettings& InParentSettings, bool Inherit = true)
 	{
diff --git a/Engine/Plugins/Runtime/nDisplay/Source/DisplayCluster/Public/Render/Viewport/Containers/DisplayClusterViewport_RenderSettingsICVFX.h b/Engine/Plugins/Runtime/nDisplay/Source/DisplayCluster/Public/Render/Viewport/Containers/DisplayClusterViewport_RenderSettingsICVFX.h
--- a/Engine/Plugins/Runtime/nDisplay/Source/DisplayCluster/Public/Render/Viewport/Containers/DisplayClusterViewport_RenderSettingsICVFX.h
+++ b/Engine/Plugins/Runtime/nDisplay/Source/DisplayCluster/Public/Render/Viewport/Containers/DisplayClusterViewport_RenderSettingsICVFX.h
@@ -6,44 +6,54 @@
 
 #include "ShaderParameters/DisplayClusterShaderParameters_ICVFX.h"
 
-enum EDisplayClusterViewportICVFXFlags
+//***************************************************
+//* Runtime configuration from DCRA
+//***************************************************
+enum class EDisplayClusterViewportICVFXFlags : uint8
 {
-	ViewportICVFX_None = 0,
+	None = 0,
 
 	// Allow to use ICVFX for this viewport (Must be supported by projection policy)
-	ViewportICVFX_Enable = 1 << 0,
+	Enable = 1 << 0,
+
 	// Disable incamera render to this viewport
-	ViewportICVFX_DisableCamera = 1 << 1,
+	DisableCamera = 1 << 1,
+
 	// Disable chromakey render to this viewport
-	ViewportICVFX_DisableChromakey = 1 << 2,
+	DisableChromakey = 1 << 2,
+
 	// Disable chromakey markers render to this viewport
-	ViewportICVFX_DisableChromakeyMarkers = 1 << 3,
+	DisableChromakeyMarkers = 1 << 3,
+
 	// Disable lightcard render to this viewport
-	ViewportICVFX_DisableLightcard = 1 << 4,
+	DisableLightcard = 1 << 4,
+
 	// Use unique lightcard mode for this viewport from param 'OverrideLightcardMode'
-	ViewportICVFX_OverrideLightcardMode = 1 << 5,
+	OverrideLightcardMode = 1 << 5,
 };
 ENUM_CLASS_FLAGS(EDisplayClusterViewportICVFXFlags);
 
 //***************************************************
 //* This flag raised only from icvfx manager
 //***************************************************
-enum EDisplayClusterViewportRuntimeICVFXFlags
+enum class EDisplayClusterViewportRuntimeICVFXFlags: uint8
 {
-	ViewportRuntimeICVFX_None = 0,
+	None = 0,
 
 	// Enable use icvfx only from projection policy for this viewport.
-	ViewportRuntime_ICVFXTarget = 1 << 0,
+	Target = 1 << 0,
 
 	// viewport ICVFX usage
-	ViewportRuntime_ICVFXIncamera       = 1 << 1,
-	ViewportRuntime_ICVFXChromakey      = 1 << 2,
-	ViewportRuntime_ICVFXLightcard      = 1 << 3,
+	InCamera    = 1 << 1,
+	Chromakey   = 1 << 2,
+	Lightcard   = 1 << 3,
+	UVLightcard = 1 << 4,
 
 	// This viewport used as internal icvfx composing resource (created and deleted inside icvfx logic)
-	ViewportRuntime_InternalResource = 1 << 30,
+	InternalResource = 1 << 5,
+
 	// Mark unused icvfx dynamic viewports
-	ViewportRuntime_Unused = 1 << 31,
+	Unused = 1 << 6,
 };
 ENUM_CLASS_FLAGS(EDisplayClusterViewportRuntimeICVFXFlags);
 
@@ -51,10 +61,10 @@ class FDisplayClusterViewport_RenderSettingsICVFX
 {
 public:
 	// ICVFX logic flags
-	EDisplayClusterViewportICVFXFlags         Flags = ViewportICVFX_None;
+	EDisplayClusterViewportICVFXFlags         Flags = EDisplayClusterViewportICVFXFlags::None;
 
 	// Internal logic ICVFX flags
-	EDisplayClusterViewportRuntimeICVFXFlags  RuntimeFlags = ViewportRuntimeICVFX_None;
+	EDisplayClusterViewportRuntimeICVFXFlags  RuntimeFlags = EDisplayClusterViewportRuntimeICVFXFlags::None;
 
 	// Compact ICVFX viewport settings
 	FDisplayClusterShaderParameters_ICVFX ICVFX;
@@ -62,12 +72,12 @@ public:
 public:
 	inline void BeginUpdateSettings()
 	{
-		Flags = ViewportICVFX_None;
+		Flags = EDisplayClusterViewportICVFXFlags::None;
 
 		ICVFX.Reset();
 
 		// Reset target flag
-		RuntimeFlags &= ~(ViewportRuntime_ICVFXTarget);
+		EnumRemoveFlags(RuntimeFlags, EDisplayClusterViewportRuntimeICVFXFlags::Target);
 	}
 
 	// Implement copy ref and arrays
diff --git a/Engine/Plugins/Runtime/nDisplay/Source/DisplayCluster/Public/Render/Viewport/IDisplayClusterViewportManager.h b/Engine/Plugins/Runtime/nDisplay/Source/DisplayCluster/Public/Render/Viewport/IDisplayClusterViewportManager.h
--- a/Engine/Plugins/Runtime/nDisplay/Source/DisplayCluster/Public/Render/Viewport/IDisplayClusterViewportManager.h
+++ b/Engine/Plugins/Runtime/nDisplay/Source/DisplayCluster/Public/Render/Viewport/IDisplayClusterViewportManager.h
@@ -15,7 +15,6 @@ class FSceneViewFamilyContext;
 class ADisplayClusterRootActor;
 class UDisplayClusterConfigurationViewport;
 class IDisplayClusterViewportManagerProxy;
-class IDisplayClusterViewportLightCardManager;
 class FReferenceCollector;
 
 class DISPLAYCLUSTER_API IDisplayClusterViewportManager
@@ -77,7 +76,14 @@ public:
 	virtual bool BeginNewFrame(FViewport* InViewport, UWorld* InWorld, FDisplayClusterRenderFrame& OutRenderFrame) = 0;
 
 	/**
-	* Finalize frame logic for viewports on game thread
+	* Initialize frame for render on game thread
+	* [Game thread func]
+	*
+	*/
+	virtual void InitializeNewFrame() = 0;
+
+	/**
+	* Finalize frame for render on game thread
 	* [Game thread func]
 	*
 	*/
@@ -163,12 +169,6 @@ public:
 	*/
 	virtual const TArrayView<IDisplayClusterViewport*> GetViewports() const = 0;
 
-	/**
-	* Return the light card manager, used to manage and render UV light cards
-	* [Game thread func]
-	*/
-	virtual TSharedPtr<IDisplayClusterViewportLightCardManager, ESPMode::ThreadSafe> GetLightCardManager() const = 0;
-
 	/**
 	* Mark the geometry of the referenced component(s) as dirty (ProceduralMesh, etc)
 	* [Game thread func]
diff --git a/Engine/Plugins/Runtime/nDisplay/Source/DisplayCluster/Public/Render/Viewport/IDisplayClusterViewportManagerProxy.h b/Engine/Plugins/Runtime/nDisplay/Source/DisplayCluster/Public/Render/Viewport/IDisplayClusterViewportManagerProxy.h
--- a/Engine/Plugins/Runtime/nDisplay/Source/DisplayCluster/Public/Render/Viewport/IDisplayClusterViewportManagerProxy.h
+++ b/Engine/Plugins/Runtime/nDisplay/Source/DisplayCluster/Public/Render/Viewport/IDisplayClusterViewportManagerProxy.h
@@ -5,7 +5,6 @@
 #include "CoreMinimal.h"
 
 #include "Render/Viewport/IDisplayClusterViewportProxy.h"
-#include "Render/Viewport/IDisplayClusterViewportLightCardManager.h"
 
 class DISPLAYCLUSTER_API IDisplayClusterViewportManagerProxy
 {
@@ -66,11 +65,5 @@ public:
 	* @return - true if success
 	*/
 	virtual bool ResolveFrameTargetToBackBuffer_RenderThread(FRHICommandListImmediate& RHICmdList, const uint32 InContextNum, const int32 DestArrayIndex, FRHITexture2D* DstBackBuffer, FVector2D WindowSize) const = 0;
-
-	/**
-	* Return the light card manager, used to manage and render UV light cards
-	* [Rendering thread func]
-	*/
-	virtual TSharedPtr<IDisplayClusterViewportLightCardManager, ESPMode::ThreadSafe> GetLightCardManager_RenderThread() const = 0;
 };
 
diff --git a/Engine/Plugins/Runtime/nDisplay/Source/DisplayCluster/Public/Render/Viewport/IDisplayClusterViewportProxy.h b/Engine/Plugins/Runtime/nDisplay/Source/DisplayCluster/Public/Render/Viewport/IDisplayClusterViewportProxy.h
--- a/Engine/Plugins/Runtime/nDisplay/Source/DisplayCluster/Public/Render/Viewport/IDisplayClusterViewportProxy.h
+++ b/Engine/Plugins/Runtime/nDisplay/Source/DisplayCluster/Public/Render/Viewport/IDisplayClusterViewportProxy.h
@@ -30,24 +30,44 @@ public:
 	virtual const FDisplayClusterViewport_RenderSettingsICVFX& GetRenderSettingsICVFX_RenderThread() const = 0;
 	virtual const FDisplayClusterViewport_PostRenderSettings& GetPostRenderSettings_RenderThread() const = 0;
 
+	/** Return viewport proxy contexts data. */
 	virtual const TArray<FDisplayClusterViewport_Context>& GetContexts_RenderThread() const = 0;
 
-	//  Return viewport scene proxy resources by type
+	/** Get viewport resources by type
+	 *
+	 * @param InResourceType - resource type (RTT, Shader, MIPS, etc)
+	 * @param OutResources   - [Out] RHI resources array for all contexts
+	 *
+	 * @return - true if success
+	 */
 	virtual bool GetResources_RenderThread(const EDisplayClusterViewportResourceType InResourceType, TArray<FRHITexture2D*>& OutResources) const = 0;
+
+	/** Get viewport resources with rects by type
+	 *
+	 * @param InResourceType - resource type (RTT, Shader, MIPS, etc)
+	 * @param OutResources   - [Out] RHI resources array for all contexts
+	 * @param OutRects       - [Out] RHI resources rects array for all contexts
+	 *
+	 * @return - true if success
+	 */
 	virtual bool GetResourcesWithRects_RenderThread(const EDisplayClusterViewportResourceType InResourceType, TArray<FRHITexture2D*>& OutResources, TArray<FIntRect>& OutRects) const = 0;
 
-	// Resolve resource contexts
+	/** Copy resource contexts by type
+	 *
+	 * @param RHICmdList         - RHIinterface
+	 * @param InputResourceType  - Input resource type (RTT, Shader, MIPS, etc)
+	 * @param OutputResourceType - Output resource type (RTT, Shader, MIPS, etc)
+	 * @param InContextNum       - [optional] the type of context to copy (by default, all contexts are copied).
+	 *
+	 * @return - true if success
+	 */
 	virtual bool ResolveResources_RenderThread(FRHICommandListImmediate& RHICmdList, const EDisplayClusterViewportResourceType InputResourceType, const EDisplayClusterViewportResourceType OutputResourceType, const int32 InContextNum = INDEX_NONE) const = 0;
 
+	/** Return output resource type (support preview, remap, etc). */
 	virtual EDisplayClusterViewportResourceType   GetOutputResourceType_RenderThread() const = 0;
 
 	virtual const class IDisplayClusterViewportManagerProxy& GetOwner_RenderThread() const = 0;
 
 	virtual void SetRenderSettings_RenderThread(const FDisplayClusterViewport_RenderSettings& InRenderSettings) const = 0;
 	virtual void SetContexts_RenderThread(const TArray<FDisplayClusterViewport_Context>& InContexts) const = 0;
-
-	virtual void OnResolvedSceneColor_RenderThread(class FRDGBuilder& GraphBuilder, const struct FSceneTextures& SceneTextures, const FDisplayClusterViewportProxy_Context& InProxyContext) = 0;
-
-	// Callback. After the viewfamily has finished rendering
-	virtual void PostRenderViewFamily_RenderThread(class FRDGBuilder& InGraphBuilder, class FSceneViewFamily& InViewFamily, const class FSceneView& InSceneView, const FDisplayClusterViewportProxy_Context& InProxyContext) = 0;
 };
diff --git a/Engine/Plugins/Runtime/nDisplay/Source/DisplayClusterColorGrading/Private/DataModelGenerators/DisplayClusterColorGradingGenerator_RootActor.cpp b/Engine/Plugins/Runtime/nDisplay/Source/DisplayClusterColorGrading/Private/DataModelGenerators/DisplayClusterColorGradingGenerator_RootActor.cpp
--- a/Engine/Plugins/Runtime/nDisplay/Source/DisplayClusterColorGrading/Private/DataModelGenerators/DisplayClusterColorGradingGenerator_RootActor.cpp
+++ b/Engine/Plugins/Runtime/nDisplay/Source/DisplayClusterColorGrading/Private/DataModelGenerators/DisplayClusterColorGradingGenerator_RootActor.cpp
@@ -185,7 +185,7 @@ public:
 		ColorGradingCategoryBuilder.AddProperty(DetailBuilder.GetProperty(GET_MEMBER_NAME_CHECKED(UDisplayClusterConfigurationData, StageSettings.EntireClusterColorGrading), ConfigDataClass));
 		ColorGradingCategoryBuilder.AddProperty(DetailBuilder.GetProperty(GET_MEMBER_NAME_CHECKED(UDisplayClusterConfigurationData, StageSettings.PerViewportColorGrading), ConfigDataClass));
 		ColorGradingCategoryBuilder.AddProperty(DetailBuilder.GetProperty(GET_MEMBER_NAME_CHECKED(UDisplayClusterConfigurationData, StageSettings.bEnableInnerFrustums), ConfigDataClass));
-		ColorGradingCategoryBuilder.AddProperty(DetailBuilder.GetProperty(GET_MEMBER_NAME_CHECKED(UDisplayClusterConfigurationData, StageSettings.bUseOverallClusterOCIOConfiguration), ConfigDataClass));
+		ColorGradingCategoryBuilder.AddProperty(DetailBuilder.GetProperty(GET_MEMBER_NAME_CHECKED(UDisplayClusterConfigurationData, StageSettings.ViewportOCIO.AllViewportsOCIOConfiguration.bIsEnabled), ConfigDataClass));
 
 		AddColorGradingDetailProperties(DetailBuilder);
 		AddDetailsPanelProperties(DetailBuilder);
@@ -267,7 +267,7 @@ private:
 		InnerFrustumCategoryBuilder.AddProperty(GET_MEMBER_NAME_CHECKED(ADisplayClusterRootActor, InnerFrustumPriority), ADisplayClusterRootActor::StaticClass());
 
 		IDetailCategoryBuilder& AllViewportsOCIOCategoryBuilder = DetailBuilder.EditCategory(TEXT("CustomAllViewportsOCIOCategory"), LOCTEXT("CustomAllViewportsOCIOCategoryLabel", "All Viewports"));
-		AllViewportsOCIOCategoryBuilder.AddProperty(TEXT("ClusterOCIOColorConfigurationRef"), ADisplayClusterRootActor::StaticClass());
+		AllViewportsOCIOCategoryBuilder.AddProperty(TEXT("AllViewportColorConfigurationRef"), ADisplayClusterRootActor::StaticClass());
 
 		IDetailCategoryBuilder& PerViewportOCIOCategoryBuilder = DetailBuilder.EditCategory(TEXT("CustomPerViewportOCIOCategory"), LOCTEXT("CustomPerViewportOCIOCategoryLabel", "Per-Viewport"));
 		PerViewportOCIOCategoryBuilder.AddProperty(TEXT("PerViewportOCIOProfilesRef"), ADisplayClusterRootActor::StaticClass());
@@ -501,7 +501,7 @@ void FDisplayClusterColorGradingGenerator_RootActor::GenerateDataModel(IProperty
 
 	FDisplayClusterColorGradingDataModel::FDetailsSection OCIODetailsSection;
 	OCIODetailsSection.DisplayName = LOCTEXT("OCIODetailsSectionLabel", "OCIO");
-	OCIODetailsSection.EditConditionPropertyHandle = FindPropertyHandle(PropertyRowGenerator, CREATE_PROPERTY_PATH(UDisplayClusterConfigurationData, StageSettings.bUseOverallClusterOCIOConfiguration));
+	OCIODetailsSection.EditConditionPropertyHandle = FindPropertyHandle(PropertyRowGenerator, CREATE_PROPERTY_PATH(UDisplayClusterConfigurationData, StageSettings.ViewportOCIO.AllViewportsOCIOConfiguration.bIsEnabled));
 
 	FDisplayClusterColorGradingDataModel::FDetailsSubsection AllViewportsOCIODetailsSubsection;
 	AllViewportsOCIODetailsSubsection.DisplayName = LOCTEXT("AllViewportsOCIOSubsectionLabel", "All Viewports");
@@ -735,9 +735,9 @@ public:
 
 		CategoryBuilder.AddProperty(DetailBuilder.GetProperty(GET_MEMBER_NAME_CHECKED(UDisplayClusterICVFXCameraComponent, CameraSettings.bEnable)));
 		CategoryBuilder.AddProperty(DetailBuilder.GetProperty(GET_MEMBER_NAME_CHECKED(UDisplayClusterICVFXCameraComponent, CameraSettings.Chromakey.bEnable)));
-		CategoryBuilder.AddProperty(DetailBuilder.GetProperty(GET_MEMBER_NAME_CHECKED(UDisplayClusterICVFXCameraComponent, CameraSettings.AllNodesOCIOConfiguration.bIsEnabled)));
-		CategoryBuilder.AddProperty(DetailBuilder.GetProperty(GET_MEMBER_NAME_CHECKED(UDisplayClusterICVFXCameraComponent, CameraSettings.AllNodesColorGrading)));
-		CategoryBuilder.AddProperty(DetailBuilder.GetProperty(GET_MEMBER_NAME_CHECKED(UDisplayClusterICVFXCameraComponent, CameraSettings.PerNodeColorGrading)));
+		CategoryBuilder.AddProperty(DetailBuilder.GetProperty(GET_MEMBER_NAME_CHECKED(UDisplayClusterICVFXCameraComponent, CameraSettings.CameraOCIO.AllNodesOCIOConfiguration.bIsEnabled)));
+		CategoryBuilder.AddProperty(DetailBuilder.GetProperty(GET_MEMBER_NAME_CHECKED(UDisplayClusterICVFXCameraComponent, CameraSettings.CameraOCIO.AllNodesOCIOConfiguration)));
+		CategoryBuilder.AddProperty(DetailBuilder.GetProperty(GET_MEMBER_NAME_CHECKED(UDisplayClusterICVFXCameraComponent, CameraSettings.CameraOCIO.PerNodeOCIOProfiles)));
 
 		AddColorGradingDetailProperties(DetailBuilder);
 		AddDetailsPanelProperties(DetailBuilder);
@@ -847,7 +847,7 @@ private:
 		AddProperty(ChromakeyCustomCategoryBuilder, TEXT("ChromakeyRenderTextureRef"), true);
 
 		IDetailCategoryBuilder& OCIOAllNodesCategoryBuilder = DetailBuilder.EditCategory(TEXT("CustomAllNodesOCIOCategory"), LOCTEXT("CustomAllNodesOCIOCategoryLabel", "All Nodes"));
-		AddProperty(OCIOAllNodesCategoryBuilder, TEXT("OCIOColorConfiguratonRef"));
+		AddProperty(OCIOAllNodesCategoryBuilder, TEXT("AllNodesOCIOConfigurationRef"));
 
 		IDetailCategoryBuilder& OCIOPerNodeCategoryBuilder = DetailBuilder.EditCategory(TEXT("CustomPerNodeOCIOCategory"), LOCTEXT("CustomPerNodeOCIOCategoryLabel", "Per-Node"));
 		AddProperty(OCIOPerNodeCategoryBuilder, TEXT("PerNodeOCIOProfilesRef"));
@@ -1108,7 +1108,8 @@ void FDisplayClusterColorGradingGenerator_ICVFXCamera::GenerateDataModel(IProper
 	{
 		FDisplayClusterColorGradingDataModel::FDetailsSection OCIODetailsSection;
 		OCIODetailsSection.DisplayName = LOCTEXT("OCIODetailsSectionLabel", "OCIO");
-		OCIODetailsSection.EditConditionPropertyHandle = FindPropertyHandle(PropertyRowGenerator, CREATE_PROPERTY_PATH(UDisplayClusterICVFXCameraComponent, CameraSettings.AllNodesOCIOConfiguration.bIsEnabled));
+		OCIODetailsSection.EditConditionPropertyHandle = FindPropertyHandle(PropertyRowGenerator,
+			CREATE_PROPERTY_PATH(UDisplayClusterICVFXCameraComponent, CameraSettings.CameraOCIO.AllNodesOCIOConfiguration.bIsEnabled));
 
 		FDisplayClusterColorGradingDataModel::FDetailsSubsection AllNodesOCIODetailsSubsection;
 		AllNodesOCIODetailsSubsection.DisplayName = LOCTEXT("AllNodesOCIODetailsSubsectionLabel", "All Nodes");
diff --git a/Engine/Plugins/Runtime/nDisplay/Source/DisplayClusterConfiguration/Private/DisplayClusterConfigurationTypes.cpp b/Engine/Plugins/Runtime/nDisplay/Source/DisplayClusterConfiguration/Private/DisplayClusterConfigurationTypes.cpp
--- a/Engine/Plugins/Runtime/nDisplay/Source/DisplayClusterConfiguration/Private/DisplayClusterConfigurationTypes.cpp
+++ b/Engine/Plugins/Runtime/nDisplay/Source/DisplayClusterConfiguration/Private/DisplayClusterConfigurationTypes.cpp
@@ -450,29 +450,3 @@ UDisplayClusterConfigurationData* UDisplayClusterConfigurationData::CreateNewCon
 
 	return NewConfigData;
 }
-
-FDisplayClusterConfigurationICVFX_CameraRenderSettings::FDisplayClusterConfigurationICVFX_CameraRenderSettings()
-{
-	// Setup incamera defaults:
-	GenerateMips.bAutoGenerateMips = true;
-}
-
-FDisplayClusterConfigurationOCIOConfiguration::FDisplayClusterConfigurationOCIOConfiguration()
-{
-	OCIOConfiguration.bIsEnabled = true;
-}
-
-FDisplayClusterConfigurationOCIOProfile::FDisplayClusterConfigurationOCIOProfile()
-{
-	OCIOConfiguration.bIsEnabled = true;
-}
-
-FDisplayClusterConfigurationICVFX_CameraCustomFrustum::FDisplayClusterConfigurationICVFX_CameraCustomFrustum():
-	EstimatedOverscanResolution(ForceInitToZero), InnerFrustumResolution(ForceInitToZero), OverscanPixelsIncrease(0.f)
-{
-}
-
-FDisplayClusterConfigurationICVFX_CameraSettings::FDisplayClusterConfigurationICVFX_CameraSettings()
-{
-	AllNodesColorGrading.bEnableEntireClusterColorGrading = true;
-}
diff --git a/Engine/Plugins/Runtime/nDisplay/Source/DisplayClusterConfiguration/Private/DisplayClusterConfigurationTypes_ICVFX.cpp b/Engine/Plugins/Runtime/nDisplay/Source/DisplayClusterConfiguration/Private/DisplayClusterConfigurationTypes_ICVFX.cpp
--- a/Engine/Plugins/Runtime/nDisplay/Source/DisplayClusterConfiguration/Private/DisplayClusterConfigurationTypes_ICVFX.cpp
+++ b/Engine/Plugins/Runtime/nDisplay/Source/DisplayClusterConfiguration/Private/DisplayClusterConfigurationTypes_ICVFX.cpp
@@ -2,7 +2,9 @@
 
 #include "DisplayClusterConfigurationTypes_ICVFX.h"
 
-
+///////////////////////////////////////////////////////////////////////////////////////
+// FDisplayClusterConfigurationICVFX_ChromakeyMarkers
+///////////////////////////////////////////////////////////////////////////////////////
 FDisplayClusterConfigurationICVFX_ChromakeyMarkers::FDisplayClusterConfigurationICVFX_ChromakeyMarkers()
 	: MarkerColor(0, 0.25f, 0)
 	, MarkerTileOffset(0)
@@ -13,3 +15,206 @@ FDisplayClusterConfigurationICVFX_ChromakeyMarkers::FDisplayClusterConfiguration
 		MarkerTileRGBA = Cast<UTexture2D>(FSoftObjectPath(TexturePath).TryLoad());
 	}
 }
+
+///////////////////////////////////////////////////////////////////////////////////////
+// FDisplayClusterConfigurationICVFX_CameraRenderSettings
+///////////////////////////////////////////////////////////////////////////////////////
+FDisplayClusterConfigurationICVFX_CameraRenderSettings::FDisplayClusterConfigurationICVFX_CameraRenderSettings()
+{
+	// Setup incamera defaults:
+	GenerateMips.bAutoGenerateMips = true;
+}
+
+///////////////////////////////////////////////////////////////////////////////////////
+// FDisplayClusterConfigurationICVFX_CameraCustomFrustum
+///////////////////////////////////////////////////////////////////////////////////////
+FDisplayClusterConfigurationICVFX_CameraCustomFrustum::FDisplayClusterConfigurationICVFX_CameraCustomFrustum() :
+	EstimatedOverscanResolution(ForceInitToZero), InnerFrustumResolution(ForceInitToZero), OverscanPixelsIncrease(0.f)
+{
+}
+
+///////////////////////////////////////////////////////////////////////////////////////
+// FDisplayClusterConfigurationICVFX_LightcardCustomOCIO
+///////////////////////////////////////////////////////////////////////////////////////
+const FOpenColorIOColorConversionSettings* FDisplayClusterConfigurationICVFX_LightcardCustomOCIO::FindOCIOConfiguration(const FString& InViewportId) const
+{
+	if (AllViewportsOCIOConfiguration.bIsEnabled)
+	{
+		// Per viewport OCIO:
+		for (const FDisplayClusterConfigurationOCIOProfile& OCIOProfileIt : PerViewportOCIOProfiles)
+		{
+			if (OCIOProfileIt.IsEnabledForObject(InViewportId))
+			{
+				return &OCIOProfileIt.ColorConfiguration;
+			}
+		}
+
+		if (AllViewportsOCIOConfiguration.IsEnabled())
+		{
+			return &AllViewportsOCIOConfiguration.ColorConfiguration;
+		}
+	}
+
+	return nullptr;
+}
+
+///////////////////////////////////////////////////////////////////////////////////////
+// FDisplayClusterConfigurationICVFX_ViewportOCIO
+///////////////////////////////////////////////////////////////////////////////////////
+const FOpenColorIOColorConversionSettings* FDisplayClusterConfigurationICVFX_ViewportOCIO::FindOCIOConfiguration(const FString& InViewportId) const
+{
+	if (AllViewportsOCIOConfiguration.bIsEnabled)
+	{
+		// Per viewport OCIO:
+		for (const FDisplayClusterConfigurationOCIOProfile& OCIOProfileIt : PerViewportOCIOProfiles)
+		{
+			if (OCIOProfileIt.IsEnabledForObject(InViewportId))
+			{
+				return &OCIOProfileIt.ColorConfiguration;
+			}
+		}
+
+		if (AllViewportsOCIOConfiguration.IsEnabled())
+		{
+			return &AllViewportsOCIOConfiguration.ColorConfiguration;
+		}
+	}
+
+	return nullptr;
+}
+
+///////////////////////////////////////////////////////////////////////////////////////
+// FDisplayClusterConfigurationICVFX_CameraOCIO
+///////////////////////////////////////////////////////////////////////////////////////
+const FOpenColorIOColorConversionSettings* FDisplayClusterConfigurationICVFX_CameraOCIO::FindOCIOConfiguration(const FString& InClusterNodeId) const
+{
+	if (AllNodesOCIOConfiguration.bIsEnabled)
+	{
+		// Per node OCIO:
+		for (const FDisplayClusterConfigurationOCIOProfile& OCIOProfileIt : PerNodeOCIOProfiles)
+		{
+			if (OCIOProfileIt.IsEnabledForObject(InClusterNodeId))
+			{
+				return &OCIOProfileIt.ColorConfiguration;
+			}
+		}
+
+		if (AllNodesOCIOConfiguration.IsEnabled())
+		{
+			return &AllNodesOCIOConfiguration.ColorConfiguration;
+		}
+	}
+
+	return nullptr;
+}
+
+#if WITH_EDITOR
+bool FDisplayClusterConfigurationICVFX_CameraOCIO::IsChromakeyViewportSettingsEqual_Editor(const FString& InClusterNodeId1, const FString& InClusterNodeId2) const
+{
+	return IsInnerFrustumViewportSettingsEqual_Editor(InClusterNodeId1, InClusterNodeId2);
+}
+
+bool FDisplayClusterConfigurationICVFX_CameraOCIO::IsInnerFrustumViewportSettingsEqual_Editor(const FString& InClusterNodeId1, const FString& InClusterNodeId2) const
+{
+	if (AllNodesOCIOConfiguration.bIsEnabled)
+	{
+		for (const FDisplayClusterConfigurationOCIOProfile& OCIOProfileIt : PerNodeOCIOProfiles)
+		{
+			if (OCIOProfileIt.IsEnabled())
+			{
+				const FString* CustomNode1 = OCIOProfileIt.ApplyOCIOToObjects.FindByPredicate([ClusterNodeId = InClusterNodeId1](const FString& InClusterNodeId)
+					{
+						return ClusterNodeId.Equals(InClusterNodeId, ESearchCase::IgnoreCase);
+					});
+
+				const FString* CustomNode2 = OCIOProfileIt.ApplyOCIOToObjects.FindByPredicate([ClusterNodeId = InClusterNodeId2](const FString& InClusterNodeId)
+					{
+						return ClusterNodeId.Equals(InClusterNodeId, ESearchCase::IgnoreCase);
+					});
+
+				if (CustomNode1 && CustomNode2)
+				{
+					// equal custom settings
+					return true;
+				}
+
+				if (CustomNode1 || CustomNode2)
+				{
+					// one of node has custom settings
+					return false;
+				}
+			}
+		}
+	}
+
+	return true;
+}
+#endif
+
+///////////////////////////////////////////////////////////////////////////////////////
+// FDisplayClusterConfigurationICVFX_LightcardOCIO
+///////////////////////////////////////////////////////////////////////////////////////
+const FOpenColorIOColorConversionSettings* FDisplayClusterConfigurationICVFX_LightcardOCIO::FindOCIOConfiguration(const FString& InViewportId, const FDisplayClusterConfigurationICVFX_ViewportOCIO& InViewportOCIO) const
+{
+	switch (LightcardOCIOMode)
+	{
+	case EDisplayClusterConfigurationViewportLightcardOCIOMode::nDisplay:
+		// Use Viewport OCIO
+		return InViewportOCIO.FindOCIOConfiguration(InViewportId);
+
+	case EDisplayClusterConfigurationViewportLightcardOCIOMode::Custom:
+		// Use custom OCIO
+		return CustomOCIO.FindOCIOConfiguration(InViewportId);
+
+	default:
+		// No OCIO for Light Cards
+		break;
+	}
+
+	return nullptr;
+}
+
+///////////////////////////////////////////////////////////////////////////////////////
+// FDisplayClusterConfigurationICVFX_CameraSettings
+///////////////////////////////////////////////////////////////////////////////////////
+FDisplayClusterConfigurationICVFX_CameraSettings::FDisplayClusterConfigurationICVFX_CameraSettings()
+{
+	AllNodesColorGrading.bEnableEntireClusterColorGrading = true;
+}
+
+const FOpenColorIOColorConversionSettings* FDisplayClusterConfigurationICVFX_CameraSettings::FindInnerFrustumOCIOConfiguration(const FString& InClusterNodeId) const
+{
+	return CameraOCIO.FindOCIOConfiguration(InClusterNodeId);
+}
+
+const FOpenColorIOColorConversionSettings* FDisplayClusterConfigurationICVFX_CameraSettings::FindChromakeyOCIOConfiguration(const FString& InClusterNodeId) const
+{
+	// Always use incamera OCIO
+	return CameraOCIO.FindOCIOConfiguration(InClusterNodeId);
+}
+
+#if WITH_EDITOR
+bool FDisplayClusterConfigurationICVFX_CameraSettings::IsInnerFrustumViewportSettingsEqual_Editor(const FString& InClusterNodeId1, const FString& InClusterNodeId2) const
+{
+	return CameraOCIO.IsInnerFrustumViewportSettingsEqual_Editor(InClusterNodeId1, InClusterNodeId2);
+}
+
+bool FDisplayClusterConfigurationICVFX_CameraSettings::IsChromakeyViewportSettingsEqual_Editor(const FString& InClusterNodeId1, const FString& InClusterNodeId2) const
+{
+	return CameraOCIO.IsChromakeyViewportSettingsEqual_Editor(InClusterNodeId1, InClusterNodeId2);
+}
+
+#endif
+
+///////////////////////////////////////////////////////////////////////////////////////
+// FDisplayClusterConfigurationICVFX_StageSettings
+///////////////////////////////////////////////////////////////////////////////////////
+const FOpenColorIOColorConversionSettings* FDisplayClusterConfigurationICVFX_StageSettings::FindViewportOCIOConfiguration(const FString& InViewportId) const
+{
+	return ViewportOCIO.FindOCIOConfiguration(InViewportId);
+}
+
+const FOpenColorIOColorConversionSettings* FDisplayClusterConfigurationICVFX_StageSettings::FindLightcardOCIOConfiguration(const FString& InViewportId) const
+{
+	return Lightcard.LightcardOCIO.FindOCIOConfiguration(InViewportId, ViewportOCIO);
+}
diff --git a/Engine/Plugins/Runtime/nDisplay/Source/DisplayClusterConfiguration/Public/DisplayClusterConfigurationTypes_Enums.h b/Engine/Plugins/Runtime/nDisplay/Source/DisplayClusterConfiguration/Public/DisplayClusterConfigurationTypes_Enums.h
--- a/Engine/Plugins/Runtime/nDisplay/Source/DisplayClusterConfiguration/Public/DisplayClusterConfigurationTypes_Enums.h
+++ b/Engine/Plugins/Runtime/nDisplay/Source/DisplayClusterConfiguration/Public/DisplayClusterConfigurationTypes_Enums.h
@@ -164,3 +164,16 @@ enum class EDisplayClusterConfigurationViewportCustomFrustumMode : uint8
 	Percent  UMETA(DisplayName = "Percentage"),
 	Pixels   UMETA(DisplayName = "Pixels")
 };
+
+UENUM()
+enum class EDisplayClusterConfigurationViewportLightcardOCIOMode : uint8
+{
+	/** Use nDisplay Viewport OCIO. */
+	nDisplay UMETA(DisplayName = "Use nDisplay Viewport OCIO"),
+
+	/** Use Custom Light Cards OCIO. */
+	Custom UMETA(DisplayName = "Use Custom Light Cards OCIO"),
+
+	/** None. */
+	None UMETA(DisplayName = "None"),
+};
diff --git a/Engine/Plugins/Runtime/nDisplay/Source/DisplayClusterConfiguration/Public/DisplayClusterConfigurationTypes_ICVFX.h b/Engine/Plugins/Runtime/nDisplay/Source/DisplayClusterConfiguration/Public/DisplayClusterConfigurationTypes_ICVFX.h
--- a/Engine/Plugins/Runtime/nDisplay/Source/DisplayClusterConfiguration/Public/DisplayClusterConfigurationTypes_ICVFX.h
+++ b/Engine/Plugins/Runtime/nDisplay/Source/DisplayClusterConfiguration/Public/DisplayClusterConfigurationTypes_ICVFX.h
@@ -19,6 +19,90 @@
 
 #include "DisplayClusterConfigurationTypes_ICVFX.generated.h"
 
+USTRUCT(Blueprintable)
+struct DISPLAYCLUSTERCONFIGURATION_API FDisplayClusterConfigurationICVFX_LightcardCustomOCIO
+{
+	GENERATED_BODY()
+
+public:
+	/** Return LightCard OCIO configuration for the specified viewport. Return nullptr if no OCIO. */
+	const FOpenColorIOColorConversionSettings* FindOCIOConfiguration(const FString& InViewportId) const;
+
+public:
+	/** Apply this OpenColorIO configuration to all viewports. */
+	UPROPERTY(EditAnywhere, BlueprintReadWrite, Category = "Light Card", meta = (DisplayName = "All Viewports Color Configuration", ConfigurationMode = "Viewports"))
+	FDisplayClusterConfigurationOCIOConfiguration AllViewportsOCIOConfiguration;
+
+	/** Apply an OpenColorIO configuration on a per-viewport or group-of-viewports basis. */
+	UPROPERTY(EditAnywhere, BlueprintReadWrite, Category = "Light Card", meta = (DisplayName = "Per-Viewport OCIO Overrides", ConfigurationMode = "Viewports"))
+	TArray<FDisplayClusterConfigurationOCIOProfile> PerViewportOCIOProfiles;
+};
+
+USTRUCT(Blueprintable)
+struct DISPLAYCLUSTERCONFIGURATION_API FDisplayClusterConfigurationICVFX_ViewportOCIO
+{
+	GENERATED_BODY()
+
+public:
+	/** Return OCIO configuration for the specified viewport. Return nullptr if no OCIO. */
+	const FOpenColorIOColorConversionSettings* FindOCIOConfiguration(const FString& InViewportId) const;
+
+public:
+	/** Apply this OpenColorIO configuration to all viewports. */
+	UPROPERTY(EditAnywhere, BlueprintReadWrite, Category = "OCIO", meta = (DisplayName = "All Viewports Color Configuration", ConfigurationMode = "Viewports"))
+	FDisplayClusterConfigurationOCIOConfiguration AllViewportsOCIOConfiguration;
+
+	/** Apply an OpenColorIO configuration on a per-viewport or group-of-viewports basis. */
+	UPROPERTY(EditAnywhere, BlueprintReadWrite, Category = "OCIO", meta = (DisplayName = "Per-Viewport OCIO Overrides", ConfigurationMode = "Viewports"))
+	TArray<FDisplayClusterConfigurationOCIOProfile> PerViewportOCIOProfiles;
+};
+
+USTRUCT(Blueprintable)
+struct DISPLAYCLUSTERCONFIGURATION_API FDisplayClusterConfigurationICVFX_CameraOCIO
+{
+	GENERATED_BODY()
+
+public:
+	/** Return InCamera OCIO configuration for the specified cluster node. Return nullptr if no OCIO. */
+	const FOpenColorIOColorConversionSettings* FindOCIOConfiguration(const FString& InClusterNodeId) const;
+
+#if WITH_EDITOR
+	/** Returns true if the InCamera OCIO configuration is the same for the input nodes. */
+	bool IsInnerFrustumViewportSettingsEqual_Editor(const FString& InClusterNodeId1, const FString& InClusterNodeId2) const;
+
+	/** Returns true if the Chromakey OCIO configuration is the same for the input nodes. */
+	bool IsChromakeyViewportSettingsEqual_Editor(const FString& InClusterNodeId1, const FString& InClusterNodeId2) const;
+#endif
+
+public:
+	/** OCIO Display look configuration for all nodes */
+	UPROPERTY(EditAnywhere, BlueprintReadWrite, Category = "OCIO", meta = (DisplayName = "All Nodes Color Configuration", ConfigurationMode = "ClusterNodes"))
+	FDisplayClusterConfigurationOCIOConfiguration AllNodesOCIOConfiguration;
+
+	/** Apply an OpenColorIO configuration on a per-node or group-of-nodes basis. */
+	UPROPERTY(EditAnywhere, BlueprintReadWrite, Category = "OCIO", meta = (DisplayName = "Per-Node OCIO Overrides", ConfigurationMode = "ClusterNodes"))
+	TArray<FDisplayClusterConfigurationOCIOProfile> PerNodeOCIOProfiles;
+};
+
+USTRUCT(Blueprintable)
+struct DISPLAYCLUSTERCONFIGURATION_API FDisplayClusterConfigurationICVFX_LightcardOCIO
+{
+	GENERATED_BODY()
+
+public:
+	/** Return LightCard OCIO configuration for the specified viewport. Return nullptr if no OCIO. */
+	const FOpenColorIOColorConversionSettings* FindOCIOConfiguration(const FString& InViewportId, const FDisplayClusterConfigurationICVFX_ViewportOCIO& InViewportOCIO) const;
+
+public:
+	/** Light Cards OCIO mode. */
+	UPROPERTY(EditAnywhere, BlueprintReadWrite, Category = "Light Card|OCIO", meta = (DisplayName = "Light Cards OCIO"))
+	EDisplayClusterConfigurationViewportLightcardOCIOMode LightcardOCIOMode;
+
+	/** Custom OpenColorIO configuration for Light Cards. */
+	UPROPERTY(EditAnywhere, BlueprintReadWrite, Category = "Light Card|OCIO", meta = (DisplayName = "Custom Light Cards OCIO"))
+	FDisplayClusterConfigurationICVFX_LightcardCustomOCIO CustomOCIO;
+};
+
 USTRUCT(Blueprintable, meta = (DefaultSubstitutions = "LayersTooltip = Actor Layers, ActorsTooltip = Actor references"))
 struct DISPLAYCLUSTERCONFIGURATION_API FDisplayClusterConfigurationICVFX_VisibilityList
 {
@@ -256,9 +340,9 @@ public:
 	UPROPERTY(BlueprintReadWrite, Category = NDisplay, meta = (EditCondition = "bEnable"))
 	FDisplayClusterConfigurationICVFX_LightcardRenderSettings RenderSettings;
 
-	// Enable using outer viewport OCIO from DCRA for lightcard rendering
-	UPROPERTY()
-	bool bEnableOuterViewportOCIO = false;
+	/** OpenColorIO configuration for the lightcards. */
+	UPROPERTY(EditAnywhere, BlueprintReadWrite, Category = "OCIO", meta = (DisplayName = "Lightcard OCIO", EditCondition = "bEnable"))
+	FDisplayClusterConfigurationICVFX_LightcardOCIO LightcardOCIO;
 
 	// Enable using outer viewport Color Grading from DCRA for lightcard rendering
 	UPROPERTY()
@@ -465,8 +549,20 @@ struct DISPLAYCLUSTERCONFIGURATION_API FDisplayClusterConfigurationICVFX_CameraS
 {
 	GENERATED_BODY()
 
+public:
 	FDisplayClusterConfigurationICVFX_CameraSettings();
 
+	/** Return InCamera OCIO configuration for the specified cluster node. Return nullptr if no OCIO. */
+	const FOpenColorIOColorConversionSettings* FindInnerFrustumOCIOConfiguration(const FString& InClusterNodeId) const;
+	
+	/** Return Chromakey OCIO configuration for the specified cluster node. Return nullptr if no OCIO. */
+	const FOpenColorIOColorConversionSettings* FindChromakeyOCIOConfiguration(const FString& InClusterNodeId) const;
+
+#if WITH_EDITOR
+	bool IsInnerFrustumViewportSettingsEqual_Editor(const FString& InClusterNodeId1, const FString& InClusterNodeId2) const;
+	bool IsChromakeyViewportSettingsEqual_Editor(const FString& InClusterNodeId1, const FString& InClusterNodeId2) const;
+#endif
+
 public:
 	/** Render the inner frustum for this ICVFX camera. */
 	UPROPERTY(EditAnywhere, BlueprintReadWrite, Category = "In Camera VFX", meta = (DisplayName = "Enable Inner Frustum"))
@@ -511,13 +607,8 @@ public:
 	UPROPERTY(EditAnywhere, BlueprintReadWrite, Category = Chromakey, meta = (EditCondition = "bEnable"))
 	FDisplayClusterConfigurationICVFX_ChromakeySettings Chromakey;
 
-	/** OCIO Display look configuration for this camera */
-	UPROPERTY(EditAnywhere, BlueprintReadWrite, Category = "OCIO", meta = (DisplayName = "All Nodes OCIO Configuration", EditCondition = "bEnable"))
-	FDisplayClusterConfigurationOCIOConfiguration AllNodesOCIOConfiguration;
-
-	/** Apply an OpenColorIO configuration on a per-node or group-of-nodes basis. */
-	UPROPERTY(EditAnywhere, BlueprintReadWrite, Category = "OCIO", meta = (DisplayName = "Per-Node OCIO Overrides", ConfigurationMode = "ClusterNodes", EditCondition = "bEnable"))
-	TArray<FDisplayClusterConfigurationOCIOProfile> PerNodeOCIOProfiles;
+	UPROPERTY(EditAnywhere, BlueprintReadWrite, Category = OCIO, meta = (EditCondition = "bEnable"))
+	FDisplayClusterConfigurationICVFX_CameraOCIO CameraOCIO;
 
 	/** Entire Cluster Color Grading */
 	UPROPERTY(EditAnywhere, BlueprintReadWrite, Category = "Inner Frustum Color Grading", meta = (DisplayName = "Enable Inner Frustum Color Grading"))
@@ -549,6 +640,13 @@ struct DISPLAYCLUSTERCONFIGURATION_API FDisplayClusterConfigurationICVFX_StageSe
 {
 	GENERATED_BODY()
 
+public:
+	/** Return OCIO configuration for the specified viewport. Return nullptr if no OCIO. */
+	const FOpenColorIOColorConversionSettings* FindViewportOCIOConfiguration(const FString& InViewportId) const;
+
+	/** Return LightCard OCIO configuration for the specified viewport. Return nullptr if no OCIO. */
+	const FOpenColorIOColorConversionSettings* FindLightcardOCIOConfiguration(const FString& InViewportId) const;
+
 public:
 	/** Enable/disable the inner frustum on all ICVFX cameras. */
 	UPROPERTY(EditAnywhere, BlueprintReadWrite, Category = "In Camera VFX", meta = (DisplayName = "Enable Inner Frustum"))
@@ -586,15 +684,7 @@ public:
 	UPROPERTY(EditAnywhere, BlueprintReadWrite, Category = "Viewport Color Grading", meta = (DisplayName = "Per-Viewport Color Grading", ConfigurationMode = "Viewports"))
 	TArray<FDisplayClusterConfigurationViewport_PerViewportColorGrading> PerViewportColorGrading;
 
-	/** Enable the application of an OpenColorIO configuration to all viewports. */
-	UPROPERTY(EditAnywhere, BlueprintReadWrite, Category = "OCIO", meta = (DisplayName = "Enable Viewport OCIO"))
-	bool bUseOverallClusterOCIOConfiguration = true;
-
-	/** Apply this OpenColorIO configuration to all viewports. */
-	UPROPERTY(EditAnywhere, BlueprintReadWrite, Category = "OCIO", meta = (DisplayName = "All Viewports Color Configuration", EditCondition = "bUseOverallClusterOCIOConfiguration"))
-	FDisplayClusterConfigurationOCIOConfiguration AllViewportsOCIOConfiguration;
-
-	/** Apply an OpenColorIO configuration on a per-viewport or group-of-viewports basis. */
-	UPROPERTY(EditAnywhere, BlueprintReadWrite, Category = "OCIO", meta = (DisplayName = "Per-Viewport OCIO Overrides", ConfigurationMode = "Viewports", EditCondition = "bUseOverallClusterOCIOConfiguration"))
-	TArray<FDisplayClusterConfigurationOCIOProfile> PerViewportOCIOProfiles;
+	/** OpenColorIO configuration for the Outer viewports. */
+	UPROPERTY(EditAnywhere, BlueprintReadWrite, Category = "OCIO", meta = (DisplayName = "Outer Viewports OCIO"))
+	FDisplayClusterConfigurationICVFX_ViewportOCIO ViewportOCIO;
 };
diff --git a/Engine/Plugins/Runtime/nDisplay/Source/DisplayClusterConfiguration/Public/DisplayClusterConfigurationTypes_OCIO.h b/Engine/Plugins/Runtime/nDisplay/Source/DisplayClusterConfiguration/Public/DisplayClusterConfigurationTypes_OCIO.h
--- a/Engine/Plugins/Runtime/nDisplay/Source/DisplayClusterConfiguration/Public/DisplayClusterConfigurationTypes_OCIO.h
+++ b/Engine/Plugins/Runtime/nDisplay/Source/DisplayClusterConfiguration/Public/DisplayClusterConfigurationTypes_OCIO.h
@@ -10,42 +10,61 @@
 #include "DisplayClusterConfigurationTypes_OCIO.generated.h"
 
 /*
- * OCIO profile structure. Can be configured for viewports or cluster nodes.
- * To enable viewport configuration when using as a UPROPERTY set meta = (ConfigurationMode = "Viewports")
- * To enable cluster node configuration when using as a UPROPERTY set meta = (ConfigurationMode = "ClusterNodes")
+ * OCIO configuration structure.
  */
 USTRUCT(Blueprintable)
 struct FDisplayClusterConfigurationOCIOConfiguration
 {
 	GENERATED_BODY()
 
-	FDisplayClusterConfigurationOCIOConfiguration();
-
-	/** Enable the application of an OpenColorIO configuration to all viewports. */
-	UPROPERTY(EditAnywhere, BlueprintReadWrite, Category = "OCIO", meta = (DisplayName = "Enable Viewport OCIO"))
+	/**  Enable the application of an OpenColorIO configuration to all viewports. */
+	UPROPERTY(EditAnywhere, BlueprintReadWrite, Category = "OCIO")
 	bool bIsEnabled = true;
 
-	/** OpenColorIO Configuration */
-	UPROPERTY(EditAnywhere, BlueprintReadWrite, Category = "OCIO", meta = (DisplayName = "OCIO Configuration", EditCondition = "bIsEnabled"))
-	FOpenColorIODisplayConfiguration OCIOConfiguration;
+	/** "This property has been deprecated. */
+	UPROPERTY(meta = (DeprecatedProperty, DeprecationMessage = "Use the ColorConfiguration property instead"))
+	FOpenColorIODisplayConfiguration OCIOConfiguration_DEPRECATED;
+
+	/** Conversion to apply when this display is enabled */
+	UPROPERTY(EditAnywhere, BlueprintReadWrite, Category = "OCIO")
+	FOpenColorIOColorConversionSettings ColorConfiguration;
+
+public:
+	/** Return true if configuration valid */
+	bool IsEnabled() const;
 };
 
+/*
+ * OCIO profile structure. Can be configured for viewports or cluster nodes.
+ * To enable viewport configuration when using as a UPROPERTY set meta = (ConfigurationMode = "Viewports")
+ * To enable cluster node configuration when using as a UPROPERTY set meta = (ConfigurationMode = "ClusterNodes")
+ */
 USTRUCT(Blueprintable)
 struct FDisplayClusterConfigurationOCIOProfile
 {
 	GENERATED_BODY()
 
-	FDisplayClusterConfigurationOCIOProfile();
-
 	/** Enable the application of an OpenColorIO configuration for the viewport(s) specified. */
-	UPROPERTY(EditAnywhere, BlueprintReadWrite, Category = "OCIO", meta = (DisplayName = "Enable OCIO"))
+	UPROPERTY(EditAnywhere, BlueprintReadWrite, Category = "OCIO")
 	bool bIsEnabled = true;
 
-	/** OpenColorIO Configuration */
-	UPROPERTY(EditAnywhere, BlueprintReadWrite, Category = "OCIO", meta = (DisplayName = "OCIO Configuration", EditCondition = "bIsEnabled"))
-	FOpenColorIODisplayConfiguration OCIOConfiguration;
+	/** "This property has been deprecated. */
+	UPROPERTY(meta = (DeprecatedProperty, DeprecationMessage = "Use the ColorConfiguration property instead"))
+	FOpenColorIODisplayConfiguration OCIOConfiguration_DEPRECATED;
 
 	/** Specify the viewports to apply this OpenColorIO configuration. */
 	UPROPERTY(EditAnywhere, BlueprintReadOnly, Category = "OCIO", meta = (EditCondition = "bIsEnabled"))
 	TArray<FString> ApplyOCIOToObjects;
+
+	/** Conversion to apply when this display is enabled */
+	UPROPERTY(EditAnywhere, BlueprintReadWrite, Category = "OCIO")
+	FOpenColorIOColorConversionSettings ColorConfiguration;
+
+public:
+	/** Return true if the configuration is valid for the input object */
+	bool IsEnabledForObject(const FString& InObjectId) const;
+
+	/** Return true if configuration valid */
+	bool IsEnabled() const;
 };
+
diff --git a/Engine/Plugins/Runtime/nDisplay/Source/DisplayClusterConfigurator/Private/Views/Details/Cluster/DisplayClusterConfiguratorNodeSelectionCustomization.cpp b/Engine/Plugins/Runtime/nDisplay/Source/DisplayClusterConfigurator/Private/Views/Details/Cluster/DisplayClusterConfiguratorNodeSelectionCustomization.cpp
--- a/Engine/Plugins/Runtime/nDisplay/Source/DisplayClusterConfigurator/Private/Views/Details/Cluster/DisplayClusterConfiguratorNodeSelectionCustomization.cpp
+++ b/Engine/Plugins/Runtime/nDisplay/Source/DisplayClusterConfigurator/Private/Views/Details/Cluster/DisplayClusterConfiguratorNodeSelectionCustomization.cpp
@@ -250,8 +250,7 @@ void FDisplayClusterConfiguratorOCIOProfileCustomization::SetChildren(const TSha
 	if (!bIsDefaultDetailsDisplay)
 	{
 		TSharedPtr<IPropertyHandle> EnableOCIOHandle = GET_CHILD_HANDLE(FDisplayClusterConfigurationOCIOProfile, bIsEnabled);
-		TSharedPtr<IPropertyHandle> OCIOConfigurationHandle = GET_CHILD_HANDLE(FDisplayClusterConfigurationOCIOProfile, OCIOConfiguration);
-		TSharedPtr<IPropertyHandle> OCIOHandle = GET_CHILD_HANDLE(FOpenColorIODisplayConfiguration, ColorConfiguration);
+		TSharedPtr<IPropertyHandle> OCIOHandle = GET_CHILD_HANDLE(FDisplayClusterConfigurationOCIOProfile, ColorConfiguration);
 		TSharedPtr<IPropertyHandle> ArrayHandle = GET_CHILD_HANDLE(FDisplayClusterConfigurationOCIOProfile, ApplyOCIOToObjects);
 
 		EnableOCIOHandle->SetPropertyDisplayName(Mode == FDisplayClusterConfiguratorNodeSelection::EOperationMode::Viewports ?
diff --git a/Engine/Plugins/Runtime/nDisplay/Source/DisplayClusterProjection/Private/Policy/MPCDI/DisplayClusterProjectionMPCDIPolicy.cpp b/Engine/Plugins/Runtime/nDisplay/Source/DisplayClusterProjection/Private/Policy/MPCDI/DisplayClusterProjectionMPCDIPolicy.cpp
--- a/Engine/Plugins/Runtime/nDisplay/Source/DisplayClusterProjection/Private/Policy/MPCDI/DisplayClusterProjectionMPCDIPolicy.cpp
+++ b/Engine/Plugins/Runtime/nDisplay/Source/DisplayClusterProjection/Private/Policy/MPCDI/DisplayClusterProjectionMPCDIPolicy.cpp
@@ -259,7 +259,7 @@ void FDisplayClusterProjectionMPCDIPolicy::ApplyWarpBlend_RenderThread(FRHIComma
 	IDisplayClusterShaders& ShadersAPI = IDisplayClusterShaders::Get();
 
 	const FDisplayClusterViewport_RenderSettingsICVFX& SettingsICVFX = InViewportProxy->GetRenderSettingsICVFX_RenderThread();
-	if ((SettingsICVFX.RuntimeFlags & ViewportRuntime_ICVFXTarget) != 0)
+	if (EnumHasAllFlags(SettingsICVFX.RuntimeFlags, EDisplayClusterViewportRuntimeICVFXFlags::Target))
 	{
 		// This viewport used as icvfx target:
 		TArray<FDisplayClusterShaderParametersICVFX_ViewportResource*> ViewportResources;
@@ -318,12 +318,6 @@ void FDisplayClusterProjectionMPCDIPolicy::ApplyWarpBlend_RenderThread(FRHIComma
 					}
 				}
 
-				TSharedPtr<IDisplayClusterViewportLightCardManager, ESPMode::ThreadSafe> LightCardManager = RefViewportManagerProxy.GetLightCardManager_RenderThread();
-				if (LightCardManager.IsValid())
-				{
-					ShaderICVFX.UVLightCardMap = LightCardManager->GetUVLightCardMap_RenderThread();
-				}
-
 				// cleanup camera list for render
 				ShaderICVFX.CleanupCamerasForRender();
 
diff --git a/Engine/Plugins/Runtime/nDisplay/Source/DisplayClusterShaders/Private/DisplayClusterShadersModule.cpp b/Engine/Plugins/Runtime/nDisplay/Source/DisplayClusterShaders/Private/DisplayClusterShadersModule.cpp
--- a/Engine/Plugins/Runtime/nDisplay/Source/DisplayClusterShaders/Private/DisplayClusterShadersModule.cpp
+++ b/Engine/Plugins/Runtime/nDisplay/Source/DisplayClusterShaders/Private/DisplayClusterShadersModule.cpp
@@ -42,9 +42,9 @@ bool FDisplayClusterShadersModule::RenderWarpBlend_ICVFX(FRHICommandListImmediat
 	return FDisplayClusterShadersWarpblend_ICVFX::RenderWarpBlend_ICVFX(RHICmdList, InWarpBlendParameters, InICVFXParameters);
 }
 
-bool FDisplayClusterShadersModule::RenderPreprocess_UVLightCards(FRHICommandListImmediate& RHICmdList, FSceneInterface* InScene, FRenderTarget* InRenderTarget, float ProjectionPlaneSize) const
+bool FDisplayClusterShadersModule::RenderPreprocess_UVLightCards(FRHICommandListImmediate& RHICmdList, FSceneInterface* InScene, FRenderTarget* InRenderTarget, float ProjectionPlaneSize, bool bRenderFinalColor) const
 {
-	return FDisplayClusterShadersPreprocess_UVLightCards::RenderPreprocess_UVLightCards(RHICmdList, InScene, InRenderTarget, ProjectionPlaneSize);
+	return FDisplayClusterShadersPreprocess_UVLightCards::RenderPreprocess_UVLightCards(RHICmdList, InScene, InRenderTarget, ProjectionPlaneSize, bRenderFinalColor);
 }
 
 bool FDisplayClusterShadersModule::RenderPostprocess_OutputRemap(FRHICommandListImmediate& RHICmdList, FRHITexture2D* InSourceTexture, FRHITexture2D* InRenderTargetableDestTexture, const IDisplayClusterRender_MeshComponentProxy& MeshProxy) const
diff --git a/Engine/Plugins/Runtime/nDisplay/Source/DisplayClusterShaders/Private/DisplayClusterShadersModule.h b/Engine/Plugins/Runtime/nDisplay/Source/DisplayClusterShaders/Private/DisplayClusterShadersModule.h
--- a/Engine/Plugins/Runtime/nDisplay/Source/DisplayClusterShaders/Private/DisplayClusterShadersModule.h
+++ b/Engine/Plugins/Runtime/nDisplay/Source/DisplayClusterShaders/Private/DisplayClusterShadersModule.h
@@ -19,7 +19,7 @@ public:
 public:
 	virtual bool RenderWarpBlend_MPCDI(FRHICommandListImmediate& RHICmdList, const FDisplayClusterShaderParameters_WarpBlend& InWarpBlendParameters) const override;
 	virtual bool RenderWarpBlend_ICVFX(FRHICommandListImmediate& RHICmdList, const FDisplayClusterShaderParameters_WarpBlend& InWarpBlendParameters, const FDisplayClusterShaderParameters_ICVFX& InICVFXParameters) const override;
-	virtual bool RenderPreprocess_UVLightCards(FRHICommandListImmediate& RHICmdList, FSceneInterface* InScene, FRenderTarget* InRenderTarget, float ProjectionPlaneSize) const override;
+	virtual bool RenderPreprocess_UVLightCards(FRHICommandListImmediate& RHICmdList, FSceneInterface* InScene, FRenderTarget* InRenderTarget, float ProjectionPlaneSize, bool bRenderFinalColor) const override;
 	virtual bool RenderPostprocess_OutputRemap(FRHICommandListImmediate& RHICmdList, FRHITexture2D* InSourceTexture, FRHITexture2D* InRenderTargetableDestTexture, const IDisplayClusterRender_MeshComponentProxy& MeshProxy) const override;
 	virtual bool RenderPostprocess_Blur(FRHICommandListImmediate& RHICmdList, FRHITexture2D* InSourceTexture, FRHITexture2D* InRenderTargetableDestTexture, const FDisplayClusterShaderParameters_PostprocessBlur& InSettings) const override;
 	virtual bool GenerateMips(FRHICommandListImmediate& RHICmdList, FRHITexture2D* InOutMipsTexture, const FDisplayClusterShaderParameters_GenerateMips& InSettings) const override;
diff --git a/Engine/Plugins/Runtime/nDisplay/Source/DisplayClusterShaders/Private/Shaders/DisplayClusterShadersPreprocess_UVLightCards.cpp b/Engine/Plugins/Runtime/nDisplay/Source/DisplayClusterShaders/Private/Shaders/DisplayClusterShadersPreprocess_UVLightCards.cpp
--- a/Engine/Plugins/Runtime/nDisplay/Source/DisplayClusterShaders/Private/Shaders/DisplayClusterShadersPreprocess_UVLightCards.cpp
+++ b/Engine/Plugins/Runtime/nDisplay/Source/DisplayClusterShaders/Private/Shaders/DisplayClusterShadersPreprocess_UVLightCards.cpp
@@ -171,7 +171,7 @@ END_SHADER_PARAMETER_STRUCT()
 
 DECLARE_GPU_STAT_NAMED(nDisplay_UVLightCards_Render, TEXT("nDisplay UVLightCards::Render"));
 
-bool FDisplayClusterShadersPreprocess_UVLightCards::RenderPreprocess_UVLightCards(FRHICommandListImmediate& RHICmdList, FSceneInterface* InScene, FRenderTarget* InRenderTarget, float ProjectionPlaneSize)
+bool FDisplayClusterShadersPreprocess_UVLightCards::RenderPreprocess_UVLightCards(FRHICommandListImmediate& RHICmdList, FSceneInterface* InScene, FRenderTarget* InRenderTarget, float ProjectionPlaneSize, bool bRenderFinalColor)
 {
 	check(IsInRenderingThread());
 
@@ -201,6 +201,12 @@ bool FDisplayClusterShadersPreprocess_UVLightCards::RenderPreprocess_UVLightCard
 		EngineShowFlags.SetLumenReflections(0);
 		EngineShowFlags.SetLumenGlobalIllumination(0);
 		EngineShowFlags.SetGlobalIllumination(0);
+
+		EngineShowFlags.SetScreenSpaceAO(0);
+		EngineShowFlags.SetAmbientOcclusion(0);
+		EngineShowFlags.SetDeferredLighting(0);
+		EngineShowFlags.SetVirtualTexturePrimitives(0);
+		EngineShowFlags.SetRectLights(0);
 	}
 
 	FSceneViewFamilyContext ViewFamily(FSceneViewFamily::ConstructionValues(
@@ -253,7 +259,17 @@ bool FDisplayClusterShadersPreprocess_UVLightCards::RenderPreprocess_UVLightCard
 	ViewInitOptions.BackgroundColor = FLinearColor::Black;
 
 	GetRendererModule().CreateAndInitSingleView(RHICmdList, &ViewFamily, &ViewInitOptions);
-	const FSceneView* View = ViewFamily.Views[0];
+	FViewInfo* View = (FViewInfo*)ViewFamily.Views[0];
+
+	if (!bRenderFinalColor)
+	{
+		ViewFamily.SceneCaptureSource = ESceneCaptureSource::SCS_SceneColorHDR;
+	}
+
+	ViewFamily.EngineShowFlags.SetToneCurve(false);
+	// This flags sets tonampper to output to ETonemapperOutputDevice::LinearNoToneCurve
+	View->FinalPostProcessSettings.bOverride_ToneCurveAmount = 1;
+	View->FinalPostProcessSettings.ToneCurveAmount = 0.0;
 
 	FRDGTextureRef OutputTexture = GraphBuilder.RegisterExternalTexture(CreateRenderTarget(InRenderTarget->GetRenderTargetTexture(), TEXT("UVLightCardRenderTarget")));
 	FRenderTargetBinding OutputRenderTargetBinding(OutputTexture, ERenderTargetLoadAction::EClear);
diff --git a/Engine/Plugins/Runtime/nDisplay/Source/DisplayClusterShaders/Private/Shaders/DisplayClusterShadersPreprocess_UVLightCards.h b/Engine/Plugins/Runtime/nDisplay/Source/DisplayClusterShaders/Private/Shaders/DisplayClusterShadersPreprocess_UVLightCards.h
--- a/Engine/Plugins/Runtime/nDisplay/Source/DisplayClusterShaders/Private/Shaders/DisplayClusterShadersPreprocess_UVLightCards.h
+++ b/Engine/Plugins/Runtime/nDisplay/Source/DisplayClusterShaders/Private/Shaders/DisplayClusterShadersPreprocess_UVLightCards.h
@@ -12,5 +12,5 @@ class FRenderTarget;
 class FDisplayClusterShadersPreprocess_UVLightCards
 {
 public:
-	static bool RenderPreprocess_UVLightCards(FRHICommandListImmediate& RHICmdList, FSceneInterface* InScene, FRenderTarget* InRenderTarget, float ProjectionPlaneSize);
+	static bool RenderPreprocess_UVLightCards(FRHICommandListImmediate& RHICmdList, FSceneInterface* InScene, FRenderTarget* InRenderTarget, float ProjectionPlaneSize, bool bRenderFinalColor);
 };
\ No newline at end of file
