nDisplay: ICVFX OCIO

diff --git a/Engine/Plugins/Runtime/nDisplay/Source/DisplayClusterShaders/Private/Shaders/DisplayClusterShadersWarpblend_ICVFX.cpp b/Engine/Plugins/Runtime/nDisplay/Source/DisplayClusterShaders/Private/Shaders/DisplayClusterShadersWarpblend_ICVFX.cpp
--- a/Engine/Plugins/Runtime/nDisplay/Source/DisplayClusterShaders/Private/Shaders/DisplayClusterShadersWarpblend_ICVFX.cpp
+++ b/Engine/Plugins/Runtime/nDisplay/Source/DisplayClusterShaders/Private/Shaders/DisplayClusterShadersWarpblend_ICVFX.cpp
@@ -47,9 +47,17 @@ static TAutoConsoleVariable<int32> CVarIcvfxMPCDIShaderType(
 
 int32 GDisplayClusterShadersICVFXEnableLightCard = 1;
 static FAutoConsoleVariableRef CVarDisplayClusterShadersICVFXEnableLightCard(
-	TEXT("nDisplay.render.icvfx.lightcard"),
+	TEXT("nDisplay.render.icvfx.LightCard"),
 	GDisplayClusterShadersICVFXEnableLightCard,
-	TEXT("Enable ICVFX lightcards render (0 - disable).\n"),
+	TEXT("Enable ICVFX LightCards render (0 - disable).\n"),
+	ECVF_RenderThreadSafe
+);
+
+int32 GDisplayClusterShadersICVFXEnableUVLightCard = 1;
+static FAutoConsoleVariableRef CVarDisplayClusterShadersICVFXEnableUVLightCard(
+	TEXT("nDisplay.render.icvfx.UVLightCard"),
+	GDisplayClusterShadersICVFXEnableUVLightCard,
+	TEXT("Enable ICVFX UVLightCards render (0 - disable).\n"),
 	ECVF_RenderThreadSafe
 );
 
@@ -66,9 +74,11 @@ namespace IcvfxShaderPermutation
 	class FIcvfxShaderViewportInput      : SHADER_PERMUTATION_BOOL("VIEWPORT_INPUT");
 	class FIcvfxShaderViewportInputAlpha : SHADER_PERMUTATION_BOOL("VIEWPORT_INPUT_ALPHA");
 
-	class FIcvfxShaderOverlayUnder     : SHADER_PERMUTATION_BOOL("OVERLAY_UNDER");
-	class FIcvfxShaderOverlayOver      : SHADER_PERMUTATION_BOOL("OVERLAY_OVER");
-	class FIcvfxShaderOverlayAlpha     : SHADER_PERMUTATION_BOOL("OVERLAY_ALPHA");
+	class FIcvfxShaderLightCardUnder     : SHADER_PERMUTATION_BOOL("LIGHTCARD_UNDER");
+	class FIcvfxShaderLightCardOver      : SHADER_PERMUTATION_BOOL("LIGHTCARD_OVER");
+
+	class FIcvfxShaderUVLightCardUnder : SHADER_PERMUTATION_BOOL("UVLIGHTCARD_UNDER");
+	class FIcvfxShaderUVLightCardOver : SHADER_PERMUTATION_BOOL("UVLIGHTCARD_OVER");
 
 	class FIcvfxShaderInnerCamera      : SHADER_PERMUTATION_BOOL("INNER_CAMERA");
 
@@ -82,17 +92,16 @@ namespace IcvfxShaderPermutation
 
 	class FIcvfxShaderMeshWarp         : SHADER_PERMUTATION_BOOL("MESH_WARP");
 
-	class FIcvfxShaderUVLightCards : SHADER_PERMUTATION_BOOL("UV_LIGHT_CARDS");
-
 	using FCommonVSDomain = TShaderPermutationDomain<FIcvfxShaderMeshWarp>;
 
 	using FCommonPSDomain = TShaderPermutationDomain<
 		FIcvfxShaderViewportInput,
 		FIcvfxShaderViewportInputAlpha,
 
-		FIcvfxShaderOverlayUnder,
-		FIcvfxShaderOverlayOver,
-		FIcvfxShaderOverlayAlpha,
+		FIcvfxShaderLightCardUnder,
+		FIcvfxShaderLightCardOver,
+		FIcvfxShaderUVLightCardUnder,
+		FIcvfxShaderUVLightCardOver,
 
 		FIcvfxShaderInnerCamera,
 		FIcvfxShaderChromakey,
@@ -101,20 +110,11 @@ namespace IcvfxShaderPermutation
 
 		FIcvfxShaderAlphaMapBlending,
 		FIcvfxShaderBetaMapBlending,
-		FIcvfxShaderMeshWarp,
-		FIcvfxShaderUVLightCards
+		FIcvfxShaderMeshWarp
 	>;
 	
 	bool ShouldCompileCommonPSPermutation(const FGlobalShaderPermutationParameters& Parameters, const FCommonPSDomain& PermutationVector)
 	{
-		if (PermutationVector.Get<FIcvfxShaderOverlayAlpha>())
-		{
-			if (!PermutationVector.Get<FIcvfxShaderOverlayUnder>() && !PermutationVector.Get<FIcvfxShaderOverlayOver>())
-			{
-				return false;
-			}
-		}
-
 		if (!PermutationVector.Get<FIcvfxShaderMeshWarp>())
 		{
 			if (PermutationVector.Get<FIcvfxShaderChromakey>() || PermutationVector.Get<FIcvfxShaderChromakeyFrameColor>())
@@ -130,7 +130,12 @@ namespace IcvfxShaderPermutation
 
 		if (!PermutationVector.Get<FIcvfxShaderViewportInput>())
 		{
-			if (PermutationVector.Get<FIcvfxShaderOverlayUnder>() || (PermutationVector.Get<FIcvfxShaderInnerCamera>() == PermutationVector.Get<FIcvfxShaderOverlayOver>()))
+			if (PermutationVector.Get<FIcvfxShaderLightCardUnder>() || (PermutationVector.Get<FIcvfxShaderInnerCamera>() == PermutationVector.Get<FIcvfxShaderLightCardOver>()))
+			{
+				return false;
+			}
+
+			if (PermutationVector.Get<FIcvfxShaderUVLightCardUnder>() || (PermutationVector.Get<FIcvfxShaderInnerCamera>() == PermutationVector.Get<FIcvfxShaderUVLightCardOver>()))
 			{
 				return false;
 			}
@@ -162,9 +167,13 @@ namespace IcvfxShaderPermutation
 			return false;
 		}
 
-		if (PermutationVector.Get<FIcvfxShaderUVLightCards>() && !PermutationVector.Get<FIcvfxShaderMeshWarp>())
+		// UVLightCard require UV_Chromakey (Mesh Warp)
+		if (!PermutationVector.Get<FIcvfxShaderMeshWarp>())
 		{
-			return false;
+			if (PermutationVector.Get<FIcvfxShaderUVLightCardUnder>() || PermutationVector.Get<FIcvfxShaderUVLightCardOver>())
+			{
+				return false;
+			}
 		}
 
 		return true;
@@ -189,23 +198,24 @@ BEGIN_SHADER_PARAMETER_STRUCT(FIcvfxPixelShaderParameters, )
 	SHADER_PARAMETER_TEXTURE(Texture2D, WarpMapTexture)
 	SHADER_PARAMETER_TEXTURE(Texture2D, AlphaMapTexture)
 	SHADER_PARAMETER_TEXTURE(Texture2D, BetaMapTexture)
-	SHADER_PARAMETER_TEXTURE(Texture2D, OverlayTexture)
-	SHADER_PARAMETER_TEXTURE(Texture2D, OverlayAlphaTexture)
 	SHADER_PARAMETER_TEXTURE(Texture2D, InnerCameraTexture)
 	SHADER_PARAMETER_TEXTURE(Texture2D, ChromakeyCameraTexture)
 	SHADER_PARAMETER_TEXTURE(Texture2D, ChromakeyMarkerTexture)
-	SHADER_PARAMETER_TEXTURE(Texture2D, UVLightCardMapTexture)
+
+	SHADER_PARAMETER_TEXTURE(Texture2D, LightCardTexture)
+	SHADER_PARAMETER_TEXTURE(Texture2D, UVLightCardTexture)
 
 	SHADER_PARAMETER_SAMPLER(SamplerState, InputSampler)
 	SHADER_PARAMETER_SAMPLER(SamplerState, WarpMapSampler)
 	SHADER_PARAMETER_SAMPLER(SamplerState, AlphaMapSampler)
 	SHADER_PARAMETER_SAMPLER(SamplerState, BetaMapSampler)
-	SHADER_PARAMETER_SAMPLER(SamplerState, OverlaySampler)
-	SHADER_PARAMETER_SAMPLER(SamplerState, OverlayAlphaSampler)
+	
 	SHADER_PARAMETER_SAMPLER(SamplerState, InnerCameraSampler)
 	SHADER_PARAMETER_SAMPLER(SamplerState, ChromakeyCameraSampler)
 	SHADER_PARAMETER_SAMPLER(SamplerState, ChromakeyMarkerSampler)
-	SHADER_PARAMETER_SAMPLER(SamplerState, UVLightCardMapSampler)
+
+	SHADER_PARAMETER_SAMPLER(SamplerState, LightCardSampler)
+	SHADER_PARAMETER_SAMPLER(SamplerState, UVLightCardSampler)
 
 	SHADER_PARAMETER(FMatrix44f, ViewportTextureProjectionMatrix)
 	SHADER_PARAMETER(FMatrix44f, OverlayProjectionMatrix)
@@ -306,6 +316,28 @@ private:
 	bool bForceMultiCameraPass = false;
 	int32 CameraIndex = 0;
 
+	//LightCards
+private:
+	inline bool IsLightcardOverUsed() const
+	{
+		return ICVFXParameters.IsLightcardOverUsed() && GDisplayClusterShadersICVFXEnableLightCard;
+	}
+
+	inline bool IsLightcardUnderUsed() const
+	{
+		return ICVFXParameters.IsLightcardUnderUsed() && GDisplayClusterShadersICVFXEnableLightCard;
+	}
+
+	inline bool IsUVLightcardOverUsed() const
+	{
+		return ICVFXParameters.IsUVLightcardOverUsed() && GDisplayClusterShadersICVFXEnableUVLightCard;
+	}
+
+	inline bool IsUVLightcardUnderUsed() const
+	{
+		return ICVFXParameters.IsUVLightcardUnderUsed() && GDisplayClusterShadersICVFXEnableUVLightCard;
+	}
+
 private:
 	inline bool IsMultiPassRender() const
 	{
@@ -329,7 +361,7 @@ private:
 			// 2 - second camera additive pass
 			// N - camera N
 			// N+1 - overlayOver (optional)
-			return (ICVFXParameters.IsLightcardOverUsed() && GDisplayClusterShadersICVFXEnableLightCard) ? ( CameraIndex == (TotalUsedCameras + 1) ) : ( CameraIndex == TotalUsedCameras );
+			return (IsLightcardOverUsed() || IsUVLightcardOverUsed()) ? ( CameraIndex == (TotalUsedCameras + 1) ) : ( CameraIndex == TotalUsedCameras );
 		}
 
 		if (TotalUsedCameras>1)
@@ -339,28 +371,45 @@ private:
 			// 1   - camera2
 			// N-1 - camera N
 			// N   - overlayOver (optional)
-			return (ICVFXParameters.IsLightcardOverUsed() && GDisplayClusterShadersICVFXEnableLightCard) ? ( CameraIndex == TotalUsedCameras ) : ( CameraIndex == (TotalUsedCameras-1) );
+			return (IsLightcardOverUsed() || IsUVLightcardOverUsed()) ? ( CameraIndex == TotalUsedCameras ) : ( CameraIndex == (TotalUsedCameras-1) );
 		}
 
 		// By default render single camera in one pass
 		return CameraIndex==0;
 	}
 
-	inline bool IsOverlayUnderUsed() const
+	inline bool IsLightCardRenderUnderInCamera() const
 	{
 		// Render OverlayUnder only on first pass
-		return IsFirstPass() && ICVFXParameters.IsLightcardUnderUsed() && GDisplayClusterShadersICVFXEnableLightCard;
+		return IsFirstPass() && IsLightcardUnderUsed();
 	}
 
-	inline bool IsOverlayOverUsed() const
+	inline bool IsLightCardRenderOverInCamera() const
 	{
 		if (IsMultiPassRender())
 		{
 			// Render OverlayOver only on last additive pass
-			return IsLastPass() && ICVFXParameters.IsLightcardOverUsed() && GDisplayClusterShadersICVFXEnableLightCard;
+			return IsLastPass() && IsLightcardOverUsed();
 		}
 		
-		return ICVFXParameters.IsLightcardOverUsed() && GDisplayClusterShadersICVFXEnableLightCard;
+		return IsLightcardOverUsed();
+	}
+
+	inline bool IsUVLightCardRenderUnderInCamera() const
+	{
+		// Render UVLightCard Under only on first pass
+		return IsFirstPass() && IsUVLightcardUnderUsed();
+	}
+
+	inline bool IsUVLightCardRenderOverInCamera() const
+	{
+		if (IsMultiPassRender())
+		{
+			// Render UVLightCard Over only on last additive pass
+			return IsLastPass() && IsUVLightcardOverUsed();
+		}
+
+		return IsUVLightcardOverUsed();
 	}
 
 	inline int32 GetUsedCameraIndex() const
@@ -522,29 +571,54 @@ public:
 		return false;
 	}
 
-	bool GetOverlayParameters(FIcvfxRenderPassData& RenderPassData)
+	bool GetLightCardParameters(FIcvfxRenderPassData& RenderPassData)
 	{
 		RenderPassData.PSParameters.OverlayProjectionMatrix = FMatrix44f(LocalUVMatrix);
 
-		if (IsOverlayUnderUsed() || IsOverlayOverUsed())
+		if (IsLightCardRenderUnderInCamera() || IsLightCardRenderOverInCamera())
 		{
+			RenderPassData.PSParameters.LightCardSampler = TStaticSamplerState<SF_Trilinear, AM_Clamp, AM_Clamp, AM_Clamp>::GetRHI();
+			RenderPassData.PSParameters.LightCardTexture = ICVFXParameters.Lightcard.Texture;
 
-				RenderPassData.PSParameters.OverlaySampler = TStaticSamplerState<SF_Trilinear, AM_Clamp, AM_Clamp, AM_Clamp>::GetRHI();
-				RenderPassData.PSParameters.OverlayTexture = ICVFXParameters.Lightcard.Texture;
-
-			if (IsOverlayUnderUsed())
+			if (IsLightCardRenderUnderInCamera())
 			{
-				RenderPassData.PSPermutationVector.Set<IcvfxShaderPermutation::FIcvfxShaderOverlayUnder>(true);
+				RenderPassData.PSPermutationVector.Set<IcvfxShaderPermutation::FIcvfxShaderLightCardUnder>(true);
 			}
-			if (IsOverlayOverUsed())
+			if (IsLightCardRenderOverInCamera())
 			{
-				RenderPassData.PSPermutationVector.Set<IcvfxShaderPermutation::FIcvfxShaderOverlayOver>(true);
+				RenderPassData.PSPermutationVector.Set<IcvfxShaderPermutation::FIcvfxShaderLightCardOver>(true);
 			}
 		}
 
 		return true;
 	}
 
+	bool GetUVLightCardParameters(FIcvfxRenderPassData& RenderPassData)
+	{
+		if (IsUVLightCardRenderUnderInCamera() || IsUVLightCardRenderOverInCamera())
+		{
+			// UVLightCard require Mesh warp (UV_Chromakey)
+			if (RenderPassData.VSPermutationVector.Get<IcvfxShaderPermutation::FIcvfxShaderMeshWarp>())
+			{
+				RenderPassData.PSParameters.UVLightCardSampler = TStaticSamplerState<SF_Trilinear, AM_Clamp, AM_Clamp, AM_Clamp>::GetRHI();
+				RenderPassData.PSParameters.UVLightCardTexture = ICVFXParameters.UVLightcard.Texture;
+
+				if (IsUVLightCardRenderUnderInCamera())
+				{
+					RenderPassData.PSPermutationVector.Set<IcvfxShaderPermutation::FIcvfxShaderUVLightCardUnder>(true);
+				}
+				if (IsUVLightCardRenderOverInCamera())
+				{
+					RenderPassData.PSPermutationVector.Set<IcvfxShaderPermutation::FIcvfxShaderUVLightCardOver>(true);
+				}
+
+				return true;
+			}
+		}
+
+		return false;
+	}
+
 	bool GetCameraParameters(FIcvfxRenderPassData& RenderPassData)
 	{
 		const int32 CurrentCameraIndex = GetUsedCameraIndex();
@@ -639,20 +713,6 @@ public:
 		return false;
 	}
 
-	bool GetUVLightCardParameters(FIcvfxRenderPassData& RenderPassData)
-	{
-		if (ICVFXParameters.UVLightCardMap.IsValid() && GDisplayClusterShadersICVFXEnableLightCard && RenderPassData.PSPermutationVector.Get<IcvfxShaderPermutation::FIcvfxShaderMeshWarp>())
-		{
-			RenderPassData.PSParameters.UVLightCardMapSampler = TStaticSamplerState<SF_Trilinear, AM_Clamp, AM_Clamp, AM_Clamp>::GetRHI();
-			RenderPassData.PSParameters.UVLightCardMapTexture = ICVFXParameters.UVLightCardMap;
-
-			RenderPassData.PSPermutationVector.Set<IcvfxShaderPermutation::FIcvfxShaderUVLightCards>(true);
-			return true;
-		}
-
-		return false;
-	}
-
 	void InitRenderPass(FIcvfxRenderPassData& RenderPassData)
 	{
 		// Forward input viewport
@@ -668,7 +728,9 @@ public:
 				GetWarpBlendParameters(RenderPassData);
 			}
 
-			GetOverlayParameters(RenderPassData);
+			GetLightCardParameters(RenderPassData);
+			GetUVLightCardParameters(RenderPassData);
+
 			if (GetCameraParameters(RenderPassData))
 			{
 				// Chromakey inside cam:
@@ -678,8 +740,6 @@ public:
 					GetCameraChromakeyMarkerParameters(RenderPassData);
 				}
 			}
-
-			GetUVLightCardParameters(RenderPassData);
 		}
 	}
 
diff --git a/Engine/Plugins/Runtime/nDisplay/Source/DisplayClusterShaders/Public/IDisplayClusterShaders.h b/Engine/Plugins/Runtime/nDisplay/Source/DisplayClusterShaders/Public/IDisplayClusterShaders.h
--- a/Engine/Plugins/Runtime/nDisplay/Source/DisplayClusterShaders/Public/IDisplayClusterShaders.h
+++ b/Engine/Plugins/Runtime/nDisplay/Source/DisplayClusterShaders/Public/IDisplayClusterShaders.h
@@ -90,7 +90,7 @@ public:
 	*
 	* @return - true if success
 	*/
-	virtual bool RenderPreprocess_UVLightCards(FRHICommandListImmediate& RHICmdList, FSceneInterface* InScene, FRenderTarget* InRenderTarget, float ProjectionPlaneSize) const = 0;
+	virtual bool RenderPreprocess_UVLightCards(FRHICommandListImmediate& RHICmdList, FSceneInterface* InScene, FRenderTarget* InRenderTarget, float ProjectionPlaneSize, bool bRenderFinalColor) const = 0;
 
 	/**
 	* Render postprocess OutputRemap
diff --git a/Engine/Plugins/Runtime/nDisplay/Source/DisplayClusterShaders/Public/ShaderParameters/DisplayClusterShaderParameters_ICVFX.h b/Engine/Plugins/Runtime/nDisplay/Source/DisplayClusterShaders/Public/ShaderParameters/DisplayClusterShaderParameters_ICVFX.h
--- a/Engine/Plugins/Runtime/nDisplay/Source/DisplayClusterShaders/Public/ShaderParameters/DisplayClusterShaderParameters_ICVFX.h
+++ b/Engine/Plugins/Runtime/nDisplay/Source/DisplayClusterShaders/Public/ShaderParameters/DisplayClusterShaderParameters_ICVFX.h
@@ -73,11 +73,6 @@ public:
 	}
 
 public:
-	bool IsAnyLightcardUsed() const
-	{
-		return IsLightcardOverUsed() || IsLightcardUnderUsed();
-	}
-
 	bool IsLightcardOverUsed() const
 	{
 		return Lightcard.IsValid() && LightcardMode == EDisplayClusterShaderParametersICVFX_LightcardRenderMode::Over;
@@ -88,6 +83,16 @@ public:
 		return Lightcard.IsValid() && LightcardMode == EDisplayClusterShaderParametersICVFX_LightcardRenderMode::Under;
 	}
 
+	bool IsUVLightcardOverUsed() const
+	{
+		return UVLightcard.IsValid() && LightcardMode == EDisplayClusterShaderParametersICVFX_LightcardRenderMode::Over;
+	}
+
+	bool IsUVLightcardUnderUsed() const
+	{
+		return UVLightcard.IsValid() && LightcardMode == EDisplayClusterShaderParametersICVFX_LightcardRenderMode::Under;
+	}
+
 	bool IsCameraUsed(int32 CameraIndex) const
 	{
 		return (Cameras.Num() > CameraIndex) && CameraIndex >= 0 && Cameras[CameraIndex].IsUsed();
@@ -105,16 +110,16 @@ public:
 
 	bool IsValid()
 	{
-		return IsAnyLightcardUsed() || IsAnyCameraUsed();
+		return Lightcard.IsValid() || UVLightcard.IsValid() || IsAnyCameraUsed();
 	}
 
 public:
 	void Reset()
 	{
 		Cameras.Empty();
-		Lightcard.Reset();
 
-		UVLightCardMap = nullptr;
+		UVLightcard.Reset();
+		Lightcard.Reset();
 	}
 
 	// Implement copy ref and arrays
@@ -124,10 +129,9 @@ public:
 
 		Cameras = InParameters.Cameras;
 
-		Lightcard      = InParameters.Lightcard;
+		UVLightcard   = InParameters.UVLightcard;
+		Lightcard     = InParameters.Lightcard;
 		LightcardMode = InParameters.LightcardMode;
-
-		UVLightCardMap = InParameters.UVLightCardMap;
 	}
 
 	void CollectRefViewports(TArray<FDisplayClusterShaderParametersICVFX_ViewportResource*>& Dst)
@@ -137,6 +141,11 @@ public:
 			Dst.Add(&Lightcard);
 		}
 
+		if (UVLightcard.IsDefined())
+		{
+			Dst.Add(&UVLightcard);
+		}
+
 		for (FCameraSettings& CameraIt : Cameras)
 		{
 			CameraIt.CollectRefViewports(Dst);
@@ -241,9 +250,7 @@ public:
 	TArray<FCameraSettings> Cameras;
 
 	// Lightcard settings
+	FDisplayClusterShaderParametersICVFX_ViewportResource    UVLightcard;
 	FDisplayClusterShaderParametersICVFX_ViewportResource    Lightcard;
 	EDisplayClusterShaderParametersICVFX_LightcardRenderMode LightcardMode = EDisplayClusterShaderParametersICVFX_LightcardRenderMode::Under;
-
-	/** Texture containing a UV map of the rendered UV light cards */
-	FTextureRHIRef UVLightCardMap;
 };
diff --git a/Engine/Plugins/Runtime/nDisplay/Source/DisplayClusterTests/Private/DisplayCluster.spec.cpp b/Engine/Plugins/Runtime/nDisplay/Source/DisplayClusterTests/Private/DisplayCluster.spec.cpp
--- a/Engine/Plugins/Runtime/nDisplay/Source/DisplayClusterTests/Private/DisplayCluster.spec.cpp
+++ b/Engine/Plugins/Runtime/nDisplay/Source/DisplayClusterTests/Private/DisplayCluster.spec.cpp
@@ -447,7 +447,12 @@ void FDisplayClusterSpec::RegisterPropertyPropagationTests(bool bDoInstanceOverr
 		TArray({ GET_MEMBER_NAME_CHECKED(UDisplayClusterConfigurationData, StageSettings), GET_MEMBER_NAME_CHECKED(FDisplayClusterConfigurationICVFX_StageSettings, bEnableInnerFrustums) }));
 	
 	PROPAGATION_TEST_SIMPLE_NO_INSTANCE_OVERRIDE(ClusterConfig, false,
-		TArray({ GET_MEMBER_NAME_CHECKED(UDisplayClusterConfigurationData, StageSettings), GET_MEMBER_NAME_CHECKED(FDisplayClusterConfigurationICVFX_StageSettings, bUseOverallClusterOCIOConfiguration) }));
+		TArray({
+			GET_MEMBER_NAME_CHECKED(UDisplayClusterConfigurationData, StageSettings),
+			GET_MEMBER_NAME_CHECKED(FDisplayClusterConfigurationICVFX_StageSettings, ViewportOCIO),
+			GET_MEMBER_NAME_CHECKED(FDisplayClusterConfigurationICVFX_ViewportOCIO, AllViewportsOCIOConfiguration),
+			GET_MEMBER_NAME_CHECKED(FDisplayClusterConfigurationOCIOConfiguration, bIsEnabled)
+			}));
 
 	// StageSettings.HideList
 	PROPAGATION_TEST_SIMPLE(ClusterConfig, FName("TestName"), FName("OverrideName"),
@@ -476,32 +481,104 @@ void FDisplayClusterSpec::RegisterPropertyPropagationTests(bool bDoInstanceOverr
 		TArray({ GET_MEMBER_NAME_CHECKED(UDisplayClusterConfigurationViewport, RenderSettings), GET_MEMBER_NAME_CHECKED(FDisplayClusterConfigurationViewport_RenderSettings, Overscan),
 			GET_MEMBER_NAME_CHECKED(FDisplayClusterConfigurationViewport_Overscan, Right) }));
 
-	// StageSettings.AllViewportsOCIOConfiguration
+	// StageSettings.ViewportOCIO.AllViewportsOCIOConfiguration
 	Describe("AllViewportsOCIOConfiguration", [this, BeforeTestAction, bDoInstanceOverrideTest]()
 	{
 		PROPAGATION_TEST_OCIO_CONFIG(ClusterConfig,
-			TArray({ GET_MEMBER_NAME_CHECKED(UDisplayClusterConfigurationData, StageSettings), GET_MEMBER_NAME_CHECKED(FDisplayClusterConfigurationICVFX_StageSettings, AllViewportsOCIOConfiguration) }));
+			TArray({
+				GET_MEMBER_NAME_CHECKED(UDisplayClusterConfigurationData, StageSettings),
+				GET_MEMBER_NAME_CHECKED(FDisplayClusterConfigurationICVFX_StageSettings, ViewportOCIO),
+				GET_MEMBER_NAME_CHECKED(FDisplayClusterConfigurationICVFX_ViewportOCIO, AllViewportsOCIOConfiguration)
+				}));
 	});
 
-	// StageSettings.PerViewportOCIOProfiles
+	// StageSettings.ViewportOCIO.PerViewportOCIOProfiles
 	Describe("PerViewportOCIOProfiles", [this, BeforeTestAction, bDoInstanceOverrideTest]()
 	{
 		BeforeEach([this]()
 		{
 			const bool bSuccessfullyEnabled = DisplayClusterTestUtils::SetBlueprintPropertyValue(AssetClusterConfig, ClusterAsset,
-				{ GET_MEMBER_NAME_CHECKED(UDisplayClusterConfigurationData, StageSettings), GET_MEMBER_NAME_CHECKED(FDisplayClusterConfigurationICVFX_StageSettings, PerViewportOCIOProfiles),
-					GET_MEMBER_NAME_CHECKED(FDisplayClusterConfigurationOCIOProfile, bIsEnabled) },
+				{
+					GET_MEMBER_NAME_CHECKED(UDisplayClusterConfigurationData, StageSettings),
+					GET_MEMBER_NAME_CHECKED(FDisplayClusterConfigurationICVFX_StageSettings, ViewportOCIO),
+					GET_MEMBER_NAME_CHECKED(FDisplayClusterConfigurationICVFX_ViewportOCIO, PerViewportOCIOProfiles),
+					GET_MEMBER_NAME_CHECKED(FDisplayClusterConfigurationOCIOProfile, bIsEnabled)
+				},
 				true);
+
 			TestTrue("Successfully enabled OCIO profile", bSuccessfullyEnabled);
 		});
 		
 		PROPAGATION_TEST_OCIO_CONFIG(ClusterConfig,
-			TArray({ GET_MEMBER_NAME_CHECKED(UDisplayClusterConfigurationData, StageSettings), GET_MEMBER_NAME_CHECKED(FDisplayClusterConfigurationICVFX_StageSettings, PerViewportOCIOProfiles) }));
+			TArray({
+				GET_MEMBER_NAME_CHECKED(UDisplayClusterConfigurationData, StageSettings),
+				GET_MEMBER_NAME_CHECKED(FDisplayClusterConfigurationICVFX_StageSettings, ViewportOCIO),
+				GET_MEMBER_NAME_CHECKED(FDisplayClusterConfigurationICVFX_ViewportOCIO, PerViewportOCIOProfiles)
+				}));
 		
 		PROPAGATION_TEST_SIMPLE(ClusterConfig, FString("TestEntry"), FString("OverrideEntry"),
-			TArray({ GET_MEMBER_NAME_CHECKED(UDisplayClusterConfigurationData, StageSettings), GET_MEMBER_NAME_CHECKED(FDisplayClusterConfigurationICVFX_StageSettings, PerViewportOCIOProfiles),
-				GET_MEMBER_NAME_CHECKED(FDisplayClusterConfigurationOCIOProfile, ApplyOCIOToObjects) }));
+			TArray({
+				GET_MEMBER_NAME_CHECKED(UDisplayClusterConfigurationData, StageSettings),
+				GET_MEMBER_NAME_CHECKED(FDisplayClusterConfigurationICVFX_StageSettings, ViewportOCIO),
+				GET_MEMBER_NAME_CHECKED(FDisplayClusterConfigurationICVFX_ViewportOCIO, PerViewportOCIOProfiles),
+				GET_MEMBER_NAME_CHECKED(FDisplayClusterConfigurationOCIOProfile, ApplyOCIOToObjects)
+				}));
 	});
+
+	// Light Cards OCIO
+	{
+		// StageSettings.Lightcard.LightcardOCIO.CustomOCIO.AllViewportsOCIOConfiguration
+		Describe("Lightcard.AllViewportsOCIOConfiguration", [this, BeforeTestAction, bDoInstanceOverrideTest]()
+			{
+				PROPAGATION_TEST_OCIO_CONFIG(ClusterConfig,
+					TArray({
+						GET_MEMBER_NAME_CHECKED(UDisplayClusterConfigurationData, StageSettings),
+						GET_MEMBER_NAME_CHECKED(FDisplayClusterConfigurationICVFX_StageSettings, Lightcard),
+						GET_MEMBER_NAME_CHECKED(FDisplayClusterConfigurationICVFX_LightcardSettings, LightcardOCIO),
+						GET_MEMBER_NAME_CHECKED(FDisplayClusterConfigurationICVFX_LightcardOCIO, CustomOCIO),
+						GET_MEMBER_NAME_CHECKED(FDisplayClusterConfigurationICVFX_LightcardCustomOCIO, AllViewportsOCIOConfiguration)
+					}));
+			});
+
+		// StageSettings.Lightcard.LightcardOCIO.CustomOCIO.PerViewportOCIOProfiles
+		Describe("Lightcard.PerViewportOCIOProfiles", [this, BeforeTestAction, bDoInstanceOverrideTest]()
+			{
+				BeforeEach([this]()
+					{
+						const bool bSuccessfullyEnabled = DisplayClusterTestUtils::SetBlueprintPropertyValue(AssetClusterConfig, ClusterAsset,
+							{
+								GET_MEMBER_NAME_CHECKED(UDisplayClusterConfigurationData, StageSettings),
+								GET_MEMBER_NAME_CHECKED(FDisplayClusterConfigurationICVFX_StageSettings, Lightcard),
+								GET_MEMBER_NAME_CHECKED(FDisplayClusterConfigurationICVFX_LightcardSettings, LightcardOCIO),
+								GET_MEMBER_NAME_CHECKED(FDisplayClusterConfigurationICVFX_LightcardOCIO, CustomOCIO),
+								GET_MEMBER_NAME_CHECKED(FDisplayClusterConfigurationICVFX_LightcardCustomOCIO, PerViewportOCIOProfiles),
+								GET_MEMBER_NAME_CHECKED(FDisplayClusterConfigurationOCIOProfile, bIsEnabled)
+							},
+							true);
+
+						TestTrue("Successfully enabled OCIO profile", bSuccessfullyEnabled);
+					});
+
+		PROPAGATION_TEST_OCIO_CONFIG(ClusterConfig,
+			TArray({
+				GET_MEMBER_NAME_CHECKED(UDisplayClusterConfigurationData, StageSettings),
+				GET_MEMBER_NAME_CHECKED(FDisplayClusterConfigurationICVFX_StageSettings, Lightcard),
+				GET_MEMBER_NAME_CHECKED(FDisplayClusterConfigurationICVFX_LightcardSettings, LightcardOCIO),
+				GET_MEMBER_NAME_CHECKED(FDisplayClusterConfigurationICVFX_LightcardOCIO, CustomOCIO),
+				GET_MEMBER_NAME_CHECKED(FDisplayClusterConfigurationICVFX_LightcardCustomOCIO, PerViewportOCIOProfiles)
+				}));
+
+		PROPAGATION_TEST_SIMPLE(ClusterConfig, FString("TestEntry"), FString("OverrideEntry"),
+			TArray({
+				GET_MEMBER_NAME_CHECKED(UDisplayClusterConfigurationData, StageSettings),
+				GET_MEMBER_NAME_CHECKED(FDisplayClusterConfigurationICVFX_StageSettings, Lightcard),
+				GET_MEMBER_NAME_CHECKED(FDisplayClusterConfigurationICVFX_LightcardSettings, LightcardOCIO),
+				GET_MEMBER_NAME_CHECKED(FDisplayClusterConfigurationICVFX_LightcardOCIO, CustomOCIO),
+				GET_MEMBER_NAME_CHECKED(FDisplayClusterConfigurationICVFX_LightcardCustomOCIO, PerViewportOCIOProfiles),
+				GET_MEMBER_NAME_CHECKED(FDisplayClusterConfigurationOCIOProfile, ApplyOCIOToObjects)
+				}));
+			});
+	}
 	
 	// StageSettings.EntireClusterColorGrading
 	Describe("EntireClusterColorGrading", [this, BeforeTestAction, bDoInstanceOverrideTest]()
