nDisplay[SharedMemoryMedia]: Explicitly avoids unregistering a unique name that the given player instance had not registered by itself.

diff --git a/Engine/Plugins/Runtime/nDisplay/Source/DisplayClusterMedia/Private/Input/SharedMemoryMediaPlayer.cpp b/Engine/Plugins/Runtime/nDisplay/Source/DisplayClusterMedia/Private/Input/SharedMemoryMediaPlayer.cpp
--- a/Engine/Plugins/Runtime/nDisplay/Source/DisplayClusterMedia/Private/Input/SharedMemoryMediaPlayer.cpp
+++ b/Engine/Plugins/Runtime/nDisplay/Source/DisplayClusterMedia/Private/Input/SharedMemoryMediaPlayer.cpp
@@ -36,7 +36,7 @@ static FAutoConsoleVariableRef CVarDisplayClusterSharedMemoryLatency(
 );
 
 
-TSet<FString>FSharedMemoryMediaPlayer::UniqueNamesRegistered;
+TSet<FString> FSharedMemoryMediaPlayer::UniqueNamesRegistered;
 
 FSharedMemoryMediaPlayer::FSharedMemoryMediaPlayer()
 	: Super()
@@ -52,6 +52,12 @@ FSharedMemoryMediaPlayer::~FSharedMemoryMediaPlayer()
 
 bool FSharedMemoryMediaPlayer::RegisterUniqueName(FString& InUniqueName)
 {
+	// We don't allow empty unique names
+	if (!InUniqueName.Len())
+	{
+		return false;
+	}
+
 	bool bIsAlreadyInSet = false;
 
 	UniqueNamesRegistered.Add(InUniqueName, &bIsAlreadyInSet);
@@ -120,6 +126,7 @@ bool FSharedMemoryMediaPlayer::Open(const FString& Url, const IMediaOptions* Opt
 	if (!RegisterUniqueName(UniqueName))
 	{
 		UE_LOG(LogDisplayClusterMedia, Error, TEXT("Only one SharedMemoryMediaPlayer the UniqueName '%s' can play at a time"), *UniqueName);
+		UniqueName.Reset(); // Reset the unique name to avoid unregistering when not open
 		return false;
 	}
 
