nDisplay: removed UE4 mentionings

diff --git a/Engine/Plugins/Runtime/nDisplay/Source/DisplayCluster/Private/Misc/DisplayClusterAppExit.cpp b/Engine/Plugins/Runtime/nDisplay/Source/DisplayCluster/Private/Misc/DisplayClusterAppExit.cpp
--- a/Engine/Plugins/Runtime/nDisplay/Source/DisplayCluster/Private/Misc/DisplayClusterAppExit.cpp
+++ b/Engine/Plugins/Runtime/nDisplay/Source/DisplayCluster/Private/Misc/DisplayClusterAppExit.cpp
@@ -35,7 +35,7 @@ void FDisplayClusterAppExit::ExitApplication(const FString& Msg, EExitType ExitT
 				}
 				else
 				{
-					// For some reason UE4 generates crash info if FPlatformMisc::RequestExit gets called
+					// For some reason UE generates crash info if FPlatformMisc::RequestExit gets called
 					// from a thread other than GameThread. Since it may be called from the networking
 					// session threads (failover pipeline), we don't want to generate unnecessary crash reports.
 					AsyncTask(ENamedThreads::GameThread, []()
diff --git a/Engine/Plugins/Runtime/nDisplay/Source/DisplayCluster/Private/Render/Viewport/DisplayClusterViewport_Math.cpp b/Engine/Plugins/Runtime/nDisplay/Source/DisplayCluster/Private/Render/Viewport/DisplayClusterViewport_Math.cpp
--- a/Engine/Plugins/Runtime/nDisplay/Source/DisplayCluster/Private/Render/Viewport/DisplayClusterViewport_Math.cpp
+++ b/Engine/Plugins/Runtime/nDisplay/Source/DisplayCluster/Private/Render/Viewport/DisplayClusterViewport_Math.cpp
@@ -181,7 +181,7 @@ FMatrix IDisplayClusterViewport::MakeProjectionMatrix(float l, float r, float t,
 		FPlane(ma, mb, mc, me),
 		FPlane(0, 0, md, 0));
 
-	// Invert Z-axis (UE4 uses Z-inverted LHS)
+	// Invert Z-axis (UE uses Z-inverted LHS)
 	static const FMatrix flipZ = FMatrix(
 		FPlane(1, 0, 0, 0),
 		FPlane(0, 1, 0, 0),
