nDisplay: Added flag to prevent the preview output on a level instance stage actor if its Enable Preview flag is set to false even if the ICVFX panel is open for that stage actor

diff --git a/Engine/Plugins/Runtime/nDisplay/Source/DisplayCluster/Private/Game/EngineClasses/Scene/DisplayClusterPreviewComponent.cpp b/Engine/Plugins/Runtime/nDisplay/Source/DisplayCluster/Private/Game/EngineClasses/Scene/DisplayClusterPreviewComponent.cpp
--- a/Engine/Plugins/Runtime/nDisplay/Source/DisplayCluster/Private/Game/EngineClasses/Scene/DisplayClusterPreviewComponent.cpp
+++ b/Engine/Plugins/Runtime/nDisplay/Source/DisplayCluster/Private/Game/EngineClasses/Scene/DisplayClusterPreviewComponent.cpp
@@ -102,7 +102,7 @@ bool UDisplayClusterPreviewComponent::InitializePreviewComponent(ADisplayCluster
 
 bool UDisplayClusterPreviewComponent::IsPreviewEnabled() const
 {
-	return (ViewportConfig && RootActor && RootActor->IsPreviewEnabled());
+	return (ViewportConfig && RootActor && RootActor->IsPreviewDrawnToScreens());
 }
 
 void UDisplayClusterPreviewComponent::RestorePreviewMeshMaterial()
@@ -163,7 +163,7 @@ bool UDisplayClusterPreviewComponent::UpdatePreviewMesh()
 
 	UpdatePreviewMeshReference();
 
-	if (IsPreviewEnabled())
+	if (IsPreviewEnabled() || OverrideTexture)
 	{
 		check(ViewportConfig);
 
diff --git a/Engine/Plugins/Runtime/nDisplay/Source/DisplayCluster/Private/Game/EngineClasses/Scene/DisplayClusterRootActorEditor.cpp b/Engine/Plugins/Runtime/nDisplay/Source/DisplayCluster/Private/Game/EngineClasses/Scene/DisplayClusterRootActorEditor.cpp
--- a/Engine/Plugins/Runtime/nDisplay/Source/DisplayCluster/Private/Game/EngineClasses/Scene/DisplayClusterRootActorEditor.cpp
+++ b/Engine/Plugins/Runtime/nDisplay/Source/DisplayCluster/Private/Game/EngineClasses/Scene/DisplayClusterRootActorEditor.cpp
@@ -219,6 +219,14 @@ bool ADisplayClusterRootActor::IsPreviewEnabled() const
 	return bPreviewEnable || !PreviewEnableOverriders.IsEmpty();
 }
 
+bool ADisplayClusterRootActor::IsPreviewDrawnToScreens() const
+{
+	// The preview output is drawn to the screen when preview rendering is enabled (either manually or through an override) AND when
+	// the bPreviewEnable flag is specifically set to true. 
+	return IsPreviewEnabled() && bPreviewEnable;
+}
+
+
 // Return all RTT RHI resources for preview
 void ADisplayClusterRootActor::GetPreviewRenderTargetableTextures(const TArray<FString>& InViewportNames, TArray<FTextureRHIRef>& OutTextures)
 {
diff --git a/Engine/Plugins/Runtime/nDisplay/Source/DisplayCluster/Public/DisplayClusterRootActor.h b/Engine/Plugins/Runtime/nDisplay/Source/DisplayCluster/Public/DisplayClusterRootActor.h
--- a/Engine/Plugins/Runtime/nDisplay/Source/DisplayCluster/Public/DisplayClusterRootActor.h
+++ b/Engine/Plugins/Runtime/nDisplay/Source/DisplayCluster/Public/DisplayClusterRootActor.h
@@ -431,6 +431,9 @@ public:
 	// return true, if preview enabled for this actor
 	bool IsPreviewEnabled() const;
 
+	/** Gets whether the preview output is displayed onto the stage actor's screen meshes */
+	bool IsPreviewDrawnToScreens() const;
+
 	void Constructor_Editor();
 	void Destructor_Editor();
 
