nDisplay: Labels are now rendered separately from the preview in-editor. This fixes having to adjust translucency values of the light card when toggling labels.

diff --git a/Engine/Plugins/Runtime/nDisplay/Source/DisplayCluster/Private/Game/EngineClasses/Scene/DisplayClusterLabelComponent.cpp b/Engine/Plugins/Runtime/nDisplay/Source/DisplayCluster/Private/Game/EngineClasses/Scene/DisplayClusterLabelComponent.cpp
--- a/Engine/Plugins/Runtime/nDisplay/Source/DisplayCluster/Private/Game/EngineClasses/Scene/DisplayClusterLabelComponent.cpp
+++ b/Engine/Plugins/Runtime/nDisplay/Source/DisplayCluster/Private/Game/EngineClasses/Scene/DisplayClusterLabelComponent.cpp
@@ -21,7 +21,7 @@ UDisplayClusterLabelComponent::UDisplayClusterLabelComponent()
 	// Keeps this visible even if its clipping through the owner.
 	// Requires the owner have a transparent blend mode on its material.
 	WidgetComponent->SetBlendMode(EWidgetBlendMode::Transparent);
-	WidgetComponent->SetTranslucentSortPriority(1.f);
+	WidgetComponent->SetTranslucencySortDistanceOffset(-1000.f);
 	WidgetComponent->SetUsingAbsoluteScale(true);
 	WidgetComponent->SetCollisionEnabled(ECollisionEnabled::NoCollision);
 	WidgetComponent->SetCollisionResponseToAllChannels(ECollisionResponse::ECR_Ignore);
@@ -41,6 +41,12 @@ void UDisplayClusterLabelComponent::SetRootActor(ADisplayClusterRootActor* InAct
 void UDisplayClusterLabelComponent::SetWidgetScale(float NewValue)
 {
 	WidgetScale = NewValue;
+	WidgetComponent->SetWidgetScale(WidgetScale);
+}
+
+float UDisplayClusterLabelComponent::GetWidgetScale() const
+{
+	return WidgetScale;
 }
 
 void UDisplayClusterLabelComponent::TickComponent(float DeltaTime, ELevelTick TickType,
@@ -66,7 +72,7 @@ void UDisplayClusterLabelComponent::OnRegister()
 	WidgetComponent->SetWidgetClass(WidgetClass.LoadSynchronous());
 }
 
-UDisplayClusterWidgetComponent* UDisplayClusterLabelComponent::GetWidgetComponent() const
+UWidgetComponent* UDisplayClusterLabelComponent::GetWidgetComponent() const
 {
 	return WidgetComponent.Get();
 }
@@ -102,27 +108,36 @@ void UDisplayClusterLabelComponent::UpdateWidgetComponent()
 			// be set from the DCRA instead
 			SetRootActor(LightCardActor->GetRootActorOwner().Get());
 		}
-		
-		const USceneComponent* Viewpoint = RootActor.IsValid() ? RootActor->GetCommonViewPoint() : nullptr;
-		if (Viewpoint != nullptr)
-		{
-			const FTransform LightCardTransform = LightCardActor->GetStageActorTransform();
-			const FVector LightCardLocation = LightCardTransform.GetLocation();
-			const FVector DestinationPoint = Viewpoint->GetComponentLocation();
 
-			// Keep the text facing up.
-			FRotator Rotation = UKismetMathLibrary::FindLookAtRotation(LightCardLocation, DestinationPoint);
-			Rotation.Roll = 0.f;
+#if WITH_EDITOR
+		const bool bInGame = GEditor == nullptr;
+#else
+		constexpr bool bInGame = true;
+#endif
+
+		if (bInGame)
+		{
+			const USceneComponent* Viewpoint = RootActor.IsValid() ? RootActor->GetCommonViewPoint() : nullptr;
+			if (Viewpoint != nullptr)
+			{
+				const FTransform LightCardTransform = LightCardActor->GetStageActorTransform();
+				const FVector LightCardLocation = LightCardTransform.GetLocation();
+				const FVector DestinationPoint = Viewpoint->GetComponentLocation();
+
+				// Keep the text facing up.
+				FRotator Rotation = UKismetMathLibrary::FindLookAtRotation(LightCardLocation, DestinationPoint);
+				Rotation.Roll = 0.f;
 			
-			WidgetComponent->SetWorldRotation(Rotation);
+				WidgetComponent->SetWorldRotation(Rotation);
 
-			// Keep the label the same size regardless of distance to the view origin.
-			const float Distance = (DestinationPoint - LightCardLocation).Length();
+				// Keep the label the same size regardless of distance to the view origin.
+				const float Distance = (DestinationPoint - LightCardLocation).Length();
 
-			const float BaseLabelScale = 0.0025f;
-			const FVector VectorScale(Distance * BaseLabelScale * WidgetScale);
+				const float BaseLabelScale = 0.0025f;
+				const FVector VectorScale(Distance * BaseLabelScale * WidgetScale);
 			
-			WidgetComponent->SetWorldScale3D(VectorScale);
+				WidgetComponent->SetWorldScale3D(VectorScale);
+			}
 		}
 	}
 	else
@@ -143,7 +158,7 @@ void UDisplayClusterLabelComponent::UpdateWidgetComponent()
 	}
 
 #if WITH_EDITOR
-	WidgetComponent->SetVisibleInSceneCaptureOnly(GEditor != nullptr);
+	WidgetComponent->SetHiddenInSceneCapture(true);
 #endif
 }
 
diff --git a/Engine/Plugins/Runtime/nDisplay/Source/DisplayCluster/Private/Game/EngineClasses/Scene/DisplayClusterLightCardActor.cpp b/Engine/Plugins/Runtime/nDisplay/Source/DisplayCluster/Private/Game/EngineClasses/Scene/DisplayClusterLightCardActor.cpp
--- a/Engine/Plugins/Runtime/nDisplay/Source/DisplayCluster/Private/Game/EngineClasses/Scene/DisplayClusterLightCardActor.cpp
+++ b/Engine/Plugins/Runtime/nDisplay/Source/DisplayCluster/Private/Game/EngineClasses/Scene/DisplayClusterLightCardActor.cpp
@@ -6,10 +6,13 @@
 #include "DisplayClusterRootActor.h"
 #include "Features/IModularFeatures.h"
 #include "Misc/DisplayClusterLog.h"
+#include "DisplayClusterWidgetComponent.h"
+
 #include "IDisplayClusterLightCardActorExtender.h"
 
 #include "Components/DisplayClusterLabelComponent.h"
 #include "Components/StaticMeshComponent.h"
+#include "Features/IModularFeatures.h"
 #include "GameFramework/SpringArmComponent.h"
 #include "Materials/MaterialInstanceDynamic.h"
 #include "Engine/StaticMesh.h"
@@ -152,6 +155,10 @@ void ADisplayClusterLightCardActor::Tick(float DeltaSeconds)
 	UpdateLightCardMaterialInstance();
 	UpdateLightCardVisibility();
 	UpdateUVIndicator();
+	if (LabelComponent->IsVisible())
+	{
+		LabelComponent->GetWidgetComponent()->SetTranslucentSortPriority(LightCardComponent->TranslucencySortPriority);
+	}
 }
 
 #if WITH_EDITOR
@@ -484,25 +491,6 @@ void ADisplayClusterLightCardActor::ShowLightCardLabel(bool bValue, float ScaleV
 	LabelComponent->SetVisibility(bValue, true);
 	LabelComponent->SetRootActor(InRootActor);
 	LabelComponent->SetWidgetScale(ScaleValue);
-	
-	const float LightCardSortPriorityLow = -1.f;
-	
-	if (bValue)
-	{
-		// Change the translucent sort priority of the light card component so it will always be shown behind the
-		// label component. When left at 0 the renderer will choose the sort order which may be wrong.
-		// This does rely on the light card material being translucent as well.
-		SavedTranslucencySortPriority = LightCardComponent->TranslucencySortPriority;
-		LightCardComponent->SetTranslucentSortPriority(LightCardSortPriorityLow);
-	}
-	else if (SavedTranslucencySortPriority.IsSet() && *SavedTranslucencySortPriority != LightCardComponent->TranslucencySortPriority &&
-			LightCardComponent->TranslucencySortPriority == LightCardSortPriorityLow)
-	{
-		// Change the sort order back to the original value and only if the user hasn't modified it since activating the label.
-		// It's highly unlikely this will ever be anything besides 0.
-		LightCardComponent->SetTranslucentSortPriority(*SavedTranslucencySortPriority);
-		SavedTranslucencySortPriority.Reset();
-	}
 }
 
 void ADisplayClusterLightCardActor::SetRootActorOwner(ADisplayClusterRootActor* InRootActor)
diff --git a/Engine/Plugins/Runtime/nDisplay/Source/DisplayCluster/Private/Game/EngineClasses/Scene/DisplayClusterWidgetComponent.cpp b/Engine/Plugins/Runtime/nDisplay/Source/DisplayCluster/Private/Game/EngineClasses/Scene/DisplayClusterWidgetComponent.cpp
--- a/Engine/Plugins/Runtime/nDisplay/Source/DisplayCluster/Private/Game/EngineClasses/Scene/DisplayClusterWidgetComponent.cpp
+++ b/Engine/Plugins/Runtime/nDisplay/Source/DisplayCluster/Private/Game/EngineClasses/Scene/DisplayClusterWidgetComponent.cpp
@@ -2,12 +2,241 @@
 
 #include "DisplayClusterWidgetComponent.h"
 
+#include "DynamicMeshBuilder.h"
 #include "Engine/Engine.h"
+#include "Engine/TextureRenderTarget2D.h"
+#include "Materials/MaterialInstanceDynamic.h"
 #include "Slate/WidgetRenderer.h"
 #if !UE_SERVER
 #include "RHI.h"
 #endif
 
+#if WITH_EDITOR
+#include "Editor.h"
+#include "LocalVertexFactory.h"
+#include "PrimitiveSceneProxy.h"
+#include "SceneInterface.h"
+#include "StaticMeshResources.h"
+
+/** Represents a billboard widget which displays a widget render target & material. Based on FMaterialSpriteSceneProxy. */
+class FEditorWidgetBillboardProxy final : public FPrimitiveSceneProxy
+{
+public:
+	SIZE_T GetTypeHash() const override
+	{
+		static size_t UniquePointer;
+		return reinterpret_cast<size_t>(&UniquePointer);
+	}
+
+	/** Initialization constructor. */
+	FEditorWidgetBillboardProxy(const UDisplayClusterWidgetComponent* InComponent)
+		: FPrimitiveSceneProxy(InComponent)
+		  , VertexFactory(GetScene().GetFeatureLevel(), "FEditorWidgetBillboardProxy")
+		  , Component(InComponent)
+	{
+		bWillEverBeLit = false;
+
+		FColor NewPropertyColor;
+		GEngine->GetPropertyColorationColor((UObject*)InComponent, NewPropertyColor);
+		SetPropertyColor(NewPropertyColor);
+
+		StaticMeshVertexBuffers.PositionVertexBuffer.Init(1);
+		StaticMeshVertexBuffers.StaticMeshVertexBuffer.Init(1, 1);
+		StaticMeshVertexBuffers.ColorVertexBuffer.Init(1);
+
+		ENQUEUE_RENDER_COMMAND(FWidgetBillboardProxyInit)(
+			[this](FRHICommandListImmediate& RHICmdList)
+		{
+			this->StaticMeshVertexBuffers.PositionVertexBuffer.InitResource();
+			this->StaticMeshVertexBuffers.StaticMeshVertexBuffer.InitResource();
+			this->StaticMeshVertexBuffers.ColorVertexBuffer.InitResource();
+
+			FLocalVertexFactory::FDataType Data;
+			this->StaticMeshVertexBuffers.PositionVertexBuffer.BindPositionVertexBuffer(&this->VertexFactory, Data);
+			this->StaticMeshVertexBuffers.StaticMeshVertexBuffer.BindTangentVertexBuffer(&this->VertexFactory, Data);
+			this->StaticMeshVertexBuffers.StaticMeshVertexBuffer.BindPackedTexCoordVertexBuffer(&this->VertexFactory, Data);
+			this->StaticMeshVertexBuffers.StaticMeshVertexBuffer.BindLightMapVertexBuffer(&this->VertexFactory, Data, 0);
+			this->StaticMeshVertexBuffers.ColorVertexBuffer.BindColorVertexBuffer(&this->VertexFactory, Data);
+			this->VertexFactory.SetData(Data);
+
+			this->VertexFactory.InitResource();
+		});
+	}
+
+	virtual ~FEditorWidgetBillboardProxy() override
+	{
+		VertexFactory.ReleaseResource();
+		StaticMeshVertexBuffers.PositionVertexBuffer.ReleaseResource();
+		StaticMeshVertexBuffers.StaticMeshVertexBuffer.ReleaseResource();
+		StaticMeshVertexBuffers.ColorVertexBuffer.ReleaseResource();
+	}
+
+	// FPrimitiveSceneProxy interface.
+	
+	virtual void GetDynamicMeshElements(const TArray<const FSceneView*>& Views, const FSceneViewFamily& ViewFamily, uint32 VisibilityMap, FMeshElementCollector& Collector) const override
+	{
+		QUICK_SCOPE_CYCLE_COUNTER(STAT_DisplayClusterWidgetComponent_GetDynamicMeshElements);
+		
+		const int32 ElementStride = 4;
+		const int32 WorstCaseVertexBufferSize = ElementStride * Views.Num();
+		
+		if (!Component.IsValid() || WorstCaseVertexBufferSize <= 0)
+		{
+			return;
+		}
+		
+		StaticMeshVertexBuffers.PositionVertexBuffer.Init(WorstCaseVertexBufferSize);
+		StaticMeshVertexBuffers.StaticMeshVertexBuffer.Init(WorstCaseVertexBufferSize, 1);
+		StaticMeshVertexBuffers.ColorVertexBuffer.Init(WorstCaseVertexBufferSize);
+		
+		if (const UTextureRenderTarget2D* RenderTarget = Component->GetRenderTarget())
+		{
+			if (const UMaterialInstanceDynamic* MaterialInstance = Component->GetMaterialInstance())
+			{
+				if (const FMaterialRenderProxy* ParentMaterialProxy = MaterialInstance->GetRenderProxy())
+				{
+					for (int32 ViewIndex = 0; ViewIndex < Views.Num(); ++ViewIndex)
+					{
+						if (VisibilityMap & (1 << ViewIndex))
+						{
+							// Render the material of the render target as a billboard always facing the view
+						
+							const FSceneView* View = Views[ViewIndex];
+
+							const FMatrix WorldToLocal = GetLocalToWorld().InverseFast();
+							const FMatrix ViewToLocal = View->ViewMatrices.GetInvViewMatrix() * WorldToLocal;
+							const FVector TangentX = -ViewToLocal.TransformVector(FVector(1.0f, 0.0f, 0.0f));
+							const FVector TangentY = -ViewToLocal.TransformVector(FVector(0.0f, 1.0f, 0.0f));
+							const FVector TangentZ = -ViewToLocal.TransformVector(FVector(0.0f, 0.0f, 1.0f));
+
+							// Evaluate the size of the sprite.
+							const float SizeX = (RenderTarget->SizeX * Component->GetWidgetScale()) * 0.25f;
+							const float SizeY = (RenderTarget->SizeY * Component->GetWidgetScale()) * 0.25f;
+
+							const float WorldSizeX = SizeX;
+							const float WorldSizeY = SizeY;
+
+							FLinearColor Color = FLinearColor::White;
+				
+							const int WriteOffset = ElementStride * ViewIndex;
+							for (uint32 VertexIndex = 0; VertexIndex < 4; ++VertexIndex)
+							{
+								const int WriteIndex = WriteOffset + VertexIndex;
+								StaticMeshVertexBuffers.StaticMeshVertexBuffer.SetVertexTangents(WriteIndex,
+									static_cast<FVector3f>(-TangentX),
+									static_cast<FVector3f>(TangentY),
+									static_cast<FVector3f>(TangentZ));
+								StaticMeshVertexBuffers.ColorVertexBuffer.VertexColor(WriteIndex) = Color.ToFColor(true);
+							}
+
+							// Set up the sprite vertex positions and texture coordinates.
+							StaticMeshVertexBuffers.PositionVertexBuffer.VertexPosition(WriteOffset + 0) = FVector3f(-WorldSizeX * TangentY + +WorldSizeY * TangentX);
+							StaticMeshVertexBuffers.PositionVertexBuffer.VertexPosition(WriteOffset + 1) = FVector3f(+WorldSizeX * TangentY + +WorldSizeY * TangentX);
+							StaticMeshVertexBuffers.PositionVertexBuffer.VertexPosition(WriteOffset + 2) = FVector3f(-WorldSizeX * TangentY + -WorldSizeY * TangentX);
+							StaticMeshVertexBuffers.PositionVertexBuffer.VertexPosition(WriteOffset + 3) = FVector3f(+WorldSizeX * TangentY + -WorldSizeY * TangentX);
+
+							StaticMeshVertexBuffers.StaticMeshVertexBuffer.SetVertexUV(WriteOffset + 0, 0, FVector2f(0, 0));
+							StaticMeshVertexBuffers.StaticMeshVertexBuffer.SetVertexUV(WriteOffset + 1, 0, FVector2f(0, 1));
+							StaticMeshVertexBuffers.StaticMeshVertexBuffer.SetVertexUV(WriteOffset + 2, 0, FVector2f(1, 0));
+							StaticMeshVertexBuffers.StaticMeshVertexBuffer.SetVertexUV(WriteOffset + 3, 0, FVector2f(1, 1));
+
+							// Set up the FMeshElement.
+							FMeshBatch& Mesh = Collector.AllocateMesh();
+
+							Mesh.VertexFactory = &VertexFactory;
+							Mesh.MaterialRenderProxy = ParentMaterialProxy;
+							Mesh.LCI = nullptr;
+							Mesh.ReverseCulling = IsLocalToWorldDeterminantNegative();
+							Mesh.CastShadow = false;
+							Mesh.DepthPriorityGroup = GetDepthPriorityGroup(View);
+							Mesh.Type = PT_TriangleStrip;
+							Mesh.bDisableBackfaceCulling = true;
+
+							// Set up the FMeshBatchElement.
+							FMeshBatchElement& BatchElement = Mesh.Elements[0];
+							BatchElement.IndexBuffer = nullptr;
+							BatchElement.FirstIndex = 0;
+							BatchElement.MinVertexIndex = 0;
+							BatchElement.MaxVertexIndex = 3;
+							BatchElement.NumPrimitives = 2;
+							BatchElement.BaseVertexIndex = WriteOffset;
+							BatchElement.bPreserveInstanceOrder = false;
+
+							Mesh.bCanApplyViewModeOverrides = true;
+							Mesh.bUseWireframeSelectionColoring = IsSelected();
+
+							Collector.AddMesh(ViewIndex, Mesh);
+						}
+					}
+				}
+			}
+
+			FLocalVertexFactory* VertexFactoryPtr = &VertexFactory;
+			ENQUEUE_RENDER_COMMAND(FWidgetBillboardProxyLegacyInit)(
+				[VertexFactoryPtr, this](FRHICommandListImmediate& RHICmdList)
+			{
+				this->StaticMeshVertexBuffers.PositionVertexBuffer.UpdateRHI();
+				this->StaticMeshVertexBuffers.StaticMeshVertexBuffer.UpdateRHI();
+				this->StaticMeshVertexBuffers.ColorVertexBuffer.UpdateRHI();
+
+				FLocalVertexFactory::FDataType Data;
+				this->StaticMeshVertexBuffers.PositionVertexBuffer.BindPositionVertexBuffer(VertexFactoryPtr, Data);
+				this->StaticMeshVertexBuffers.StaticMeshVertexBuffer.BindTangentVertexBuffer(VertexFactoryPtr, Data);
+				this->StaticMeshVertexBuffers.StaticMeshVertexBuffer.BindPackedTexCoordVertexBuffer(VertexFactoryPtr, Data);
+				this->StaticMeshVertexBuffers.StaticMeshVertexBuffer.BindLightMapVertexBuffer(VertexFactoryPtr, Data, 0);
+				this->StaticMeshVertexBuffers.ColorVertexBuffer.BindColorVertexBuffer(VertexFactoryPtr, Data);
+				VertexFactoryPtr->SetData(Data);
+
+				VertexFactoryPtr->UpdateRHI();
+			});
+		}
+	}
+
+	virtual FPrimitiveViewRelevance GetViewRelevance(const FSceneView* View) const override
+	{
+		FPrimitiveViewRelevance Result;
+		if (Component.IsValid() && Component->GetMaterialInstance())
+		{
+			Component->GetMaterialInstance()->GetRelevance_Concurrent(GetScene().GetFeatureLevel()).SetPrimitiveViewRelevance(Result);
+		}
+		
+		Result.bDrawRelevance = IsShown(View) && View->Family->EngineShowFlags.WidgetComponents;
+		Result.bDynamicRelevance = true;
+		Result.bShadowRelevance = IsShadowCast(View);
+		Result.bEditorPrimitiveRelevance = UseEditorCompositing(View);
+		
+		return Result;
+	}
+
+	virtual bool CanBeOccluded() const override
+	{
+		return false;
+	}
+
+	virtual void GetLightRelevance(const FLightSceneProxy* LightSceneProxy, bool& bDynamic, bool& bRelevant, bool& bLightMapped, bool& bShadowMapped) const override
+	{
+		bDynamic = false;
+		bRelevant = false;
+		bLightMapped = false;
+		bShadowMapped = false;
+	}
+	
+	virtual uint32 GetMemoryFootprint(void) const override { return(sizeof(*this) + GetAllocatedSize()); }
+	
+	uint32 GetAllocatedSize(void) const { return( FPrimitiveSceneProxy::GetAllocatedSize() ); }
+	
+	// End of FPrimitiveSceneProxy interface.
+	
+private:
+	/** The buffer containing vertex data. */
+	mutable FStaticMeshVertexBuffers StaticMeshVertexBuffers;
+	mutable FLocalVertexFactory VertexFactory;
+	/** The owning component. */
+	TWeakObjectPtr<const UDisplayClusterWidgetComponent> Component;
+};
+
+#endif
+
 UDisplayClusterWidgetComponent::UDisplayClusterWidgetComponent()
 {
 #if WITH_EDITOR
@@ -22,6 +251,11 @@ UDisplayClusterWidgetComponent::~UDisplayClusterWidgetComponent()
 #endif
 }
 
+void UDisplayClusterWidgetComponent::SetWidgetScale(float NewValue)
+{
+	WidgetScale = NewValue;
+}
+
 void UDisplayClusterWidgetComponent::OnRegister()
 {
 	// Set this prior to registering the scene component so that bounds are calculated correctly.
@@ -76,6 +310,25 @@ void UDisplayClusterWidgetComponent::OnUnregister()
 	UPrimitiveComponent::OnUnregister();
 }
 
+FPrimitiveSceneProxy* UDisplayClusterWidgetComponent::CreateSceneProxy()
+{
+#if WITH_EDITOR
+	if (GEditor && !GetOwner()->HasAnyFlags(RF_Transient) && WidgetRenderer)
+	{
+		// Use our custom widget billboard proxy if this is the level instance in the editor,
+		// otherwise this is either handled by the default scene proxy (-game) or manually in
+		// FDisplayClusterLightCardEditorViewportClient.
+		
+		RequestRenderUpdate();
+		LastWidgetRenderTime = 0;
+	
+		return new FEditorWidgetBillboardProxy(this);
+	}
+#endif
+	
+	return Super::CreateSceneProxy();
+}
+
 #if WITH_EDITOR
 void UDisplayClusterWidgetComponent::OnWorldCleanup(UWorld* InWorld, bool bSessionEnded, bool bCleanupResources)
 {
diff --git a/Engine/Plugins/Runtime/nDisplay/Source/DisplayCluster/Private/Game/EngineClasses/Scene/DisplayClusterWidgetComponent.h b/Engine/Plugins/Runtime/nDisplay/Source/DisplayCluster/Private/Game/EngineClasses/Scene/DisplayClusterWidgetComponent.h
--- a/Engine/Plugins/Runtime/nDisplay/Source/DisplayCluster/Private/Game/EngineClasses/Scene/DisplayClusterWidgetComponent.h
+++ b/Engine/Plugins/Runtime/nDisplay/Source/DisplayCluster/Private/Game/EngineClasses/Scene/DisplayClusterWidgetComponent.h
@@ -17,11 +17,18 @@ public:
 	UDisplayClusterWidgetComponent();
 	virtual ~UDisplayClusterWidgetComponent() override;
 
+	void SetWidgetScale(float NewValue);
+	float GetWidgetScale() const { return WidgetScale; }
 private:
 	virtual void OnRegister() override;
 	virtual void OnUnregister() override;
 
+	virtual FPrimitiveSceneProxy* CreateSceneProxy() override;
+
 #if WITH_EDITOR
 	void OnWorldCleanup(UWorld* InWorld, bool bSessionEnded, bool bCleanupResources);
 #endif
+
+private:
+	float WidgetScale = 1.f;
 };
diff --git a/Engine/Plugins/Runtime/nDisplay/Source/DisplayCluster/Public/Components/DisplayClusterLabelComponent.h b/Engine/Plugins/Runtime/nDisplay/Source/DisplayCluster/Public/Components/DisplayClusterLabelComponent.h
--- a/Engine/Plugins/Runtime/nDisplay/Source/DisplayCluster/Public/Components/DisplayClusterLabelComponent.h
+++ b/Engine/Plugins/Runtime/nDisplay/Source/DisplayCluster/Public/Components/DisplayClusterLabelComponent.h
@@ -33,14 +33,17 @@ public:
 
 	/** Set the widget scale */
 	void SetWidgetScale(float NewValue);
-	
+
+	/** Return the widget scale */
+	float GetWidgetScale() const;
+
 	virtual void TickComponent(float DeltaTime, ELevelTick TickType,
 	                           FActorComponentTickFunction* ThisTickFunction) override;
 
 	virtual void OnRegister() override;
 
 	/** Returns the widget component */
-	UDisplayClusterWidgetComponent* GetWidgetComponent() const;
+	UWidgetComponent* GetWidgetComponent() const;
 
 #if WITH_EDITOR
 public:
diff --git a/Engine/Plugins/Runtime/nDisplay/Source/DisplayCluster/Public/DisplayClusterLightCardActor.h b/Engine/Plugins/Runtime/nDisplay/Source/DisplayCluster/Public/DisplayClusterLightCardActor.h
--- a/Engine/Plugins/Runtime/nDisplay/Source/DisplayCluster/Public/DisplayClusterLightCardActor.h
+++ b/Engine/Plugins/Runtime/nDisplay/Source/DisplayCluster/Public/DisplayClusterLightCardActor.h
@@ -82,6 +82,9 @@ public:
 	/** Returns the current static mesh used by this light card */
 	UStaticMesh* GetStaticMesh() const;
 
+	/** Return the label component */
+	UDisplayClusterLabelComponent* GetLabelComponent() const { return LabelComponent; }
+	
 	/** Sets a new static mesh for the light card */
 	void SetStaticMesh(UStaticMesh* InStaticMesh);
 
@@ -301,8 +304,4 @@ protected:
 	UPROPERTY()
 	bool bIsLightCardFlag = false;
 	
-private:
-	/** Stores the user translucency value when labels are displayed */
-	TOptional<int32> SavedTranslucencySortPriority;
-	
 };
diff --git a/Engine/Plugins/Runtime/nDisplay/Source/DisplayClusterLightCardEditor/DisplayClusterLightCardEditor.Build.cs b/Engine/Plugins/Runtime/nDisplay/Source/DisplayClusterLightCardEditor/DisplayClusterLightCardEditor.Build.cs
--- a/Engine/Plugins/Runtime/nDisplay/Source/DisplayClusterLightCardEditor/DisplayClusterLightCardEditor.Build.cs
+++ b/Engine/Plugins/Runtime/nDisplay/Source/DisplayClusterLightCardEditor/DisplayClusterLightCardEditor.Build.cs
@@ -34,6 +34,7 @@ public class DisplayClusterLightCardEditor : ModuleRules
 				"SceneOutliner",
 				"Slate",
 				"SlateCore",
+				"UMG",
 				"UnrealEd",
 				"ToolMenus",
 				"ToolWidgets",
diff --git a/Engine/Plugins/Runtime/nDisplay/Source/DisplayClusterLightCardEditor/Private/Viewport/DisplayClusterLightCardEditorViewportClient.cpp b/Engine/Plugins/Runtime/nDisplay/Source/DisplayClusterLightCardEditor/Private/Viewport/DisplayClusterLightCardEditorViewportClient.cpp
--- a/Engine/Plugins/Runtime/nDisplay/Source/DisplayClusterLightCardEditor/Private/Viewport/DisplayClusterLightCardEditorViewportClient.cpp
+++ b/Engine/Plugins/Runtime/nDisplay/Source/DisplayClusterLightCardEditor/Private/Viewport/DisplayClusterLightCardEditorViewportClient.cpp
@@ -30,6 +30,7 @@
 #include "Components/PostProcessComponent.h"
 #include "Components/PrimitiveComponent.h"
 #include "Components/SkyLightComponent.h"
+#include "Components/WidgetComponent.h"
 #include "Debug/DebugDrawService.h"
 #include "Editor/UnrealEdEngine.h"
 #include "EditorModeManager.h"
@@ -56,6 +57,7 @@
 #include "TextureResource.h"
 #include "UnrealEdGlobals.h"
 #include "UnrealWidget.h"
+#include "Components/DisplayClusterLabelComponent.h"
 #include "Components/DisplayClusterStageGeometryComponent.h"
 #include "Widgets/Docking/SDockTab.h"
 
@@ -327,7 +329,7 @@ void FDisplayClusterLightCardEditorViewportClient::Draw(FViewport* InViewport, F
 
 	Canvas->Clear(FLinearColor::Black);
 
-	FSceneViewInitOptions BillboardSceneViewInitOptions;
+	FSceneViewInitOptions SpriteSceneViewInitOptions;
 	
 	if (!bDisableCustomRenderer)
 	{
@@ -341,24 +343,24 @@ void FDisplayClusterLightCardEditorViewportClient::Draw(FViewport* InViewport, F
 
 		GetSceneViewInitOptions(RenderSettings.ViewInitOptions);
 
-		BillboardSceneViewInitOptions = RenderSettings.ViewInitOptions;
+		SpriteSceneViewInitOptions = RenderSettings.ViewInitOptions;
 		
 		IDisplayClusterScenePreview::Get().Render(PreviewRendererId, RenderSettings, *Canvas);
 	}
 	else
 	{
-		if (BillboardComponentProxies.Num() > 0)
+		if (BillboardComponentProxies.Num() > 0 || WidgetComponentProxies.Num() > 0)
 		{
-			GetSceneViewInitOptions(BillboardSceneViewInitOptions);
+			GetSceneViewInitOptions(SpriteSceneViewInitOptions);
 		}
 		
 		GetRendererModule().BeginRenderingViewFamily(Canvas, &ViewFamily);
 	}
 
-	if (BillboardComponentProxies.Num() > 0)
+	if (BillboardComponentProxies.Num() > 0 || WidgetComponentProxies.Num() > 0)
 	{
-		FScopeLock Lock(&BillboardComponentCS);
-		BillboardViewMatrices = FViewMatrices(BillboardSceneViewInitOptions);
+		FScopeLock Lock(&SpriteComponentCS);
+		SpriteViewMatrices = FViewMatrices(SpriteSceneViewInitOptions);
 	}
 
 	if (View)
@@ -473,7 +475,7 @@ void FDisplayClusterLightCardEditorViewportClient::Draw(const FSceneView* View,
 	// we can't render the gizmo over the icons in that case. Canvas items drawn onto the viewport will draw onto the
 	// back buffer after the entire custom render pipeline is drawn.
 	{
-		FScopeLock Lock(&BillboardComponentCS);
+		FScopeLock Lock(&SpriteComponentCS);
 		for (const TWeakObjectPtr<UBillboardComponent>& BillboardComponent : BillboardComponentProxies)
 		{
 			if (!BillboardComponent.IsValid())
@@ -491,7 +493,7 @@ void FDisplayClusterLightCardEditorViewportClient::Draw(const FSceneView* View,
 			
 			FSpriteProxy SpriteProxy = FSpriteProxy::FromBillboard(BillboardComponent.Get());
 		
-			FVector ProjectedLocation = ProjectWorldPosition(SpriteProxy.WorldPosition, BillboardViewMatrices);
+			FVector ProjectedLocation = ProjectWorldPosition(SpriteProxy.WorldPosition, SpriteViewMatrices);
 		
 			PDI->SetHitProxy(new HActor(BillboardComponent->GetOwner(), BillboardComponent.Get()));
 		
@@ -534,6 +536,43 @@ void FDisplayClusterLightCardEditorViewportClient::Draw(const FSceneView* View,
 			PDI->SetHitProxy(nullptr);
 		}
 	}
+
+	if (!bIsUVProjection) // Labels not currently supported in UV mode
+	{
+		FScopeLock Lock(&SpriteComponentCS);
+		for (const TWeakObjectPtr<UWidgetComponent>& WidgetComponent : WidgetComponentProxies)
+		{
+			if (!WidgetComponent.IsValid())
+			{
+				continue;
+			}
+			
+			if (UTextureRenderTarget2D* RenderTarget = WidgetComponent->GetRenderTarget())
+			{
+				if (const FTextureResource* Resource = RenderTarget->GetResource())
+				{
+					const FVector2D Size = WidgetComponent->GetDrawSize();
+					const float UL = Size.X;
+					const float VL = Size.Y;
+
+					// Labels can appear large by default, especially in dome projection mode
+					const float ScaleModifier = ProjectionMode == EDisplayClusterMeshProjectionType::Azimuthal ? 0.125f : 0.25f;
+
+					const UDisplayClusterLightCardEditorProjectSettings* Settings = GetDefault<UDisplayClusterLightCardEditorProjectSettings>();
+					const float ViewedSizeX = (Settings->LightCardLabelScale * UL) * ScaleModifier;
+					const float ViewedSizeY = (Settings->LightCardLabelScale * VL) * ScaleModifier;
+
+					const FVector ProjectedLocation = ProjectWorldPosition(WidgetComponent->GetComponentLocation(), SpriteViewMatrices);
+				
+					PDI->SetHitProxy(new HActor(WidgetComponent->GetOwner(), WidgetComponent.Get()));
+					PDI->DrawSprite(ProjectedLocation, ViewedSizeX, ViewedSizeY,
+						Resource, FLinearColor::White,
+						SDPG_World, 0, 0, 0, 0, SE_BLEND_Masked);
+					PDI->SetHitProxy(nullptr);
+				}
+			}
+		}
+	}
 		
 	if (LastSelectedActor.IsValid() && LastSelectedActor->IsUVActor() == bIsUVProjection)
 	{
@@ -1150,10 +1189,25 @@ AActor* FDisplayClusterLightCardEditorViewportClient::CreateStageActorProxy(AAct
 		TArray<UBillboardComponent*> BillboardComponents;
 		ActorProxy->GetComponents<UBillboardComponent>(BillboardComponents);
 
-		FScopeLock Lock(&BillboardComponentCS);
+		FScopeLock Lock(&SpriteComponentCS);
 		BillboardComponentProxies.Append(BillboardComponents);
 	}
 
+	// Label widgets
+	{
+		TArray<UDisplayClusterLabelComponent*> LabelComponents;
+		ActorProxy->GetComponents<UDisplayClusterLabelComponent>(LabelComponents);
+
+		FScopeLock Lock(&SpriteComponentCS);
+		
+		for (UDisplayClusterLabelComponent* LabelComponent : LabelComponents)
+		{
+			UWidgetComponent* WidgetComponent = LabelComponent->GetWidgetComponent();
+			WidgetComponent->SetTickWhenOffscreen(true);
+			WidgetComponentProxies.Add(WidgetComponent);
+		}
+	}
+	
 	FActorProxy ActorProxyStruct(InLevelInstance, ActorProxy);
 	ActorProxies.Add(ActorProxyStruct);
 	ProjectionHelper->VerifyAndFixActorOrigin(InLevelInstance);
@@ -1745,7 +1799,7 @@ void FDisplayClusterLightCardEditorViewportClient::UpdateProxyTransformFromLevel
 void FDisplayClusterLightCardEditorViewportClient::DestroyProxies(
 	EDisplayClusterLightCardEditorProxyType ProxyType)
 {
-	FScopeLock Lock(&BillboardComponentCS);
+	FScopeLock Lock(&SpriteComponentCS);
 
 	// Clear the primitives from the scene renderer based on the type of proxy that is being destroyed
 	switch (ProxyType)
@@ -1813,7 +1867,7 @@ void FDisplayClusterLightCardEditorViewportClient::DestroyProxy(AActor* Actor)
 				return OtherProxy == ActualProxy;
 			});
 		
-			FScopeLock Lock(&BillboardComponentCS);
+			FScopeLock Lock(&SpriteComponentCS);
 		
 			UWorld* PreviewWorld = PreviewScene->GetWorld();
 			check(PreviewWorld);
@@ -1824,6 +1878,11 @@ void FDisplayClusterLightCardEditorViewportClient::DestroyProxy(AActor* Actor)
 			{
 				return !BillboardComponent.IsValid() || BillboardComponent->GetOwner() == ActualProxy;
 			});
+
+			WidgetComponentProxies.RemoveAll([ActualProxy](const TWeakObjectPtr<UWidgetComponent>& WidgetComponent)
+			{
+				return !WidgetComponent.IsValid() || WidgetComponent->GetOwner() == ActualProxy;
+			});
 		}
 	}
 }
diff --git a/Engine/Plugins/Runtime/nDisplay/Source/DisplayClusterLightCardEditor/Private/Viewport/DisplayClusterLightCardEditorViewportClient.h b/Engine/Plugins/Runtime/nDisplay/Source/DisplayClusterLightCardEditor/Private/Viewport/DisplayClusterLightCardEditorViewportClient.h
--- a/Engine/Plugins/Runtime/nDisplay/Source/DisplayClusterLightCardEditor/Private/Viewport/DisplayClusterLightCardEditorViewportClient.h
+++ b/Engine/Plugins/Runtime/nDisplay/Source/DisplayClusterLightCardEditor/Private/Viewport/DisplayClusterLightCardEditorViewportClient.h
@@ -14,6 +14,7 @@
 #include "DisplayClusterLightCardEditorHelper.h"
 #include "StageActor/DisplayClusterWeakStageActorPtr.h"
 
+class UWidgetComponent;
 class ADisplayClusterRootActor;
 class FDisplayClusterLightCardEditor;
 class SDisplayClusterLightCardEditorViewport;
@@ -325,13 +326,14 @@ private:
 
 	TArray<FActorProxy> ActorProxies;
 	TArray<TWeakObjectPtr<UBillboardComponent>> BillboardComponentProxies;
+	TArray<TWeakObjectPtr<UWidgetComponent>> WidgetComponentProxies;
 	TArray<FDisplayClusterWeakStageActorPtr> SelectedActors;
 
-	/** Billboards are accessed from the render thread and game thread */
-	FCriticalSection BillboardComponentCS;
+	/** Billboards and widgets are accessed from the render thread and game thread for sprite rendering */
+	FCriticalSection SpriteComponentCS;
 
-	/** View matrices set per tick if there are any billboard components */
-	FViewMatrices BillboardViewMatrices;
+	/** View matrices set per tick if there are any sprites being rendered */
+	FViewMatrices SpriteViewMatrices;
 
 	/** The actor in the selected actor list that was selected last */
 	FDisplayClusterWeakStageActorPtr LastSelectedActor;
