nDisplay: added initializer for LightcardOCIOMode EnumProperty

diff --git a/Engine/Plugins/Runtime/nDisplay/Source/DisplayClusterConfiguration/Public/DisplayClusterConfigurationTypes_ICVFX.h b/Engine/Plugins/Runtime/nDisplay/Source/DisplayClusterConfiguration/Public/DisplayClusterConfigurationTypes_ICVFX.h
--- a/Engine/Plugins/Runtime/nDisplay/Source/DisplayClusterConfiguration/Public/DisplayClusterConfigurationTypes_ICVFX.h
+++ b/Engine/Plugins/Runtime/nDisplay/Source/DisplayClusterConfiguration/Public/DisplayClusterConfigurationTypes_ICVFX.h
@@ -96,7 +96,7 @@ public:
 public:
 	/** Light Cards OCIO mode. */
 	UPROPERTY(EditAnywhere, BlueprintReadWrite, Category = "Light Card|OCIO", meta = (DisplayName = "Light Cards OCIO"))
-	EDisplayClusterConfigurationViewportLightcardOCIOMode LightcardOCIOMode;
+	EDisplayClusterConfigurationViewportLightcardOCIOMode LightcardOCIOMode = EDisplayClusterConfigurationViewportLightcardOCIOMode::nDisplay;
 
 	/** Custom OpenColorIO configuration for Light Cards. */
 	UPROPERTY(EditAnywhere, BlueprintReadWrite, Category = "Light Card|OCIO", meta = (DisplayName = "Custom Light Cards OCIO"))
