nDisplay: Fix setting CCW transform properties through non-interactive changes in the ICVFX editor not behaving properly.
Prevent the light card editor from double processing scene component transformation changes.

diff --git a/Engine/Plugins/Experimental/ColorCorrectRegions/Source/ColorCorrectRegions/Private/ColorCorrectWindow.cpp b/Engine/Plugins/Experimental/ColorCorrectRegions/Source/ColorCorrectRegions/Private/ColorCorrectWindow.cpp
--- a/Engine/Plugins/Experimental/ColorCorrectRegions/Source/ColorCorrectRegions/Private/ColorCorrectWindow.cpp
+++ b/Engine/Plugins/Experimental/ColorCorrectRegions/Source/ColorCorrectRegions/Private/ColorCorrectWindow.cpp
@@ -55,7 +55,6 @@ AColorCorrectionWindow::AColorCorrectionWindow(const FObjectInitializer& ObjectI
 #if WITH_EDITOR
 void AColorCorrectionWindow::PostEditChangeProperty(struct FPropertyChangedEvent& PropertyChangedEvent)
 {
-	Super::PostEditChangeProperty(PropertyChangedEvent);
 	const FName PropertyName = PropertyChangedEvent.GetPropertyName();
 
 	if (PropertyName == GET_MEMBER_NAME_CHECKED(AColorCorrectionWindow, WindowType))
@@ -83,6 +82,10 @@ void AColorCorrectionWindow::PostEditChangeProperty(struct FPropertyChangedEvent
 			bNotifyOnParamSetter = true;
 		}
 	}
+
+	// Call after stage actor transform is updated, so any observers will have both the correct actor transform and
+	// positional properties.
+	Super::PostEditChangeProperty(PropertyChangedEvent);
 }
 
 void AColorCorrectionWindow::PostEditMove(bool bFinished)
diff --git a/Engine/Plugins/Runtime/nDisplay/Source/DisplayClusterLightCardEditor/Private/DisplayClusterLightCardEditor.cpp b/Engine/Plugins/Runtime/nDisplay/Source/DisplayClusterLightCardEditor/Private/DisplayClusterLightCardEditor.cpp
--- a/Engine/Plugins/Runtime/nDisplay/Source/DisplayClusterLightCardEditor/Private/DisplayClusterLightCardEditor.cpp
+++ b/Engine/Plugins/Runtime/nDisplay/Source/DisplayClusterLightCardEditor/Private/DisplayClusterLightCardEditor.cpp
@@ -1775,14 +1775,22 @@ void FDisplayClusterLightCardEditor::OnActorPropertyChanged(UObject* ObjectBeing
 		IDisplayClusterStageActor* StageActor = Cast<IDisplayClusterStageActor>(ObjectBeingModified);
 		
 		const FName PropertyName = PropertyChangedEvent.GetPropertyName();
-		const bool bTransformationChanged =
-				(PropertyName == USceneComponent::GetRelativeLocationPropertyName() ||
-					PropertyName == USceneComponent::GetRelativeRotationPropertyName() ||
-					PropertyName == USceneComponent::GetRelativeScale3DPropertyName() ||
-					(StageActor && StageActor->GetPositionalPropertyNames().Contains(PropertyName)));
+		const bool bComponentTransformPropertyChanged = PropertyName == USceneComponent::GetRelativeLocationPropertyName()
+		|| PropertyName == USceneComponent::GetRelativeRotationPropertyName()
+		|| PropertyName == USceneComponent::GetRelativeScale3DPropertyName();
+		
+		const bool bTransformationChanged = (bComponentTransformPropertyChanged
+				|| (StageActor && StageActor->GetPositionalPropertyNames().Contains(PropertyName)));
+		
 		if (PropertyChangedEvent.ChangeType == EPropertyChangeType::Interactive ||
 			(bTransformationChanged && !ObjectBeingModified->IsA<ADisplayClusterRootActor>()))
 		{
+			if (bComponentTransformPropertyChanged && ObjectBeingModified->IsA<USceneComponent>())
+			{
+				// The owning scene component actor will also fire a property change. Prevent double processing.
+				return;
+			}
+			
 			// Real-time & efficient update when dragging a slider or when we know the property type doesn't
 			// require a full refresh
 			if (ViewportView.IsValid())
