nDisplay: Temporary solution to a crash caused by running the root actor's preview during a SIE session; added a CVar flag to enable SIE preview rendering that is disabled by default.

diff --git a/Engine/Plugins/Runtime/nDisplay/Source/DisplayCluster/Private/Game/EngineClasses/Scene/DisplayClusterRootActorEditor.cpp b/Engine/Plugins/Runtime/nDisplay/Source/DisplayCluster/Private/Game/EngineClasses/Scene/DisplayClusterRootActorEditor.cpp
--- a/Engine/Plugins/Runtime/nDisplay/Source/DisplayCluster/Private/Game/EngineClasses/Scene/DisplayClusterRootActorEditor.cpp
+++ b/Engine/Plugins/Runtime/nDisplay/Source/DisplayCluster/Private/Game/EngineClasses/Scene/DisplayClusterRootActorEditor.cpp
@@ -79,6 +79,13 @@ static FAutoConsoleVariableRef CVarDisplayClusterPreviewMultiGPURenderingMaxInde
 	ECVF_RenderThreadSafe
 );
 
+int32 GDisplayClusterPreviewEnableSIE = 0;
+static FAutoConsoleVariableRef CVarDisplayClusterPreviewEnableSIE(
+	TEXT("DC.Preview.EnableSIE"),
+	GDisplayClusterPreviewEnableSIE,
+	TEXT("Enable preview rendering during an SIE session"),
+	ECVF_RenderThreadSafe
+);
 //////////////////////////////////////////////////////////////////////////////////////////////
 // ADisplayClusterRootActor
 //////////////////////////////////////////////////////////////////////////////////////////////
@@ -124,7 +131,8 @@ void ADisplayClusterRootActor::Destructor_Editor()
 
 void ADisplayClusterRootActor::Tick_Editor(float DeltaSeconds)
 {
-	if (IsPreviewEnabled())
+	// Temporary fix for UE-174008, don't render preview when in SIE unless the cvar forces it to
+	if (IsPreviewEnabled() && (!GEditor->IsSimulatingInEditor() || GDisplayClusterPreviewEnableSIE))
 	{
 		if (bDeferPreviewGeneration)
 		{
