nDisplay: Fix labels being included in the preview capture. Fix labels sometimes being clipped by the preview.

diff --git a/Engine/Plugins/Runtime/nDisplay/Source/DisplayCluster/Private/Game/EngineClasses/Scene/DisplayClusterLabelComponent.cpp b/Engine/Plugins/Runtime/nDisplay/Source/DisplayCluster/Private/Game/EngineClasses/Scene/DisplayClusterLabelComponent.cpp
--- a/Engine/Plugins/Runtime/nDisplay/Source/DisplayCluster/Private/Game/EngineClasses/Scene/DisplayClusterLabelComponent.cpp
+++ b/Engine/Plugins/Runtime/nDisplay/Source/DisplayCluster/Private/Game/EngineClasses/Scene/DisplayClusterLabelComponent.cpp
@@ -21,7 +21,7 @@ UDisplayClusterLabelComponent::UDisplayClusterLabelComponent()
 	// Keeps this visible even if its clipping through the owner.
 	// Requires the owner have a transparent blend mode on its material.
 	WidgetComponent->SetBlendMode(EWidgetBlendMode::Transparent);
-	WidgetComponent->SetTranslucencySortDistanceOffset(-1000.f);
+	WidgetComponent->SetTranslucencySortDistanceOffset(-100.f);
 	WidgetComponent->SetUsingAbsoluteScale(true);
 	WidgetComponent->SetCollisionEnabled(ECollisionEnabled::NoCollision);
 	WidgetComponent->SetCollisionResponseToAllChannels(ECollisionResponse::ECR_Ignore);
diff --git a/Engine/Plugins/Runtime/nDisplay/Source/DisplayCluster/Private/Game/EngineClasses/Scene/DisplayClusterLightCardActor.cpp b/Engine/Plugins/Runtime/nDisplay/Source/DisplayCluster/Private/Game/EngineClasses/Scene/DisplayClusterLightCardActor.cpp
--- a/Engine/Plugins/Runtime/nDisplay/Source/DisplayCluster/Private/Game/EngineClasses/Scene/DisplayClusterLightCardActor.cpp
+++ b/Engine/Plugins/Runtime/nDisplay/Source/DisplayCluster/Private/Game/EngineClasses/Scene/DisplayClusterLightCardActor.cpp
@@ -157,7 +157,7 @@ void ADisplayClusterLightCardActor::Tick(float DeltaSeconds)
 	UpdateUVIndicator();
 	if (LabelComponent->IsVisible())
 	{
-		LabelComponent->GetWidgetComponent()->SetTranslucentSortPriority(LightCardComponent->TranslucencySortPriority);
+		LabelComponent->GetWidgetComponent()->SetTranslucentSortPriority(LightCardComponent->TranslucencySortPriority + 1);
 	}
 }
 
diff --git a/Engine/Plugins/Runtime/nDisplay/Source/DisplayCluster/Private/Game/EngineClasses/Scene/DisplayClusterRootActor.cpp b/Engine/Plugins/Runtime/nDisplay/Source/DisplayCluster/Private/Game/EngineClasses/Scene/DisplayClusterRootActor.cpp
--- a/Engine/Plugins/Runtime/nDisplay/Source/DisplayCluster/Private/Game/EngineClasses/Scene/DisplayClusterRootActor.cpp
+++ b/Engine/Plugins/Runtime/nDisplay/Source/DisplayCluster/Private/Game/EngineClasses/Scene/DisplayClusterRootActor.cpp
@@ -636,7 +636,11 @@ bool ADisplayClusterRootActor::GetHiddenInGamePrimitives(TSet<FPrimitiveComponen
 
 						for (UPrimitiveComponent* PrimComp : PrimitiveComponents)
 						{
-							if (((Actor->IsHidden() || PrimComp->bHiddenInGame) && !PrimComp->bCastHiddenShadow))
+							if (((Actor->IsHidden() || PrimComp->bHiddenInGame) && !PrimComp->bCastHiddenShadow)
+#if WITH_EDITOR
+								|| (GIsEditor && PrimComp->bHiddenInSceneCapture /* We are running as a scene capture for preview */)
+#endif
+								)
 							{
 								PrimitiveComponentsArray[Index].Add(PrimComp->ComponentId);
 							}
diff --git a/Engine/Plugins/Runtime/nDisplay/Source/DisplayCluster/Private/Render/Viewport/DisplayClusterViewport_VisibilitySettings.cpp b/Engine/Plugins/Runtime/nDisplay/Source/DisplayCluster/Private/Render/Viewport/DisplayClusterViewport_VisibilitySettings.cpp
--- a/Engine/Plugins/Runtime/nDisplay/Source/DisplayCluster/Private/Render/Viewport/DisplayClusterViewport_VisibilitySettings.cpp
+++ b/Engine/Plugins/Runtime/nDisplay/Source/DisplayCluster/Private/Render/Viewport/DisplayClusterViewport_VisibilitySettings.cpp
@@ -7,15 +7,15 @@
 #include "EngineUtils.h"
 #include "SceneView.h"
 
-static void GetPrimitiveComponentsFromLayers(UWorld* World, const TArray<FName>& SourceLayers, TSet<FPrimitiveComponentId>& OutPrimitives)
+static void GetPrimitiveComponentsFromLayers(UWorld* World, const TArray<FName>& SourceLayers, TSet<FPrimitiveComponentId>& OutPrimitives,
+                                             const TSet<FPrimitiveComponentId>* InExcludePrimitivesList = nullptr)
 {
 	if (SourceLayers.Num())
 	{
 		// Iterate over all actors, looking for actors in the specified layers.
 		for (const TWeakObjectPtr<AActor> WeakActor : FActorRange(World))
 		{
-			AActor* Actor = WeakActor.Get();
-			if (Actor)
+			if (const AActor* Actor = WeakActor.Get())
 			{
 				bool bActorFoundOnSourceLayers = false;
 
@@ -34,9 +34,12 @@ static void GetPrimitiveComponentsFromLayers(UWorld* World, const TArray<FName>&
 					// Save all actor components to OutPrimitives
 					for (UActorComponent* Component : Actor->GetComponents())
 					{
-						if (UPrimitiveComponent* PrimComp = Cast<UPrimitiveComponent>(Component))
+						if (const UPrimitiveComponent* PrimComp = Cast<UPrimitiveComponent>(Component))
 						{
-							OutPrimitives.Add(PrimComp->ComponentId);
+							if (!InExcludePrimitivesList || !InExcludePrimitivesList->Contains(PrimComp->ComponentId))
+							{
+								OutPrimitives.Add(PrimComp->ComponentId);
+							}
 						}
 					}
 				}
@@ -56,8 +59,16 @@ void FDisplayClusterViewport_VisibilitySettings::SetupSceneView(class UWorld* Wo
 		if (ActorLayers.Num() > 0 || AdditionalComponentsList.Num() > 0)
 		{
 			InOutView.ShowOnlyPrimitives.Emplace();
-			GetPrimitiveComponentsFromLayers(World, ActorLayers, InOutView.ShowOnlyPrimitives.GetValue());
-			InOutView.ShowOnlyPrimitives.GetValue().Append(AdditionalComponentsList);
+			GetPrimitiveComponentsFromLayers(World, ActorLayers,InOutView.ShowOnlyPrimitives.GetValue(),
+				&RootActorHidePrimitivesList);
+
+			for (const FPrimitiveComponentId& AdditionalPrimitiveId : AdditionalComponentsList)
+			{
+				if (!RootActorHidePrimitivesList.Contains(AdditionalPrimitiveId))
+				{
+					InOutView.ShowOnlyPrimitives.GetValue().Add(AdditionalPrimitiveId);
+				}
+			}
 			return;
 		}
 		break;
