nDisplay: fixed memory leak caused by missed cleanup of viewports

diff --git a/Engine/Plugins/Runtime/nDisplay/Source/DisplayCluster/Private/Game/EngineClasses/Scene/DisplayClusterPreviewComponent.cpp b/Engine/Plugins/Runtime/nDisplay/Source/DisplayCluster/Private/Game/EngineClasses/Scene/DisplayClusterPreviewComponent.cpp
--- a/Engine/Plugins/Runtime/nDisplay/Source/DisplayCluster/Private/Game/EngineClasses/Scene/DisplayClusterPreviewComponent.cpp
+++ b/Engine/Plugins/Runtime/nDisplay/Source/DisplayCluster/Private/Game/EngineClasses/Scene/DisplayClusterPreviewComponent.cpp
@@ -303,7 +303,15 @@ template<typename T>
 void UDisplayClusterPreviewComponent::ReleaseRenderTargetImpl(T* InOutRenderTarget)
 {
 	checkSlow(InOutRenderTarget);
-	*InOutRenderTarget = nullptr;
+	T& RenderTargetPtr = *InOutRenderTarget;
+
+	if (RenderTargetPtr != nullptr)
+	{
+		RenderTargetPtr->ReleaseResource();
+		RenderTargetPtr->MarkAsGarbage();
+
+		RenderTargetPtr = nullptr;
+	}
 }
 
 template<typename T>
diff --git a/Engine/Plugins/Runtime/nDisplay/Source/DisplayCluster/Private/Game/EngineClasses/Scene/DisplayClusterRootActor.cpp b/Engine/Plugins/Runtime/nDisplay/Source/DisplayCluster/Private/Game/EngineClasses/Scene/DisplayClusterRootActor.cpp
--- a/Engine/Plugins/Runtime/nDisplay/Source/DisplayCluster/Private/Game/EngineClasses/Scene/DisplayClusterRootActor.cpp
+++ b/Engine/Plugins/Runtime/nDisplay/Source/DisplayCluster/Private/Game/EngineClasses/Scene/DisplayClusterRootActor.cpp
@@ -1014,6 +1014,18 @@ void ADisplayClusterRootActor::PostActorCreated()
 	InitializeRootActor();
 }
 
+void ADisplayClusterRootActor::Destroyed()
+{
+#if WITH_EDITOR
+	Destroyed_Editor();
+#endif
+
+	// Release viewport manager with resources immediatelly
+	ViewportManager.Reset();
+
+	Super::Destroyed();
+}
+
 void ADisplayClusterRootActor::BeginDestroy()
 {
 #if WITH_EDITOR
diff --git a/Engine/Plugins/Runtime/nDisplay/Source/DisplayCluster/Private/Game/EngineClasses/Scene/DisplayClusterRootActorEditor.cpp b/Engine/Plugins/Runtime/nDisplay/Source/DisplayCluster/Private/Game/EngineClasses/Scene/DisplayClusterRootActorEditor.cpp
--- a/Engine/Plugins/Runtime/nDisplay/Source/DisplayCluster/Private/Game/EngineClasses/Scene/DisplayClusterRootActorEditor.cpp
+++ b/Engine/Plugins/Runtime/nDisplay/Source/DisplayCluster/Private/Game/EngineClasses/Scene/DisplayClusterRootActorEditor.cpp
@@ -177,6 +177,18 @@ void ADisplayClusterRootActor::PostLoad_Editor()
 
 	ResetPreviewInternals_Editor();
 }
+
+void ADisplayClusterRootActor::Destroyed_Editor()
+{
+	ResetPreviewInternals_Editor();
+	ReleasePreviewComponents();
+
+	MarkAsGarbage();
+
+	// Force garbage collector
+	GEngine->ForceGarbageCollection(true);
+}
+
 void ADisplayClusterRootActor::BeginDestroy_Editor()
 {
 	ResetPreviewInternals_Editor();
diff --git a/Engine/Plugins/Runtime/nDisplay/Source/DisplayCluster/Private/Render/Viewport/DisplayClusterViewportManager.cpp b/Engine/Plugins/Runtime/nDisplay/Source/DisplayCluster/Private/Render/Viewport/DisplayClusterViewportManager.cpp
--- a/Engine/Plugins/Runtime/nDisplay/Source/DisplayCluster/Private/Render/Viewport/DisplayClusterViewportManager.cpp
+++ b/Engine/Plugins/Runtime/nDisplay/Source/DisplayCluster/Private/Render/Viewport/DisplayClusterViewportManager.cpp
@@ -74,6 +74,13 @@ FDisplayClusterViewportManager::FDisplayClusterViewportManager()
 FDisplayClusterViewportManager::~FDisplayClusterViewportManager()
 {
 	// Remove viewports
+	for (FDisplayClusterViewport* ViewportIt : Viewports)
+	{
+		if (ViewportIt)
+		{
+			delete ViewportIt;
+		}
+	}
 	Viewports.Reset();
 	ClusterNodeViewports.Reset();
 
diff --git a/Engine/Plugins/Runtime/nDisplay/Source/DisplayCluster/Public/DisplayClusterRootActor.h b/Engine/Plugins/Runtime/nDisplay/Source/DisplayCluster/Public/DisplayClusterRootActor.h
--- a/Engine/Plugins/Runtime/nDisplay/Source/DisplayCluster/Public/DisplayClusterRootActor.h
+++ b/Engine/Plugins/Runtime/nDisplay/Source/DisplayCluster/Public/DisplayClusterRootActor.h
@@ -115,6 +115,7 @@ protected:
 	virtual void PostLoad() override;
 	virtual void PostActorCreated() override;
 	virtual void BeginDestroy() override;
+	virtual void Destroyed() override;
 #if WITH_EDITOR
 	virtual void RerunConstructionScripts() override;
 #endif
@@ -441,6 +442,7 @@ public:
 	void PostLoad_Editor();
 	void PostActorCreated_Editor();
 	void BeginDestroy_Editor();
+	void Destroyed_Editor();
 	void RerunConstructionScripts_Editor();
 
 	// Preview components free referenced meshes and materials
