nDisplay: ViewStates will be destroyed after the viewport proxy on the render thread.

diff --git a/Engine/Plugins/Runtime/nDisplay/Source/DisplayCluster/Private/Render/Viewport/Containers/DisplayClusterViewportProxyData.cpp b/Engine/Plugins/Runtime/nDisplay/Source/DisplayCluster/Private/Render/Viewport/Containers/DisplayClusterViewportProxyData.cpp
--- a/Engine/Plugins/Runtime/nDisplay/Source/DisplayCluster/Private/Render/Viewport/Containers/DisplayClusterViewportProxyData.cpp
+++ b/Engine/Plugins/Runtime/nDisplay/Source/DisplayCluster/Private/Render/Viewport/Containers/DisplayClusterViewportProxyData.cpp
@@ -38,6 +38,7 @@ FDisplayClusterViewportProxyData::FDisplayClusterViewportProxyData(const FDispla
 
 #if WITH_EDITOR
 	OutputPreviewTargetableResource = SrcViewport->OutputPreviewTargetableResource;
+	ViewStates = SrcViewport->ViewStates;
 #endif
 
 	InputShaderResources = SrcViewport->InputShaderResources;
@@ -84,6 +85,7 @@ void FDisplayClusterViewportProxyData::UpdateProxy_RenderThread() const
 
 #if WITH_EDITOR
 	DstViewportProxy->OutputPreviewTargetableResource = OutputPreviewTargetableResource;
+	DstViewportProxy->ViewStates = ViewStates;
 #endif
 
 	DstViewportProxy->InputShaderResources = InputShaderResources;
diff --git a/Engine/Plugins/Runtime/nDisplay/Source/DisplayCluster/Private/Render/Viewport/Containers/DisplayClusterViewportProxyData.h b/Engine/Plugins/Runtime/nDisplay/Source/DisplayCluster/Private/Render/Viewport/Containers/DisplayClusterViewportProxyData.h
--- a/Engine/Plugins/Runtime/nDisplay/Source/DisplayCluster/Private/Render/Viewport/Containers/DisplayClusterViewportProxyData.h
+++ b/Engine/Plugins/Runtime/nDisplay/Source/DisplayCluster/Private/Render/Viewport/Containers/DisplayClusterViewportProxyData.h
@@ -57,6 +57,7 @@ private:
 
 #if WITH_EDITOR
 	FTextureRHIRef OutputPreviewTargetableResource;
+	TArray<TSharedPtr<FSceneViewStateReference, ESPMode::ThreadSafe>> ViewStates;
 #endif
 
 	// unique viewport resources
diff --git a/Engine/Plugins/Runtime/nDisplay/Source/DisplayCluster/Private/Render/Viewport/DisplayClusterViewport.cpp b/Engine/Plugins/Runtime/nDisplay/Source/DisplayCluster/Private/Render/Viewport/DisplayClusterViewport.cpp
--- a/Engine/Plugins/Runtime/nDisplay/Source/DisplayCluster/Private/Render/Viewport/DisplayClusterViewport.cpp
+++ b/Engine/Plugins/Runtime/nDisplay/Source/DisplayCluster/Private/Render/Viewport/DisplayClusterViewport.cpp
@@ -226,14 +226,7 @@ void FDisplayClusterViewport::HandleEndScene()
 void FDisplayClusterViewport::AddReferencedObjects(FReferenceCollector& Collector)
 {
 #if WITH_EDITOR
-	for (FSceneViewStateReference& ViewState : ViewStates)
-	{
-		FSceneViewStateInterface* Ref = ViewState.GetReference();
-		if (Ref)
-		{
-			Ref->AddReferencedObjects(Collector);
-		}
-	}
+	// ViewStates released on rendering thread from viewport proxy object
 #endif
 }
 
diff --git a/Engine/Plugins/Runtime/nDisplay/Source/DisplayCluster/Private/Render/Viewport/DisplayClusterViewport.h b/Engine/Plugins/Runtime/nDisplay/Source/DisplayCluster/Private/Render/Viewport/DisplayClusterViewport.h
--- a/Engine/Plugins/Runtime/nDisplay/Source/DisplayCluster/Private/Render/Viewport/DisplayClusterViewport.h
+++ b/Engine/Plugins/Runtime/nDisplay/Source/DisplayCluster/Private/Render/Viewport/DisplayClusterViewport.h
@@ -218,7 +218,7 @@ private:
 	void CleanupViewState();
 
 private:
-	TArray<FSceneViewStateReference> ViewStates;
+	TArray<TSharedPtr<FSceneViewStateReference, ESPMode::ThreadSafe>> ViewStates;
 #endif
 
 public:
diff --git a/Engine/Plugins/Runtime/nDisplay/Source/DisplayCluster/Private/Render/Viewport/DisplayClusterViewportProxy.h b/Engine/Plugins/Runtime/nDisplay/Source/DisplayCluster/Private/Render/Viewport/DisplayClusterViewportProxy.h
--- a/Engine/Plugins/Runtime/nDisplay/Source/DisplayCluster/Private/Render/Viewport/DisplayClusterViewportProxy.h
+++ b/Engine/Plugins/Runtime/nDisplay/Source/DisplayCluster/Private/Render/Viewport/DisplayClusterViewportProxy.h
@@ -315,6 +315,7 @@ protected:
 
 #if WITH_EDITOR
 	FTextureRHIRef OutputPreviewTargetableResource;
+	TArray<TSharedPtr<FSceneViewStateReference, ESPMode::ThreadSafe>> ViewStates;
 	
 	mutable bool bPreviewReadPixels = false;
 	mutable FCriticalSection PreviewPixelsCSGuard;
diff --git a/Engine/Plugins/Runtime/nDisplay/Source/DisplayCluster/Private/Render/Viewport/DisplayClusterViewport_EditorPreview.cpp b/Engine/Plugins/Runtime/nDisplay/Source/DisplayCluster/Private/Render/Viewport/DisplayClusterViewport_EditorPreview.cpp
--- a/Engine/Plugins/Runtime/nDisplay/Source/DisplayCluster/Private/Render/Viewport/DisplayClusterViewport_EditorPreview.cpp
+++ b/Engine/Plugins/Runtime/nDisplay/Source/DisplayCluster/Private/Render/Viewport/DisplayClusterViewport_EditorPreview.cpp
@@ -33,12 +33,15 @@
 ///////////////////////////////////////////////////////////////////////////////////////
 void FDisplayClusterViewport::CleanupViewState()
 {
-	for (FSceneViewStateReference& ViewState : ViewStates)
+	for (TSharedPtr<FSceneViewStateReference, ESPMode::ThreadSafe>& ViewState : ViewStates)
 	{
-		FSceneViewStateInterface* Ref = ViewState.GetReference();
-		if (Ref != nullptr)
+		if (ViewState.IsValid())
 		{
-			Ref->ClearMIDPool();
+			FSceneViewStateInterface* Ref = ViewState->GetReference();
+			if (Ref != nullptr)
+			{
+				Ref->ClearMIDPool();
+			}
 		}
 	}
 }
@@ -46,20 +49,25 @@ void FDisplayClusterViewport::CleanupViewState()
 FSceneViewStateInterface* FDisplayClusterViewport::GetViewState(uint32 ViewIndex)
 {
 	int32 RequiredAmount = (int32)ViewIndex - ViewStates.Num() + 1;
-	if(RequiredAmount > 0)
+	if (RequiredAmount > 0)
 	{
 		ViewStates.AddDefaulted(RequiredAmount);
 	}
 
-	if (ViewStates[ViewIndex].GetReference() == NULL)
+	if (!ViewStates[ViewIndex].IsValid())
+	{
+		ViewStates[ViewIndex] = MakeShared<FSceneViewStateReference>();
+	}
+
+	if (ViewStates[ViewIndex]->GetReference() == NULL)
 	{
 		const UWorld* CurrentWorld = Owner.GetCurrentWorld();
 		const ERHIFeatureLevel::Type FeatureLevel = CurrentWorld ? CurrentWorld->FeatureLevel.GetValue() : GMaxRHIFeatureLevel;
 
-		ViewStates[ViewIndex].Allocate(FeatureLevel);
+		ViewStates[ViewIndex]->Allocate(FeatureLevel);
 	}
 
-	return ViewStates[ViewIndex].GetReference();
+	return ViewStates[ViewIndex]->GetReference();
 }
 
 FSceneView* FDisplayClusterViewport::ImplCalcScenePreview(FSceneViewFamilyContext& InOutViewFamily, uint32 InContextNum)
