nDisplay: Added trace to GetPrimitiveComponentsFromLayers.

diff --git a/Engine/Plugins/Runtime/nDisplay/Source/DisplayCluster/Private/Render/Viewport/DisplayClusterViewport_VisibilitySettings.cpp b/Engine/Plugins/Runtime/nDisplay/Source/DisplayCluster/Private/Render/Viewport/DisplayClusterViewport_VisibilitySettings.cpp
--- a/Engine/Plugins/Runtime/nDisplay/Source/DisplayCluster/Private/Render/Viewport/DisplayClusterViewport_VisibilitySettings.cpp
+++ b/Engine/Plugins/Runtime/nDisplay/Source/DisplayCluster/Private/Render/Viewport/DisplayClusterViewport_VisibilitySettings.cpp
@@ -10,6 +10,8 @@
 static void GetPrimitiveComponentsFromLayers(UWorld* World, const TArray<FName>& SourceLayers, TSet<FPrimitiveComponentId>& OutPrimitives,
                                              const TSet<FPrimitiveComponentId>* InExcludePrimitivesList = nullptr)
 {
+	TRACE_CPUPROFILER_EVENT_SCOPE(DCRootActor_GetPrimitiveComponentsFromLayers);
+
 	if (SourceLayers.Num())
 	{
 		// Iterate over all actors, looking for actors in the specified layers.
