Fixed non unity/pch compile errors reported by build farm

diff --git a/Engine/Plugins/Animation/DeformerGraph/Source/OptimusCore/Public/OptimusDiagnostic.h b/Engine/Plugins/Animation/DeformerGraph/Source/OptimusCore/Public/OptimusDiagnostic.h
--- a/Engine/Plugins/Animation/DeformerGraph/Source/OptimusCore/Public/OptimusDiagnostic.h
+++ b/Engine/Plugins/Animation/DeformerGraph/Source/OptimusCore/Public/OptimusDiagnostic.h
@@ -2,7 +2,6 @@
 
 #pragma once
 
-#include "CoreMinimal.h"
 #include "UObject/WeakObjectPtr.h"
 #include "OptimusDiagnostic.generated.h"
 
diff --git a/Engine/Plugins/Runtime/nDisplay/Source/DisplayClusterConfigurator/Private/Views/Details/Components/DisplayClusterICVFXCameraComponentDetailsCustomization.h b/Engine/Plugins/Runtime/nDisplay/Source/DisplayClusterConfigurator/Private/Views/Details/Components/DisplayClusterICVFXCameraComponentDetailsCustomization.h
--- a/Engine/Plugins/Runtime/nDisplay/Source/DisplayClusterConfigurator/Private/Views/Details/Components/DisplayClusterICVFXCameraComponentDetailsCustomization.h
+++ b/Engine/Plugins/Runtime/nDisplay/Source/DisplayClusterConfigurator/Private/Views/Details/Components/DisplayClusterICVFXCameraComponentDetailsCustomization.h
@@ -3,9 +3,8 @@
 #pragma once
 
 #include "IDetailCustomization.h"
-
-#include "CoreMinimal.h"
 #include "Types/SlateEnums.h"
+#include "UObject/WeakObjectPtr.h"
 
 class UDisplayClusterICVFXCameraComponent;
 class ACineCameraActor;
diff --git a/Engine/Source/Runtime/AutomationTest/Public/AutomationState.h b/Engine/Source/Runtime/AutomationTest/Public/AutomationState.h
--- a/Engine/Source/Runtime/AutomationTest/Public/AutomationState.h
+++ b/Engine/Source/Runtime/AutomationTest/Public/AutomationState.h
@@ -2,7 +2,7 @@
 
 #pragma once
 
-#include "CoreTypes.h"
+#include "UObject/Class.h"
 #include "UObject/ObjectMacros.h"
 #include "UObject/Class.h"
 
diff --git a/Engine/Source/Runtime/Engine/Private/Streaming/TextureStreamingHelpers.h b/Engine/Source/Runtime/Engine/Private/Streaming/TextureStreamingHelpers.h
--- a/Engine/Source/Runtime/Engine/Private/Streaming/TextureStreamingHelpers.h
+++ b/Engine/Source/Runtime/Engine/Private/Streaming/TextureStreamingHelpers.h
@@ -12,6 +12,7 @@ TextureStreamingHelpers.h: Definitions of classes used for texture streaming.
 #include "Misc/MemStack.h"
 #include "Engine/TextureDefines.h"
 
+class AActor;
 class UStreamableRenderAsset;
 
 /**
