nDisplay: bHiddenInGame enabled by default in DCRA to avoid being active on Stage

diff --git a/Engine/Plugins/Runtime/nDisplay/Source/DisplayCluster/Private/Game/EngineClasses/Scene/DisplayClusterRootActor.cpp b/Engine/Plugins/Runtime/nDisplay/Source/DisplayCluster/Private/Game/EngineClasses/Scene/DisplayClusterRootActor.cpp
--- a/Engine/Plugins/Runtime/nDisplay/Source/DisplayCluster/Private/Game/EngineClasses/Scene/DisplayClusterRootActor.cpp
+++ b/Engine/Plugins/Runtime/nDisplay/Source/DisplayCluster/Private/Game/EngineClasses/Scene/DisplayClusterRootActor.cpp
@@ -98,6 +98,9 @@ ADisplayClusterRootActor::ADisplayClusterRootActor(const FObjectInitializer& Obj
 	bReplicates = false;
 	SpawnCollisionHandlingMethod = ESpawnActorCollisionHandlingMethod::AdjustIfPossibleButAlwaysSpawn;
 
+	// Enabled by default to avoid being active on Stage
+	SetActorHiddenInGame(true);
+
 #if WITH_EDITOR
 	Constructor_Editor();
 #endif
