nDisplay: hide preview meshes for DCRA rendering (prevent viewport preview mesh from rendering into itself)

diff --git a/Engine/Plugins/Runtime/nDisplay/Source/DisplayCluster/Private/Game/EngineClasses/Scene/DisplayClusterRootActor.cpp b/Engine/Plugins/Runtime/nDisplay/Source/DisplayCluster/Private/Game/EngineClasses/Scene/DisplayClusterRootActor.cpp
--- a/Engine/Plugins/Runtime/nDisplay/Source/DisplayCluster/Private/Game/EngineClasses/Scene/DisplayClusterRootActor.cpp
+++ b/Engine/Plugins/Runtime/nDisplay/Source/DisplayCluster/Private/Game/EngineClasses/Scene/DisplayClusterRootActor.cpp
@@ -551,6 +551,21 @@ bool ADisplayClusterRootActor::GetHiddenInGamePrimitives(TSet<FPrimitiveComponen
 	GetTypedPrimitives<UDisplayClusterScreenComponent>(OutPrimitives);
 
 #if WITH_EDITOR
+	// Always hide preview meshes for preview
+	{
+		TArray<UDisplayClusterPreviewComponent*> CurrentPreviewComponents;
+		GetComponents(CurrentPreviewComponents);
+		for (UDisplayClusterPreviewComponent* PreviewComponentIt : CurrentPreviewComponents)
+		{
+			if (UMeshComponent* PreviewMesh = PreviewComponentIt->GetPreviewMesh())
+			{
+				if(UPrimitiveComponent* PrimComp = Cast<UPrimitiveComponent>(PreviewMesh))
+				{
+					OutPrimitives.Add(PrimComp->ComponentId);
+				}
+			}
+		}
+	}
 
 	// Hide visualization and hidden components from RootActor
 	{
