nDisplay: Fixed nDisplay node preview ambiguity.

diff --git a/Engine/Plugins/Runtime/nDisplay/Source/DisplayCluster/Private/Game/DisplayClusterGameManager.cpp b/Engine/Plugins/Runtime/nDisplay/Source/DisplayCluster/Private/Game/DisplayClusterGameManager.cpp
--- a/Engine/Plugins/Runtime/nDisplay/Source/DisplayCluster/Private/Game/DisplayClusterGameManager.cpp
+++ b/Engine/Plugins/Runtime/nDisplay/Source/DisplayCluster/Private/Game/DisplayClusterGameManager.cpp
@@ -205,6 +205,53 @@ ADisplayClusterRootActor* FDisplayClusterGameManager::FindRootActor(UWorld* InWo
 		}
 	}
 
+#if WITH_EDITOR
+	if (InWorld->IsPlayInEditor())
+	{
+		/** We need to fix it in the following way :
+		 * - When multiple DCRA are available in the scene, and only one has Preview Node selected, then we should see the preview
+		 * - When multiple DCRA are available in the sceneand multiple Preview Node is selected, we should write a warning in the logand print the name of the selected DCRA for the preview.
+		 */
+		TArray<ADisplayClusterRootActor*> FoundActorsForPIE;
+		for (ADisplayClusterRootActor* RootActorIt : FoundActors)
+		{
+			if (RootActorIt && !RootActorIt->PreviewNodeId.IsEmpty())
+			{
+				FoundActorsForPIE.Add(RootActorIt);
+			}
+		}
+
+		if (FoundActorsForPIE.IsEmpty())
+		{
+			// No DCRA for PIE preview
+			return nullptr;
+		}
+
+		ADisplayClusterRootActor* ActiveRootActor = FoundActorsForPIE[0];
+
+		if (FoundActorsForPIE.Num() > 1)
+		{
+			FString PreviewActorsNames;
+			for (ADisplayClusterRootActor* RootActorIt : FoundActorsForPIE)
+			{
+				if (PreviewActorsNames.IsEmpty())
+				{
+					PreviewActorsNames.Appendf(TEXT("%s"), *RootActorIt->GetName());
+				}
+				else
+				{
+					PreviewActorsNames.Appendf(TEXT(", %s"), *RootActorIt->GetName());
+
+				}
+			}
+
+			UE_LOG(LogDisplayClusterGame, Warning, TEXT("Multiple PIE Preview Nodes selected at DisplayClusterRootActors - %s. Active : %s"), *PreviewActorsNames , *ActiveRootActor->GetName());
+		}
+
+		return ActiveRootActor;
+	}
+#endif
+
 	// Now iterate over all DCRA instances we have found to pick the one that corresponds to the config data
 	for (ADisplayClusterRootActor* Actor : FoundActors)
 	{
