nDisplay[OffscreenRendering]: Auto-detect which cameras to render offscreen

diff --git a/Engine/Plugins/Runtime/nDisplay/Source/DisplayCluster/Private/Game/EngineClasses/Scene/DisplayClusterICVFXCameraComponent.cpp b/Engine/Plugins/Runtime/nDisplay/Source/DisplayCluster/Private/Game/EngineClasses/Scene/DisplayClusterICVFXCameraComponent.cpp
--- a/Engine/Plugins/Runtime/nDisplay/Source/DisplayCluster/Private/Game/EngineClasses/Scene/DisplayClusterICVFXCameraComponent.cpp
+++ b/Engine/Plugins/Runtime/nDisplay/Source/DisplayCluster/Private/Game/EngineClasses/Scene/DisplayClusterICVFXCameraComponent.cpp
@@ -10,7 +10,10 @@
 #include "Components/DisplayClusterCameraComponent.h"
 #include "DisplayClusterRootActor.h"
 
+#include "DisplayClusterConfigurationTypes.h"
+
 #include "Misc/DisplayClusterGlobals.h"
+#include "Misc/Parse.h"
 #include "DisplayClusterEnums.h"
 
 
@@ -37,20 +40,42 @@ FString UDisplayClusterICVFXCameraComponent::GetCameraUniqueId() const
 
 bool UDisplayClusterICVFXCameraComponent::IsICVFXEnabled() const
 {
-	// In runtime, we also need to check if this camera is allowed to render on a specific node
-	static const bool bIsRunningCluster = (GDisplayCluster->GetOperationMode() == EDisplayClusterOperationMode::Cluster);
+	// When rendering offscreen, we have an extended logic for camera rendering activation
+	static const bool bIsRunningClusterModeOffscreen =
+		(GDisplayCluster->GetOperationMode() == EDisplayClusterOperationMode::Cluster) &&
+		FParse::Param(FCommandLine::Get(), TEXT("RenderOffscreen"));
 
-	if (bIsRunningCluster)
+	// If cluster mode + rendering offscreen, discover media output settings
+	if (bIsRunningClusterModeOffscreen)
 	{
+		// This cluster node ID
 		static const FString NodeId = GDisplayCluster->GetPrivateClusterMgr()->GetNodeId();
-		const bool bCameraIsAllowedOnThisNode = !CameraSettings.HiddenICVFXNodes.ItemNames.ContainsByPredicate([](const FString& Item)
+
+		// First condition to render offscreen: it has media output assigned
+		const bool bUsesMediaOutput = CameraSettings.RenderSettings.Media.bEnable ? !!CameraSettings.RenderSettings.Media.GetMediaOutput(NodeId) : false;
+
+		// Get backbuffer media settings
+		const FDisplayClusterConfigurationMedia* BackbufferMediaSettings = nullptr;
+		if (const ADisplayClusterRootActor* const RootActor = Cast<ADisplayClusterRootActor>(GetOwner()))
+		{
+			if (const UDisplayClusterConfigurationData* const ConfigData = RootActor->GetConfigData())
 			{
-				return Item.Equals(NodeId, ESearchCase::IgnoreCase);
-			});
+				if (const UDisplayClusterConfigurationClusterNode* const NodeCfg = ConfigData->Cluster->GetNode(NodeId))
+				{
+					BackbufferMediaSettings = &NodeCfg->Media;
+				}
+			}
+		}
+
+		// Second condition to render offscreen: the backbuffer has media output assigned.
+		// This means the whole frame including ICVFX cameras need to be rendered.
+		const bool bIsBackbufferBeingCaptured = BackbufferMediaSettings ? !!BackbufferMediaSettings->MediaOutput : false;
 
-		return CameraSettings.bEnable && bCameraIsAllowedOnThisNode;
+		// Finally make a decision if the camera should be rendered
+		return CameraSettings.bEnable && (bUsesMediaOutput || bIsBackbufferBeingCaptured);
 	}
 
+	// Otherwise the on/off condition only
 	return CameraSettings.bEnable;
 }
 
diff --git a/Engine/Plugins/Runtime/nDisplay/Source/DisplayCluster/Public/Components/DisplayClusterICVFXCameraComponent.h b/Engine/Plugins/Runtime/nDisplay/Source/DisplayCluster/Public/Components/DisplayClusterICVFXCameraComponent.h
--- a/Engine/Plugins/Runtime/nDisplay/Source/DisplayCluster/Public/Components/DisplayClusterICVFXCameraComponent.h
+++ b/Engine/Plugins/Runtime/nDisplay/Source/DisplayCluster/Public/Components/DisplayClusterICVFXCameraComponent.h
@@ -88,9 +88,6 @@ private:
 	UPROPERTY(EditAnywhere, Transient, Category = "In Camera VFX", meta = (PropertyPath = "CameraSettings.HiddenICVFXViewports"))
 	FDisplayClusterEditorPropertyReference HiddenICVFXViewportsRef;
 
-	UPROPERTY(EditAnywhere, Transient, Category = "In Camera VFX", meta = (PropertyPath = "CameraSettings.HiddenICVFXNodes"))
-	FDisplayClusterEditorPropertyReference HiddenICVFXNodesRef;
-
 	UPROPERTY(EditAnywhere, Transient, Category = "In Camera VFX", meta = (PropertyPath = "CameraSettings.ExternalCameraActor"))
 	FDisplayClusterEditorPropertyReference ExternalCameraActorRef;
 
diff --git a/Engine/Plugins/Runtime/nDisplay/Source/DisplayClusterConfiguration/Public/DisplayClusterConfigurationTypes_ICVFX.h b/Engine/Plugins/Runtime/nDisplay/Source/DisplayClusterConfiguration/Public/DisplayClusterConfigurationTypes_ICVFX.h
--- a/Engine/Plugins/Runtime/nDisplay/Source/DisplayClusterConfiguration/Public/DisplayClusterConfigurationTypes_ICVFX.h
+++ b/Engine/Plugins/Runtime/nDisplay/Source/DisplayClusterConfiguration/Public/DisplayClusterConfigurationTypes_ICVFX.h
@@ -629,10 +629,6 @@ public:
 	/** A list of viewports that the inner frustum is not rendered to. */
 	UPROPERTY(EditAnywhere, BlueprintReadWrite, Category = NDisplay, meta = (ClusterItemType = Viewports, DisplayName = "Inner Frustum Hidden in Viewports", EditCondition = "bEnable"))
 	FDisplayClusterConfigurationClusterItemReferenceList HiddenICVFXViewports;
-
-	/** A list of nodes where the inner frustum is not rendered. */
-	UPROPERTY(EditAnywhere, BlueprintReadWrite, Category = NDisplay, meta = (ClusterItemType = ClusterNodes, DisplayName = "Inner Frustum Hidden on Nodes", EditCondition = "bEnable"))
-	FDisplayClusterConfigurationClusterItemReferenceList HiddenICVFXNodes;
 };
 
 USTRUCT(Blueprintable)
