nDisplay: Fixed issue that was causing a crash when closing a SIE session due to the root actor cleaning up rendering resources during the normal GC period of the SIE world cleanup, which causes a null reference crash in FEdModeTools when the UV light card manager's preview world is destroyed.

diff --git a/Engine/Plugins/Runtime/nDisplay/Source/DisplayCluster/Private/Game/EngineClasses/Scene/DisplayClusterRootActor.cpp b/Engine/Plugins/Runtime/nDisplay/Source/DisplayCluster/Private/Game/EngineClasses/Scene/DisplayClusterRootActor.cpp
--- a/Engine/Plugins/Runtime/nDisplay/Source/DisplayCluster/Private/Game/EngineClasses/Scene/DisplayClusterRootActor.cpp
+++ b/Engine/Plugins/Runtime/nDisplay/Source/DisplayCluster/Private/Game/EngineClasses/Scene/DisplayClusterRootActor.cpp
@@ -1032,6 +1032,15 @@ void ADisplayClusterRootActor::PostActorCreated()
 	InitializeRootActor();
 }
 
+void ADisplayClusterRootActor::EndPlay(const EEndPlayReason::Type EndPlayReason)
+{
+#if WITH_EDITOR
+	EndPlay_Editor(EndPlayReason);
+#endif
+
+	Super::EndPlay(EndPlayReason);
+}
+
 void ADisplayClusterRootActor::Destroyed()
 {
 #if WITH_EDITOR
diff --git a/Engine/Plugins/Runtime/nDisplay/Source/DisplayCluster/Private/Game/EngineClasses/Scene/DisplayClusterRootActorEditor.cpp b/Engine/Plugins/Runtime/nDisplay/Source/DisplayCluster/Private/Game/EngineClasses/Scene/DisplayClusterRootActorEditor.cpp
--- a/Engine/Plugins/Runtime/nDisplay/Source/DisplayCluster/Private/Game/EngineClasses/Scene/DisplayClusterRootActorEditor.cpp
+++ b/Engine/Plugins/Runtime/nDisplay/Source/DisplayCluster/Private/Game/EngineClasses/Scene/DisplayClusterRootActorEditor.cpp
@@ -79,13 +79,6 @@ static FAutoConsoleVariableRef CVarDisplayClusterPreviewMultiGPURenderingMaxInde
 	ECVF_RenderThreadSafe
 );
 
-int32 GDisplayClusterPreviewEnableSIE = 0;
-static FAutoConsoleVariableRef CVarDisplayClusterPreviewEnableSIE(
-	TEXT("DC.Preview.EnableSIE"),
-	GDisplayClusterPreviewEnableSIE,
-	TEXT("Enable preview rendering during an SIE session"),
-	ECVF_RenderThreadSafe
-);
 //////////////////////////////////////////////////////////////////////////////////////////////
 // ADisplayClusterRootActor
 //////////////////////////////////////////////////////////////////////////////////////////////
@@ -131,8 +124,7 @@ void ADisplayClusterRootActor::Destructor_Editor()
 
 void ADisplayClusterRootActor::Tick_Editor(float DeltaSeconds)
 {
-	// Temporary fix for UE-174008, don't render preview when in SIE unless the cvar forces it to
-	if (IsPreviewEnabled() && (!GEditor->IsSimulatingInEditor() || GDisplayClusterPreviewEnableSIE))
+	if (IsPreviewEnabled())
 	{
 		if (bDeferPreviewGeneration)
 		{
@@ -178,6 +170,20 @@ void ADisplayClusterRootActor::PostLoad_Editor()
 	ResetPreviewInternals_Editor();
 }
 
+void ADisplayClusterRootActor::EndPlay_Editor(const EEndPlayReason::Type EndPlayReason)
+{
+	// UE-174008: When a SIE session is ended, we can't clean up the actor's rendering resources in BeginDestroy
+	// because of an issue cleaning up the UV light card manager's preview world during the normal GC pass,
+	// so in this specific case, destroy the rendering resources here.
+	if (EndPlayReason == EEndPlayReason::EndPlayInEditor && GEditor->bIsSimulatingInEditor)
+	{
+		ResetPreviewInternals_Editor();
+		ReleasePreviewComponents();
+
+		ViewportManager.Reset();
+	}
+}
+
 void ADisplayClusterRootActor::Destroyed_Editor()
 {
 	ResetPreviewInternals_Editor();
diff --git a/Engine/Plugins/Runtime/nDisplay/Source/DisplayCluster/Public/DisplayClusterRootActor.h b/Engine/Plugins/Runtime/nDisplay/Source/DisplayCluster/Public/DisplayClusterRootActor.h
--- a/Engine/Plugins/Runtime/nDisplay/Source/DisplayCluster/Public/DisplayClusterRootActor.h
+++ b/Engine/Plugins/Runtime/nDisplay/Source/DisplayCluster/Public/DisplayClusterRootActor.h
@@ -114,6 +114,7 @@ protected:
 	virtual void Tick(float DeltaSeconds) override;
 	virtual void PostLoad() override;
 	virtual void PostActorCreated() override;
+	virtual void EndPlay(const EEndPlayReason::Type EndPlayReason) override;
 	virtual void BeginDestroy() override;
 	virtual void Destroyed() override;
 #if WITH_EDITOR
@@ -441,6 +442,7 @@ public:
 	void Tick_Editor(float DeltaSeconds);
 	void PostLoad_Editor();
 	void PostActorCreated_Editor();
+	void EndPlay_Editor(const EEndPlayReason::Type EndPlayReason);
 	void BeginDestroy_Editor();
 	void Destroyed_Editor();
 	void RerunConstructionScripts_Editor();
