nDisplay: Fixed issues that were preventing light card flush placement from working correctly on display cluster screen components, as well as adding a higher vertex screen mesh for them that gives more accurate results when rendered in the Dome projection of the ICVFX panel.

diff --git a/Engine/Plugins/Runtime/nDisplay/Source/DisplayCluster/Private/Game/EngineClasses/Scene/DisplayClusterScreenComponent.cpp b/Engine/Plugins/Runtime/nDisplay/Source/DisplayCluster/Private/Game/EngineClasses/Scene/DisplayClusterScreenComponent.cpp
--- a/Engine/Plugins/Runtime/nDisplay/Source/DisplayCluster/Private/Game/EngineClasses/Scene/DisplayClusterScreenComponent.cpp
+++ b/Engine/Plugins/Runtime/nDisplay/Source/DisplayCluster/Private/Game/EngineClasses/Scene/DisplayClusterScreenComponent.cpp
@@ -14,7 +14,7 @@
 UDisplayClusterScreenComponent::UDisplayClusterScreenComponent(const FObjectInitializer& ObjectInitializer)
 	: Super(ObjectInitializer)
 {
-	static ConstructorHelpers::FObjectFinder<UStaticMesh> ScreenMesh(TEXT("/nDisplay/Meshes/plane_1x1"));
+	static ConstructorHelpers::FObjectFinder<UStaticMesh> ScreenMesh(TEXT("/nDisplay/Meshes/plane_hd_1x1"));
 	SetStaticMesh(ScreenMesh.Object);
 
 	SetMobility(EComponentMobility::Movable);
diff --git a/Engine/Plugins/Runtime/nDisplay/Source/DisplayCluster/Private/Game/EngineClasses/Scene/DisplayClusterStageGeometryComponent.cpp b/Engine/Plugins/Runtime/nDisplay/Source/DisplayCluster/Private/Game/EngineClasses/Scene/DisplayClusterStageGeometryComponent.cpp
--- a/Engine/Plugins/Runtime/nDisplay/Source/DisplayCluster/Private/Game/EngineClasses/Scene/DisplayClusterStageGeometryComponent.cpp
+++ b/Engine/Plugins/Runtime/nDisplay/Source/DisplayCluster/Private/Game/EngineClasses/Scene/DisplayClusterStageGeometryComponent.cpp
@@ -255,8 +255,22 @@ void UDisplayClusterStageGeometryComponent::UpdateStageGeometry()
 			return bIsProjectionMesh || bIsScreen;
 		});
 
-		const FBox BoundingBox = RootActor->GetComponentsBoundingBox().ShiftBy(-RootActor->GetActorLocation());
-		StageBoundingRadius = FMath::Max(BoundingBox.Min.Length(), BoundingBox.Max.Length());
+		// Get the bounding box for all of the stage's screens
+		FBox StageBoundingBox(ForceInit);
+		RootActor->ForEachComponent<UPrimitiveComponent>(false, [&StageBoundingBox, &ProjectionMeshNames](const UPrimitiveComponent* PrimitiveComponent)
+		{
+			// Filter out any primitive component that isn't a projection mesh (a static mesh that has a Mesh projection configured for it) or a screen component
+			const bool bIsProjectionMesh = PrimitiveComponent->IsA<UStaticMeshComponent>() && ProjectionMeshNames.Contains(PrimitiveComponent->GetName());
+			const bool bIsScreen = PrimitiveComponent->IsA<UDisplayClusterScreenComponent>();
+
+			if (PrimitiveComponent->IsRegistered() && (bIsProjectionMesh || bIsScreen))
+			{
+				StageBoundingBox += PrimitiveComponent->Bounds.GetBox();
+			}
+		});
+
+		StageBoundingBox = StageBoundingBox.ShiftBy(-RootActor->GetActorLocation());
+		StageBoundingRadius = FMath::Max(StageBoundingBox.Min.Length(), StageBoundingBox.Max.Length());
 	}
 }
 
diff --git a/Engine/Plugins/Runtime/nDisplay/Source/DisplayClusterLightCardEditor/Private/Viewport/DisplayClusterLightCardEditorViewportClient.cpp b/Engine/Plugins/Runtime/nDisplay/Source/DisplayClusterLightCardEditor/Private/Viewport/DisplayClusterLightCardEditorViewportClient.cpp
--- a/Engine/Plugins/Runtime/nDisplay/Source/DisplayClusterLightCardEditor/Private/Viewport/DisplayClusterLightCardEditorViewportClient.cpp
+++ b/Engine/Plugins/Runtime/nDisplay/Source/DisplayClusterLightCardEditor/Private/Viewport/DisplayClusterLightCardEditorViewportClient.cpp
@@ -71,7 +71,13 @@ static TAutoConsoleVariable<bool> CVarICVFXPanelAutoPan(
 	true,
 	TEXT("When true, the stage view will automatically tilt and pan on supported map projections as you drag objects near the edges of the screen."));
 
-
+int32 GICVFXPanelDisplayNormalMapVisualization = 0;
+static FAutoConsoleVariableRef CVarICVFXPanelDisplayNormalMapVisualization(
+	TEXT("nDisplay.panel.DisplayNormalMapVisualization"),
+	GICVFXPanelDisplayNormalMapVisualization,
+	TEXT("Displays the normal map visualization of the stage's geometry map, which is used to determine physical stage geometry"),
+	ECVF_RenderThreadSafe
+);
 //////////////////////////////////////////////////////////////////////////
 // FDisplayClusterLightCardEditorViewportClient
 
@@ -368,7 +374,7 @@ void FDisplayClusterLightCardEditorViewportClient::Draw(FViewport* InViewport, F
 		DrawCanvas(*Viewport, *View, *Canvas);
 	}
 
-	if (bDisplayNormalMapVisualization)
+	if (bDisplayNormalMapVisualization || GICVFXPanelDisplayNormalMapVisualization)
 	{
 		auto DrawNormalMap = [Canvas, this](bool bShowNorthMap, FVector2D Position)
 		{
@@ -1048,6 +1054,7 @@ void FDisplayClusterLightCardEditorViewportClient::CreateDrawnLightCard(const TA
 	FScopedTransaction Transaction(LOCTEXT("AddNewLightCard", "Add New Light Card"));
 
 	ADisplayClusterLightCardActor* LightCard = LightCardEditorPtr.Pin()->SpawnActorAs<ADisplayClusterLightCardActor>(TEXT("LightCard"));
+	ProjectionHelper->VerifyAndFixActorOrigin(LightCard);
 
 	if (!LightCard)
 	{
diff --git a/Engine/Plugins/Runtime/nDisplay/Source/DisplayClusterScenePreview/Private/DisplayClusterLightCardEditorHelper.cpp b/Engine/Plugins/Runtime/nDisplay/Source/DisplayClusterScenePreview/Private/DisplayClusterLightCardEditorHelper.cpp
--- a/Engine/Plugins/Runtime/nDisplay/Source/DisplayClusterScenePreview/Private/DisplayClusterLightCardEditorHelper.cpp
+++ b/Engine/Plugins/Runtime/nDisplay/Source/DisplayClusterScenePreview/Private/DisplayClusterLightCardEditorHelper.cpp
@@ -976,8 +976,7 @@ ADisplayClusterRootActor* FDisplayClusterLightCardEditorHelper::UpdateRootActor(
 
 	if (NewRootActor)
 	{
-		const FBox BoundingBox = NewRootActor->GetComponentsBoundingBox();
-		RootActorBoundingRadius = FMath::Max(BoundingBox.Min.Length(), BoundingBox.Max.Length());
+		RootActorBoundingRadius = NewRootActor->GetStageGeometryComponent()->GetStageBoundingRadius();
 
 #if WITH_EDITOR
 		FCoreUObjectDelegates::OnObjectPropertyChanged.AddRaw(this, &FDisplayClusterLightCardEditorHelper::OnActorPropertyChanged);
