nDisplay: Added an additional rule for an empty preview node.

diff --git a/Engine/Plugins/Runtime/nDisplay/Source/DisplayCluster/Private/Game/DisplayClusterGameManager.cpp b/Engine/Plugins/Runtime/nDisplay/Source/DisplayCluster/Private/Game/DisplayClusterGameManager.cpp
--- a/Engine/Plugins/Runtime/nDisplay/Source/DisplayCluster/Private/Game/DisplayClusterGameManager.cpp
+++ b/Engine/Plugins/Runtime/nDisplay/Source/DisplayCluster/Private/Game/DisplayClusterGameManager.cpp
@@ -215,7 +215,7 @@ ADisplayClusterRootActor* FDisplayClusterGameManager::FindRootActor(UWorld* InWo
 		TArray<ADisplayClusterRootActor*> FoundActorsForPIE;
 		for (ADisplayClusterRootActor* RootActorIt : FoundActors)
 		{
-			if (RootActorIt && !RootActorIt->PreviewNodeId.IsEmpty())
+			if (RootActorIt && !(RootActorIt->PreviewNodeId.IsEmpty() || RootActorIt->PreviewNodeId == DisplayClusterConfigurationStrings::gui::preview::PreviewNodeNone))
 			{
 				FoundActorsForPIE.Add(RootActorIt);
 			}
