nDisplay: Fixed issue causing the editor to crash when opening any level with an nDisplay root actor. Crash was caused by a missing shader parameter for the Scene uniform buffer in the mesh projection renderer.

diff --git a/Engine/Plugins/Runtime/nDisplay/Source/DisplayClusterLightCardEditorShaders/Private/DisplayClusterMeshProjectionRenderer.cpp b/Engine/Plugins/Runtime/nDisplay/Source/DisplayClusterLightCardEditorShaders/Private/DisplayClusterMeshProjectionRenderer.cpp
--- a/Engine/Plugins/Runtime/nDisplay/Source/DisplayClusterLightCardEditorShaders/Private/DisplayClusterMeshProjectionRenderer.cpp
+++ b/Engine/Plugins/Runtime/nDisplay/Source/DisplayClusterLightCardEditorShaders/Private/DisplayClusterMeshProjectionRenderer.cpp
@@ -11,6 +11,7 @@
 #include "CanvasTypes.h"
 #include "EngineModule.h"
 #include "SceneManagement.h"
+#include "SceneUniformBuffer.h"
 #include "SceneViewExtension.h"
 #include "ScreenPass.h"
 #include "Components/PrimitiveComponent.h"
@@ -25,6 +26,7 @@
 
 BEGIN_SHADER_PARAMETER_STRUCT(FMeshProjectionPassParameters, )
 	SHADER_PARAMETER_STRUCT_REF(FViewUniformShaderParameters, View)
+	SHADER_PARAMETER_RDG_UNIFORM_BUFFER(FSceneUniformParameters, Scene)
 	SHADER_PARAMETER_RDG_UNIFORM_BUFFER(FInstanceCullingGlobalUniforms, InstanceCulling)
 	SHADER_PARAMETER_STRUCT_INCLUDE(FInstanceCullingDrawParams, InstanceCullingDrawParams)
 	RENDER_TARGET_BINDING_SLOTS()
@@ -275,6 +277,7 @@ FMeshProjectionPassParameters* CreateProjectionPassParameters(FRDGBuilder& Graph
 {
 	FMeshProjectionPassParameters* PassParameters = GraphBuilder.AllocParameters<FMeshProjectionPassParameters>();
 	PassParameters->View = View->ViewUniformBuffer;
+	PassParameters->Scene = View->GetSceneUniforms().GetBuffer(GraphBuilder);
 	PassParameters->InstanceCulling = FInstanceCullingContext::CreateDummyInstanceCullingUniformBuffer(GraphBuilder);
 
 	return PassParameters;
