nDisplay: Followup to 24058738. Fixed issue causing the editor to crash when opening any level with an nDisplay root actor and turning on its preview in Editor.
Crash was caused by a missing shader parameter for the Scene uniform buffer in the mesh projection renderer.

diff --git a/Engine/Plugins/Runtime/nDisplay/Source/DisplayClusterShaders/Private/Shaders/DisplayClusterShadersPreprocess_UVLightCards.cpp b/Engine/Plugins/Runtime/nDisplay/Source/DisplayClusterShaders/Private/Shaders/DisplayClusterShadersPreprocess_UVLightCards.cpp
--- a/Engine/Plugins/Runtime/nDisplay/Source/DisplayClusterShaders/Private/Shaders/DisplayClusterShadersPreprocess_UVLightCards.cpp
+++ b/Engine/Plugins/Runtime/nDisplay/Source/DisplayClusterShaders/Private/Shaders/DisplayClusterShadersPreprocess_UVLightCards.cpp
@@ -15,6 +15,7 @@
 #include "RenderGraphResources.h"
 #include "ScenePrivate.h"
 #include "SceneRenderTargetParameters.h"
+#include "SceneUniformBuffer.h"
 #include "MeshPassProcessor.inl"
 #include "UnrealClient.h"
 
@@ -164,6 +165,7 @@ private:
 
 BEGIN_SHADER_PARAMETER_STRUCT(FUVLightCardPassParameters, )
 	SHADER_PARAMETER_STRUCT_REF(FViewUniformShaderParameters, View)
+	SHADER_PARAMETER_RDG_UNIFORM_BUFFER(FSceneUniformParameters, Scene)
 	SHADER_PARAMETER_RDG_UNIFORM_BUFFER(FInstanceCullingGlobalUniforms, InstanceCulling)
 	SHADER_PARAMETER_STRUCT_INCLUDE(FInstanceCullingDrawParams, InstanceCullingDrawParams)
 	RENDER_TARGET_BINDING_SLOTS()
@@ -276,6 +278,7 @@ bool FDisplayClusterShadersPreprocess_UVLightCards::RenderPreprocess_UVLightCard
 
 	FUVLightCardPassParameters* PassParameters = GraphBuilder.AllocParameters<FUVLightCardPassParameters>();
 	PassParameters->View = View->ViewUniformBuffer;
+	PassParameters->Scene = View->GetSceneUniforms().GetBuffer(GraphBuilder);
 	PassParameters->InstanceCulling = FInstanceCullingContext::CreateDummyInstanceCullingUniformBuffer(GraphBuilder);
 	PassParameters->RenderTargets[0] = OutputRenderTargetBinding;
 
