nDisplay: Disable UV chromakey card creation.

diff --git a/Engine/Plugins/Runtime/nDisplay/Source/DisplayClusterLightCardEditor/Private/DisplayClusterLightCardEditor.cpp b/Engine/Plugins/Runtime/nDisplay/Source/DisplayClusterLightCardEditor/Private/DisplayClusterLightCardEditor.cpp
--- a/Engine/Plugins/Runtime/nDisplay/Source/DisplayClusterLightCardEditor/Private/DisplayClusterLightCardEditor.cpp
+++ b/Engine/Plugins/Runtime/nDisplay/Source/DisplayClusterLightCardEditor/Private/DisplayClusterLightCardEditor.cpp
@@ -490,7 +490,20 @@ void FDisplayClusterLightCardEditor::AddLightCardsToActor(const TArray<ADisplayC
 
 bool FDisplayClusterLightCardEditor::CanAddNewActor() const
 {
-	return ActiveRootActor.IsValid() && ActiveRootActor->GetWorld() != nullptr;
+	return CanAddNewActor(nullptr);
+}
+
+bool FDisplayClusterLightCardEditor::CanAddNewActor(UClass* InClass) const
+{
+	bool bClassAllowed = true;
+	if (InClass && InClass->IsChildOf(ADisplayClusterChromakeyCardActor::StaticClass()))
+	{
+		// @todo uv chromakey cards
+		bClassAllowed = !ViewportView.IsValid() ||
+			ViewportView->GetLightCardEditorViewportClient()->GetProjectionMode() != EDisplayClusterMeshProjectionType::UV;
+	}
+
+	return bClassAllowed && ActiveRootActor.IsValid() && ActiveRootActor->GetWorld() != nullptr;
 }
 
 void FDisplayClusterLightCardEditor::CutSelectedActors()
@@ -1112,7 +1125,7 @@ TSharedRef<SWidget> FDisplayClusterLightCardEditor::GeneratePlaceActorsMenu()
 			FSlateIcon StageActorIcon = FSlateIconFinder::FindIconForClass(Class);
 			MenuBuilder.AddMenuEntry(Label, LOCTEXT("AddStageActorHeader", "Add a stage actor to the scene"), StageActorIcon,
 				FUIAction(FExecuteAction::CreateSP(this, &FDisplayClusterLightCardEditor::AddNewDynamic, Class),
-					FCanExecuteAction::CreateSP(this, &FDisplayClusterLightCardEditor::CanAddNewActor)));
+					FCanExecuteAction::CreateSP(this, &FDisplayClusterLightCardEditor::CanAddNewActor, Class)));
 		}
 
 		if (CanAddNewActor())
diff --git a/Engine/Plugins/Runtime/nDisplay/Source/DisplayClusterLightCardEditor/Private/DisplayClusterLightCardEditor.h b/Engine/Plugins/Runtime/nDisplay/Source/DisplayClusterLightCardEditor/Private/DisplayClusterLightCardEditor.h
--- a/Engine/Plugins/Runtime/nDisplay/Source/DisplayClusterLightCardEditor/Private/DisplayClusterLightCardEditor.h
+++ b/Engine/Plugins/Runtime/nDisplay/Source/DisplayClusterLightCardEditor/Private/DisplayClusterLightCardEditor.h
@@ -118,6 +118,7 @@ public:
 
 	/** If a Light Card can currently be added */
 	bool CanAddNewActor() const;
+	bool CanAddNewActor(UClass* InClass) const;
 
 	/** Copies any selected actors to the clipboard, and then deletes them */
 	void CutSelectedActors();
diff --git a/Engine/Plugins/Runtime/nDisplay/Source/DisplayClusterLightCardEditor/Private/Viewport/DisplayClusterLightcardEditorViewport.cpp b/Engine/Plugins/Runtime/nDisplay/Source/DisplayClusterLightCardEditor/Private/Viewport/DisplayClusterLightcardEditorViewport.cpp
--- a/Engine/Plugins/Runtime/nDisplay/Source/DisplayClusterLightCardEditor/Private/Viewport/DisplayClusterLightcardEditorViewport.cpp
+++ b/Engine/Plugins/Runtime/nDisplay/Source/DisplayClusterLightCardEditor/Private/Viewport/DisplayClusterLightcardEditorViewport.cpp
@@ -760,7 +760,7 @@ void SDisplayClusterLightCardEditorViewport::MakePlaceActorsSubMenu(FMenuBuilder
 			LOCTEXT("AddStageActorHeader", "Add a stage actor to the scene"), 
 			StageActorIcon,
 			FUIAction(FExecuteAction::CreateSP(this, &SDisplayClusterLightCardEditorViewport::AddStageActorHere, Class),
-				FCanExecuteAction::CreateSP(this, &SDisplayClusterLightCardEditorViewport::CanPlaceActorHere)));
+				FCanExecuteAction::CreateSP(this, &SDisplayClusterLightCardEditorViewport::CanPlaceActorHere, Class)));
 	}
 }
 
@@ -929,7 +929,12 @@ void SDisplayClusterLightCardEditorViewport::AddStageActorHere(UClass* InClass)
 
 bool SDisplayClusterLightCardEditorViewport::CanPlaceActorHere() const
 {
-	return LightCardEditorPtr.IsValid() && LightCardEditorPtr.Pin()->CanAddNewActor();
+	return CanPlaceActorHere(nullptr);
+}
+
+bool SDisplayClusterLightCardEditorViewport::CanPlaceActorHere(UClass* Class) const
+{
+	return LightCardEditorPtr.IsValid() && LightCardEditorPtr.Pin()->CanAddNewActor(Class);
 }
 
 void SDisplayClusterLightCardEditorViewport::ToggleLabels()
diff --git a/Engine/Plugins/Runtime/nDisplay/Source/DisplayClusterLightCardEditor/Private/Viewport/DisplayClusterLightcardEditorViewport.h b/Engine/Plugins/Runtime/nDisplay/Source/DisplayClusterLightCardEditor/Private/Viewport/DisplayClusterLightcardEditorViewport.h
--- a/Engine/Plugins/Runtime/nDisplay/Source/DisplayClusterLightCardEditor/Private/Viewport/DisplayClusterLightcardEditorViewport.h
+++ b/Engine/Plugins/Runtime/nDisplay/Source/DisplayClusterLightCardEditor/Private/Viewport/DisplayClusterLightcardEditorViewport.h
@@ -112,6 +112,7 @@ private:
 
 	/** If placing actors is allowed */
 	bool CanPlaceActorHere() const;
+	bool CanPlaceActorHere(UClass* Class) const;
 
 	/** Toggles light card labels on or off */
 	void ToggleLabels();
