nDisplay: The UE memory manager may reallocate the elements of a TArray at run time. Reworked previous code with saved ptrs on array elements. Replaced by iterators with predicates.

diff --git a/Engine/Plugins/Runtime/nDisplay/Source/DisplayClusterProjection/Private/Policy/MPCDI/DisplayClusterProjectionMPCDIPolicy.cpp b/Engine/Plugins/Runtime/nDisplay/Source/DisplayClusterProjection/Private/Policy/MPCDI/DisplayClusterProjectionMPCDIPolicy.cpp
--- a/Engine/Plugins/Runtime/nDisplay/Source/DisplayClusterProjection/Private/Policy/MPCDI/DisplayClusterProjectionMPCDIPolicy.cpp
+++ b/Engine/Plugins/Runtime/nDisplay/Source/DisplayClusterProjection/Private/Policy/MPCDI/DisplayClusterProjectionMPCDIPolicy.cpp
@@ -261,96 +261,88 @@ void FDisplayClusterProjectionMPCDIPolicy::ApplyWarpBlend_RenderThread(FRHIComma
 	const FDisplayClusterViewport_RenderSettingsICVFX& SettingsICVFX = InViewportProxy->GetRenderSettingsICVFX_RenderThread();
 	if (EnumHasAllFlags(SettingsICVFX.RuntimeFlags, EDisplayClusterViewportRuntimeICVFXFlags::Target))
 	{
-		// This viewport used as icvfx target:
-		TArray<FDisplayClusterShaderParametersICVFX_ViewportResource*> ViewportResources;
+		// This viewport used as the target of ICVFX.
 		FDisplayClusterShaderParameters_ICVFX ShaderICVFX(SettingsICVFX.ICVFX);
 
-		const IDisplayClusterViewportManagerProxy& RefViewportManagerProxy = InViewportProxy->GetOwner_RenderThread();
-
-		// Collect refs
-		ShaderICVFX.CollectRefViewports(ViewportResources);
+		// Iterate over all warped viewport contexts (support for stereo views).
+		for (int32 ContextNum = 0; ContextNum < InputTextures.Num(); ContextNum++)
 		{
-			// ICVFX warp:
-			for (int32 ContextNum = 0; ContextNum < InputTextures.Num(); ContextNum++)
+			// Update referenced resources and math in the ShaderICVFX.
+			ShaderICVFX.IterateViewportResourcesByPredicate([ContextNum, InViewportProxy, &ShaderICVFX](FDisplayClusterShaderParametersICVFX_ViewportResource& ViewportResourceIt)
 			{
-				// Update ref viewport resources:
-				for (FDisplayClusterShaderParametersICVFX_ViewportResource* It : ViewportResources)
-				{
-					// reset prev resource ref
-					It->Texture = nullptr;
+				// reset prev resource reference
+				ViewportResourceIt.Texture = nullptr;
 
-					const IDisplayClusterViewportProxy* RefViewportProxy = RefViewportManagerProxy.FindViewport_RenderThread(It->ViewportId);
-					if (RefViewportProxy)
+				if(const IDisplayClusterViewportProxy* SrcViewportProxy = InViewportProxy->GetOwner_RenderThread().FindViewport_RenderThread(ViewportResourceIt.ViewportId))
+				{
+					if (!SrcViewportProxy->GetRenderSettings_RenderThread().bSkipRendering)
 					{
-						const FDisplayClusterViewport_RenderSettings& RefViewportRenderSettings = RefViewportProxy->GetRenderSettings_RenderThread();
-						if (RefViewportRenderSettings.bSkipRendering == false)
+						TArray<FRHITexture2D*> RefTextures;
+						// Use for input first MipsShader texture if enabled in viewport render settings
+						if (SrcViewportProxy->GetResources_RenderThread(EDisplayClusterViewportResourceType::MipsShaderResource, RefTextures) ||
+							SrcViewportProxy->GetResources_RenderThread(EDisplayClusterViewportResourceType::InputShaderResource, RefTextures))
 						{
-							TArray<FRHITexture2D*> RefTextures;
-							// Use for input first MipsShader texture if enabled in viewport render settings
-							if (RefViewportProxy->GetResources_RenderThread(EDisplayClusterViewportResourceType::MipsShaderResource, RefTextures) ||
-								RefViewportProxy->GetResources_RenderThread(EDisplayClusterViewportResourceType::InputShaderResource, RefTextures))
+							if (ContextNum < RefTextures.Num())
 							{
-								if (ContextNum < RefTextures.Num())
-								{
-									It->Texture = RefTextures[ContextNum];
-								}
+								ViewportResourceIt.Texture = RefTextures[ContextNum];
 							}
 						}
+					}
 
-						const TArray<FDisplayClusterViewport_Context>& RefViewportContexts = RefViewportProxy->GetContexts_RenderThread();
-
+					// Override camera context from referenced viewport:
+					if(FDisplayClusterShaderParameters_ICVFX::FCameraSettings* CameraSettings = ShaderICVFX.FindCameraByName(ViewportResourceIt.ViewportId))
+					{
 						// Support stereo icvfx
-						for (FDisplayClusterShaderParameters_ICVFX::FCameraSettings& CameraIt : ShaderICVFX.Cameras)
+						const TArray<FDisplayClusterViewport_Context>& SrcViewportContexts = SrcViewportProxy->GetContexts_RenderThread();
+						if (SrcViewportContexts.IsValidIndex(ContextNum))
 						{
-							if (It->ViewportId == CameraIt.Resource.ViewportId && ContextNum < RefViewportContexts.Num())
-							{
-								// matched camera reference, update context for current eye
-								const FDisplayClusterViewport_Context& InContext = RefViewportContexts[ContextNum];
+							// matched camera reference, update context for current eye
+							const FDisplayClusterViewport_Context& InContext = SrcViewportContexts[ContextNum];
 
-								FDisplayClusterShaderParametersICVFX_CameraContext CameraContext;
-								CameraContext.CameraViewLocation = InContext.ViewLocation;
-								CameraContext.CameraViewRotation = InContext.ViewRotation;
-								CameraContext.CameraPrjMatrix = InContext.ProjectionMatrix;
+							FDisplayClusterShaderParametersICVFX_CameraContext CameraContext;
+							CameraContext.CameraViewLocation = InContext.ViewLocation;
+							CameraContext.CameraViewRotation = InContext.ViewRotation;
+							CameraContext.CameraPrjMatrix = InContext.ProjectionMatrix;
 
-								CameraIt.UpdateCameraContext(CameraContext);
-							}
+							CameraSettings->UpdateCameraContext(CameraContext);
 						}
 					}
 				}
+			});
 
-				// cleanup camera list for render
-				ShaderICVFX.CleanupCamerasForRender();
+			// cleanup camera list for render
+			ShaderICVFX.CleanupCamerasForRender();
 
-				// Initialize shader input data
-				FDisplayClusterShaderParameters_WarpBlend WarpBlendParameters;
+			// Initialize shader input data
+			FDisplayClusterShaderParameters_WarpBlend WarpBlendParameters;
 
-				WarpBlendParameters.Context = WarpBlendContexts_Proxy[ContextNum];
-				WarpBlendParameters.WarpInterface = WarpBlendInterface_Proxy;
+			WarpBlendParameters.Context = WarpBlendContexts_Proxy[ContextNum];
+			WarpBlendParameters.WarpInterface = WarpBlendInterface_Proxy;
 
-				WarpBlendParameters.Src.Set(InputTextures[ContextNum], InputRects[ContextNum]);
-				WarpBlendParameters.Dest.Set(OutputTextures[ContextNum], OutputRects[ContextNum]);
+			WarpBlendParameters.Src.Set(InputTextures[ContextNum], InputRects[ContextNum]);
+			WarpBlendParameters.Dest.Set(OutputTextures[ContextNum], OutputRects[ContextNum]);
 
-				WarpBlendParameters.bRenderAlphaChannel = InViewportProxy->GetRenderSettings_RenderThread().bWarpBlendRenderAlphaChannel;
+			WarpBlendParameters.bRenderAlphaChannel = InViewportProxy->GetRenderSettings_RenderThread().bWarpBlendRenderAlphaChannel;
 
-				// Before starting the whole ICVFX shaders pipeline, we need to provide ability to modify any ICVFX shader parameters.
-				// One of use cases is the latency queue that needs to substitute shader parameters.
-				IDisplayCluster::Get().GetCallbacks().OnDisplayClusterPreProcessIcvfx_RenderThread().Broadcast(RHICmdList, InViewportProxy, WarpBlendParameters, ShaderICVFX);
+			// Before starting the whole ICVFX shaders pipeline, we need to provide ability to modify any ICVFX shader parameters.
+			// One of use cases is the latency queue that needs to substitute shader parameters.
+			IDisplayCluster::Get().GetCallbacks().OnDisplayClusterPreProcessIcvfx_RenderThread().Broadcast(RHICmdList, InViewportProxy, WarpBlendParameters, ShaderICVFX);
 
-				// Start ICVFX pipeline
-				if (!ShadersAPI.RenderWarpBlend_ICVFX(RHICmdList, WarpBlendParameters, ShaderICVFX))
+			// Start ICVFX pipeline
+			if (!ShadersAPI.RenderWarpBlend_ICVFX(RHICmdList, WarpBlendParameters, ShaderICVFX))
+			{
+				if (!IsEditorOperationMode_RenderThread(InViewportProxy))
 				{
-					if (!IsEditorOperationMode_RenderThread(InViewportProxy))
-					{
-						UE_LOG(LogDisplayClusterProjectionMPCDI, Warning, TEXT("Couldn't apply icvfx warp&blend"));
-					}
-
-					return;
+					UE_LOG(LogDisplayClusterProjectionMPCDI, Warning, TEXT("Couldn't apply icvfx warp&blend"));
 				}
-			}
 
-			// Finish ICVFX warp
-			return;
+				// Break iteration of target viewport contexts
+				break;
+			}
 		}
+
+		// ICVFX warpblend ready
+		return;
 	}
 	
 	// Mesh warp:
diff --git a/Engine/Plugins/Runtime/nDisplay/Source/DisplayClusterShaders/Public/ShaderParameters/DisplayClusterShaderParameters_ICVFX.h b/Engine/Plugins/Runtime/nDisplay/Source/DisplayClusterShaders/Public/ShaderParameters/DisplayClusterShaderParameters_ICVFX.h
--- a/Engine/Plugins/Runtime/nDisplay/Source/DisplayClusterShaders/Public/ShaderParameters/DisplayClusterShaderParameters_ICVFX.h
+++ b/Engine/Plugins/Runtime/nDisplay/Source/DisplayClusterShaders/Public/ShaderParameters/DisplayClusterShaderParameters_ICVFX.h
@@ -134,24 +134,6 @@ public:
 		LightcardMode = InParameters.LightcardMode;
 	}
 
-	void CollectRefViewports(TArray<FDisplayClusterShaderParametersICVFX_ViewportResource*>& Dst)
-	{
-		if (Lightcard.IsDefined())
-		{
-			Dst.Add(&Lightcard);
-		}
-
-		if (UVLightcard.IsDefined())
-		{
-			Dst.Add(&UVLightcard);
-		}
-
-		for (FCameraSettings& CameraIt : Cameras)
-		{
-			CameraIt.CollectRefViewports(Dst);
-		}
-	}
-
 	inline void SortCamerasRenderOrder()
 	{
 		Cameras.Sort([](const FCameraSettings& It1, const FCameraSettings& It2)
@@ -168,19 +150,6 @@ public:
 	// ICVFX Target only data
 	struct FCameraSettings
 	{
-		void CollectRefViewports(TArray<FDisplayClusterShaderParametersICVFX_ViewportResource*>& Dst)
-		{
-			if (Resource.IsDefined())
-			{
-				Dst.Add(&Resource);
-			}
-
-			if (Chromakey.IsDefined())
-			{
-				Dst.Add(&Chromakey);
-			}
-		}
-
 		bool IsUsed() const
 		{
 			return (ChromakeySource == EDisplayClusterShaderParametersICVFX_ChromakeySource::FrameColor) || Resource.IsValid();
@@ -199,6 +168,23 @@ public:
 			CameraPrjMatrix = InContext.CameraPrjMatrix;
 		}
 
+		/**
+		 * Iterate over all defined viewport resources with a predicate functor.
+		 */
+		template <typename Predicate>
+		void IterateViewportResourcesByPredicate(Predicate Pred)
+		{
+			if (Resource.IsDefined())
+			{
+				::Invoke(Pred, Resource);
+			}
+
+			if (Chromakey.IsDefined())
+			{
+				::Invoke(Pred, Chromakey);
+			}
+		}
+
 	public:
 		FDisplayClusterShaderParametersICVFX_ViewportResource Resource;
 		FVector4 SoftEdge;
@@ -230,6 +216,39 @@ public:
 		int32 RenderOrder = -1;
 	};
 
+	/**
+	 * Iterate over all defined viewport resources with a predicate functor.
+	 */
+	template <typename Predicate>
+	void IterateViewportResourcesByPredicate(Predicate Pred)
+	{
+		if (Lightcard.IsDefined())
+		{
+			::Invoke(Pred, Lightcard);
+		}
+
+		if (UVLightcard.IsDefined())
+		{
+			::Invoke(Pred, UVLightcard);
+		}
+
+		for (FCameraSettings& CameraIt : Cameras)
+		{
+			CameraIt.IterateViewportResourcesByPredicate(Pred);
+		}
+	}
+
+	/**
+	 * Find camera settings by viewport name
+	 */
+	inline FCameraSettings* FindCameraByName(const FString& InViewportId)
+	{
+		return Cameras.FindByPredicate([InViewportId](const FCameraSettings& CameraIt)
+			{
+				return CameraIt.Resource.IsValid() && CameraIt.Resource.ViewportId == InViewportId;
+			});
+	}
+
 	// Remove unused cameras from render
 	bool CleanupCamerasForRender()
 	{
