nDisplay: GPU memory optimization for DCRA preview:
   - Added preview components resource deallocation calls.
   - Added deallocation of DCRA render resources when preview is unchecked.
   - Added CVars to allow using ViewState for preview in scene and configurator: "nDisplay.preview.EnableViewState",  "nDisplay.preview.EnableConfiguratorViewState".

diff --git a/Engine/Plugins/Runtime/nDisplay/Source/DisplayCluster/Private/Game/EngineClasses/Scene/DisplayClusterPreviewComponent.cpp b/Engine/Plugins/Runtime/nDisplay/Source/DisplayCluster/Private/Game/EngineClasses/Scene/DisplayClusterPreviewComponent.cpp
--- a/Engine/Plugins/Runtime/nDisplay/Source/DisplayCluster/Private/Game/EngineClasses/Scene/DisplayClusterPreviewComponent.cpp
+++ b/Engine/Plugins/Runtime/nDisplay/Source/DisplayCluster/Private/Game/EngineClasses/Scene/DisplayClusterPreviewComponent.cpp
@@ -115,6 +115,9 @@ void UDisplayClusterPreviewComponent::RestorePreviewMeshMaterial()
 		PreviewMesh->SetMaterial(0, OriginalMaterial);
 		OriginalMaterial = nullptr;
 	}
+
+	// Release RTTs
+	ReleasePreviewRenderTarget();
 }
 
 void UDisplayClusterPreviewComponent::SetPreviewMeshMaterial()
diff --git a/Engine/Plugins/Runtime/nDisplay/Source/DisplayCluster/Private/Game/EngineClasses/Scene/DisplayClusterRootActorEditor.cpp b/Engine/Plugins/Runtime/nDisplay/Source/DisplayCluster/Private/Game/EngineClasses/Scene/DisplayClusterRootActorEditor.cpp
--- a/Engine/Plugins/Runtime/nDisplay/Source/DisplayCluster/Private/Game/EngineClasses/Scene/DisplayClusterRootActorEditor.cpp
+++ b/Engine/Plugins/Runtime/nDisplay/Source/DisplayCluster/Private/Game/EngineClasses/Scene/DisplayClusterRootActorEditor.cpp
@@ -36,6 +36,7 @@
 #include "Render/Viewport/IDisplayClusterViewport.h"
 #include "Render/Viewport/IDisplayClusterViewportProxy.h"
 #include "Render/Viewport/RenderFrame/DisplayClusterRenderFrame.h"
+#include "Render/Viewport/RenderFrame/DisplayClusterRenderFrameSettings.h"
 
 #include "Engine/TextureRenderTarget2D.h"
 
@@ -126,6 +127,12 @@ void ADisplayClusterRootActor::Tick_Editor(float DeltaSeconds)
 {
 	if (IsPreviewEnabled())
 	{
+		// Restore ViewportManager
+		if (ViewportManager.IsValid() == false)
+		{
+			ViewportManager = MakeUnique<FDisplayClusterViewportManager>();
+		}
+
 		if (bDeferPreviewGeneration)
 		{
 			// Hack to generate preview components on instances during map load.
@@ -152,6 +159,21 @@ void ADisplayClusterRootActor::Tick_Editor(float DeltaSeconds)
 	else
 	{
 		ResetPreviewInternals_Editor();
+		if (ViewportManager.IsValid())
+		{
+			if (FDisplayClusterViewportManager* ViewportManagerPrivate = static_cast<FDisplayClusterViewportManager*>(ViewportManager.Get()))
+			{
+				switch (ViewportManagerPrivate->GetRenderFrameSettings().RenderMode)
+				{
+				case EDisplayClusterRenderFrameMode::PreviewInScene:
+					// Release viewport manager with resources immediatelly for preview in scene
+					ViewportManager.Reset();
+					break;
+				default:
+					break;
+				}
+			}
+		}
 	}
 	
 	SetChromakeyCardsOwner();
@@ -190,9 +212,6 @@ void ADisplayClusterRootActor::Destroyed_Editor()
 	ReleasePreviewComponents();
 
 	MarkAsGarbage();
-
-	// Force garbage collector
-	GEngine->ForceGarbageCollection(true);
 }
 
 void ADisplayClusterRootActor::BeginDestroy_Editor()
diff --git a/Engine/Plugins/Runtime/nDisplay/Source/DisplayCluster/Private/Render/Viewport/DisplayClusterViewportManager.h b/Engine/Plugins/Runtime/nDisplay/Source/DisplayCluster/Private/Render/Viewport/DisplayClusterViewportManager.h
--- a/Engine/Plugins/Runtime/nDisplay/Source/DisplayCluster/Private/Render/Viewport/DisplayClusterViewportManager.h
+++ b/Engine/Plugins/Runtime/nDisplay/Source/DisplayCluster/Private/Render/Viewport/DisplayClusterViewportManager.h
@@ -79,6 +79,9 @@ public:
 
 		return Viewports;
 	}
+
+	// Current world is a tools window (configurator, other tools)
+	bool IsEditorPreviewWorld() const;
 #endif
 
 	virtual IDisplayClusterViewport* FindViewport(const FString& InViewportId) const override
diff --git a/Engine/Plugins/Runtime/nDisplay/Source/DisplayCluster/Private/Render/Viewport/DisplayClusterViewportManager_EditorPreview.cpp b/Engine/Plugins/Runtime/nDisplay/Source/DisplayCluster/Private/Render/Viewport/DisplayClusterViewportManager_EditorPreview.cpp
--- a/Engine/Plugins/Runtime/nDisplay/Source/DisplayCluster/Private/Render/Viewport/DisplayClusterViewportManager_EditorPreview.cpp
+++ b/Engine/Plugins/Runtime/nDisplay/Source/DisplayCluster/Private/Render/Viewport/DisplayClusterViewportManager_EditorPreview.cpp
@@ -32,6 +32,22 @@
 ///////////////////////////////////////////////////////////////////////////////////////
 // FDisplayClusterViewportManager
 ///////////////////////////////////////////////////////////////////////////////////////
+bool FDisplayClusterViewportManager::IsEditorPreviewWorld() const
+{
+	if (UWorld* CurrentWorld = GetCurrentWorld())
+	{
+		switch (CurrentWorld->WorldType)
+		{
+		case EWorldType::EditorPreview:
+			return true;
+
+		default:
+			break;
+		}
+	}
+
+	return false;
+}
 
 void FDisplayClusterViewportManager::ImplUpdatePreviewRTTResources()
 {
diff --git a/Engine/Plugins/Runtime/nDisplay/Source/DisplayCluster/Private/Render/Viewport/DisplayClusterViewport_EditorPreview.cpp b/Engine/Plugins/Runtime/nDisplay/Source/DisplayCluster/Private/Render/Viewport/DisplayClusterViewport_EditorPreview.cpp
--- a/Engine/Plugins/Runtime/nDisplay/Source/DisplayCluster/Private/Render/Viewport/DisplayClusterViewport_EditorPreview.cpp
+++ b/Engine/Plugins/Runtime/nDisplay/Source/DisplayCluster/Private/Render/Viewport/DisplayClusterViewport_EditorPreview.cpp
@@ -28,6 +28,23 @@
 
 #include "Render/Viewport/Configuration/DisplayClusterViewportConfigurationHelpers_Postprocess.h"
 
+///////////////////////////////////////////////////////////////////////////////////////
+int32 GDisplayClusterPreviewEnableViewState = 0;
+static FAutoConsoleVariableRef CVarDisplayClusterPreviewEnableViewState(
+	TEXT("nDisplay.preview.EnableViewState"),
+	GDisplayClusterPreviewEnableViewState,
+	TEXT("Enable view state for preview (0 - disable).\n"),
+	ECVF_RenderThreadSafe
+);
+
+int32 GDisplayClusterPreviewEnableConfiguratorViewState = 0;
+static FAutoConsoleVariableRef CVarDisplayClusterPreviewEnableConfiguratorViewState(
+	TEXT("nDisplay.preview.EnableConfiguratorViewState"),
+	GDisplayClusterPreviewEnableConfiguratorViewState,
+	TEXT("Enable view state for preview in Configurator window (0 - disable).\n"),
+	ECVF_RenderThreadSafe
+);
+
 ///////////////////////////////////////////////////////////////////////////////////////
 //          FDisplayClusterViewport
 ///////////////////////////////////////////////////////////////////////////////////////
@@ -48,6 +65,14 @@ void FDisplayClusterViewport::CleanupViewState()
 
 FSceneViewStateInterface* FDisplayClusterViewport::GetViewState(uint32 ViewIndex)
 {
+	if (GDisplayClusterPreviewEnableViewState == 0 || (GDisplayClusterPreviewEnableConfiguratorViewState == 0 && Owner.IsEditorPreviewWorld()))
+	{
+		// Disable ViewState
+		ViewStates.Empty();
+
+		return nullptr;
+	}
+
 	int32 RequiredAmount = (int32)ViewIndex - ViewStates.Num() + 1;
 	if (RequiredAmount > 0)
 	{
