nDisplay: Fixed issue where opening multiple nDisplay config editors would gernerate duplicate Add buttons for each of them in the component panel

diff --git a/Engine/Plugins/Runtime/nDisplay/Source/DisplayClusterConfigurator/Private/DisplayClusterConfiguratorBlueprintEditor.cpp b/Engine/Plugins/Runtime/nDisplay/Source/DisplayClusterConfigurator/Private/DisplayClusterConfiguratorBlueprintEditor.cpp
--- a/Engine/Plugins/Runtime/nDisplay/Source/DisplayClusterConfigurator/Private/DisplayClusterConfiguratorBlueprintEditor.cpp
+++ b/Engine/Plugins/Runtime/nDisplay/Source/DisplayClusterConfigurator/Private/DisplayClusterConfiguratorBlueprintEditor.cpp
@@ -71,6 +71,11 @@
 
 #define LOCTEXT_NAMESPACE "DisplayClusterConfiguratorBlueprintEditor"
 
+namespace UE::DisplayClusterConfiguratorBlueprintEditorConstants
+{
+	const static FName SCSEditorExtensionIdentifier = TEXT("DisplayClusterSCSEditorExtension");
+}
+
 /*
  * Extend the SCS Editor. We hide the Add Component button and replace it with our own.
  */
@@ -97,6 +102,50 @@ private:
 };
 
 TSharedPtr<FDisplayClusterBlueprintEditorSCSEditorUICustomization> FDisplayClusterBlueprintEditorSCSEditorUICustomization::Instance;
+TArray<TWeakPtr<FDisplayClusterConfiguratorBlueprintEditor>> FDisplayClusterConfiguratorBlueprintEditor::EditorsRegisteredForPanelExtension;
+
+void FDisplayClusterConfiguratorBlueprintEditor::RegisterPanelExtensionFactory()
+{
+	if (UPanelExtensionSubsystem* PanelExtensionSubsystem = GEditor->GetEditorSubsystem<UPanelExtensionSubsystem>())
+	{
+		if (!PanelExtensionSubsystem->IsPanelFactoryRegistered(UE::DisplayClusterConfiguratorBlueprintEditorConstants::SCSEditorExtensionIdentifier))
+		{
+			FPanelExtensionFactory SCSEditorExtensionWidget;
+			SCSEditorExtensionWidget.CreateExtensionWidget = FPanelExtensionFactory::FCreateExtensionWidget::CreateStatic(&FDisplayClusterConfiguratorBlueprintEditor::OnExtendAddComponentButton);
+			SCSEditorExtensionWidget.Identifier = UE::DisplayClusterConfiguratorBlueprintEditorConstants::SCSEditorExtensionIdentifier;
+
+			PanelExtensionSubsystem->RegisterPanelFactory(TEXT("SCSEditor.NextToAddComponentButton"), SCSEditorExtensionWidget);
+		}
+	}
+}
+
+void FDisplayClusterConfiguratorBlueprintEditor::UnregisterPanelExtensionFactory()
+{
+	if (UPanelExtensionSubsystem* PanelExtensionSubsystem = GEditor->GetEditorSubsystem<UPanelExtensionSubsystem>())
+	{
+		// Unregister all panel factories associated with the DisplayCluster configurator
+		PanelExtensionSubsystem->UnregisterPanelFactory(UE::DisplayClusterConfiguratorBlueprintEditorConstants::SCSEditorExtensionIdentifier);
+	}
+}
+
+TSharedRef<SWidget> FDisplayClusterConfiguratorBlueprintEditor::OnExtendAddComponentButton(FWeakObjectPtr ExtensionContext)
+{
+	if (const USubobjectEditorExtensionContext* SubobjectEditorExtensionContext = Cast<USubobjectEditorExtensionContext>(ExtensionContext.Get()))
+	{
+		if (const TSharedPtr<SSubobjectEditor> SubobjectEditor = SubobjectEditorExtensionContext->GetSubobjectEditor().Pin())
+		{
+			for (const TWeakPtr<FDisplayClusterConfiguratorBlueprintEditor>& BlueprintEditor : EditorsRegisteredForPanelExtension)
+			{
+				if (BlueprintEditor.IsValid() && BlueprintEditor.Pin()->SubobjectEditor == SubobjectEditor)
+				{
+					return BlueprintEditor.Pin()->CreateAddComponentButton();
+				}
+			}
+		}
+	}
+
+	return SNullWidget::NullWidget;
+}
 
 FDisplayClusterConfiguratorBlueprintEditor::~FDisplayClusterConfiguratorBlueprintEditor()
 {
@@ -127,8 +176,6 @@ void FDisplayClusterConfiguratorBlueprintEditor::InitDisplayClusterBlueprintEdit
 {
 	LoadedBlueprint = Blueprint;
 	FDisplayClusterConfiguratorVersionUtils::SetToLatestVersion(Blueprint);
-
-	SCSEditorExtensionIdentifier = *(FString("DisplayClusterSCSEditorExtension") + FGuid::NewGuid().ToString());
 	
 	TSharedPtr<FDisplayClusterConfiguratorBlueprintEditor> Editor = SharedThis(this);
 	
@@ -169,7 +216,10 @@ void FDisplayClusterConfiguratorBlueprintEditor::InitDisplayClusterBlueprintEdit
 
 	TSharedRef<FApplicationMode> ConfigurationMode = MakeShared<FDisplayClusterConfiguratorEditorConfigurationMode>(Editor);
 	AddApplicationMode(ConfigurationMode->GetModeName(), ConfigurationMode);
-	
+
+	// Register this blueprint editor instance to the list of blueprint editors that can receive panel expansion callbacks
+	EditorsRegisteredForPanelExtension.Add(StaticCastWeakPtr<FDisplayClusterConfiguratorBlueprintEditor>(AsWeak()));
+
 	const TArray<UBlueprint*> Blueprints{ Blueprint };
 	CommonInitialization(Blueprints, false);
 
@@ -177,18 +227,6 @@ void FDisplayClusterConfiguratorBlueprintEditor::InitDisplayClusterBlueprintEdit
 	
 	CreateSCSEditorWrapper();
 
-	if (UPanelExtensionSubsystem* PanelExtensionSubsystem = GEditor->GetEditorSubsystem<UPanelExtensionSubsystem>())
-	{
-		if (!PanelExtensionSubsystem->IsPanelFactoryRegistered(SCSEditorExtensionIdentifier))
-		{
-			FPanelExtensionFactory SCSEditorExtensionWidget;
-			SCSEditorExtensionWidget.CreateExtensionWidget = FPanelExtensionFactory::FCreateExtensionWidget::CreateStatic(&FDisplayClusterConfiguratorBlueprintEditor::CreateSCSEditorExtensionWidget);
-			SCSEditorExtensionWidget.Identifier = SCSEditorExtensionIdentifier;
-
-			PanelExtensionSubsystem->RegisterPanelFactory("SCSEditor.NextToAddComponentButton", SCSEditorExtensionWidget);
-		}
-	}
-	
 	ExtendMenu();
 	ExtendToolbar();
 	
@@ -1277,140 +1315,137 @@ void FDisplayClusterConfiguratorBlueprintEditor::ShutdownDCSCSEditors()
 {
 	ViewportTabContent.Reset();
 	ViewportTab.Reset();
-	
-	if (UPanelExtensionSubsystem* PanelExtensionSubsystem = GEditor->GetEditorSubsystem<UPanelExtensionSubsystem>())
+
+	// Remove this blueprint editor and any registered blueprint editors that are no longer valid (as is the case when in a destructor)
+	// from the list of registered blueprint editors
+	EditorsRegisteredForPanelExtension.RemoveAll([=](const TWeakPtr<FDisplayClusterConfiguratorBlueprintEditor>& BlueprintEditor)
 	{
-		PanelExtensionSubsystem->UnregisterPanelFactory(SCSEditorExtensionIdentifier, "SCSEditor.NextToAddComponentButton");
-	}
+		return !BlueprintEditor.IsValid() || BlueprintEditor.Pin().Get() == this;
+	});
 }
 
-TSharedRef<SWidget> FDisplayClusterConfiguratorBlueprintEditor::CreateSCSEditorExtensionWidget(FWeakObjectPtr ExtensionContext)
+TSharedRef<SWidget> FDisplayClusterConfiguratorBlueprintEditor::CreateAddComponentButton()
 {
 	auto PerformComboAddClass = [=](TSubclassOf<UActorComponent> ComponentClass, EComponentCreateAction::Type ComponentCreateAction,
 		UObject* AssetOverride) -> FSubobjectDataHandle
 	{
 		FSubobjectDataHandle NewComponentHandle = FSubobjectDataHandle::InvalidHandle;
+
+		UClass* NewClass = ComponentClass;
 		
-		if (const USubobjectEditorExtensionContext* SubobjectEditorExtensionContext = Cast<USubobjectEditorExtensionContext>(ExtensionContext.Get()))
+		if (NewClass != nullptr)
 		{
-			if (const TSharedPtr<SSubobjectEditor> SubobjectEditor = SubobjectEditorExtensionContext->GetSubobjectEditor().Pin())
+			TUniquePtr<FScopedTransaction> AddTransaction = MakeUnique<FScopedTransaction>(LOCTEXT("AddComponent", "Add Component"));
+			
+			FEditorDelegates::LoadSelectedAssetsIfNeeded.Broadcast();
+			USelection* Selection = GEditor->GetSelectedObjects();
+
+			bool bAddedComponent = false;
+
+			// Rename the asset to our display names by creating a template.
+			// Otherwise the DisplayCluster class name is used, not the custom nDisplay display name.
+			if (AssetOverride == nullptr && NewClass->IsClassGroupName(TEXT("DisplayCluster")))
 			{
-				UClass* NewClass = ComponentClass;
-				
-				if (NewClass != nullptr)
+				const FString DisplayName = FDisplayClusterConfiguratorUtils::FormatNDisplayComponentName(NewClass);
+				AssetOverride = NewObject<UActorComponent>(GetTransientPackage(), NewClass, *DisplayName);
+			}
+			
+			USubobjectDataSubsystem* SubobjectSystem = USubobjectDataSubsystem::Get();
+			FAddNewSubobjectParams AddSubobjectParams;
+			AddSubobjectParams.BlueprintContext = SubobjectEditor->GetBlueprint();
+			AddSubobjectParams.NewClass = NewClass;
+								
+			// This adds components according to the type selected in the drop down. If the user
+			// has the appropriate objects selected in the content browser then those are added,
+			// else we go down the previous route of adding components by type.
+			//
+			// Furthermore don't try to match up assets for USceneComponent it will match lots of things and doesn't have any nice behavior for asset adds 
+
+			FText OutFailReason;
+			
+			if (Selection->Num() > 0 && !AssetOverride && NewClass != USceneComponent::StaticClass())
+			{
+				for (FSelectionIterator ObjectIter(*Selection); ObjectIter; ++ObjectIter)
 				{
-					TUniquePtr<FScopedTransaction> AddTransaction = MakeUnique<FScopedTransaction>(LOCTEXT("AddComponent", "Add Component"));
-					
-					FEditorDelegates::LoadSelectedAssetsIfNeeded.Broadcast();
-					USelection* Selection = GEditor->GetSelectedObjects();
+					UObject* Object = *ObjectIter;
+					TArray<TSubclassOf<UActorComponent>> ComponentClasses = FComponentAssetBrokerage::GetComponentsForAsset(Object);
 
-					bool bAddedComponent = false;
-
-					// Rename the asset to our display names by creating a template.
-					// Otherwise the DisplayCluster class name is used, not the custom nDisplay display name.
-					if (AssetOverride == nullptr && NewClass->IsClassGroupName(TEXT("DisplayCluster")))
-					{
-						const FString DisplayName = FDisplayClusterConfiguratorUtils::FormatNDisplayComponentName(NewClass);
-						AssetOverride = NewObject<UActorComponent>(GetTransientPackage(), NewClass, *DisplayName);
-					}
-					
-					USubobjectDataSubsystem* SubobjectSystem = USubobjectDataSubsystem::Get();
-					FAddNewSubobjectParams AddSubobjectParams;
-					AddSubobjectParams.BlueprintContext = SubobjectEditor->GetBlueprint();
-					AddSubobjectParams.NewClass = NewClass;
-										
-					// This adds components according to the type selected in the drop down. If the user
-					// has the appropriate objects selected in the content browser then those are added,
-					// else we go down the previous route of adding components by type.
-					//
-					// Furthermore don't try to match up assets for USceneComponent it will match lots of things and doesn't have any nice behavior for asset adds 
-
-					FText OutFailReason;
-					
-					if (Selection->Num() > 0 && !AssetOverride && NewClass != USceneComponent::StaticClass())
+					// if the selected asset supports the selected component type then go ahead and add it
+					for (int32 ComponentIndex = 0; ComponentIndex < ComponentClasses.Num(); ComponentIndex++)
 					{
-						for (FSelectionIterator ObjectIter(*Selection); ObjectIter; ++ObjectIter)
+						if (ComponentClasses[ComponentIndex]->IsChildOf(NewClass))
 						{
-							UObject* Object = *ObjectIter;
-							TArray<TSubclassOf<UActorComponent>> ComponentClasses = FComponentAssetBrokerage::GetComponentsForAsset(Object);
-
-							// if the selected asset supports the selected component type then go ahead and add it
-							for (int32 ComponentIndex = 0; ComponentIndex < ComponentClasses.Num(); ComponentIndex++)
+							const FSubobjectDataHandle ParentHandle = SubobjectSystem->FindHandleForObject(SubobjectEditor->GetSceneRootNode()->GetDataHandle(), Object);
+							if(ParentHandle.IsValid())
 							{
-								if (ComponentClasses[ComponentIndex]->IsChildOf(NewClass))
-								{
-									const FSubobjectDataHandle ParentHandle = SubobjectSystem->FindHandleForObject(SubobjectEditor->GetSceneRootNode()->GetDataHandle(), Object);
-									if(ParentHandle.IsValid())
-									{
-										AddSubobjectParams.AssetOverride = Object;
-										AddSubobjectParams.ParentHandle = ParentHandle;
-										NewComponentHandle = SubobjectSystem->AddNewSubobject(AddSubobjectParams, OutFailReason);
-										bAddedComponent = true;
-										break;
-									}
-								}
+								AddSubobjectParams.AssetOverride = Object;
+								AddSubobjectParams.ParentHandle = ParentHandle;
+								NewComponentHandle = SubobjectSystem->AddNewSubobject(AddSubobjectParams, OutFailReason);
+								bAddedComponent = true;
+								break;
 							}
 						}
 					}
+				}
+			}
 
-					if (!bAddedComponent)
+			if (!bAddedComponent)
+			{
+				// Attach this component to the override asset first, but if none is given then use the actor context			
+				FSubobjectDataHandle ParentHandle = SubobjectSystem->FindHandleForObject(SubobjectEditor->GetSceneRootNode()->GetDataHandle(), AssetOverride);
+	
+				if(!ParentHandle.IsValid())
+				{
+					// If we have something selected, then we should attach it to that
+					TArray<FSubobjectEditorTreeNodePtrType> SelectedTreeNodes = SubobjectEditor->GetSelectedNodes();
+					if(SelectedTreeNodes.Num() > 0)
 					{
-						// Attach this component to the override asset first, but if none is given then use the actor context			
-						FSubobjectDataHandle ParentHandle = SubobjectSystem->FindHandleForObject(SubobjectEditor->GetSceneRootNode()->GetDataHandle(), AssetOverride);
-			
-						if(!ParentHandle.IsValid())
-						{
-							// If we have something selected, then we should attach it to that
-							TArray<FSubobjectEditorTreeNodePtrType> SelectedTreeNodes = SubobjectEditor->GetSelectedNodes();
-							if(SelectedTreeNodes.Num() > 0)
-							{
-								ParentHandle = SelectedTreeNodes[0]->GetDataHandle();
-							}
-							// Otherwise fall back to the root node
-							else
-							{
-								const TArray<FSubobjectEditorTreeNodePtrType>& RootNodes = SubobjectEditor->GetRootNodes();
-								ParentHandle = RootNodes.Num() > 0
-												? RootNodes[0]->GetDataHandle()
-												: FSubobjectDataHandle::InvalidHandle;					
-							}
-						}
-			
-						if(ParentHandle.IsValid())
-						{
-							AddSubobjectParams.AssetOverride = AssetOverride;
-							AddSubobjectParams.ParentHandle = ParentHandle;
-							NewComponentHandle = SubobjectSystem->AddNewSubobject(AddSubobjectParams, OutFailReason);
-						}
-					}
-					
-					if(!NewComponentHandle.IsValid())
-					{
-						if (OutFailReason.IsEmpty())
-						{
-							OutFailReason = LOCTEXT("AddComponentFailed_Generic", "Failed to add component!");
-						}
-						FNotificationInfo Info(OutFailReason);
-						Info.Image = FAppStyle::GetBrush(TEXT("Icons.Error"));
-						Info.bFireAndForget = true;
-						Info.bUseSuccessFailIcons = false;
-						Info.ExpireDuration = 5.0f;
-
-						FSlateNotificationManager::Get().AddNotification(Info);
+						ParentHandle = SelectedTreeNodes[0]->GetDataHandle();
 					}
+					// Otherwise fall back to the root node
 					else
 					{
-						SubobjectEditor->UpdateTree();
-						// Set focus to the newly created subobject
-						if(const FSubobjectEditorTreeNodePtrType NewNode = SubobjectEditor->FindSlateNodeForHandle(NewComponentHandle))
-						{
-							SubobjectEditor->GetDragDropTree()->SetSelection(NewNode);
-							SubobjectEditor->GetCommandList()->ExecuteAction(FGenericCommands::Get().Rename.ToSharedRef());
-						}
+						const TArray<FSubobjectEditorTreeNodePtrType>& RootNodes = SubobjectEditor->GetRootNodes();
+						ParentHandle = RootNodes.Num() > 0
+										? RootNodes[0]->GetDataHandle()
+										: FSubobjectDataHandle::InvalidHandle;					
 					}
 				}
+	
+				if(ParentHandle.IsValid())
+				{
+					AddSubobjectParams.AssetOverride = AssetOverride;
+					AddSubobjectParams.ParentHandle = ParentHandle;
+					NewComponentHandle = SubobjectSystem->AddNewSubobject(AddSubobjectParams, OutFailReason);
+				}
+			}
+			
+			if(!NewComponentHandle.IsValid())
+			{
+				if (OutFailReason.IsEmpty())
+				{
+					OutFailReason = LOCTEXT("AddComponentFailed_Generic", "Failed to add component!");
+				}
+				FNotificationInfo Info(OutFailReason);
+				Info.Image = FAppStyle::GetBrush(TEXT("Icons.Error"));
+				Info.bFireAndForget = true;
+				Info.bUseSuccessFailIcons = false;
+				Info.ExpireDuration = 5.0f;
+
+				FSlateNotificationManager::Get().AddNotification(Info);
+			}
+			else
+			{
+				SubobjectEditor->UpdateTree();
+				// Set focus to the newly created subobject
+				if(const FSubobjectEditorTreeNodePtrType NewNode = SubobjectEditor->FindSlateNodeForHandle(NewComponentHandle))
+				{
+					SubobjectEditor->GetDragDropTree()->SetSelection(NewNode);
+					SubobjectEditor->GetCommandList()->ExecuteAction(FGenericCommands::Get().Rename.ToSharedRef());
+				}
 			}
 		}
+
 		return NewComponentHandle;
 	};
 	
@@ -1422,22 +1457,6 @@ TSharedRef<SWidget> FDisplayClusterConfiguratorBlueprintEditor::CreateSCSEditorE
 		[
 			SNew(SDisplayClusterConfiguratorComponentClassCombo)
 			.OnSubobjectClassSelected(FSubobjectClassSelected::CreateLambda(PerformComboAddClass))
-			.Visibility_Lambda([ExtensionContext]
-			{
-				if(USubobjectEditorExtensionContext* SubobjectEditorExtensionContext = Cast<USubobjectEditorExtensionContext>(ExtensionContext.Get()))
-				{
-					if (SubobjectEditorExtensionContext->GetSubobjectEditor().IsValid())
-					{
-						if (const ADisplayClusterRootActor* RootActor = Cast<ADisplayClusterRootActor>(SubobjectEditorExtensionContext->GetSubobjectEditor().Pin()->GetObjectContext()))
-						{
-							// Check that this context is for our root actor and that we are editing just the CDO.
-							// Without this we could show up for unrelated blueprints or in the level editor when selecting our actor.
-							return RootActor->IsTemplate(RF_ClassDefaultObject) ? EVisibility::Visible : EVisibility::Collapsed;
-						}
-					}
-				}
-				return EVisibility::Collapsed;
-			})
 			.AddMetaData<FTagMetaData>(FTagMetaData(TEXT("Actor.AddComponent")))
 			.ToolTipText(LOCTEXT("AddComponent_Tooltip", "Adds a new component to this actor"))
 		];
diff --git a/Engine/Plugins/Runtime/nDisplay/Source/DisplayClusterConfigurator/Private/DisplayClusterConfiguratorBlueprintEditor.h b/Engine/Plugins/Runtime/nDisplay/Source/DisplayClusterConfigurator/Private/DisplayClusterConfiguratorBlueprintEditor.h
--- a/Engine/Plugins/Runtime/nDisplay/Source/DisplayClusterConfigurator/Private/DisplayClusterConfiguratorBlueprintEditor.h
+++ b/Engine/Plugins/Runtime/nDisplay/Source/DisplayClusterConfigurator/Private/DisplayClusterConfiguratorBlueprintEditor.h
@@ -28,6 +28,19 @@ class FSubobjectEditorTreeNode;
 class FDisplayClusterConfiguratorBlueprintEditor
 	: public IDisplayClusterConfiguratorBlueprintEditor
 {
+public:
+	/** Registers any panel extension factories needed for the blueprint editor */
+	static void RegisterPanelExtensionFactory();
+
+	/** Unregisters any panel extension factories that were registered for the blueprint editor */
+	static void UnregisterPanelExtensionFactory();
+
+private:
+	/** Raised when the blueprint editor's subobject editor is creating extensions next to the Add Component button */
+	static TSharedRef<SWidget> OnExtendAddComponentButton(FWeakObjectPtr ExtensionContext);
+
+	/** A list of instantiated blueprint editors that are queried when the AddComponentButton panel extension is creating widgets */
+	static TArray<TWeakPtr<FDisplayClusterConfiguratorBlueprintEditor>> EditorsRegisteredForPanelExtension;
 
 public:
 	~FDisplayClusterConfiguratorBlueprintEditor();
@@ -145,7 +158,7 @@ protected:
 	void CreateDCSCSEditors();
 	void ShutdownDCSCSEditors();
 
-	static TSharedRef<SWidget> CreateSCSEditorExtensionWidget(FWeakObjectPtr ExtensionContext);
+	TSharedRef<SWidget> CreateAddComponentButton();
 
 	void CreateSCSEditorWrapper();
 
@@ -243,8 +256,6 @@ private:
 	/** The currently loaded blueprint. */
 	TWeakObjectPtr<UDisplayClusterBlueprint> LoadedBlueprint;
 
-	FName SCSEditorExtensionIdentifier;
-
 	FDelegateHandle RenameVariableHandle;
 	FDelegateHandle FocusChangedHandle;
 };
\ No newline at end of file
diff --git a/Engine/Plugins/Runtime/nDisplay/Source/DisplayClusterConfigurator/Private/DisplayClusterConfiguratorModule.cpp b/Engine/Plugins/Runtime/nDisplay/Source/DisplayClusterConfigurator/Private/DisplayClusterConfiguratorModule.cpp
--- a/Engine/Plugins/Runtime/nDisplay/Source/DisplayClusterConfigurator/Private/DisplayClusterConfiguratorModule.cpp
+++ b/Engine/Plugins/Runtime/nDisplay/Source/DisplayClusterConfigurator/Private/DisplayClusterConfiguratorModule.cpp
@@ -39,6 +39,7 @@
 #include "ISettingsModule.h"
 #include "AssetRegistry/AssetRegistryModule.h"
 #include "AssetTypeCategories.h"
+#include "DisplayClusterConfiguratorBlueprintEditor.h"
 #include "HAL/IConsoleManager.h"
 #include "Modules/ModuleManager.h"
 #include "ActorFactories/ActorFactoryBlueprint.h"
@@ -61,6 +62,8 @@
 
 void FDisplayClusterConfiguratorModule::StartupModule()
 {
+	FCoreDelegates::OnPostEngineInit.AddRaw(this, &FDisplayClusterConfiguratorModule::OnPostEngineInit);
+
 	IAssetTools& AssetTools = FModuleManager::LoadModuleChecked<FAssetToolsModule>("AssetTools").Get();
 
 	/*
@@ -115,6 +118,8 @@ void FDisplayClusterConfiguratorModule::StartupModule()
 
 void FDisplayClusterConfiguratorModule::ShutdownModule()
 {
+	FCoreDelegates::OnPostEngineInit.RemoveAll(this);
+
 	if (FAssetToolsModule* AssetTools = FModuleManager::GetModulePtr<FAssetToolsModule>("AssetTools"))
 	{
 		for (int32 IndexAction = 0; IndexAction < CreatedAssetTypeActions.Num(); ++IndexAction)
@@ -130,6 +135,11 @@ void FDisplayClusterConfiguratorModule::ShutdownModule()
 	MenuExtensibilityManager.Reset();
 	ToolBarExtensibilityManager.Reset();
 
+	if (GEditor)
+	{
+		FDisplayClusterConfiguratorBlueprintEditor::UnregisterPanelExtensionFactory();
+	}
+
 	IKismetCompilerInterface& KismetCompilerModule = FModuleManager::GetModuleChecked<IKismetCompilerInterface>("KismetCompiler");
 	KismetCompilerModule.GetCompilers().Remove(&BlueprintCompiler);
 
@@ -140,6 +150,11 @@ void FDisplayClusterConfiguratorModule::ShutdownModule()
 	}
 }
 
+void FDisplayClusterConfiguratorModule::OnPostEngineInit()
+{
+	FDisplayClusterConfiguratorBlueprintEditor::RegisterPanelExtensionFactory();
+}
+
 const FDisplayClusterConfiguratorCommands& FDisplayClusterConfiguratorModule::GetCommands() const
 {
 	return FDisplayClusterConfiguratorCommands::Get();
diff --git a/Engine/Plugins/Runtime/nDisplay/Source/DisplayClusterConfigurator/Private/DisplayClusterConfiguratorModule.h b/Engine/Plugins/Runtime/nDisplay/Source/DisplayClusterConfigurator/Private/DisplayClusterConfiguratorModule.h
--- a/Engine/Plugins/Runtime/nDisplay/Source/DisplayClusterConfigurator/Private/DisplayClusterConfiguratorModule.h
+++ b/Engine/Plugins/Runtime/nDisplay/Source/DisplayClusterConfigurator/Private/DisplayClusterConfiguratorModule.h
@@ -22,6 +22,8 @@ public:
 	virtual void StartupModule() override;
 	virtual void ShutdownModule() override;
 
+	void OnPostEngineInit();
+
 	/** Gets the extensibility managers for outside entities to extend this editor's menus and toolbars */
 	virtual TSharedPtr<FExtensibilityManager> GetMenuExtensibilityManager() const override { return MenuExtensibilityManager; }
 	virtual TSharedPtr<FExtensibilityManager> GetToolBarExtensibilityManager() const override { return ToolBarExtensibilityManager; }
