nDisplay: Fixed regression: Math for a camera with empty RTT must be passed to the icvfx rendering pipeline.

diff --git a/Engine/Plugins/Runtime/nDisplay/Source/DisplayClusterShaders/Public/ShaderParameters/DisplayClusterShaderParameters_ICVFX.h b/Engine/Plugins/Runtime/nDisplay/Source/DisplayClusterShaders/Public/ShaderParameters/DisplayClusterShaderParameters_ICVFX.h
--- a/Engine/Plugins/Runtime/nDisplay/Source/DisplayClusterShaders/Public/ShaderParameters/DisplayClusterShaderParameters_ICVFX.h
+++ b/Engine/Plugins/Runtime/nDisplay/Source/DisplayClusterShaders/Public/ShaderParameters/DisplayClusterShaderParameters_ICVFX.h
@@ -245,7 +245,7 @@ public:
 	{
 		return Cameras.FindByPredicate([InViewportId](const FCameraSettings& CameraIt)
 			{
-				return CameraIt.Resource.IsValid() && CameraIt.Resource.ViewportId == InViewportId;
+				return CameraIt.Resource.ViewportId == InViewportId;
 			});
 	}
 
