nDisplay: Fixed possible crashes and changed LK rules for ICVFX shader permutations.

diff --git a/Engine/Plugins/Runtime/nDisplay/Source/DisplayClusterShaders/Private/Shaders/DisplayClusterShadersWarpblend_ICVFX.cpp b/Engine/Plugins/Runtime/nDisplay/Source/DisplayClusterShaders/Private/Shaders/DisplayClusterShadersWarpblend_ICVFX.cpp
--- a/Engine/Plugins/Runtime/nDisplay/Source/DisplayClusterShaders/Private/Shaders/DisplayClusterShadersWarpblend_ICVFX.cpp
+++ b/Engine/Plugins/Runtime/nDisplay/Source/DisplayClusterShaders/Private/Shaders/DisplayClusterShadersWarpblend_ICVFX.cpp
@@ -113,7 +113,7 @@ namespace IcvfxShaderPermutation
 		FIcvfxShaderMeshWarp
 	>;
 	
-	bool ShouldCompileCommonPSPermutation(const FGlobalShaderPermutationParameters& Parameters, const FCommonPSDomain& PermutationVector)
+	bool ShouldCompileCommonPSPermutation(const FCommonPSDomain& PermutationVector)
 	{
 		if (!PermutationVector.Get<FIcvfxShaderMeshWarp>())
 		{
@@ -130,14 +130,17 @@ namespace IcvfxShaderPermutation
 
 		if (!PermutationVector.Get<FIcvfxShaderViewportInput>())
 		{
-			if (PermutationVector.Get<FIcvfxShaderLightCardUnder>() || (PermutationVector.Get<FIcvfxShaderInnerCamera>() == PermutationVector.Get<FIcvfxShaderLightCardOver>()))
+			if (PermutationVector.Get<FIcvfxShaderLightCardUnder>() || PermutationVector.Get<FIcvfxShaderUVLightCardUnder>())
 			{
 				return false;
 			}
 
-			if (PermutationVector.Get<FIcvfxShaderUVLightCardUnder>() || (PermutationVector.Get<FIcvfxShaderInnerCamera>() == PermutationVector.Get<FIcvfxShaderUVLightCardOver>()))
+			if (PermutationVector.Get<FIcvfxShaderLightCardOver>() || PermutationVector.Get<FIcvfxShaderUVLightCardOver>())
 			{
-				return false;
+				if (PermutationVector.Get<FIcvfxShaderInnerCamera>())
+				{
+					return false;
+				}
 			}
 
 			if (PermutationVector.Get<FIcvfxShaderViewportInputAlpha>())
@@ -179,7 +182,7 @@ namespace IcvfxShaderPermutation
 		return true;
 	}
 
-	bool ShouldCompileCommonVSPermutation(const FGlobalShaderPermutationParameters& Parameters, const FCommonVSDomain& PermutationVector)
+	bool ShouldCompileCommonVSPermutation(const FCommonVSDomain& PermutationVector)
 	{
 		return true;
 	}
@@ -248,7 +251,7 @@ public:
 
 	static bool ShouldCompilePermutation(const FGlobalShaderPermutationParameters& Parameters)
 	{
-		return IcvfxShaderPermutation::ShouldCompileCommonVSPermutation(Parameters, FPermutationDomain(Parameters.PermutationId));
+		return IcvfxShaderPermutation::ShouldCompileCommonVSPermutation(FPermutationDomain(Parameters.PermutationId));
 	}
 };
 
@@ -265,7 +268,7 @@ public:
 
 	static bool ShouldCompilePermutation(const FGlobalShaderPermutationParameters& Parameters)
 	{
-		return IcvfxShaderPermutation::ShouldCompileCommonPSPermutation(Parameters, FPermutationDomain(Parameters.PermutationId));
+		return IcvfxShaderPermutation::ShouldCompileCommonPSPermutation(FPermutationDomain(Parameters.PermutationId));
 	}
 };
 
@@ -849,7 +852,22 @@ public:
 		{
 			return false;
 		}
-		
+
+		// Check the permutation vectors. This should prevent a crash when no shader permutation is found.
+		if (!IcvfxShaderPermutation::ShouldCompileCommonPSPermutation(RenderPassData.PSPermutationVector))
+		{
+			UE_LOG(LogDisplayClusterShaders, Warning, TEXT("Invalid permutation vector %d for shader FIcvfxWarpPS"), RenderPassData.PSPermutationVector.ToDimensionValueId());
+
+			return false;
+		}
+
+		if (!IcvfxShaderPermutation::ShouldCompileCommonVSPermutation(RenderPassData.VSPermutationVector))
+		{
+			UE_LOG(LogDisplayClusterShaders, Warning, TEXT("Invalid permutation vector %d for shader FIcvfxWarpVS"), RenderPassData.VSPermutationVector.ToDimensionValueId());
+
+			return false;
+		}
+
 		// Get mutable shaders:
 		FGlobalShaderMap* ShaderMap = GetGlobalShaderMap(GMaxRHIFeatureLevel);
 		TShaderMapRef<FIcvfxWarpVS> VertexShader(ShaderMap, RenderPassData.VSPermutationVector);
@@ -857,16 +875,7 @@ public:
 
 		if (!VertexShader.IsValid() || !PixelShader.IsValid())
 		{
-			// Invalid permutation vector
-			if (!PixelShader.IsValid())
-			{
-				UE_LOG(LogDisplayClusterShaders, Warning, TEXT("Invalid permutation vector %d for shader FIcvfxWarpPS"), RenderPassData.PSPermutationVector.ToDimensionValueId());
-			}
-
-			if (!VertexShader.IsValid())
-			{
-				UE_LOG(LogDisplayClusterShaders, Warning, TEXT("Invalid permutation vector %d for shader FIcvfxWarpVS"), RenderPassData.VSPermutationVector.ToDimensionValueId());
-			}
+			UE_LOG(LogDisplayClusterShaders, Warning, TEXT("ICVFX shaders are not initialized properly."));
 
 			return false;
 		}
