nDisplay: Fixed UVLC+LC math in ICVFX shader.

diff --git a/Engine/Plugins/Runtime/nDisplay/Shaders/Private/MPCDIOverlayShaders.usf b/Engine/Plugins/Runtime/nDisplay/Shaders/Private/MPCDIOverlayShaders.usf
--- a/Engine/Plugins/Runtime/nDisplay/Shaders/Private/MPCDIOverlayShaders.usf
+++ b/Engine/Plugins/Runtime/nDisplay/Shaders/Private/MPCDIOverlayShaders.usf
@@ -55,7 +55,7 @@ float3 ComposeOverlayColor(float4 InColor, float4 InOverlayColor)
 
 float4 AddOverlayColor(float4 InColor, float4 InOverlayColor)
 {
-	return float4(ComposeOverlayColor(InColor, InOverlayColor), max(InColor.w, InOverlayColor.w));
+	return float4(ComposeOverlayColor(InColor, InOverlayColor), min(InColor.w, InOverlayColor.w));
 }
 
 ////////////////////////////////////////////////////////////////////////////////
@@ -235,7 +235,7 @@ OutputPS IcvfxWarpPS(OutputVS IN)
 	#if INNER_CAMERA
 		float4 OutColor = GetInnerCameraColor(WarpedUV, IN);
 	#else
-		float4 OutColor = float4(0, 0, 0, 0);
+		float4 OutColor = float4(0, 0, 0, 1);
 
 		#if LIGHTCARD_OVER
 			OutColor = AddOverlayColor(OutColor, GetLightCardTextureColor(OverlayUV));
