nDisplay: Fixed issue where the ICVFX panel wasn't displaying the level instance root actor's preview when preview was disabled on the root actor.

diff --git a/Engine/Plugins/Runtime/nDisplay/Source/DisplayCluster/Private/Game/EngineClasses/Scene/DisplayClusterPreviewComponent.cpp b/Engine/Plugins/Runtime/nDisplay/Source/DisplayCluster/Private/Game/EngineClasses/Scene/DisplayClusterPreviewComponent.cpp
--- a/Engine/Plugins/Runtime/nDisplay/Source/DisplayCluster/Private/Game/EngineClasses/Scene/DisplayClusterPreviewComponent.cpp
+++ b/Engine/Plugins/Runtime/nDisplay/Source/DisplayCluster/Private/Game/EngineClasses/Scene/DisplayClusterPreviewComponent.cpp
@@ -102,7 +102,12 @@ bool UDisplayClusterPreviewComponent::InitializePreviewComponent(ADisplayCluster
 
 bool UDisplayClusterPreviewComponent::IsPreviewEnabled() const
 {
-	return (ViewportConfig && RootActor && RootActor->IsPreviewDrawnToScreens());
+	return ViewportConfig && RootActor && RootActor->IsPreviewEnabled();
+}
+
+bool UDisplayClusterPreviewComponent::IsPreviewDrawnToScreen() const
+{
+	return ViewportConfig && RootActor && (RootActor->IsPreviewDrawnToScreens() || OverrideTexture);
 }
 
 void UDisplayClusterPreviewComponent::RestorePreviewMeshMaterial()
@@ -117,7 +122,13 @@ void UDisplayClusterPreviewComponent::RestorePreviewMeshMaterial()
 	}
 
 	// Release RTTs
-	ReleasePreviewRenderTarget();
+	// To fix UE-176749, we only want to do this if the root actor is not rendering previews, which can happen even if the previews
+	// are not being output to the meshes (for example, if the ICVFX panel is open)
+	// TODO: We will eventually want to separate out the preview rendering and resources from the preview mesh and material management
+	if (!IsPreviewEnabled())
+	{
+		ReleasePreviewRenderTarget();
+	}
 }
 
 void UDisplayClusterPreviewComponent::SetPreviewMeshMaterial()
@@ -166,7 +177,7 @@ bool UDisplayClusterPreviewComponent::UpdatePreviewMesh()
 
 	UpdatePreviewMeshReference();
 
-	if (IsPreviewEnabled() || OverrideTexture)
+	if (IsPreviewDrawnToScreen())
 	{
 		check(ViewportConfig);
 
@@ -420,6 +431,10 @@ void UDisplayClusterPreviewComponent::SetOverrideTexture(UTexture* InOverrideTex
 	if (OverrideTexture != InOverrideTexture)
 	{
 		OverrideTexture = InOverrideTexture;
+
+		// Update the entire preview mesh here, as the preview mesh may not be configured for rendering the override texture if
+		// the owning root actor has previews disabled.
+		UpdatePreviewMesh();
 		UpdatePreviewMaterial();
 	}
 }
diff --git a/Engine/Plugins/Runtime/nDisplay/Source/DisplayCluster/Public/Components/DisplayClusterPreviewComponent.h b/Engine/Plugins/Runtime/nDisplay/Source/DisplayCluster/Public/Components/DisplayClusterPreviewComponent.h
--- a/Engine/Plugins/Runtime/nDisplay/Source/DisplayCluster/Public/Components/DisplayClusterPreviewComponent.h
+++ b/Engine/Plugins/Runtime/nDisplay/Source/DisplayCluster/Public/Components/DisplayClusterPreviewComponent.h
@@ -90,8 +90,12 @@ public:
 	void SetOverrideTexture(UTexture* InOverrideTexture);
 
 protected:
+	/** Gets whether the root actor is requesting a preview render */
 	bool IsPreviewEnabled() const;
 
+	/** Gets whether the preview render or override texture should be output to the screen meshes */
+	bool IsPreviewDrawnToScreen() const;
+
 	class IDisplayClusterViewport* GetCurrentViewport() const;
 	bool GetPreviewTextureSettings(FIntPoint& OutSize, EPixelFormat& OutTextureFormat, float& OutGamma, bool& bOutSRGB) const;
 
diff --git a/Engine/Plugins/Runtime/nDisplay/Source/DisplayClusterScenePreview/Private/DisplayClusterScenePreviewModule.cpp b/Engine/Plugins/Runtime/nDisplay/Source/DisplayClusterScenePreview/Private/DisplayClusterScenePreviewModule.cpp
--- a/Engine/Plugins/Runtime/nDisplay/Source/DisplayClusterScenePreview/Private/DisplayClusterScenePreviewModule.cpp
+++ b/Engine/Plugins/Runtime/nDisplay/Source/DisplayClusterScenePreview/Private/DisplayClusterScenePreviewModule.cpp
@@ -338,6 +338,14 @@ bool FDisplayClusterScenePreviewModule::InternalRenderImmediate(FRendererConfig&
 				{
 					PreviewComponentOverrideTextures.Add(PreviewComp, PreviewComp->GetOverrideTexture());
 					PreviewComp->SetOverrideTexture(PreviewComp->GetRenderTargetTexturePostProcess());
+
+					// Force the preview mesh to recreate its render state immediately, in case setting the override texture
+					// caused the preview component to create and set the preview material on the mesh (as is the case when
+					// previews are disabled on the root actor)
+					if (UMeshComponent* PreviewMesh = PreviewComp->GetPreviewMesh())
+					{
+						PreviewMesh->RecreateRenderState_Concurrent();
+					}
 				}
 			}
 		}
