nDisplay: Fixed issues with LC\uvLC.
* LC overlay blending math changed.
* fixed potential crash in LCManager when deleting LCActor

diff --git a/Engine/Plugins/Runtime/nDisplay/Shaders/Private/MPCDIOverlayShaders.usf b/Engine/Plugins/Runtime/nDisplay/Shaders/Private/MPCDIOverlayShaders.usf
--- a/Engine/Plugins/Runtime/nDisplay/Shaders/Private/MPCDIOverlayShaders.usf
+++ b/Engine/Plugins/Runtime/nDisplay/Shaders/Private/MPCDIOverlayShaders.usf
@@ -53,9 +53,12 @@ float3 ComposeOverlayColor(float4 InColor, float4 InOverlayColor)
 	return InOverlayColor.xyz + InColor.xyz * InOverlayColor.w;
 }
 
-float4 AddOverlayColor(float4 InColor, float4 InOverlayColor)
+// Additive render pass with 2 overlays
+// a = in\out, b,c = ordered overlays (a + b + c)
+// c.xyz + (b.xyz + a.xyz * b.w)*c.w = a.xyz * (b.w *c.w) + b.xyz*c.w + c.xyz
+float4 ComposeRenderPassTwoOverlays(float4 InOverlay1, float4 InOverlay2)
 {
-	return float4(ComposeOverlayColor(InColor, InOverlayColor), min(InColor.w, InOverlayColor.w));
+	return float4(ComposeOverlayColor(InOverlay1, InOverlay2), InOverlay1.w * InOverlay2.w);
 }
 
 ////////////////////////////////////////////////////////////////////////////////
@@ -200,6 +203,20 @@ OutputPS IcvfxWarpPS(OutputVS IN)
 	float4 OverlayUVW = mul(WarpedUV, OverlayProjectionMatrix);
 	float2 OverlayUV = OverlayUVW.xy / OverlayUVW.w;
 
+	#if LIGHTCARD_OVER || UVLIGHTCARD_OVER || LIGHTCARD_UNDER || UVLIGHTCARD_UNDER
+		#if (LIGHTCARD_OVER && UVLIGHTCARD_OVER) || (LIGHTCARD_UNDER && UVLIGHTCARD_UNDER)
+			float4 LightCardColor = ComposeRenderPassTwoOverlays(GetLightCardTextureColor(OverlayUV), GetUVLightCardTextureColor(IN.UV_Chromakey.xy));
+		#else
+			// Render overlays in separate render-pass
+			#if LIGHTCARD_UNDER || LIGHTCARD_OVER
+				float4 LightCardColor = GetLightCardTextureColor(OverlayUV);
+			#endif
+			#if UVLIGHTCARD_UNDER || UVLIGHTCARD_OVER
+				float4 LightCardColor = GetUVLightCardTextureColor(IN.UV_Chromakey.xy);
+			#endif
+		#endif
+	#endif
+
 	// Compose all:
 #if VIEWPORT_INPUT
 
@@ -211,11 +228,8 @@ OutputPS IcvfxWarpPS(OutputVS IN)
 		float4 OutColor = float4(InputTexture.Sample(InputSampler, ViewportUV).rgb, 1.0f);
 	#endif
 
-	#if LIGHTCARD_UNDER
-		OutColor.xyz = ComposeOverlayColor(OutColor, GetLightCardTextureColor(OverlayUV));
-	#endif
-	#if UVLIGHTCARD_UNDER
-		OutColor.xyz = ComposeOverlayColor(OutColor, GetUVLightCardTextureColor(IN.UV_Chromakey.xy));
+	#if LIGHTCARD_UNDER || UVLIGHTCARD_UNDER
+		OutColor.xyz = ComposeOverlayColor(OutColor, LightCardColor);
 	#endif
 	
 	#if INNER_CAMERA
@@ -223,11 +237,8 @@ OutputPS IcvfxWarpPS(OutputVS IN)
 		OutColor.xyz = lerp(OutColor.xyz, CameraColor.xyz, CameraColor.w);
 	#endif
 
-	#if LIGHTCARD_OVER
-		OutColor.xyz = ComposeOverlayColor(OutColor, GetLightCardTextureColor(OverlayUV));
-	#endif
-	#if UVLIGHTCARD_OVER
-		OutColor.xyz = ComposeOverlayColor(OutColor, GetUVLightCardTextureColor(IN.UV_Chromakey.xy));
+	#if LIGHTCARD_OVER || UVLIGHTCARD_OVER
+		OutColor.xyz = ComposeOverlayColor(OutColor, LightCardColor);
 	#endif
 
 #else 
@@ -235,16 +246,13 @@ OutputPS IcvfxWarpPS(OutputVS IN)
 	#if INNER_CAMERA
 		float4 OutColor = GetInnerCameraColor(WarpedUV, IN);
 	#else
-		float4 OutColor = float4(0, 0, 0, 1);
-
-		#if LIGHTCARD_OVER
-			OutColor = AddOverlayColor(OutColor, GetLightCardTextureColor(OverlayUV));
-		#endif
-		#if UVLIGHTCARD_OVER
-			OutColor = AddOverlayColor(OutColor, GetUVLightCardTextureColor(IN.UV_Chromakey.xy));
+		#if LIGHTCARD_OVER || UVLIGHTCARD_OVER
+			float4 OutColor = LightCardColor;
+		#else
+			float4 OutColor = float4(0,0,0,0);
 		#endif
-
 	#endif
+
 #endif
 
 	//@todo: add LUT here
diff --git a/Engine/Plugins/Runtime/nDisplay/Source/DisplayCluster/Private/Render/Viewport/LightCard/DisplayClusterViewportLightCardManager.cpp b/Engine/Plugins/Runtime/nDisplay/Source/DisplayCluster/Private/Render/Viewport/LightCard/DisplayClusterViewportLightCardManager.cpp
--- a/Engine/Plugins/Runtime/nDisplay/Source/DisplayCluster/Private/Render/Viewport/LightCard/DisplayClusterViewportLightCardManager.cpp
+++ b/Engine/Plugins/Runtime/nDisplay/Source/DisplayCluster/Private/Render/Viewport/LightCard/DisplayClusterViewportLightCardManager.cpp
@@ -57,14 +57,12 @@ void FDisplayClusterViewportLightCardManager::Release()
 ///////////////////////////////////////////////////////////////////////////////////////////////
 void FDisplayClusterViewportLightCardManager::UpdateConfiguration()
 {
-	UpdateUVLightCardData();
+
 }
 
 void FDisplayClusterViewportLightCardManager::HandleStartScene()
 {
 	InitializePreviewWorld();
-
-	UpdateUVLightCardData();
 }
 
 void FDisplayClusterViewportLightCardManager::HandleEndScene()
@@ -76,6 +74,7 @@ void FDisplayClusterViewportLightCardManager::HandleEndScene()
 
 void FDisplayClusterViewportLightCardManager::RenderFrame()
 {
+	UpdateUVLightCardData();
 	RenderUVLightCard();
 }
 
diff --git a/Engine/Plugins/Runtime/nDisplay/Source/DisplayClusterShaders/Private/Shaders/DisplayClusterShadersWarpblend_ICVFX.cpp b/Engine/Plugins/Runtime/nDisplay/Source/DisplayClusterShaders/Private/Shaders/DisplayClusterShadersWarpblend_ICVFX.cpp
--- a/Engine/Plugins/Runtime/nDisplay/Source/DisplayClusterShaders/Private/Shaders/DisplayClusterShadersWarpblend_ICVFX.cpp
+++ b/Engine/Plugins/Runtime/nDisplay/Source/DisplayClusterShaders/Private/Shaders/DisplayClusterShadersWarpblend_ICVFX.cpp
@@ -128,6 +128,12 @@ namespace IcvfxShaderPermutation
 			return false;
 		}
 
+		// LightCard can be only 'over' or 'under'
+		if ((PermutationVector.Get<FIcvfxShaderLightCardUnder>() || PermutationVector.Get<FIcvfxShaderUVLightCardUnder>()) && (PermutationVector.Get<FIcvfxShaderUVLightCardOver>() || PermutationVector.Get<FIcvfxShaderLightCardOver>()))
+		{
+			return false;
+		}
+
 		if (!PermutationVector.Get<FIcvfxShaderViewportInput>())
 		{
 			if (PermutationVector.Get<FIcvfxShaderLightCardUnder>() || PermutationVector.Get<FIcvfxShaderUVLightCardUnder>())
@@ -901,15 +907,16 @@ public:
 			}
 			else
 			{
+				// Extra render pass: don't change target alpha channel
 				if (RenderPassData.PSPermutationVector.Get<IcvfxShaderPermutation::FIcvfxShaderInnerCamera>())
 				{
-					// Render additive camera frame
-					GraphicsPSOInit.BlendState = TStaticBlendState <CW_RGBA, BO_Add, BF_SourceAlpha, BF_InverseSourceAlpha, BO_Add, BF_InverseDestAlpha, BF_One>::GetRHI();
+					// Multicam rendering
+					GraphicsPSOInit.BlendState = TStaticBlendState <CW_RGBA, BO_Add, BF_SourceAlpha, BF_InverseSourceAlpha, BO_Add, BF_Zero, BF_One>::GetRHI();
 				}
 				else
 				{
-					// Render additive overlay frame
-					GraphicsPSOInit.BlendState = TStaticBlendState <CW_RGBA, BO_Add, BF_InverseSourceAlpha, BF_SourceAlpha, BO_Add, BF_InverseDestAlpha, BF_One>::GetRHI();
+					// [optional] Final render pass for LC overlays
+					GraphicsPSOInit.BlendState = TStaticBlendState <CW_RGBA, BO_Add, BF_One, BF_SourceAlpha, BO_Add, BF_Zero, BF_One>::GetRHI();
 				}
 			}
 
@@ -963,6 +970,11 @@ public:
 		case EVarIcvfxMPCDIShaderType::Disable:
 			return false;
 			break;
+
+		case EVarIcvfxMPCDIShaderType::Default:
+		default:
+			// Use default icvfx render
+			break;
 		};
 
 		return true;
