Removing almost all uses of Renderer/Private from plugins and modules.

diff --git a/Engine/Plugins/Compositing/OpenColorIO/Source/OpenColorIO/OpenColorIO.Build.cs b/Engine/Plugins/Compositing/OpenColorIO/Source/OpenColorIO/OpenColorIO.Build.cs
--- a/Engine/Plugins/Compositing/OpenColorIO/Source/OpenColorIO/OpenColorIO.Build.cs
+++ b/Engine/Plugins/Compositing/OpenColorIO/Source/OpenColorIO/OpenColorIO.Build.cs
@@ -35,7 +35,6 @@ namespace UnrealBuildTool.Rules
 
 			PrivateIncludePaths.AddRange(
 				new string[] {
-					System.IO.Path.Combine(GetModuleDirectory("Renderer"), "Private"), //required for FPostProcessMaterialInputs
 				});
 
 			PublicDependencyModuleNames.AddRange(
diff --git a/Engine/Plugins/Compositing/OpenColorIO/Source/OpenColorIO/Private/OpenColorIODisplayExtension.cpp b/Engine/Plugins/Compositing/OpenColorIO/Source/OpenColorIO/Private/OpenColorIODisplayExtension.cpp
--- a/Engine/Plugins/Compositing/OpenColorIO/Source/OpenColorIO/Private/OpenColorIODisplayExtension.cpp
+++ b/Engine/Plugins/Compositing/OpenColorIO/Source/OpenColorIO/Private/OpenColorIODisplayExtension.cpp
@@ -16,11 +16,9 @@
 #include "SceneView.h"
 #include "ScreenPass.h"
 #include "CommonRenderResources.h"
-#include "PostProcess/PostProcessing.h"
 #include "Containers/DynamicRHIResourceArray.h"
 // for FPostProcessMaterialInputs
-#include "PostProcess/PostProcessMaterial.h"
-
+#include "PostProcess/PostProcessMaterialInputs.h"
 
 float FOpenColorIODisplayExtension::DefaultDisplayGamma = 2.2f;
 
@@ -112,14 +110,13 @@ FScreenPassTexture FOpenColorIODisplayExtension::PostProcessPassAfterTonemap_Ren
 {
 	const FScreenPassTexture& SceneColor = InOutInputs.GetInput(EPostProcessMaterialInput::SceneColor);
 	check(SceneColor.IsValid());
-	checkSlow(View.bIsViewInfo);
-	const FViewInfo& ViewInfo = static_cast<const FViewInfo&>(View);
+
 	FScreenPassRenderTarget Output = InOutInputs.OverrideOutput;
 
 	// If the override output is provided, it means that this is the last pass in post processing.
 	if (!Output.IsValid())
 	{
-		Output = FScreenPassRenderTarget::CreateFromInput(GraphBuilder, SceneColor, ViewInfo.GetOverwriteLoadAction(), TEXT("OCIORenderTarget"));
+		Output = FScreenPassRenderTarget::CreateFromInput(GraphBuilder, SceneColor, View.GetOverwriteLoadAction(), TEXT("OCIORenderTarget"));
 	}
 
 	const float EngineDisplayGamma = View.Family->RenderTarget->GetDisplayGamma();
@@ -129,7 +126,7 @@ FScreenPassTexture FOpenColorIODisplayExtension::PostProcessPassAfterTonemap_Ren
 
 	FOpenColorIORendering::AddPass_RenderThread(
 		GraphBuilder,
-		ViewInfo,
+		View,
 		SceneColor,
 		Output,
 		CachedResourcesRenderThread,
diff --git a/Engine/Plugins/Compositing/OpenColorIO/Source/OpenColorIO/Private/OpenColorIORendering.cpp b/Engine/Plugins/Compositing/OpenColorIO/Source/OpenColorIO/Private/OpenColorIORendering.cpp
--- a/Engine/Plugins/Compositing/OpenColorIO/Source/OpenColorIO/Private/OpenColorIORendering.cpp
+++ b/Engine/Plugins/Compositing/OpenColorIO/Source/OpenColorIO/Private/OpenColorIORendering.cpp
@@ -20,7 +20,6 @@
 #include "SceneInterface.h"
 #include "ScreenPass.h"
 #include "TextureResource.h"
-#include "ScenePrivate.h"
 
 
 namespace {
@@ -43,7 +42,7 @@ namespace {
 // static
 void FOpenColorIORendering::AddPass_RenderThread(
 	FRDGBuilder& GraphBuilder,
-	const FViewInfo& View,
+	const FSceneView& View,
 	const FScreenPassTexture& Input,
 	const FScreenPassRenderTarget& Output,
 	const FOpenColorIORenderPassResources& InPassResource,
diff --git a/Engine/Plugins/Compositing/OpenColorIO/Source/OpenColorIO/Public/OpenColorIORendering.h b/Engine/Plugins/Compositing/OpenColorIO/Source/OpenColorIO/Public/OpenColorIORendering.h
--- a/Engine/Plugins/Compositing/OpenColorIO/Source/OpenColorIO/Public/OpenColorIORendering.h
+++ b/Engine/Plugins/Compositing/OpenColorIO/Source/OpenColorIO/Public/OpenColorIORendering.h
@@ -10,7 +10,7 @@ class FRDGBuilder;
 struct FScreenPassRenderTarget;
 struct FScreenPassTexture;
 class FTextureResource;
-class FViewInfo;
+class FSceneView;
 class UTexture;
 class UTextureRenderTarget2D;
 class UWorld;
@@ -50,7 +50,7 @@ public:
 	 */
 	static void AddPass_RenderThread(
 		FRDGBuilder& GraphBuilder,
-		const FViewInfo& View,
+		const FSceneView& View,
 		const FScreenPassTexture& Input,
 		const FScreenPassRenderTarget& Output,
 		const FOpenColorIORenderPassResources& InPassInfo,
diff --git a/Engine/Plugins/Experimental/ColorCorrectRegions/Source/ColorCorrectRegions/Private/ColorCorrectRegionsSceneViewExtension.cpp b/Engine/Plugins/Experimental/ColorCorrectRegions/Source/ColorCorrectRegions/Private/ColorCorrectRegionsSceneViewExtension.cpp
--- a/Engine/Plugins/Experimental/ColorCorrectRegions/Source/ColorCorrectRegions/Private/ColorCorrectRegionsSceneViewExtension.cpp
+++ b/Engine/Plugins/Experimental/ColorCorrectRegions/Source/ColorCorrectRegions/Private/ColorCorrectRegionsSceneViewExtension.cpp
@@ -735,7 +735,7 @@ void FColorCorrectRegionsSceneViewExtension::PrePostProcessPass_RenderThread(FRD
 		// Because we are not using proxy material, but plain global shader, we need to setup Scene textures ourselves.
 		// We don't need to do this per region.
 		check(View.bIsViewInfo);
-		FSceneTextureShaderParameters SceneTextures = CreateSceneTextureShaderParameters(GraphBuilder, ((const FViewInfo&)View).GetSceneTexturesChecked(), View.GetFeatureLevel(), ESceneTextureSetupMode::All);
+		FSceneTextureShaderParameters SceneTextures = CreateSceneTextureShaderParameters(GraphBuilder, View, ESceneTextureSetupMode::All);
 
 		WorldSubsystem->SortRegionsByDistance(View.ViewLocation);
 		{
diff --git a/Engine/Plugins/Experimental/RemoteSession/Source/RemoteSession/Private/Channels/RemoteSessionARCameraChannel.cpp b/Engine/Plugins/Experimental/RemoteSession/Source/RemoteSession/Private/Channels/RemoteSessionARCameraChannel.cpp
--- a/Engine/Plugins/Experimental/RemoteSession/Source/RemoteSession/Private/Channels/RemoteSessionARCameraChannel.cpp
+++ b/Engine/Plugins/Experimental/RemoteSession/Source/RemoteSession/Private/Channels/RemoteSessionARCameraChannel.cpp
@@ -23,7 +23,12 @@
 #include "IImageWrapperModule.h"
 #include "UObject/Package.h"
 #include "PostProcess/SceneFilterRendering.h"
-#include "ScenePrivate.h"
+
+#include "MaterialShader.h"
+#include "Containers/DynamicRHIResourceArray.h"
+#include "CommonRenderResources.h"
+#include "ScreenPass.h"
+#include "SceneRenderTargetParameters.h"
 
 #define CAMERA_MESSAGE_ADDRESS TEXT("/ARCamera")
 
@@ -254,7 +259,7 @@ void FARCameraSceneViewExtension::RenderARCamera_RenderThread(FRDGBuilder& Graph
 	auto* PassParameters = GraphBuilder.AllocParameters<FRenderARCameraPassParameters>();
 	PassParameters->RenderTargets[0] = FRenderTargetBinding(ViewFamilyTexture, ERenderTargetLoadAction::ELoad);
 	check(InView.bIsViewInfo);
-	PassParameters->SceneTextures = CreateSceneTextureUniformBuffer(GraphBuilder, ((const FViewInfo&)InView).GetSceneTexturesChecked(), InView.FeatureLevel, ESceneTextureSetupMode::None);
+	PassParameters->SceneTextures = CreateSceneTextureUniformBuffer(GraphBuilder, InView, ESceneTextureSetupMode::None);
 
 	GraphBuilder.AddPass(
 		RDG_EVENT_NAME("ARCameraOverlay"),
diff --git a/Engine/Plugins/Experimental/SimpleHMD/Source/SimpleHMD/Private/SimpleHMD.cpp b/Engine/Plugins/Experimental/SimpleHMD/Source/SimpleHMD/Private/SimpleHMD.cpp
--- a/Engine/Plugins/Experimental/SimpleHMD/Source/SimpleHMD/Private/SimpleHMD.cpp
+++ b/Engine/Plugins/Experimental/SimpleHMD/Source/SimpleHMD/Private/SimpleHMD.cpp
@@ -6,7 +6,8 @@
 #include "ISimpleHMDPlugin.h"
 #include "PostProcess/PostProcessHMD.h"
 #include "GameFramework/WorldSettings.h"
-#include "ScenePrivate.h"
+//#include "ScenePrivate.h"
+#include "SceneView.h"
 
 //---------------------------------------------------
 // SimpleHMD Plugin Implementation
@@ -171,7 +172,7 @@ void FSimpleHMD::DrawDistortionMesh_RenderThread(struct FHeadMountedDisplayPassC
 	float ClipSpaceQuadZ = 0.0f;
 	FMatrix QuadTexTransform = FMatrix::Identity;
 	FMatrix QuadPosTransform = FMatrix::Identity;
-	const FViewInfo& View = Context.View;
+	const FSceneView& View = Context.View;
 	const FIntRect SrcRect = View.UnscaledViewRect;
 
 	FRHICommandListImmediate& RHICmdList = Context.RHICmdList;
diff --git a/Engine/Plugins/Experimental/SimpleHMD/Source/SimpleHMD/SimpleHMD.Build.cs b/Engine/Plugins/Experimental/SimpleHMD/Source/SimpleHMD/SimpleHMD.Build.cs
--- a/Engine/Plugins/Experimental/SimpleHMD/Source/SimpleHMD/SimpleHMD.Build.cs
+++ b/Engine/Plugins/Experimental/SimpleHMD/Source/SimpleHMD/SimpleHMD.Build.cs
@@ -8,7 +8,6 @@ namespace UnrealBuildTool.Rules
 		{
 			PrivateIncludePaths.AddRange(
 				new string[] {
-					System.IO.Path.Combine(GetModuleDirectory("Renderer"), "Private"),
 				}
 				);
 
diff --git a/Engine/Plugins/Experimental/Water/Source/Runtime/Private/WaterInfoRendering.cpp b/Engine/Plugins/Experimental/Water/Source/Runtime/Private/WaterInfoRendering.cpp
--- a/Engine/Plugins/Experimental/Water/Source/Runtime/Private/WaterInfoRendering.cpp
+++ b/Engine/Plugins/Experimental/Water/Source/Runtime/Private/WaterInfoRendering.cpp
@@ -10,7 +10,6 @@
 #include "RHIStaticStates.h"
 #include "WaterBodySceneProxy.h"
 #include "Engine/TextureRenderTarget2D.h"
-#include "ScenePrivate.h"
 #include "SceneCaptureRendering.h"
 #include "PostProcess/SceneFilterRendering.h"
 #include "Math/OrthoMatrix.h"
@@ -21,6 +20,10 @@
 #include "LandscapeModule.h"
 #include "TextureResource.h"
 #include "WaterBodyComponent.h"
+#include "Containers/StridedView.h"
+
+#include "ScenePrivate.h"
+#include "SceneRendering.h"
 
 static int32 RenderCaptureNextWaterInfoDraws = 0;
 static FAutoConsoleVariableRef CVarRenderCaptureNextWaterInfoDraws(
@@ -466,14 +469,12 @@ static void UpdateWaterInfoRendering_RenderThread(
 		
 		if (DepthRenderer->Scene->GetShadingPath() == EShadingPath::Mobile)
 		{
-			const FMinimalSceneTextures& SceneTextures = View.GetSceneTextures();
 			RDG_EVENT_SCOPE(GraphBuilder, "CaptureSceneColor");
 			CopySceneCaptureComponentToTarget(
 				GraphBuilder,
-				SceneTextures,
 				DepthTexture,
 				DepthRenderer->ViewFamily,
-				DepthRenderer->Views);
+				MakeStridedViewOfBase<const FSceneView>(MakeArrayView(DepthRenderer->Views)));
 		}
 		else
 		{
@@ -530,10 +531,9 @@ static void UpdateWaterInfoRendering_RenderThread(
 			RDG_EVENT_SCOPE(GraphBuilder, "CaptureSceneColor");
 			CopySceneCaptureComponentToTarget(
 				GraphBuilder,
-				SceneTextures,
 				ColorTexture,
 				ColorRenderer->ViewFamily,
-				ColorRenderer->Views);
+				MakeStridedViewOfBase<const FSceneView>(MakeArrayView(ColorRenderer->Views)));
 		}
 
 		// We currently can't have multiple scene renderers run within the same RDGBuilder. Therefore, we must
@@ -587,10 +587,9 @@ static void UpdateWaterInfoRendering_RenderThread(
 			RDG_EVENT_SCOPE(GraphBuilder, "CaptureSceneColor");
 			CopySceneCaptureComponentToTarget(
 				GraphBuilder,
-				SceneTextures,
 				DilationTexture,
 				DilationRenderer->ViewFamily,
-				DilationRenderer->Views);
+				MakeStridedViewOfBase<const FSceneView>(MakeArrayView(DilationRenderer->Views)));
 		}
 		else
 		{
diff --git a/Engine/Plugins/Media/PixelStreaming/Source/PixelStreaming/PixelStreaming.Build.cs b/Engine/Plugins/Media/PixelStreaming/Source/PixelStreaming/PixelStreaming.Build.cs
--- a/Engine/Plugins/Media/PixelStreaming/Source/PixelStreaming/PixelStreaming.Build.cs
+++ b/Engine/Plugins/Media/PixelStreaming/Source/PixelStreaming/PixelStreaming.Build.cs
@@ -73,7 +73,6 @@ namespace UnrealBuildTool.Rules
 			PrivateIncludePaths.AddRange(new string[]
 			{
 				Path.Combine(GetModuleDirectory("AudioMixer"), "Private"),
-				Path.Combine(GetModuleDirectory("Renderer"), "Private"),
 			});
 
 			PrivateDependencyModuleNames.AddRange(new string[]
diff --git a/Engine/Plugins/Media/PixelStreaming/Source/PixelStreaming/Private/UtilsRender.h b/Engine/Plugins/Media/PixelStreaming/Source/PixelStreaming/Private/UtilsRender.h
--- a/Engine/Plugins/Media/PixelStreaming/Source/PixelStreaming/Private/UtilsRender.h
+++ b/Engine/Plugins/Media/PixelStreaming/Source/PixelStreaming/Private/UtilsRender.h
@@ -13,7 +13,6 @@
 #include "MediaShaders.h"
 #include "RHI.h"
 #include "ScreenPass.h"
-#include "PostProcess/SceneFilterRendering.h"
 
 namespace UE::PixelStreaming
 {
@@ -171,95 +170,4 @@ namespace UE::PixelStreaming
 		}
 		GraphBuilder.Execute();
 	}
-
-
-
-	inline void CopyTextureToReadbackTexture(FTextureRHIRef SourceTexture, TSharedPtr<FRHIGPUTextureReadback> GPUTextureReadback, void* OutBuffer)
-	{
-		FRHICommandListImmediate& RHICmdList = FRHICommandListExecutor::GetImmediateCommandList();
-		FRDGBuilder GraphBuilder(RHICmdList);
-
-		FRDGTextureRef RDGSourceTexture = RegisterExternalTexture(GraphBuilder, SourceTexture, TEXT("SourceCopyTextureToReadbackTexture"));
-		FRDGTextureRef RDGStagingTexture = nullptr;
-
-		if (RDGSourceTexture->Desc.Format != EPixelFormat::PF_B8G8R8A8)
-		{
-			// We need the pixel format to be BGRA8 so we first draw it to a staging texture
-			{
-				const FRDGTextureDesc Desc(FRDGTextureDesc::Create2D(RDGSourceTexture->Desc.Extent, PF_B8G8R8A8, FClearValueBinding::None, TexCreate_ShaderResource | TexCreate_RenderTargetable));
-				RDGStagingTexture = GraphBuilder.CreateTexture(Desc, TEXT("StagingCopyTextureToReadbackTexture"));
-			}
-
-			FGlobalShaderMap* ShaderMap = GetGlobalShaderMap(GMaxRHIFeatureLevel);
-
-			TShaderMapRef<FScreenPassVS> VertexShader(ShaderMap);
-
-			// Setup the pixel shader
-			TShaderMapRef<FCopyRectPS> PixelShader(ShaderMap);
-
-			FCopyRectPS::FParameters* PixelShaderParameters = GraphBuilder.AllocParameters<FCopyRectPS::FParameters>();
-			PixelShaderParameters->InputTexture = RDGSourceTexture;
-			PixelShaderParameters->InputSampler = TStaticSamplerState<SF_Bilinear, AM_Clamp, AM_Clamp, AM_Clamp>::GetRHI();
-			PixelShaderParameters->RenderTargets[0] = FRenderTargetBinding(RDGStagingTexture, ERenderTargetLoadAction::ELoad);
-
-			ClearUnusedGraphResources(PixelShader, PixelShaderParameters);
-
-			// We are not doing any clever blending stuff so we just use defaults here
-			FRHIBlendState* BlendState = FScreenPassPipelineState::FDefaultBlendState::GetRHI();
-			FRHIDepthStencilState* DepthStencilState = FScreenPassPipelineState::FDefaultDepthStencilState::GetRHI();
-
-			// Create the pipline state that will execute
-			const FScreenPassPipelineState PipelineState(VertexShader, PixelShader, BlendState, DepthStencilState);
-
-			// Add the pass the the graph builder
-			GraphBuilder.AddPass(
-				RDG_EVENT_NAME("PixelStreamingChangePixelFormat"),
-				PixelShaderParameters,
-				ERDGPassFlags::Raster,
-				[PipelineState, Extent = RDGSourceTexture->Desc.Extent, PixelShader, PixelShaderParameters](FRHICommandList& RHICmdList) {
-					PipelineState.Validate();
-
-					RHICmdList.SetViewport(0.0f, 0.0f, 0.0f, Extent.X, Extent.Y, 1.0f);
-
-					SetScreenPassPipelineState(RHICmdList, PipelineState);
-
-					SetShaderParameters(RHICmdList, PixelShader, PixelShader.GetPixelShader(), *PixelShaderParameters);
-
-					DrawRectangle(
-						RHICmdList,
-						0, 0, Extent.X, Extent.Y,
-						0, 0, Extent.X, Extent.Y,
-						Extent,
-						Extent,
-						PipelineState.VertexShader,
-						EDRF_UseTriangleOptimization);
-				});
-		}
-		else
-		{
-			// Otherwise if the PixelFormat is already BGRA8 we can just use the SourceTexture as the staging texture
-			RDGStagingTexture = RDGSourceTexture;
-		}
-
-		// Do the copy from staging RBGA8 texture into the readback
-		AddEnqueueCopyPass(GraphBuilder, GPUTextureReadback.Get(), RDGStagingTexture);
-
-		GraphBuilder.Execute();
-
-		// Lock and copy out the content of the TextureReadback to the CPU
-		int32 OutRowPitchInPixels;
-		int32 BlockSize = GPixelFormats[RDGStagingTexture->Desc.Format].BlockBytes;
-		void* ResultsBuffer = GPUTextureReadback->Lock(OutRowPitchInPixels);
-		if (RDGSourceTexture->Desc.Extent.X == OutRowPitchInPixels)
-		{
-			// Source pixel width is the same as the stride of the result buffer (ie no padding), we can do a plain memcpy
-			FPlatformMemory::Memcpy(OutBuffer, ResultsBuffer, (RDGSourceTexture->Desc.Extent.X * RDGSourceTexture->Desc.Extent.Y * BlockSize));
-		}
-		else
-		{
-			// Source pixel width differs from the stride of the result buffer (ie padding), do a memcpy that accounts for this
-			MemCpyStride(OutBuffer, ResultsBuffer, RDGSourceTexture->Desc.Extent.X * BlockSize, OutRowPitchInPixels * BlockSize, RDGSourceTexture->Desc.Extent.Y);
-		}
-		GPUTextureReadback->Unlock();
-	}
 } // namespace UE::PixelStreaming
diff --git a/Engine/Plugins/Runtime/HairStrands/Source/HairStrandsCore/HairStrandsCore.Build.cs b/Engine/Plugins/Runtime/HairStrands/Source/HairStrandsCore/HairStrandsCore.Build.cs
--- a/Engine/Plugins/Runtime/HairStrands/Source/HairStrandsCore/HairStrandsCore.Build.cs
+++ b/Engine/Plugins/Runtime/HairStrands/Source/HairStrandsCore/HairStrandsCore.Build.cs
@@ -10,7 +10,6 @@ namespace UnrealBuildTool.Rules
 			PrivateIncludePaths.AddRange(
 				new string[] {
 					ModuleDirectory + "/Private",
-					System.IO.Path.Combine(GetModuleDirectory("Renderer"), "Private"),
 					EngineDirectory + "/Shaders/Shared",
 				});
 			PublicIncludePaths.Add(ModuleDirectory + "/Public");
diff --git a/Engine/Plugins/Runtime/HairStrands/Source/HairStrandsCore/Private/GroomTextureBuilder.cpp b/Engine/Plugins/Runtime/HairStrands/Source/HairStrandsCore/Private/GroomTextureBuilder.cpp
--- a/Engine/Plugins/Runtime/HairStrands/Source/HairStrandsCore/Private/GroomTextureBuilder.cpp
+++ b/Engine/Plugins/Runtime/HairStrands/Source/HairStrandsCore/Private/GroomTextureBuilder.cpp
@@ -12,7 +12,7 @@
 #include "RHIStaticStates.h"
 #include "ShaderParameterStruct.h"
 #include "GlobalShader.h"
-#include "ShaderPrint.h"
+//#include "ShaderPrint.h"
 #include "CommonRenderResources.h"
 #include "HairStrandsMeshProjection.h"
 #include "Engine/StaticMesh.h"
diff --git a/Engine/Plugins/Runtime/HairStrands/Source/HairStrandsCore/Private/HairCardsBuilder.cpp b/Engine/Plugins/Runtime/HairStrands/Source/HairStrandsCore/Private/HairCardsBuilder.cpp
--- a/Engine/Plugins/Runtime/HairStrands/Source/HairStrandsCore/Private/HairCardsBuilder.cpp
+++ b/Engine/Plugins/Runtime/HairStrands/Source/HairStrandsCore/Private/HairCardsBuilder.cpp
@@ -19,7 +19,7 @@
 #include "Containers/ResourceArray.h"
 #include "Rendering/SkeletalMeshRenderData.h"
 #include "HAL/ConsoleManager.h"
-#include "ShaderPrint.h"
+#include "ShaderPrintParameters.h"
 #include "Async/ParallelFor.h"
 #include "Misc/ScopedSlowTask.h"
 #include "CommonRenderResources.h"
diff --git a/Engine/Plugins/Runtime/HairStrands/Source/HairStrandsCore/Private/HairStrandsRendering.cpp b/Engine/Plugins/Runtime/HairStrands/Source/HairStrandsCore/Private/HairStrandsRendering.cpp
--- a/Engine/Plugins/Runtime/HairStrands/Source/HairStrandsCore/Private/HairStrandsRendering.cpp
+++ b/Engine/Plugins/Runtime/HairStrands/Source/HairStrandsCore/Private/HairStrandsRendering.cpp
@@ -16,7 +16,7 @@
 #include "Containers/ResourceArray.h"
 #include "Rendering/SkeletalMeshRenderData.h"
 #include "HAL/ConsoleManager.h"
-#include "ShaderPrint.h"
+#include "ShaderPrintParameters.h"
 #include "Async/ParallelFor.h"
 #include "RenderTargetPool.h"
 #include "GroomTextureBuilder.h"
diff --git a/Engine/Plugins/Runtime/Oculus/OculusVR/Source/OculusHMD/OculusHMD.Build.cs b/Engine/Plugins/Runtime/Oculus/OculusVR/Source/OculusHMD/OculusHMD.Build.cs
--- a/Engine/Plugins/Runtime/Oculus/OculusVR/Source/OculusHMD/OculusHMD.Build.cs
+++ b/Engine/Plugins/Runtime/Oculus/OculusVR/Source/OculusHMD/OculusHMD.Build.cs
@@ -12,7 +12,6 @@ namespace UnrealBuildTool.Rules
 			PrivateIncludePaths.AddRange(
 				new string[] {
 					// Relative to Engine\Plugins\Runtime\Oculus\OculusVR\Source
-					System.IO.Path.Combine(GetModuleDirectory("Renderer"), "Private"),
 					"../../../../../Source/Runtime/Engine/Classes/Components",
 					"../../../../../Source/Runtime/Engine/Classes/Kismet",
 				});
diff --git a/Engine/Plugins/Runtime/Oculus/OculusVR/Source/OculusHMD/Private/OculusHMD.cpp b/Engine/Plugins/Runtime/Oculus/OculusVR/Source/OculusHMD/Private/OculusHMD.cpp
--- a/Engine/Plugins/Runtime/Oculus/OculusVR/Source/OculusHMD/Private/OculusHMD.cpp
+++ b/Engine/Plugins/Runtime/Oculus/OculusVR/Source/OculusHMD/Private/OculusHMD.cpp
@@ -8,7 +8,7 @@
 #include "AnalyticsEventAttribute.h"
 #include "Slate/SceneViewport.h"
 #include "PostProcess/PostProcessHMD.h"
-#include "PostProcess/SceneRenderTargets.h"
+//#include "PostProcess/SceneRenderTargets.h"
 #include "ScreenRendering.h"
 #include "GameFramework/PlayerController.h"
 #include "Math/UnrealMathUtility.h"
diff --git a/Engine/Plugins/Runtime/Oculus/OculusVR/Source/OculusHMD/Private/OculusHMD_Layer.cpp b/Engine/Plugins/Runtime/Oculus/OculusVR/Source/OculusHMD/Private/OculusHMD_Layer.cpp
--- a/Engine/Plugins/Runtime/Oculus/OculusVR/Source/OculusHMD/Private/OculusHMD_Layer.cpp
+++ b/Engine/Plugins/Runtime/Oculus/OculusVR/Source/OculusHMD/Private/OculusHMD_Layer.cpp
@@ -15,7 +15,7 @@
 #include "OculusHMD.h"
 #include "OculusHMDModule.h"
 #include "OculusHMDPrivate.h"
-#include "PostProcess/SceneRenderTargets.h"
+//#include "PostProcess/SceneRenderTargets.h"
 #include "SceneUtils.h"
 #include "UObject/ConstructorHelpers.h"
 #include "XRThreadUtils.h"
diff --git a/Engine/Plugins/Runtime/Oculus/OculusVR/Source/OculusInput/OculusInput.Build.cs b/Engine/Plugins/Runtime/Oculus/OculusVR/Source/OculusInput/OculusInput.Build.cs
--- a/Engine/Plugins/Runtime/Oculus/OculusVR/Source/OculusInput/OculusInput.Build.cs
+++ b/Engine/Plugins/Runtime/Oculus/OculusVR/Source/OculusInput/OculusInput.Build.cs
@@ -34,7 +34,6 @@ namespace UnrealBuildTool.Rules
 				new string[] {
 					// Relative to Engine\Plugins\Runtime\Oculus\OculusVR\Source
 					"OculusHMD/Private",
-					System.IO.Path.Combine(GetModuleDirectory("Renderer"), "Private"),
 					"../../../../../Source/Runtime/Engine/Classes/Components",
 				});
 
diff --git a/Engine/Plugins/Runtime/Oculus/OculusVR/Source/OculusMR/OculusMR.Build.cs b/Engine/Plugins/Runtime/Oculus/OculusVR/Source/OculusMR/OculusMR.Build.cs
--- a/Engine/Plugins/Runtime/Oculus/OculusVR/Source/OculusMR/OculusMR.Build.cs
+++ b/Engine/Plugins/Runtime/Oculus/OculusVR/Source/OculusMR/OculusMR.Build.cs
@@ -38,7 +38,6 @@ namespace UnrealBuildTool.Rules
 					// Relative to Engine\Plugins\Runtime\Oculus\OculusVR\Source
 					"OculusHMD/Private",
                     "OculusInput/Private",
-					System.IO.Path.Combine(GetModuleDirectory("Renderer"), "Private"),
 					"../../../../../Source/Runtime/Engine/Classes/Components",
 					System.IO.Path.Combine(GetModuleDirectory("MediaAssets"), "Private"),
                 });
diff --git a/Engine/Plugins/Runtime/Oculus/OculusVR/Source/OculusMR/Private/OculusMRModule.cpp b/Engine/Plugins/Runtime/Oculus/OculusVR/Source/OculusMR/Private/OculusMRModule.cpp
--- a/Engine/Plugins/Runtime/Oculus/OculusVR/Source/OculusMR/Private/OculusMRModule.cpp
+++ b/Engine/Plugins/Runtime/Oculus/OculusVR/Source/OculusMR/Private/OculusMRModule.cpp
@@ -9,7 +9,7 @@
 #include "SceneCaptureComponent2D.h"
 #include "Engine/TextureRenderTarget2D.h"
 #include "EngineUtils.h"
-#include "PostProcess/SceneRenderTargets.h"
+//#include "PostProcess/SceneRenderTargets.h"
 #include "Kismet/GameplayStatics.h"
 #include "OculusHMDModule.h"
 #include "OculusHMD.h"
diff --git a/Engine/Plugins/Runtime/Oculus/OculusVR/Source/OculusOpenXRHMD/OculusOpenXRHMD.Build.cs b/Engine/Plugins/Runtime/Oculus/OculusVR/Source/OculusOpenXRHMD/OculusOpenXRHMD.Build.cs
--- a/Engine/Plugins/Runtime/Oculus/OculusVR/Source/OculusOpenXRHMD/OculusOpenXRHMD.Build.cs
+++ b/Engine/Plugins/Runtime/Oculus/OculusVR/Source/OculusOpenXRHMD/OculusOpenXRHMD.Build.cs
@@ -13,7 +13,6 @@ namespace UnrealBuildTool.Rules
                 new string[] {
                     // Relative to Engine\Plugins\Runtime\Oculus\OculusOpenXR\Source
                     "../../../OpenXR/Source/OpenXRHMD/Private",
-					System.IO.Path.Combine(GetModuleDirectory("Renderer"), "Private"),
 					"../../../../../Source/Runtime/Engine/Classes/Components",
                     "../../../../../Source/Runtime/Engine/Classes/Kismet",
                 });
diff --git a/Engine/Plugins/Runtime/OpenXR/Source/OpenXRHMD/OpenXRHMD.Build.cs b/Engine/Plugins/Runtime/OpenXR/Source/OpenXRHMD/OpenXRHMD.Build.cs
--- a/Engine/Plugins/Runtime/OpenXR/Source/OpenXRHMD/OpenXRHMD.Build.cs
+++ b/Engine/Plugins/Runtime/OpenXR/Source/OpenXRHMD/OpenXRHMD.Build.cs
@@ -15,7 +15,6 @@ namespace UnrealBuildTool.Rules
 				new string[] {
 					"OpenXRHMD/Private",
                     EngineDir + "/Source/ThirdParty/OpenXR/include",
-					System.IO.Path.Combine(GetModuleDirectory("Renderer"), "Private"),
 					// ... add other private include paths required here ...
 				}
 				);
diff --git a/Engine/Plugins/Runtime/OpenXR/Source/OpenXRHMD/Private/OpenXRHMD.cpp b/Engine/Plugins/Runtime/OpenXR/Source/OpenXRHMD/Private/OpenXRHMD.cpp
--- a/Engine/Plugins/Runtime/OpenXR/Source/OpenXRHMD/Private/OpenXRHMD.cpp
+++ b/Engine/Plugins/Runtime/OpenXR/Source/OpenXRHMD/Private/OpenXRHMD.cpp
@@ -16,7 +16,7 @@
 #include "Engine/Engine.h"
 #include "GameFramework/PlayerController.h"
 #include "Engine/LocalPlayer.h"
-#include "SceneRendering.h"
+//#include "SceneRendering.h"
 #include "PostProcess/PostProcessHMD.h"
 #include "GameFramework/WorldSettings.h"
 #include "Misc/CString.h"
diff --git a/Engine/Plugins/Runtime/Steam/SteamVR/Source/SteamVR/Private/SteamVRHMD.cpp b/Engine/Plugins/Runtime/Steam/SteamVR/Source/SteamVR/Private/SteamVRHMD.cpp
--- a/Engine/Plugins/Runtime/Steam/SteamVR/Source/SteamVR/Private/SteamVRHMD.cpp
+++ b/Engine/Plugins/Runtime/Steam/SteamVR/Source/SteamVR/Private/SteamVRHMD.cpp
@@ -7,7 +7,7 @@
 #include "Misc/CoreDelegates.h"
 #include "Misc/EngineVersion.h"
 #include "HardwareInfo.h"
-#include "RendererPrivate.h"
+//#include "RendererPrivate.h"
 #include "Slate/SceneViewport.h"
 #include "PostProcess/PostProcessHMD.h"
 #include "SteamVRFunctionLibrary.h"
@@ -15,6 +15,7 @@
 #include "GameFramework/PlayerController.h"
 #include "GameFramework/WorldSettings.h"
 #include "IHeadMountedDisplayVulkanExtensions.h"
+#include "DataDrivenShaderPlatformInfo.h"
 
 #include "EngineAnalytics.h"
 #include "Interfaces/IAnalyticsProvider.h"
diff --git a/Engine/Plugins/Runtime/Steam/SteamVR/Source/SteamVR/Private/SteamVRRender.cpp b/Engine/Plugins/Runtime/Steam/SteamVR/Source/SteamVR/Private/SteamVRRender.cpp
--- a/Engine/Plugins/Runtime/Steam/SteamVR/Source/SteamVR/Private/SteamVRRender.cpp
+++ b/Engine/Plugins/Runtime/Steam/SteamVR/Source/SteamVR/Private/SteamVRRender.cpp
@@ -7,8 +7,7 @@
 
 #include "SteamVRHMD.h"
 
-#include "RendererPrivate.h"
-#include "PostProcess/PostProcessHMD.h"
+#include "DataDrivenShaderPlatformInfo.h"
 #include "PipelineStateCache.h"
 #include "ClearQuad.h"
 #include "DefaultSpectatorScreenController.h"
diff --git a/Engine/Plugins/Runtime/Steam/SteamVR/Source/SteamVR/Private/SteamVRSpectatorScreen.cpp b/Engine/Plugins/Runtime/Steam/SteamVR/Source/SteamVR/Private/SteamVRSpectatorScreen.cpp
--- a/Engine/Plugins/Runtime/Steam/SteamVR/Source/SteamVR/Private/SteamVRSpectatorScreen.cpp
+++ b/Engine/Plugins/Runtime/Steam/SteamVR/Source/SteamVR/Private/SteamVRSpectatorScreen.cpp
@@ -7,11 +7,11 @@
 
 #include "SteamVRHMD.h"
 
-#include "RendererPrivate.h"
-#include "PostProcess/PostProcessHMD.h"
-#include "PipelineStateCache.h"
 #include "ClearQuad.h"
+#include "CommonRenderResources.h"
 #include "DefaultSpectatorScreenController.h"
+#include "PipelineStateCache.h"
+#include "RHIStaticStates.h"
 
 #if PLATFORM_LINUX
 #include "IVulkanDynamicRHI.h"
diff --git a/Engine/Plugins/Runtime/Steam/SteamVR/Source/SteamVR/SteamVR.Build.cs b/Engine/Plugins/Runtime/Steam/SteamVR/Source/SteamVR/SteamVR.Build.cs
--- a/Engine/Plugins/Runtime/Steam/SteamVR/Source/SteamVR/SteamVR.Build.cs
+++ b/Engine/Plugins/Runtime/Steam/SteamVR/Source/SteamVR/SteamVR.Build.cs
@@ -9,7 +9,6 @@ namespace UnrealBuildTool.Rules
 			PrivateIncludePaths.AddRange(
 				new string[] {
 					"SteamVR/Private",
-					System.IO.Path.Combine(GetModuleDirectory("Renderer"), "Private"),
 					// ... add other private include paths required here ...
 				}
 				);
diff --git a/Engine/Plugins/Runtime/nDisplay/Source/DisplayClusterLightCardEditor/DisplayClusterLightCardEditor.Build.cs b/Engine/Plugins/Runtime/nDisplay/Source/DisplayClusterLightCardEditor/DisplayClusterLightCardEditor.Build.cs
--- a/Engine/Plugins/Runtime/nDisplay/Source/DisplayClusterLightCardEditor/DisplayClusterLightCardEditor.Build.cs
+++ b/Engine/Plugins/Runtime/nDisplay/Source/DisplayClusterLightCardEditor/DisplayClusterLightCardEditor.Build.cs
@@ -43,7 +43,6 @@ public class DisplayClusterLightCardEditor : ModuleRules
 
 		PrivateIncludePaths.AddRange(
 			new string[] {
-				System.IO.Path.Combine(GetModuleDirectory("Renderer"), "Private"),
 			}
 		);
 	}
diff --git a/Engine/Plugins/Runtime/nDisplay/Source/DisplayClusterLightCardEditorShaders/DisplayClusterLightCardEditorShaders.Build.cs b/Engine/Plugins/Runtime/nDisplay/Source/DisplayClusterLightCardEditorShaders/DisplayClusterLightCardEditorShaders.Build.cs
--- a/Engine/Plugins/Runtime/nDisplay/Source/DisplayClusterLightCardEditorShaders/DisplayClusterLightCardEditorShaders.Build.cs
+++ b/Engine/Plugins/Runtime/nDisplay/Source/DisplayClusterLightCardEditorShaders/DisplayClusterLightCardEditorShaders.Build.cs
@@ -20,7 +20,6 @@ public class DisplayClusterLightCardEditorShaders : ModuleRules
 
 		PrivateIncludePaths.AddRange(
 			new string[] {
-				System.IO.Path.Combine(GetModuleDirectory("Renderer"), "Private"),
 			});
 	}
 }
diff --git a/Engine/Plugins/Runtime/nDisplay/Source/DisplayClusterLightCardEditorShaders/Private/DisplayClusterMeshProjectionRenderer.cpp b/Engine/Plugins/Runtime/nDisplay/Source/DisplayClusterLightCardEditorShaders/Private/DisplayClusterMeshProjectionRenderer.cpp
--- a/Engine/Plugins/Runtime/nDisplay/Source/DisplayClusterLightCardEditorShaders/Private/DisplayClusterMeshProjectionRenderer.cpp
+++ b/Engine/Plugins/Runtime/nDisplay/Source/DisplayClusterLightCardEditorShaders/Private/DisplayClusterMeshProjectionRenderer.cpp
@@ -19,8 +19,9 @@
 #include "MaterialDomain.h"
 #include "Materials/Material.h"
 #include "DataDrivenShaderPlatformInfo.h"
-#include "PostProcess/SceneFilterRendering.h"
-#include "ScenePrivate.h"
+//#include "PostProcess/SceneFilterRendering.h"
+//#include "ScenePrivate.h"
+#include "ScreenPass.h"
 
 //////////////////////////////////////////////////////////////////////////
 // Mesh Pass Processor
@@ -274,7 +275,7 @@ using FMeshProjectionNormalPS = FMeshProjectionPS<EDisplayClusterMeshProjectionO
 IMPLEMENT_MATERIAL_SHADER_TYPE(template<>, FMeshProjectionColorPS, TEXT("/Plugin/nDisplay/Private/MeshProjectionShaders.usf"), TEXT("MainPS"), SF_Pixel);
 IMPLEMENT_MATERIAL_SHADER_TYPE(template<>, FMeshProjectionNormalPS, TEXT("/Plugin/nDisplay/Private/MeshProjectionShaders.usf"), TEXT("NormalPS"), SF_Pixel);
 
-FMeshProjectionPassParameters* CreateProjectionPassParameters(FRDGBuilder& GraphBuilder, const FViewInfo* View, const FDisplayClusterMeshProjectionRenderSettings& RenderSettings)
+FMeshProjectionPassParameters* CreateProjectionPassParameters(FRDGBuilder& GraphBuilder, const FSceneView* View, const FDisplayClusterMeshProjectionRenderSettings& RenderSettings)
 {
 	FMeshProjectionPassParameters* PassParameters = GraphBuilder.AllocParameters<FMeshProjectionPassParameters>();
 	PassParameters->View = View->ViewUniformBuffer;
@@ -720,48 +721,6 @@ protected:
 	}
 };
 
-namespace
-{
-	template<typename TSetupFunction>
-	void DrawScreenPass(
-		FRHICommandList& RHICmdList,
-		const FSceneView& View,
-		const FScreenPassTextureViewport& OutputViewport,
-		const FScreenPassTextureViewport& InputViewport,
-		const FScreenPassPipelineState& PipelineState,
-		TSetupFunction SetupFunction)
-	{
-		PipelineState.Validate();
-
-		const FIntRect InputRect = InputViewport.Rect;
-		const FIntPoint InputSize = InputViewport.Extent;
-		const FIntRect OutputRect = OutputViewport.Rect;
-		const FIntPoint OutputSize = OutputRect.Size();
-
-		RHICmdList.SetViewport(OutputRect.Min.X, OutputRect.Min.Y, 0.0f, OutputRect.Max.X, OutputRect.Max.Y, 1.0f);
-
-		SetScreenPassPipelineState(RHICmdList, PipelineState);
-
-		// Setting up buffers.
-		SetupFunction(RHICmdList);
-
-		FIntPoint LocalOutputPos(FIntPoint::ZeroValue);
-		FIntPoint LocalOutputSize(OutputSize);
-		EDrawRectangleFlags DrawRectangleFlags = EDRF_UseTriangleOptimization;
-
-		DrawPostProcessPass(
-			RHICmdList,
-			LocalOutputPos.X, LocalOutputPos.Y, LocalOutputSize.X, LocalOutputSize.Y,
-			InputRect.Min.X, InputRect.Min.Y, InputRect.Width(), InputRect.Height(),
-			OutputSize,
-			InputSize,
-			PipelineState.VertexShader,
-			View.StereoViewIndex,
-			false,
-			DrawRectangleFlags);
-	}
-} //! namespace
-
 //////////////////////////////////////////////////////////////////////////
 // FDisplayClusterMeshProjectionRenderer
 
@@ -981,7 +940,7 @@ void FDisplayClusterMeshProjectionRenderer::Render(FCanvas* Canvas, FSceneInterf
 			NewInitOptions.ViewFamily = &ViewFamily;
 
 			GetRendererModule().CreateAndInitSingleView(RHICmdList, &ViewFamily, &NewInitOptions);
-			FViewInfo* View = (FViewInfo*)ViewFamily.Views[0];
+			const FSceneView* View = ViewFamily.Views[0];
 
 			FRDGTextureRef OutputTexture = GraphBuilder.RegisterExternalTexture(CreateRenderTarget(RenderTarget->GetRenderTargetTexture(), TEXT("ViewRenderTarget")));
 			FRenderTargetBinding OutputRenderTargetBinding(OutputTexture, ERenderTargetLoadAction::ELoad);
@@ -1016,7 +975,7 @@ void FDisplayClusterMeshProjectionRenderer::Render(FCanvas* Canvas, FSceneInterf
 }
 
 void FDisplayClusterMeshProjectionRenderer::RenderColorOutput(FRDGBuilder& GraphBuilder,
-	const FViewInfo* View,
+	const FSceneView* View,
 	const FDisplayClusterMeshProjectionRenderSettings& RenderSettings,
 	FRenderTargetBinding& OutputRenderTargetBinding)
 {
@@ -1054,7 +1013,7 @@ void FDisplayClusterMeshProjectionRenderer::RenderColorOutput(FRDGBuilder& Graph
 		ScreenPassParameters->InputSampler = TStaticSamplerState<>::GetRHI();
 		ScreenPassParameters->RenderTargets[0] = OutputRenderTargetBinding;
 
-		FGlobalShaderMap* GlobalShaderMap = GetGlobalShaderMap(GMaxRHIFeatureLevel);
+		FGlobalShaderMap* GlobalShaderMap = GetGlobalShaderMap(View->FeatureLevel);
 		TShaderMapRef<FScreenPassVS> ScreenPassVS(GlobalShaderMap);
 		TShaderMapRef<FCopyRectPS> CopyPixelShader(GlobalShaderMap);
 
@@ -1073,6 +1032,7 @@ void FDisplayClusterMeshProjectionRenderer::RenderColorOutput(FRDGBuilder& Graph
 				RegionViewport,
 				RegionViewport,
 				FScreenPassPipelineState(ScreenPassVS, CopyPixelShader, DefaultBlendState),
+				EScreenPassDrawFlags::None,
 				[&](FRHICommandList&)
 			{
 				SetShaderParameters(RHICmdList, CopyPixelShader, CopyPixelShader.GetPixelShader(), *ScreenPassParameters);
@@ -1082,7 +1042,7 @@ void FDisplayClusterMeshProjectionRenderer::RenderColorOutput(FRDGBuilder& Graph
 }
 
 void FDisplayClusterMeshProjectionRenderer::RenderHitProxyOutput(FRDGBuilder& GraphBuilder,
-	const FViewInfo* View,
+	const FSceneView* View,
 	const FDisplayClusterMeshProjectionRenderSettings& RenderSettings,
 	FRenderTargetBinding& OutputRenderTargetBinding,
 	FHitProxyConsumer* HitProxyConsumer)
@@ -1107,7 +1067,7 @@ void FDisplayClusterMeshProjectionRenderer::RenderHitProxyOutput(FRDGBuilder& Gr
 		ScreenPassParameters->InputSampler = TStaticSamplerState<>::GetRHI();
 		ScreenPassParameters->RenderTargets[0] = OutputRenderTargetBinding;
 
-		FGlobalShaderMap* GlobalShaderMap = GetGlobalShaderMap(GMaxRHIFeatureLevel);
+		FGlobalShaderMap* GlobalShaderMap = GetGlobalShaderMap(View->FeatureLevel);
 		TShaderMapRef<FScreenPassVS> ScreenPassVS(GlobalShaderMap);
 		TShaderMapRef<FCopyRectPS> CopyPixelShader(GlobalShaderMap);
 
@@ -1126,6 +1086,7 @@ void FDisplayClusterMeshProjectionRenderer::RenderHitProxyOutput(FRDGBuilder& Gr
 				RegionViewport,
 				RegionViewport,
 				FScreenPassPipelineState(ScreenPassVS, CopyPixelShader, DefaultBlendState),
+				EScreenPassDrawFlags::None,
 				[&](FRHICommandList&)
 			{
 				SetShaderParameters(RHICmdList, CopyPixelShader, CopyPixelShader.GetPixelShader(), *ScreenPassParameters);
@@ -1135,7 +1096,7 @@ void FDisplayClusterMeshProjectionRenderer::RenderHitProxyOutput(FRDGBuilder& Gr
 }
 
 void FDisplayClusterMeshProjectionRenderer::RenderNormalsOutput(FRDGBuilder& GraphBuilder,
-	const FViewInfo* View,
+	const FSceneView* View,
 	const FDisplayClusterMeshProjectionRenderSettings& RenderSettings,
 	FRenderTargetBinding& OutputRenderTargetBinding)
 {
@@ -1165,7 +1126,7 @@ void FDisplayClusterMeshProjectionRenderer::RenderNormalsOutput(FRDGBuilder& Gra
 	}
 
 void FDisplayClusterMeshProjectionRenderer::AddBaseRenderPass(FRDGBuilder& GraphBuilder,
-	const FViewInfo* View,
+	const FSceneView* View,
 	const FDisplayClusterMeshProjectionRenderSettings& RenderSettings,
 	FRenderTargetBinding& OutputRenderTargetBinding,
 	FDepthStencilBinding& OutputDepthStencilBinding)
@@ -1187,7 +1148,7 @@ void FDisplayClusterMeshProjectionRenderer::AddBaseRenderPass(FRDGBuilder& Graph
 }
 
 void FDisplayClusterMeshProjectionRenderer::AddTranslucencyRenderPass(FRDGBuilder& GraphBuilder,
-	const FViewInfo* View,
+	const FSceneView* View,
 	const FDisplayClusterMeshProjectionRenderSettings& RenderSettings,
 	FRenderTargetBinding& OutputRenderTargetBinding,
 	FDepthStencilBinding& OutputDepthStencilBinding)
@@ -1209,7 +1170,7 @@ void FDisplayClusterMeshProjectionRenderer::AddTranslucencyRenderPass(FRDGBuilde
 }
 
 void FDisplayClusterMeshProjectionRenderer::AddHitProxyRenderPass(FRDGBuilder& GraphBuilder,
-	const FViewInfo* View,
+	const FSceneView* View,
 	const FDisplayClusterMeshProjectionRenderSettings& RenderSettings,
 	FRenderTargetBinding& OutputRenderTargetBinding,
 	FDepthStencilBinding& OutputDepthStencilBinding)
@@ -1231,7 +1192,7 @@ void FDisplayClusterMeshProjectionRenderer::AddHitProxyRenderPass(FRDGBuilder& G
 }
 
 void FDisplayClusterMeshProjectionRenderer::AddNormalsRenderPass(FRDGBuilder& GraphBuilder,
-	const FViewInfo* View,
+	const FSceneView* View,
 	const FDisplayClusterMeshProjectionRenderSettings& RenderSettings,
 	FRenderTargetBinding& OutputRenderTargetBinding,
 	FDepthStencilBinding& OutputDepthStencilBinding)
@@ -1279,7 +1240,8 @@ void GaussianSpatialFilter(float Sigma, TArray<float>& OutKernel)
 
 template<EMeshProjectionFilterType FilterType>
 void AddSeparableFilterPass(FRDGBuilder& GraphBuilder,
-	const FViewInfo* View,
+	FGlobalShaderMap* GlobalShaderMap,
+	const FSceneView* View,
 	FRDGTexture* SceneDepth,
 	FRDGTextureUAV* RWColorUAV,
 	FRDGTextureUAV* RWDepthUAV,
@@ -1309,7 +1271,7 @@ void AddSeparableFilterPass(FRDGBuilder& GraphBuilder,
 		GET_SCALAR_ARRAY_ELEMENT(HorizontalPassParameters->InteriorSpatialKernel, Index) = InteriorSpatialKernel[Index];
 	}
 
-	TShaderMapRef<FMeshProjectionNormalsFilterCS<FilterType>> ComputeShader(View->ShaderMap);
+	TShaderMapRef<FMeshProjectionNormalsFilterCS<FilterType>> ComputeShader(GlobalShaderMap);
 
 	FComputeShaderUtils::AddPass(
 		GraphBuilder,
@@ -1343,7 +1305,7 @@ void AddSeparableFilterPass(FRDGBuilder& GraphBuilder,
 }
 
 void FDisplayClusterMeshProjectionRenderer::AddNormalsFilterPass(FRDGBuilder& GraphBuilder,
-	const FViewInfo* View,
+	const FSceneView* View,
 	FRenderTargetBinding& OutputRenderTargetBinding,
 	FRDGTexture* SceneColor,
 	FRDGTexture* SceneDepth,
@@ -1371,6 +1333,8 @@ void FDisplayClusterMeshProjectionRenderer::AddNormalsFilterPass(FRDGBuilder& Gr
 	FRDGTextureRef RWStencil = GraphBuilder.CreateTexture(RWStencilDesc, TEXT("DisplayClusterMeshProjection.RWStencil"));
 	FRDGTextureUAVRef RWStencilUAV = GraphBuilder.CreateUAV(RWStencil);
 
+	FGlobalShaderMap* GlobalShaderMap = GetGlobalShaderMap(View->FeatureLevel);
+
 	// Copy the scene color and depth to RW buffers
 	{
 		FMeshProjectionNormalsCreateRWTexturesCS::FParameters* PassParameters = GraphBuilder.AllocParameters<FMeshProjectionNormalsCreateRWTexturesCS::FParameters>();
@@ -1383,7 +1347,7 @@ void FDisplayClusterMeshProjectionRenderer::AddNormalsFilterPass(FRDGBuilder& Gr
 		PassParameters->RWStencil = RWStencilUAV;
 		PassParameters->NormalCorrectionMatrix = NormalCorrectionMatrix;
 
-		TShaderMapRef<FMeshProjectionNormalsCreateRWTexturesCS> ComputeShader(View->ShaderMap);
+		TShaderMapRef<FMeshProjectionNormalsCreateRWTexturesCS> ComputeShader(GlobalShaderMap);
 
 		FComputeShaderUtils::AddPass(
 			GraphBuilder,
@@ -1400,17 +1364,17 @@ void FDisplayClusterMeshProjectionRenderer::AddNormalsFilterPass(FRDGBuilder& Gr
 	// Filter the RW buffers, dilating and blurring them appropriately
 	for (int32 Index = 0; Index < DepthDilatePasses; Index++)
 	{
-		AddSeparableFilterPass<EMeshProjectionFilterType::DepthDilate>(GraphBuilder, View, SceneDepth, RWColorUAV, RWDepthUAV, RWStencilUAV);
+		AddSeparableFilterPass<EMeshProjectionFilterType::DepthDilate>(GraphBuilder, GlobalShaderMap, View, SceneDepth, RWColorUAV, RWDepthUAV, RWStencilUAV);
 	}
 
 	for (int32 Index = 0; Index < DilatePasses; ++Index)
 	{
-		AddSeparableFilterPass<EMeshProjectionFilterType::Dilate>(GraphBuilder, View, SceneDepth, RWColorUAV, RWDepthUAV, RWStencilUAV);
+		AddSeparableFilterPass<EMeshProjectionFilterType::Dilate>(GraphBuilder, GlobalShaderMap, View, SceneDepth, RWColorUAV, RWDepthUAV, RWStencilUAV);
 	}
 
 	for (int32 Index = 0; Index < BlurPasses; Index++)
 	{
-		AddSeparableFilterPass<EMeshProjectionFilterType::Blur>(GraphBuilder, View, SceneDepth, RWColorUAV, RWDepthUAV, RWStencilUAV);
+		AddSeparableFilterPass<EMeshProjectionFilterType::Blur>(GraphBuilder, GlobalShaderMap, View, SceneDepth, RWColorUAV, RWDepthUAV, RWStencilUAV);
 	}
 
 	// Copy the RW buffers to the output render target
@@ -1424,8 +1388,8 @@ void FDisplayClusterMeshProjectionRenderer::AddNormalsFilterPass(FRDGBuilder& Gr
 		ScreenPassParameters->RWDepth = RWDepthUAV;
 		ScreenPassParameters->RenderTargets[0] = OutputRenderTargetBinding;
 
-		TShaderMapRef<FScreenPassVS> ScreenPassVS(View->ShaderMap);
-		TShaderMapRef<FMeshProjectionNormalsOutputPS> OutputNormalsPS(View->ShaderMap);
+		TShaderMapRef<FScreenPassVS> ScreenPassVS(GlobalShaderMap);
+		TShaderMapRef<FMeshProjectionNormalsOutputPS> OutputNormalsPS(GlobalShaderMap);
 
 		FRHIBlendState* DefaultBlendState = FScreenPassPipelineState::FDefaultBlendState::GetRHI();
 
@@ -1441,6 +1405,7 @@ void FDisplayClusterMeshProjectionRenderer::AddNormalsFilterPass(FRDGBuilder& Gr
 				InputViewport,
 				OutputViewport,
 				FScreenPassPipelineState(ScreenPassVS, OutputNormalsPS, DefaultBlendState),
+				EScreenPassDrawFlags::None,
 				[&](FRHICommandList&)
 			{
 				SetShaderParameters(RHICmdList, OutputNormalsPS, OutputNormalsPS.GetPixelShader(), *ScreenPassParameters);
@@ -1451,7 +1416,7 @@ void FDisplayClusterMeshProjectionRenderer::AddNormalsFilterPass(FRDGBuilder& Gr
 
 #if WITH_EDITOR
 void FDisplayClusterMeshProjectionRenderer::AddSelectionDepthRenderPass(FRDGBuilder& GraphBuilder, 
-	const FViewInfo* View,
+	const FSceneView* View,
 	const FDisplayClusterMeshProjectionRenderSettings& RenderSettings,
 	FDepthStencilBinding& OutputDepthStencilBinding)
 {
@@ -1471,7 +1436,7 @@ void FDisplayClusterMeshProjectionRenderer::AddSelectionDepthRenderPass(FRDGBuil
 }
 
 void FDisplayClusterMeshProjectionRenderer::AddSelectionOutlineScreenPass(FRDGBuilder& GraphBuilder,
-	const FViewInfo* View,
+	const FSceneView* View,
 	FRenderTargetBinding& OutputRenderTargetBinding,
 	FRDGTexture* SceneColor,
 	FRDGTexture* SceneDepth,
@@ -1497,8 +1462,10 @@ void FDisplayClusterMeshProjectionRenderer::AddSelectionOutlineScreenPass(FRDGBu
 	ScreenPassParameters->OutlineColor = FVector3f(View->SelectionOutlineColor);
 	ScreenPassParameters->SelectionHighlightIntensity = GEngine->SelectionHighlightIntensity;
 
-	TShaderMapRef<FScreenPassVS> ScreenPassVS(View->ShaderMap);
-	TShaderMapRef<FMeshProjectionSelectionOutlinePS> SelectionOutlinePS(View->ShaderMap);
+	FGlobalShaderMap* GlobalShaderMap = GetGlobalShaderMap(View->FeatureLevel);
+
+	TShaderMapRef<FScreenPassVS> ScreenPassVS(GlobalShaderMap);
+	TShaderMapRef<FMeshProjectionSelectionOutlinePS> SelectionOutlinePS(GlobalShaderMap);
 
 	FRHIBlendState* DefaultBlendState = FScreenPassPipelineState::FDefaultBlendState::GetRHI();
 
@@ -1514,6 +1481,7 @@ void FDisplayClusterMeshProjectionRenderer::AddSelectionOutlineScreenPass(FRDGBu
 			OutputViewport,
 			OutputViewport,
 			FScreenPassPipelineState(ScreenPassVS, SelectionOutlinePS, DefaultBlendState),
+			EScreenPassDrawFlags::None,
 			[&](FRHICommandList&)
 		{
 			SetShaderParameters(RHICmdList, SelectionOutlinePS, SelectionOutlinePS.GetPixelShader(), *ScreenPassParameters);
@@ -1523,7 +1491,7 @@ void FDisplayClusterMeshProjectionRenderer::AddSelectionOutlineScreenPass(FRDGBu
 #endif
 
 void FDisplayClusterMeshProjectionRenderer::AddSimpleElementPass(FRDGBuilder& GraphBuilder,
-	const FViewInfo* View,
+	const FSceneView* View,
 	FRenderTargetBinding& OutputRenderTargetBinding,
 	FSimpleElementCollector& ElementCollector)
 {
diff --git a/Engine/Plugins/Runtime/nDisplay/Source/DisplayClusterLightCardEditorShaders/Public/DisplayClusterMeshProjectionRenderer.h b/Engine/Plugins/Runtime/nDisplay/Source/DisplayClusterLightCardEditorShaders/Public/DisplayClusterMeshProjectionRenderer.h
--- a/Engine/Plugins/Runtime/nDisplay/Source/DisplayClusterLightCardEditorShaders/Public/DisplayClusterMeshProjectionRenderer.h
+++ b/Engine/Plugins/Runtime/nDisplay/Source/DisplayClusterLightCardEditorShaders/Public/DisplayClusterMeshProjectionRenderer.h
@@ -165,47 +165,47 @@ public:
 private:
 	/** Constructs the necessary render passes for the default output of the rendered primitives */
 	void RenderColorOutput(FRDGBuilder& GraphBuilder,
-		const FViewInfo* View,
+		const FSceneView* View,
 		const FDisplayClusterMeshProjectionRenderSettings& RenderSettings,
 		FRenderTargetBinding& OutputRenderTargetBinding);
 
 	/** Constructs the necessary render passes for the hit proxy output of the rendered primitives */
 	void RenderHitProxyOutput(FRDGBuilder& GraphBuilder,
-		const FViewInfo* View,
+		const FSceneView* View,
 		const FDisplayClusterMeshProjectionRenderSettings& RenderSettings,
 		FRenderTargetBinding& OutputRenderTargetBinding,
 		FHitProxyConsumer* HitProxyConsumer);
 
 	/** Constructs the necessary render passes for the normals/depth output of the rendered primitives */
 	void RenderNormalsOutput(FRDGBuilder& GraphBuilder,
-		const FViewInfo* View,
+		const FSceneView* View,
 		const FDisplayClusterMeshProjectionRenderSettings& RenderSettings,
 		FRenderTargetBinding& OutputRenderTargetBinding);
 
 	/** Adds a pass to perform the base render for the mesh projection. */
 	void AddBaseRenderPass(FRDGBuilder& GraphBuilder, 
-		const FViewInfo* View,
+		const FSceneView* View,
 		const FDisplayClusterMeshProjectionRenderSettings& RenderSettings,
 		FRenderTargetBinding& OutputRenderTargetBinding,
 		FDepthStencilBinding& OutputDepthStencilBinding);
 
 	/** Adds a pass to perform the translucency render for the mesh projection. */
 	void AddTranslucencyRenderPass(FRDGBuilder& GraphBuilder, 
-		const FViewInfo* View,
+		const FSceneView* View,
 		const FDisplayClusterMeshProjectionRenderSettings& RenderSettings,
 		FRenderTargetBinding& OutputRenderTargetBinding,
 		FDepthStencilBinding& OutputDepthStencilBinding);
 
 	/** Adds a pass to perform the hit proxy render for the mesh projection */
 	void AddHitProxyRenderPass(FRDGBuilder& GraphBuilder, 
-		const FViewInfo* View,
+		const FSceneView* View,
 		const FDisplayClusterMeshProjectionRenderSettings& RenderSettings,
 		FRenderTargetBinding& OutputRenderTargetBinding,
 		FDepthStencilBinding& OutputDepthStencilBinding);
 
 	/** Adds a pass to perform a render of the primitives' normals for the mesh projection */
 	void AddNormalsRenderPass(FRDGBuilder& GraphBuilder,
-		const FViewInfo* View,
+		const FSceneView* View,
 		const FDisplayClusterMeshProjectionRenderSettings& RenderSettings,
 		FRenderTargetBinding& OutputRenderTargetBinding,
 		FDepthStencilBinding& OutputDepthStencilBinding);
@@ -213,7 +213,7 @@ private:
 	/** Adds a pass to filter the output normal map for the mesh projection, dilating and blurring it
 	 * to obtain a continuous normal map from primitives into empty space */
 	void AddNormalsFilterPass(FRDGBuilder& GraphBuilder,
-		const FViewInfo* View,
+		const FSceneView* View,
 		FRenderTargetBinding& OutputRenderTargetBinding,
 		FRDGTexture* SceneColor,
 		FRDGTexture* SceneDepth,
@@ -222,13 +222,13 @@ private:
 #if WITH_EDITOR
 	/** Adds a pass to perform the depth render for any selected primitives for the mesh projection */
 	void AddSelectionDepthRenderPass(FRDGBuilder& GraphBuilder, 
-		const FViewInfo* View,
+		const FSceneView* View,
 		const FDisplayClusterMeshProjectionRenderSettings& RenderSettings,
 		FDepthStencilBinding& OutputDepthStencilBinding);
 
 	/** Adds a pass to perform the post process selection outline for the mesh projection */
 	void AddSelectionOutlineScreenPass(FRDGBuilder& GraphBuilder,
-		const FViewInfo* View,
+		const FSceneView* View,
 		FRenderTargetBinding& OutputRenderTargetBinding,
 		FRDGTexture* SceneColor,
 		FRDGTexture* SceneDepth,
@@ -237,7 +237,7 @@ private:
 
 	/** Adds a pass that renders any elements added to the PDI through the renderer's RenderSimpleElements callback */
 	void AddSimpleElementPass(FRDGBuilder& GraphBuilder,
-		const FViewInfo* View,
+		const FSceneView* View,
 		FRenderTargetBinding& OutputRenderTargetBinding,
 		FSimpleElementCollector& ElementCollector);
 
diff --git a/Engine/Plugins/Runtime/nDisplay/Source/DisplayClusterShaders/DisplayClusterShaders.Build.cs b/Engine/Plugins/Runtime/nDisplay/Source/DisplayClusterShaders/DisplayClusterShaders.Build.cs
--- a/Engine/Plugins/Runtime/nDisplay/Source/DisplayClusterShaders/DisplayClusterShaders.Build.cs
+++ b/Engine/Plugins/Runtime/nDisplay/Source/DisplayClusterShaders/DisplayClusterShaders.Build.cs
@@ -34,7 +34,6 @@ public class DisplayClusterShaders : ModuleRules
 		
 		// NOTE: General rule is not to access the private folder of another module
 		PrivateIncludePaths.AddRange(new string[] {
-			System.IO.Path.Combine(GetModuleDirectory("Renderer"), "Private"),
 		});
 
 		AddEngineThirdPartyPrivateStaticDependencies(Target, "UElibPNG");
diff --git a/Engine/Plugins/Runtime/nDisplay/Source/DisplayClusterShaders/Private/Shaders/DisplayClusterShadersPreprocess_UVLightCards.cpp b/Engine/Plugins/Runtime/nDisplay/Source/DisplayClusterShaders/Private/Shaders/DisplayClusterShadersPreprocess_UVLightCards.cpp
--- a/Engine/Plugins/Runtime/nDisplay/Source/DisplayClusterShaders/Private/Shaders/DisplayClusterShadersPreprocess_UVLightCards.cpp
+++ b/Engine/Plugins/Runtime/nDisplay/Source/DisplayClusterShaders/Private/Shaders/DisplayClusterShadersPreprocess_UVLightCards.cpp
@@ -13,7 +13,8 @@
 #include "PrimitiveSceneProxy.h"
 #include "RenderGraphBuilder.h"
 #include "RenderGraphResources.h"
-#include "ScenePrivate.h"
+//#include "ScenePrivate.h"
+#include "SceneInterface.h"
 #include "SceneRendererInterface.h"
 #include "SceneRenderTargetParameters.h"
 #include "MeshPassProcessor.inl"
@@ -177,7 +178,7 @@ bool FDisplayClusterShadersPreprocess_UVLightCards::RenderPreprocess_UVLightCard
 {
 	check(IsInRenderingThread());
 
-	if (InRenderTarget == nullptr || InScene->GetRenderScene()->PrimitiveSceneProxies.IsEmpty())
+	if (InRenderTarget == nullptr || InScene->GetPrimitiveSceneProxies().IsEmpty())
 	{
 		return false;
 	}
@@ -293,7 +294,7 @@ bool FDisplayClusterShadersPreprocess_UVLightCards::RenderPreprocess_UVLightCard
 			DrawDynamicMeshPass(*View, RHICmdList, [View, InScene](FDynamicPassMeshDrawListContext* DynamicMeshPassContext)
 			{
 				TArray<FPrimitiveSceneProxy*> SceneProxies;
-				for (FPrimitiveSceneProxy* SceneProxy : InScene->GetRenderScene()->PrimitiveSceneProxies)
+				for (FPrimitiveSceneProxy* SceneProxy : InScene->GetPrimitiveSceneProxies())
 				{
 					SceneProxies.Add(SceneProxy);
 				}
diff --git a/Engine/Plugins/VirtualProduction/CameraCalibrationCore/Source/CameraCalibrationCore/CameraCalibrationCore.Build.cs b/Engine/Plugins/VirtualProduction/CameraCalibrationCore/Source/CameraCalibrationCore/CameraCalibrationCore.Build.cs
--- a/Engine/Plugins/VirtualProduction/CameraCalibrationCore/Source/CameraCalibrationCore/CameraCalibrationCore.Build.cs
+++ b/Engine/Plugins/VirtualProduction/CameraCalibrationCore/Source/CameraCalibrationCore/CameraCalibrationCore.Build.cs
@@ -39,8 +39,6 @@ namespace UnrealBuildTool.Rules
 			PrivateIncludePaths.AddRange(
 				new string[] 
 				{
-            		//required for ScreenPass
-					System.IO.Path.Combine(GetModuleDirectory("Renderer"), "Private"),
 				}
 			);
 		}
diff --git a/Engine/Source/Editor/WorldPartitionEditor/Private/WorldPartition/SWorldPartitionEditorGrid2D.cpp b/Engine/Source/Editor/WorldPartitionEditor/Private/WorldPartition/SWorldPartitionEditorGrid2D.cpp
--- a/Engine/Source/Editor/WorldPartitionEditor/Private/WorldPartition/SWorldPartitionEditorGrid2D.cpp
+++ b/Engine/Source/Editor/WorldPartitionEditor/Private/WorldPartition/SWorldPartitionEditorGrid2D.cpp
@@ -16,7 +16,6 @@
 #include "LevelEditorViewport.h"
 #include "LocationVolume.h"
 #include "Modules/ModuleManager.h"
-#include "RendererModule.h"
 #include "Rendering/SlateRenderer.h"
 #include "SWorldPartitionViewportWidget.h"
 #include "TextureResource.h"
diff --git a/Engine/Source/Editor/WorldPartitionEditor/WorldPartitionEditor.Build.cs b/Engine/Source/Editor/WorldPartitionEditor/WorldPartitionEditor.Build.cs
--- a/Engine/Source/Editor/WorldPartitionEditor/WorldPartitionEditor.Build.cs
+++ b/Engine/Source/Editor/WorldPartitionEditor/WorldPartitionEditor.Build.cs
@@ -47,7 +47,6 @@ public class WorldPartitionEditor : ModuleRules
 
 		PrivateIncludePaths.AddRange(
 			new string[] {
-				System.IO.Path.Combine(GetModuleDirectory("Renderer"), "Private"),
 			}
 		);
 	}
diff --git a/Engine/Source/Runtime/HeadMountedDisplay/Public/IHeadMountedDisplay.h b/Engine/Source/Runtime/HeadMountedDisplay/Public/IHeadMountedDisplay.h
--- a/Engine/Source/Runtime/HeadMountedDisplay/Public/IHeadMountedDisplay.h
+++ b/Engine/Source/Runtime/HeadMountedDisplay/Public/IHeadMountedDisplay.h
@@ -15,24 +15,21 @@ struct FPostProcessSettings;
 struct FWorldContext;
 class UTexture;
 class FSceneViewFamily;
-class FViewInfo;
+class FSceneView;
 class FRHICommandListImmediate;
 class FTexture;
 
 struct FHeadMountedDisplayPassContext
 {
-	FHeadMountedDisplayPassContext(FRHICommandListImmediate& InRHICmdList, const FViewInfo& InView)
+	FHeadMountedDisplayPassContext(FRHICommandListImmediate& InRHICmdList, const FSceneView& InView)
 		: RHICmdList(InRHICmdList)
 		, View(InView)
 	{}
 
 	FRHICommandListImmediate& RHICmdList;
-	const FViewInfo& View;
+	const FSceneView& View;
 };
 
-UE_DEPRECATED(5.0, "FRenderingCompositePassContext has been refactored to FHeadMountedDisplayPassContext.")
-typedef FHeadMountedDisplayPassContext FRenderingCompositePassContext;
-
 /**
  * HMD device interface
  */
@@ -151,7 +148,7 @@ public:
 	/**
 	* When implemented, creates a new post process node to provide platform-specific HMD distortion.
 	*/
-	virtual void CreateHMDPostProcessPass_RenderThread(class FRDGBuilder& GraphBuilder, const class FViewInfo& View, const struct FHMDDistortionInputs& Inputs, struct FScreenPassTexture& OutPass) const {}
+	virtual void CreateHMDPostProcessPass_RenderThread(FRDGBuilder& GraphBuilder, const FSceneView& View, const struct FHMDDistortionInputs& Inputs, struct FScreenPassTexture& OutPass) const {}
 
 public:
 
diff --git a/Engine/Source/Runtime/Landscape/Landscape.Build.cs b/Engine/Source/Runtime/Landscape/Landscape.Build.cs
--- a/Engine/Source/Runtime/Landscape/Landscape.Build.cs
+++ b/Engine/Source/Runtime/Landscape/Landscape.Build.cs
@@ -10,7 +10,6 @@ public class Landscape : ModuleRules
 		PrivateIncludePaths.AddRange(
 			new string[] {
 				Path.Combine(GetModuleDirectory("Engine"), "Private"), // for Engine/Private/Collision/PhysXCollision.h
-				Path.Combine(GetModuleDirectory("Renderer"), "Private"),
 				"../Shaders/Shared"
             }
 		);
diff --git a/Engine/Source/Runtime/Landscape/Private/LandscapeRender.cpp b/Engine/Source/Runtime/Landscape/Private/LandscapeRender.cpp
--- a/Engine/Source/Runtime/Landscape/Private/LandscapeRender.cpp
+++ b/Engine/Source/Runtime/Landscape/Private/LandscapeRender.cpp
@@ -39,7 +39,6 @@ LandscapeRender.cpp: New terrain rendering
 #include "RHIStaticStates.h"
 #include "PrimitiveSceneInfo.h"
 #include "SceneView.h"
-#include "SceneCore.h"
 #include "ScenePrivate.h"
 #include "LandscapeProxy.h"
 #include "HAL/LowLevelMemTracker.h"
