Fixing compile errors after merge.

diff --git a/Engine/Plugins/Media/PixelStreaming/Source/PixelStreamingHMD/Private/PixelStreamingHMD.cpp b/Engine/Plugins/Media/PixelStreaming/Source/PixelStreamingHMD/Private/PixelStreamingHMD.cpp
--- a/Engine/Plugins/Media/PixelStreaming/Source/PixelStreamingHMD/Private/PixelStreamingHMD.cpp
+++ b/Engine/Plugins/Media/PixelStreaming/Source/PixelStreamingHMD/Private/PixelStreamingHMD.cpp
@@ -102,7 +102,7 @@ void FPixelStreamingHMD::DrawDistortionMesh_RenderThread(struct FHeadMountedDisp
 	float ClipSpaceQuadZ = 0.0f;
 	FMatrix QuadTexTransform = FMatrix::Identity;
 	FMatrix QuadPosTransform = FMatrix::Identity;
-	const FViewInfo& View = Context.View;
+	const FSceneView& View = Context.View;
 	const FIntRect SrcRect = View.UnscaledViewRect;
 
 	FRHICommandListImmediate& RHICmdList = Context.RHICmdList;
diff --git a/Engine/Plugins/Runtime/nDisplay/Source/DisplayClusterShaders/Private/Shaders/DisplayClusterShadersPreprocess_UVLightCards.cpp b/Engine/Plugins/Runtime/nDisplay/Source/DisplayClusterShaders/Private/Shaders/DisplayClusterShadersPreprocess_UVLightCards.cpp
--- a/Engine/Plugins/Runtime/nDisplay/Source/DisplayClusterShaders/Private/Shaders/DisplayClusterShadersPreprocess_UVLightCards.cpp
+++ b/Engine/Plugins/Runtime/nDisplay/Source/DisplayClusterShaders/Private/Shaders/DisplayClusterShadersPreprocess_UVLightCards.cpp
@@ -262,7 +262,7 @@ bool FDisplayClusterShadersPreprocess_UVLightCards::RenderPreprocess_UVLightCard
 	ViewInitOptions.BackgroundColor = FLinearColor::Black;
 
 	GetRendererModule().CreateAndInitSingleView(RHICmdList, &ViewFamily, &ViewInitOptions);
-	FViewInfo* View = (FViewInfo*)ViewFamily.Views[0];
+	FSceneView* View = (FSceneView*)ViewFamily.Views[0];
 
 	if (!bRenderFinalColor)
 	{
