Fixing errors in DisplayCluster after previous Renderer changes.

diff --git a/Engine/Plugins/Runtime/nDisplay/Source/DisplayCluster/DisplayCluster.Build.cs b/Engine/Plugins/Runtime/nDisplay/Source/DisplayCluster/DisplayCluster.Build.cs
--- a/Engine/Plugins/Runtime/nDisplay/Source/DisplayCluster/DisplayCluster.Build.cs
+++ b/Engine/Plugins/Runtime/nDisplay/Source/DisplayCluster/DisplayCluster.Build.cs
@@ -1,6 +1,7 @@
 // Copyright Epic Games, Inc. All Rights Reserved.
 
 using UnrealBuildTool;
+using System.IO;
 
 public class DisplayCluster : ModuleRules
 {
@@ -8,6 +9,11 @@ public class DisplayCluster : ModuleRules
 	{
 		PublicDefinitions.Add("WITH_OCIO=0");
 
+		PrivateIncludePaths.AddRange(
+			new string[] {
+				System.IO.Path.Combine(GetModuleDirectory("Renderer"), "Private"), //required for FPostProcessMaterialInputs
+			});
+
 		PublicDependencyModuleNames.AddRange(
 			new string[] {
 				"CinematicCamera",
diff --git a/Engine/Plugins/Runtime/nDisplay/Source/DisplayCluster/Private/Render/Viewport/DisplayClusterViewportManagerViewExtension.cpp b/Engine/Plugins/Runtime/nDisplay/Source/DisplayCluster/Private/Render/Viewport/DisplayClusterViewportManagerViewExtension.cpp
--- a/Engine/Plugins/Runtime/nDisplay/Source/DisplayCluster/Private/Render/Viewport/DisplayClusterViewportManagerViewExtension.cpp
+++ b/Engine/Plugins/Runtime/nDisplay/Source/DisplayCluster/Private/Render/Viewport/DisplayClusterViewportManagerViewExtension.cpp
@@ -11,9 +11,7 @@
 
 #include "ScreenPass.h"
 #include "CommonRenderResources.h"
-#include "PostProcess/PostProcessing.h"
-// for FPostProcessMaterialInputs
-#include "PostProcess/PostProcessMaterial.h"
+#include "PostProcess/PostProcessMaterialInputs.h"
 
 namespace DisplayClusterViewportManagerViewExtensionHelpers
 {
diff --git a/Engine/Plugins/Runtime/nDisplay/Source/DisplayCluster/Private/Render/Viewport/DisplayClusterViewportProxy.cpp b/Engine/Plugins/Runtime/nDisplay/Source/DisplayCluster/Private/Render/Viewport/DisplayClusterViewportProxy.cpp
--- a/Engine/Plugins/Runtime/nDisplay/Source/DisplayCluster/Private/Render/Viewport/DisplayClusterViewportProxy.cpp
+++ b/Engine/Plugins/Runtime/nDisplay/Source/DisplayCluster/Private/Render/Viewport/DisplayClusterViewportProxy.cpp
@@ -50,7 +50,7 @@
 #include "ScreenPass.h"
 #include "PostProcess/PostProcessing.h"
 // for FPostProcessMaterialInputs
-#include "PostProcess/PostProcessMaterial.h"
+#include "PostProcess/PostProcessMaterialInputs.h"
 
 static TAutoConsoleVariable<int32> CVarDisplayClusterRenderOverscanResolve(
 	TEXT("nDisplay.render.overscan.resolve"),
diff --git a/Engine/Plugins/Runtime/nDisplay/Source/DisplayCluster/Private/Render/Viewport/DisplayClusterViewport_OpenColorIO.cpp b/Engine/Plugins/Runtime/nDisplay/Source/DisplayCluster/Private/Render/Viewport/DisplayClusterViewport_OpenColorIO.cpp
--- a/Engine/Plugins/Runtime/nDisplay/Source/DisplayCluster/Private/Render/Viewport/DisplayClusterViewport_OpenColorIO.cpp
+++ b/Engine/Plugins/Runtime/nDisplay/Source/DisplayCluster/Private/Render/Viewport/DisplayClusterViewport_OpenColorIO.cpp
@@ -17,29 +17,6 @@
 #include "ScreenPass.h"
 #include "CommonRenderResources.h"
 
-namespace {
-	/** Create dummy view for OCIO render pass
-	 *
-	 * @param InViewRect - input texture view rect
-	 *
-	 * @return - new view info
-	 */
-	FViewInfo CreateDummyViewInfo(const FIntRect& InViewRect)
-	{
-		FSceneViewFamily ViewFamily(FSceneViewFamily::ConstructionValues(nullptr, nullptr, FEngineShowFlags(ESFIM_Game))
-			.SetTime(FGameTime())
-			.SetGammaCorrection(1.0f));
-		FSceneViewInitOptions ViewInitOptions;
-		ViewInitOptions.ViewFamily = &ViewFamily;
-		ViewInitOptions.SetViewRectangle(InViewRect);
-		ViewInitOptions.ViewOrigin = FVector::ZeroVector;
-		ViewInitOptions.ViewRotationMatrix = FMatrix::Identity;
-		ViewInitOptions.ProjectionMatrix = FMatrix::Identity;
-
-		return FViewInfo(ViewInitOptions);
-	}
-}
-
 //////////////////////////////////////////////////////////////////////////
 // FDisplayClusterViewport_OpenColorIO
 //////////////////////////////////////////////////////////////////////////
@@ -162,7 +139,8 @@ bool FDisplayClusterViewport_OpenColorIO::AddPass_RenderThread(FRDGBuilder& Grap
 
 		FOpenColorIORendering::AddPass_RenderThread(
 			GraphBuilder,
-			CreateDummyViewInfo(Output.ViewRect),
+			FScreenPassViewInfo(),
+			GMaxRHIFeatureLevel,
 			FScreenPassTexture(InputTexture),
 			Output,
 			CachedResourcesRenderThread,
