nDisplayMedia: Fixed not restarting viewport media output when stopped.

It was only restarting it when in error, but when resolution changes, it is only temporarily in error and then it settles on stopped state.

diff --git a/Engine/Plugins/Runtime/nDisplay/Source/DisplayClusterMedia/Private/Capture/DisplayClusterMediaCaptureBase.cpp b/Engine/Plugins/Runtime/nDisplay/Source/DisplayClusterMedia/Private/Capture/DisplayClusterMediaCaptureBase.cpp
--- a/Engine/Plugins/Runtime/nDisplay/Source/DisplayClusterMedia/Private/Capture/DisplayClusterMediaCaptureBase.cpp
+++ b/Engine/Plugins/Runtime/nDisplay/Source/DisplayClusterMedia/Private/Capture/DisplayClusterMediaCaptureBase.cpp
@@ -81,13 +81,17 @@ void FDisplayClusterMediaCaptureBase::OnPostClusterTick()
 {
 	if (MediaCapture && bWasCaptureStarted)
 	{
-		if (MediaCapture->GetState() == EMediaCaptureState::Error)
+		const EMediaCaptureState MediaCaptureState = MediaCapture->GetState();
+		const bool bMediaCaptureNeedsRestart = (MediaCaptureState == EMediaCaptureState::Error) || (MediaCaptureState == EMediaCaptureState::Stopped);
+
+		if (bMediaCaptureNeedsRestart)
 		{
 			constexpr double Interval = 1.0;
 			const double CurrentTime = FPlatformTime::Seconds();
+
 			if (CurrentTime - LastRestartTimestamp > Interval)
 			{
-				UE_LOG(LogDisplayClusterMedia, Log, TEXT("MediaCapture '%s' is in error, restarting it."), *GetMediaId());
+				UE_LOG(LogDisplayClusterMedia, Log, TEXT("MediaCapture '%s' is in error or stopped, restarting it."), *GetMediaId());
 
 				StartMediaCapture();
 				LastRestartTimestamp = CurrentTime;
