nDisplay: Fix hang when texture replace is enabled and media input is enabled.

Viewports will now not try to import media when texture replacement is enabled.

diff --git a/Engine/Plugins/Runtime/nDisplay/Source/DisplayCluster/Public/Render/Viewport/Containers/DisplayClusterViewport_RenderSettings.h b/Engine/Plugins/Runtime/nDisplay/Source/DisplayCluster/Public/Render/Viewport/Containers/DisplayClusterViewport_RenderSettings.h
--- a/Engine/Plugins/Runtime/nDisplay/Source/DisplayCluster/Public/Render/Viewport/Containers/DisplayClusterViewport_RenderSettings.h
+++ b/Engine/Plugins/Runtime/nDisplay/Source/DisplayCluster/Public/Render/Viewport/Containers/DisplayClusterViewport_RenderSettings.h
@@ -92,9 +92,6 @@ public:
 	EDisplayClusterViewportOverrideMode ViewportOverrideMode = EDisplayClusterViewportOverrideMode::None;
 	FString ViewportOverrideId;
 
-	// Override all viewport resources
-	bool bOverrideWholeViewport= true;
-
 public:
 	// Reset runtime values from prev frame
 	inline void BeginUpdateSettings()
diff --git a/Engine/Plugins/Runtime/nDisplay/Source/DisplayClusterMedia/Private/Input/DisplayClusterMediaInputViewport.cpp b/Engine/Plugins/Runtime/nDisplay/Source/DisplayClusterMedia/Private/Input/DisplayClusterMediaInputViewport.cpp
--- a/Engine/Plugins/Runtime/nDisplay/Source/DisplayClusterMedia/Private/Input/DisplayClusterMediaInputViewport.cpp
+++ b/Engine/Plugins/Runtime/nDisplay/Source/DisplayClusterMedia/Private/Input/DisplayClusterMediaInputViewport.cpp
@@ -45,15 +45,20 @@ void FDisplayClusterMediaInputViewport::PostCrossGpuTransfer_RenderThread(FRHICo
 
 	if (const IDisplayClusterViewportProxy* const CaptureViewport = ViewportManagerProxy->FindViewport_RenderThread(GetViewportId()))
 	{
-		TArray<FRHITexture*> Textures;
-		TArray<FIntRect>     Regions;
+		const bool bShouldImportMedia = !CaptureViewport->GetPostRenderSettings_RenderThread().Replace.IsEnabled();
 
-		if (CaptureViewport->GetResourcesWithRects_RenderThread(EDisplayClusterViewportResourceType::InternalRenderTargetResource, Textures, Regions))
+		if (bShouldImportMedia)
 		{
-			if (Textures.Num() > 0 && Regions.Num() > 0)
+			TArray<FRHITexture*> Textures;
+			TArray<FIntRect>     Regions;
+
+			if (CaptureViewport->GetResourcesWithRects_RenderThread(EDisplayClusterViewportResourceType::InternalRenderTargetResource, Textures, Regions))
 			{
-				FMediaTextureInfo TextureInfo{ Textures[0], Regions[0] };
-				ImportMediaData(RHICmdList, TextureInfo);
+				if (Textures.Num() > 0 && Regions.Num() > 0)
+				{
+					FMediaTextureInfo TextureInfo{ Textures[0], Regions[0] };
+					ImportMediaData(RHICmdList, TextureInfo);
+				}
 			}
 		}
 	}
