nDisplay: Changed the UV light card rendering pipeline to get rid of the need for a preview world to generate scene proxies for rendering the light cards.
Now, the light card primitive components are set to visible, and the scene proxies are forced to immediately be generated just before the light card rendering pass, and then set back to invisible again after the pass completes. This should resolve any continuing issues with using a preview world as part of the render pass

diff --git a/Engine/Plugins/Runtime/nDisplay/Source/DisplayCluster/Private/Game/EngineClasses/Scene/DisplayClusterRootActorEditor.cpp b/Engine/Plugins/Runtime/nDisplay/Source/DisplayCluster/Private/Game/EngineClasses/Scene/DisplayClusterRootActorEditor.cpp
--- a/Engine/Plugins/Runtime/nDisplay/Source/DisplayCluster/Private/Game/EngineClasses/Scene/DisplayClusterRootActorEditor.cpp
+++ b/Engine/Plugins/Runtime/nDisplay/Source/DisplayCluster/Private/Game/EngineClasses/Scene/DisplayClusterRootActorEditor.cpp
@@ -194,16 +194,6 @@ void ADisplayClusterRootActor::PostLoad_Editor()
 
 void ADisplayClusterRootActor::EndPlay_Editor(const EEndPlayReason::Type EndPlayReason)
 {
-	// UE-174008: When a SIE session is ended, we can't clean up the actor's rendering resources in BeginDestroy
-	// because of an issue cleaning up the UV light card manager's preview world during the normal GC pass,
-	// so in this specific case, destroy the rendering resources here.
-	if (EndPlayReason == EEndPlayReason::EndPlayInEditor && GEditor->bIsSimulatingInEditor)
-	{
-		ResetPreviewInternals_Editor();
-		ReleasePreviewComponents();
-
-		ViewportManager.Reset();
-	}
 }
 
 void ADisplayClusterRootActor::Destroyed_Editor()
diff --git a/Engine/Plugins/Runtime/nDisplay/Source/DisplayCluster/Private/Render/Viewport/LightCard/DisplayClusterViewportLightCardManager.cpp b/Engine/Plugins/Runtime/nDisplay/Source/DisplayCluster/Private/Render/Viewport/LightCard/DisplayClusterViewportLightCardManager.cpp
--- a/Engine/Plugins/Runtime/nDisplay/Source/DisplayCluster/Private/Render/Viewport/LightCard/DisplayClusterViewportLightCardManager.cpp
+++ b/Engine/Plugins/Runtime/nDisplay/Source/DisplayCluster/Private/Render/Viewport/LightCard/DisplayClusterViewportLightCardManager.cpp
@@ -3,6 +3,7 @@
 #include "DisplayClusterViewportLightCardManager.h"
 #include "DisplayClusterViewportLightCardManagerProxy.h"
 #include "DisplayClusterViewportLightCardResource.h"
+#include "ShaderParameters/DisplayClusterShaderParameters_UVLightCards.h"
 
 #include "DisplayClusterLightCardActor.h"
 #include "Blueprints/DisplayClusterBlueprintLib.h"
@@ -10,11 +11,8 @@
 #include "IDisplayClusterShaders.h"
 
 #include "Render/Viewport/DisplayClusterViewportManager.h"
-#include "Render/Viewport/RenderFrame/DisplayClusterRenderFrame.h"
 
 #include "PreviewScene.h"
-#include "SceneInterface.h"
-#include "UObject/Package.h"
 
 ///////////////////////////////////////////////////////////////////////////////////////////////
 /** Console variable used to control the size of the UV light card map texture */
@@ -43,15 +41,11 @@ FDisplayClusterViewportLightCardManager::~FDisplayClusterViewportLightCardManage
 void FDisplayClusterViewportLightCardManager::Release()
 {
 	// The destructor is usually called from the rendering thread, so Release() must be called first from the game thread.
-	const bool bIsInRenderingThread = IsInRenderingThread();
-	check(!bIsInRenderingThread || (bIsInRenderingThread && PreviewWorld == nullptr));
+	check(!IsInRenderingThread());
 
 	// Release UVLightCard
 	ReleaseUVLightCardData();
 	ReleaseUVLightCardResource();
-
-	// Deleting PreviewScene is only called from the game thread
-	DestroyPreviewWorld();
 }
 
 ///////////////////////////////////////////////////////////////////////////////////////////////
@@ -62,13 +56,10 @@ void FDisplayClusterViewportLightCardManager::UpdateConfiguration()
 
 void FDisplayClusterViewportLightCardManager::HandleStartScene()
 {
-	InitializePreviewWorld();
 }
 
 void FDisplayClusterViewportLightCardManager::HandleEndScene()
 {
-	DestroyPreviewWorld();
-
 	ReleaseUVLightCardData();
 }
 
@@ -174,22 +165,41 @@ void FDisplayClusterViewportLightCardManager::UpdateUVLightCardResource()
 void FDisplayClusterViewportLightCardManager::RenderUVLightCard()
 {
 	// Render UV LightCard:
-	if (IsUVLightCardEnabled() && PreviewWorld && LightCardManagerProxy.IsValid())
+	UWorld* World = ViewportManager.GetCurrentWorld();
+	if (IsUVLightCardEnabled() && World && LightCardManagerProxy.IsValid())
 	{
 		UpdateUVLightCardResource();
 
 		if (UVLightCardResource.IsValid())
 		{
-			for (UPrimitiveComponent* LoadedComponent : UVLightCardPrimitiveComponents)
+			FDisplayClusterShaderParameters_UVLightCards UVLightCardParameters;
+			UVLightCardParameters.ProjectionPlaneSize = ADisplayClusterLightCardActor::UVPlaneDefaultSize;
+			UVLightCardParameters.bRenderFinalColor = ViewportManager.ShouldRenderFinalColor();
+
+			// Store any components that were invisible but forced to be visible so they can be set back to invisible after the render
+			TArray<UPrimitiveComponent*> ComponentsToUnload;
+			for (UPrimitiveComponent* PrimitiveComponent : UVLightCardPrimitiveComponents)
 			{
-				PreviewWorld->Scene->AddPrimitive(LoadedComponent);
+				// Set the component's visibility to true and force it to generate its scene proxies
+				if (!PrimitiveComponent->IsVisible())
+				{
+					PrimitiveComponent->SetVisibility(true);
+					PrimitiveComponent->RecreateRenderState_Concurrent();
+					ComponentsToUnload.Add(PrimitiveComponent);
+				}
+
+				if (PrimitiveComponent->SceneProxy)
+				{
+					UVLightCardParameters.PrimitivesToRender.Add(PrimitiveComponent->SceneProxy);
+				}
 			}
 
-			LightCardManagerProxy->RenderUVLightCard(PreviewWorld->Scene, ADisplayClusterLightCardActor::UVPlaneDefaultSize, ViewportManager.ShouldRenderFinalColor());
+			LightCardManagerProxy->RenderUVLightCard(World->Scene, UVLightCardParameters);
 
-			for (UPrimitiveComponent* LoadedComponent : UVLightCardPrimitiveComponents)
+			for (UPrimitiveComponent* LoadedComponent : ComponentsToUnload)
 			{
-				PreviewWorld->Scene->RemovePrimitive(LoadedComponent);
+				LoadedComponent->SetVisibility(false);
+				LoadedComponent->RecreateRenderState_Concurrent();
 			}
 		}
 	}
@@ -203,45 +213,4 @@ void FDisplayClusterViewportLightCardManager::RenderUVLightCard()
 ///////////////////////////////////////////////////////////////////////////////////////////////
 void FDisplayClusterViewportLightCardManager::AddReferencedObjects(FReferenceCollector& Collector)
 {
-	Collector.AddReferencedObject(PreviewWorld);
-}
-
-///////////////////////////////////////////////////////////////////////////////////////////////
-void FDisplayClusterViewportLightCardManager::InitializePreviewWorld()
-{
-	if (!PreviewWorld)
-	{
-		FName UniqueWorldName = MakeUniqueObjectName(GetTransientPackage(), UWorld::StaticClass(), FName(TEXT("DisplayClusterLightCardManager_PreviewWorld")));
-		PreviewWorld = NewObject<UWorld>(GetTransientPackage(), UniqueWorldName);
-		PreviewWorld->WorldType = EWorldType::GamePreview;
-
-		FWorldContext& WorldContext = GEngine->CreateNewWorldContext(PreviewWorld->WorldType);
-		WorldContext.SetCurrentWorld(PreviewWorld);
-
-		PreviewWorld->InitializeNewWorld(UWorld::InitializationValues()
-			.AllowAudioPlayback(false)
-			.CreatePhysicsScene(false)
-			.RequiresHitProxies(false)
-			.CreateNavigation(false)
-			.CreateAISystem(false)
-			.ShouldSimulatePhysics(false)
-			.SetTransactional(false));
-	}
-}
-
-void FDisplayClusterViewportLightCardManager::DestroyPreviewWorld()
-{
-	if (PreviewWorld)
-	{
-		// Hack to avoid issue where the engine considers this world a leaked object; When UEngine loads a new map, it checks to see if there are any UWorlds
-		// still in memory that aren't what it considers "persistent" worlds, worlds with type Inactive or EditorPreview. Even if the UWorld object has been marked for
-		// GC and has no references to it, UEngine will still flag it as "leaked" unless it is one of these two types.
-		PreviewWorld->WorldType = EWorldType::Inactive;
-
-		GEngine->DestroyWorldContext(PreviewWorld);
-
-		PreviewWorld->DestroyWorld(false);
-		PreviewWorld->MarkObjectsPendingKill();
-		PreviewWorld = nullptr;
-	}
 }
diff --git a/Engine/Plugins/Runtime/nDisplay/Source/DisplayCluster/Private/Render/Viewport/LightCard/DisplayClusterViewportLightCardManager.h b/Engine/Plugins/Runtime/nDisplay/Source/DisplayCluster/Private/Render/Viewport/LightCard/DisplayClusterViewportLightCardManager.h
--- a/Engine/Plugins/Runtime/nDisplay/Source/DisplayCluster/Private/Render/Viewport/LightCard/DisplayClusterViewportLightCardManager.h
+++ b/Engine/Plugins/Runtime/nDisplay/Source/DisplayCluster/Private/Render/Viewport/LightCard/DisplayClusterViewportLightCardManager.h
@@ -54,13 +54,6 @@ public:
 	/** Render internal resoures for current frame. */
 	void RenderFrame();
 
-private:
-	/** Initializes the preview world used to render the UV light cards*/
-	void InitializePreviewWorld();
-
-	/** Destroys the preview world used to render the UV light cards */
-	void DestroyPreviewWorld();
-
 private:
 	/** Render UVLightCard */
 	void RenderUVLightCard();
@@ -84,9 +77,6 @@ private:
 	/** A reference to the owning viewport manager */
 	FDisplayClusterViewportManager& ViewportManager;
 
-	/** The preview world the UV light card proxies live in for rendering */
-	UWorld* PreviewWorld = nullptr;
-
 private:
 	/** A list of primitive components that have been added to the preview scene for rendering in the current frame */
 	TArray<UPrimitiveComponent*> UVLightCardPrimitiveComponents;
diff --git a/Engine/Plugins/Runtime/nDisplay/Source/DisplayCluster/Private/Render/Viewport/LightCard/DisplayClusterViewportLightCardManagerProxy.cpp b/Engine/Plugins/Runtime/nDisplay/Source/DisplayCluster/Private/Render/Viewport/LightCard/DisplayClusterViewportLightCardManagerProxy.cpp
--- a/Engine/Plugins/Runtime/nDisplay/Source/DisplayCluster/Private/Render/Viewport/LightCard/DisplayClusterViewportLightCardManagerProxy.cpp
+++ b/Engine/Plugins/Runtime/nDisplay/Source/DisplayCluster/Private/Render/Viewport/LightCard/DisplayClusterViewportLightCardManagerProxy.cpp
@@ -3,6 +3,7 @@
 #include "DisplayClusterViewportLightCardManagerProxy.h"
 
 #include "IDisplayClusterShaders.h"
+#include "ShaderParameters/DisplayClusterShaderParameters_UVLightCards.h"
 
 #include "SceneInterface.h"
 #include "RenderingThread.h"
@@ -39,12 +40,12 @@ void FDisplayClusterViewportLightCardManagerProxy::ReleaseUVLightCardResource()
 		});
 }
 
-void FDisplayClusterViewportLightCardManagerProxy::RenderUVLightCard(FSceneInterface* InSceneInterface, const float UVPlaneDefaultSize, const bool bRenderFinalColor) const
+void FDisplayClusterViewportLightCardManagerProxy::RenderUVLightCard(FSceneInterface* InScene, const FDisplayClusterShaderParameters_UVLightCards& InParameters) const
 {
 	ENQUEUE_RENDER_COMMAND(DisplayClusterViewportLightCardManagerProxy_RenderUVLightCard)(
-		[InProxyData = SharedThis(this), SceneInterface = InSceneInterface, UVPlaneDefaultSize, bRenderFinalColor](FRHICommandListImmediate& RHICmdList)
+		[InProxyData = SharedThis(this), Scene = InScene, Parameters = InParameters](FRHICommandListImmediate& RHICmdList)
 		{
-			InProxyData->ImplRenderUVLightCard_RenderThread(RHICmdList, SceneInterface, UVPlaneDefaultSize, bRenderFinalColor);
+			InProxyData->ImplRenderUVLightCard_RenderThread(RHICmdList, Scene, Parameters);
 		});
 }
 
@@ -74,11 +75,11 @@ void FDisplayClusterViewportLightCardManagerProxy::ImplReleaseUVLightCardResourc
 	}
 }
 
-void FDisplayClusterViewportLightCardManagerProxy::ImplRenderUVLightCard_RenderThread(FRHICommandListImmediate& RHICmdList, FSceneInterface* InSceneInterface, const float InUVPlaneDefaultSize, const bool bRenderFinalColor) const
+void FDisplayClusterViewportLightCardManagerProxy::ImplRenderUVLightCard_RenderThread(FRHICommandListImmediate& RHICmdList, FSceneInterface* InSceneInterface, const FDisplayClusterShaderParameters_UVLightCards& InParameters) const
 {
-	if(InSceneInterface && UVLightCardMapResource.IsValid())
+	if (InParameters.PrimitivesToRender.Num() && UVLightCardMapResource.IsValid())
 	{
 		IDisplayClusterShaders& ShadersAPI = IDisplayClusterShaders::Get();
-		ShadersAPI.RenderPreprocess_UVLightCards(RHICmdList, InSceneInterface, UVLightCardMapResource.Get(), InUVPlaneDefaultSize, bRenderFinalColor);
+		ShadersAPI.RenderPreprocess_UVLightCards(RHICmdList, InSceneInterface, UVLightCardMapResource.Get(), InParameters);
 	}
 }
diff --git a/Engine/Plugins/Runtime/nDisplay/Source/DisplayCluster/Private/Render/Viewport/LightCard/DisplayClusterViewportLightCardManagerProxy.h b/Engine/Plugins/Runtime/nDisplay/Source/DisplayCluster/Private/Render/Viewport/LightCard/DisplayClusterViewportLightCardManagerProxy.h
--- a/Engine/Plugins/Runtime/nDisplay/Source/DisplayCluster/Private/Render/Viewport/LightCard/DisplayClusterViewportLightCardManagerProxy.h
+++ b/Engine/Plugins/Runtime/nDisplay/Source/DisplayCluster/Private/Render/Viewport/LightCard/DisplayClusterViewportLightCardManagerProxy.h
@@ -6,11 +6,10 @@
 
 #include "DisplayClusterViewportLightCardResource.h"
 
-#include "RHI.h"
-#include "RHICommandList.h"
 #include "RHIResources.h"
 #include "Templates/SharedPointer.h"
 
+struct FDisplayClusterShaderParameters_UVLightCards;
 class FSceneInterface;
 
 /*
@@ -29,7 +28,7 @@ public:
 	void ReleaseUVLightCardResource();
 
 	/** Render UVLightCard. */
-	void RenderUVLightCard(FSceneInterface* InSceneInterface, const float InUVPlaneDefaultSize, const bool bRenderFinalColor) const;
+	void RenderUVLightCard(FSceneInterface* InSceneInterface, const FDisplayClusterShaderParameters_UVLightCards& InParameters) const;
 
 	/** Get current UVLightCard RHI resource on rendering thread. */
 	FRHITexture* GetUVLightCardRHIResource_RenderThread() const;
@@ -38,9 +37,9 @@ protected:
 	void ImplUpdateUVLightCardResource_RenderThread(const TSharedPtr<FDisplayClusterViewportLightCardResource, ESPMode::ThreadSafe>& InUVLightCardMapResource);
 	void ImplReleaseUVLightCardResource_RenderThread();
 
-	void ImplRenderUVLightCard_RenderThread(FRHICommandListImmediate& RHICmdList, FSceneInterface* InSceneInterface, const float InUVPlaneDefaultSize, const bool bRenderFinalColor) const;
+	void ImplRenderUVLightCard_RenderThread(FRHICommandListImmediate& RHICmdList, FSceneInterface* InScene, const FDisplayClusterShaderParameters_UVLightCards& InParameters) const;
 
 private:
-	/** The render thread copy of the pointer to the UV ligth card map */
+	/** The render thread copy of the pointer to the UV light card map */
 	TSharedPtr<FDisplayClusterViewportLightCardResource, ESPMode::ThreadSafe> UVLightCardMapResource;
 };
diff --git a/Engine/Plugins/Runtime/nDisplay/Source/DisplayClusterShaders/Private/DisplayClusterShadersModule.cpp b/Engine/Plugins/Runtime/nDisplay/Source/DisplayClusterShaders/Private/DisplayClusterShadersModule.cpp
--- a/Engine/Plugins/Runtime/nDisplay/Source/DisplayClusterShaders/Private/DisplayClusterShadersModule.cpp
+++ b/Engine/Plugins/Runtime/nDisplay/Source/DisplayClusterShaders/Private/DisplayClusterShadersModule.cpp
@@ -42,9 +42,9 @@ bool FDisplayClusterShadersModule::RenderWarpBlend_ICVFX(FRHICommandListImmediat
 	return FDisplayClusterShadersWarpblend_ICVFX::RenderWarpBlend_ICVFX(RHICmdList, InWarpBlendParameters, InICVFXParameters);
 }
 
-bool FDisplayClusterShadersModule::RenderPreprocess_UVLightCards(FRHICommandListImmediate& RHICmdList, FSceneInterface* InScene, FRenderTarget* InRenderTarget, float ProjectionPlaneSize, bool bRenderFinalColor) const
+bool FDisplayClusterShadersModule::RenderPreprocess_UVLightCards(FRHICommandListImmediate& RHICmdList, FSceneInterface* InScene, FRenderTarget* InRenderTarget, const FDisplayClusterShaderParameters_UVLightCards& InParameters) const
 {
-	return FDisplayClusterShadersPreprocess_UVLightCards::RenderPreprocess_UVLightCards(RHICmdList, InScene, InRenderTarget, ProjectionPlaneSize, bRenderFinalColor);
+	return FDisplayClusterShadersPreprocess_UVLightCards::RenderPreprocess_UVLightCards(RHICmdList, InScene, InRenderTarget, InParameters);
 }
 
 bool FDisplayClusterShadersModule::RenderPostprocess_OutputRemap(FRHICommandListImmediate& RHICmdList, FRHITexture2D* InSourceTexture, FRHITexture2D* InRenderTargetableDestTexture, const IDisplayClusterRender_MeshComponentProxy& MeshProxy) const
diff --git a/Engine/Plugins/Runtime/nDisplay/Source/DisplayClusterShaders/Private/DisplayClusterShadersModule.h b/Engine/Plugins/Runtime/nDisplay/Source/DisplayClusterShaders/Private/DisplayClusterShadersModule.h
--- a/Engine/Plugins/Runtime/nDisplay/Source/DisplayClusterShaders/Private/DisplayClusterShadersModule.h
+++ b/Engine/Plugins/Runtime/nDisplay/Source/DisplayClusterShaders/Private/DisplayClusterShadersModule.h
@@ -19,7 +19,8 @@ public:
 public:
 	virtual bool RenderWarpBlend_MPCDI(FRHICommandListImmediate& RHICmdList, const FDisplayClusterShaderParameters_WarpBlend& InWarpBlendParameters) const override;
 	virtual bool RenderWarpBlend_ICVFX(FRHICommandListImmediate& RHICmdList, const FDisplayClusterShaderParameters_WarpBlend& InWarpBlendParameters, const FDisplayClusterShaderParameters_ICVFX& InICVFXParameters) const override;
-	virtual bool RenderPreprocess_UVLightCards(FRHICommandListImmediate& RHICmdList, FSceneInterface* InScene, FRenderTarget* InRenderTarget, float ProjectionPlaneSize, bool bRenderFinalColor) const override;
+	virtual bool RenderPreprocess_UVLightCards(FRHICommandListImmediate& RHICmdList, FSceneInterface* InScene, FRenderTarget* InRenderTarget, float ProjectionPlaneSize, bool bRenderFinalColor) const override { return false; }
+	virtual bool RenderPreprocess_UVLightCards(FRHICommandListImmediate& RHICmdList, FSceneInterface* InScene, FRenderTarget* InRenderTarget, const FDisplayClusterShaderParameters_UVLightCards& InParameters) const override;
 	virtual bool RenderPostprocess_OutputRemap(FRHICommandListImmediate& RHICmdList, FRHITexture2D* InSourceTexture, FRHITexture2D* InRenderTargetableDestTexture, const IDisplayClusterRender_MeshComponentProxy& MeshProxy) const override;
 	virtual bool RenderPostprocess_Blur(FRHICommandListImmediate& RHICmdList, FRHITexture2D* InSourceTexture, FRHITexture2D* InRenderTargetableDestTexture, const FDisplayClusterShaderParameters_PostprocessBlur& InSettings) const override;
 	virtual bool GenerateMips(FRHICommandListImmediate& RHICmdList, FRHITexture2D* InOutMipsTexture, const FDisplayClusterShaderParameters_GenerateMips& InSettings) const override;
diff --git a/Engine/Plugins/Runtime/nDisplay/Source/DisplayClusterShaders/Private/Shaders/DisplayClusterShadersPreprocess_UVLightCards.cpp b/Engine/Plugins/Runtime/nDisplay/Source/DisplayClusterShaders/Private/Shaders/DisplayClusterShadersPreprocess_UVLightCards.cpp
--- a/Engine/Plugins/Runtime/nDisplay/Source/DisplayClusterShaders/Private/Shaders/DisplayClusterShadersPreprocess_UVLightCards.cpp
+++ b/Engine/Plugins/Runtime/nDisplay/Source/DisplayClusterShaders/Private/Shaders/DisplayClusterShadersPreprocess_UVLightCards.cpp
@@ -9,14 +9,14 @@
 #include "InstanceCulling/InstanceCullingContext.h"
 #include "MeshMaterialShader.h"
 #include "MeshPassProcessor.h"
+#include "MeshPassProcessor.inl"
 #include "PrimitiveSceneInfo.h"
 #include "PrimitiveSceneProxy.h"
 #include "RenderGraphBuilder.h"
 #include "RenderGraphResources.h"
 #include "SceneInterface.h"
 #include "SceneRendererInterface.h"
-#include "SceneRenderTargetParameters.h"
-#include "MeshPassProcessor.inl"
+#include "ShaderParameters/DisplayClusterShaderParameters_UVLightCards.h"
 #include "UnrealClient.h"
 
 class FUVLightCardVS : public FMeshMaterialShader
@@ -173,11 +173,11 @@ END_SHADER_PARAMETER_STRUCT()
 
 DECLARE_GPU_STAT_NAMED(nDisplay_UVLightCards_Render, TEXT("nDisplay UVLightCards::Render"));
 
-bool FDisplayClusterShadersPreprocess_UVLightCards::RenderPreprocess_UVLightCards(FRHICommandListImmediate& RHICmdList, FSceneInterface* InScene, FRenderTarget* InRenderTarget, float ProjectionPlaneSize, bool bRenderFinalColor)
+bool FDisplayClusterShadersPreprocess_UVLightCards::RenderPreprocess_UVLightCards(FRHICommandListImmediate& RHICmdList, FSceneInterface* InScene, FRenderTarget* InRenderTarget, const FDisplayClusterShaderParameters_UVLightCards& InParameters)
 {
 	check(IsInRenderingThread());
 
-	if (InRenderTarget == nullptr || InScene->GetPrimitiveSceneProxies().IsEmpty())
+	if (InRenderTarget == nullptr || InParameters.PrimitivesToRender.IsEmpty())
 	{
 		return false;
 	}
@@ -237,7 +237,7 @@ bool FDisplayClusterShadersPreprocess_UVLightCards::RenderPreprocess_UVLightCard
 
 	const float ZScale = 0.5f / UE_OLD_HALF_WORLD_MAX;
 	const float ZOffset = UE_OLD_HALF_WORLD_MAX;
-	const float OrthoSize = 0.5f * ProjectionPlaneSize;
+	const float OrthoSize = 0.5f * InParameters.ProjectionPlaneSize;
 
 	if ((bool)ERHIZBuffer::IsInverted)
 	{
@@ -263,13 +263,13 @@ bool FDisplayClusterShadersPreprocess_UVLightCards::RenderPreprocess_UVLightCard
 	GetRendererModule().CreateAndInitSingleView(RHICmdList, &ViewFamily, &ViewInitOptions);
 	FSceneView* View = (FSceneView*)ViewFamily.Views[0];
 
-	if (!bRenderFinalColor)
+	if (!InParameters.bRenderFinalColor)
 	{
 		ViewFamily.SceneCaptureSource = ESceneCaptureSource::SCS_SceneColorHDR;
 	}
 
 	ViewFamily.EngineShowFlags.SetToneCurve(false);
-	// This flags sets tonampper to output to ETonemapperOutputDevice::LinearNoToneCurve
+	// This flags sets tonemapper to output to ETonemapperOutputDevice::LinearNoToneCurve
 	View->FinalPostProcessSettings.bOverride_ToneCurveAmount = 1;
 	View->FinalPostProcessSettings.ToneCurveAmount = 0.0;
 
@@ -285,21 +285,15 @@ bool FDisplayClusterShadersPreprocess_UVLightCards::RenderPreprocess_UVLightCard
 	GraphBuilder.AddPass(RDG_EVENT_NAME("DisplayClusterUVLightCards::Render"),
 		PassParameters,
 		ERDGPassFlags::Raster | ERDGPassFlags::NeverCull,
-		[View, InScene](FRHICommandList& RHICmdList)
+		[View, &InParameters](FRHICommandList& RHICmdList)
 		{
 			FIntRect ViewRect = View->UnscaledViewRect;
 			RHICmdList.SetViewport(ViewRect.Min.X, ViewRect.Min.Y, 0.0f, ViewRect.Max.X, ViewRect.Max.Y, 1.0f);
 
-			DrawDynamicMeshPass(*View, RHICmdList, [View, InScene](FDynamicPassMeshDrawListContext* DynamicMeshPassContext)
+			DrawDynamicMeshPass(*View, RHICmdList, [View, &InParameters](FDynamicPassMeshDrawListContext* DynamicMeshPassContext)
 			{
-				TArray<FPrimitiveSceneProxy*> SceneProxies;
-				for (FPrimitiveSceneProxy* SceneProxy : InScene->GetPrimitiveSceneProxies())
-				{
-					SceneProxies.Add(SceneProxy);
-				}
-
 				FUVLightCardPassProcessor MeshPassProcessor(View, DynamicMeshPassContext);
-				for (FPrimitiveSceneProxy* SceneProxy : SceneProxies)
+				for (FPrimitiveSceneProxy* SceneProxy : InParameters.PrimitivesToRender)
 				{
 					if (const FMeshBatch* MeshBatch = SceneProxy->GetPrimitiveSceneInfo()->GetMeshBatch(SceneProxy->GetPrimitiveSceneInfo()->StaticMeshes.Num() - 1))
 					{
diff --git a/Engine/Plugins/Runtime/nDisplay/Source/DisplayClusterShaders/Private/Shaders/DisplayClusterShadersPreprocess_UVLightCards.h b/Engine/Plugins/Runtime/nDisplay/Source/DisplayClusterShaders/Private/Shaders/DisplayClusterShadersPreprocess_UVLightCards.h
--- a/Engine/Plugins/Runtime/nDisplay/Source/DisplayClusterShaders/Private/Shaders/DisplayClusterShadersPreprocess_UVLightCards.h
+++ b/Engine/Plugins/Runtime/nDisplay/Source/DisplayClusterShaders/Private/Shaders/DisplayClusterShadersPreprocess_UVLightCards.h
@@ -3,14 +3,14 @@
 #pragma once
 
 #include "CoreMinimal.h"
-#include "RHI.h"
 #include "RHICommandList.h"
 
+struct FDisplayClusterShaderParameters_UVLightCards;
 class FSceneInterface;
 class FRenderTarget;
 
 class FDisplayClusterShadersPreprocess_UVLightCards
 {
 public:
-	static bool RenderPreprocess_UVLightCards(FRHICommandListImmediate& RHICmdList, FSceneInterface* InScene, FRenderTarget* InRenderTarget, float ProjectionPlaneSize, bool bRenderFinalColor);
+	static bool RenderPreprocess_UVLightCards(FRHICommandListImmediate& RHICmdList, FSceneInterface* InScene, FRenderTarget* InRenderTarget, const FDisplayClusterShaderParameters_UVLightCards& InParameters);
 };
\ No newline at end of file
diff --git a/Engine/Plugins/Runtime/nDisplay/Source/DisplayClusterShaders/Public/IDisplayClusterShaders.h b/Engine/Plugins/Runtime/nDisplay/Source/DisplayClusterShaders/Public/IDisplayClusterShaders.h
--- a/Engine/Plugins/Runtime/nDisplay/Source/DisplayClusterShaders/Public/IDisplayClusterShaders.h
+++ b/Engine/Plugins/Runtime/nDisplay/Source/DisplayClusterShaders/Public/IDisplayClusterShaders.h
@@ -16,6 +16,7 @@ class FDisplayClusterShaderParameters_Override;
 struct FDisplayClusterShaderParameters_GenerateMips;
 struct FDisplayClusterShaderParameters_PostprocessBlur;
 struct FDisplayClusterShaderParameters_WarpBlend;
+struct FDisplayClusterShaderParameters_UVLightCards;
 class IDisplayClusterWarpBlendManager;
 class FSceneInterface;
 class FRenderTarget;
@@ -90,8 +91,21 @@ public:
 	*
 	* @return - true if success
 	*/
+	UE_DEPRECATED(5.2, "Use version that takes FDisplayClusterShaderParameters_UVLightCards instead")
 	virtual bool RenderPreprocess_UVLightCards(FRHICommandListImmediate& RHICmdList, FSceneInterface* InScene, FRenderTarget* InRenderTarget, float ProjectionPlaneSize, bool bRenderFinalColor) const = 0;
 
+	/**
+	* Render UV light cards to texture that can be sampled by viewports to place light cards in UV space
+	*
+	* @param RHICmdList                    - RHI command list
+	* @param InScene                       - The scene the light cards live in
+	* @param InRenderTarget                - Destination RTT texture
+	* @param InParameters                  - Parameters for the render
+	*
+	* @return - true if success
+	*/
+	virtual bool RenderPreprocess_UVLightCards(FRHICommandListImmediate& RHICmdList, FSceneInterface* InScene, FRenderTarget* InRenderTarget, const FDisplayClusterShaderParameters_UVLightCards& InParameters) const = 0;
+
 	/**
 	* Render postprocess OutputRemap
 	*
diff --git a/Engine/Plugins/Runtime/nDisplay/Source/DisplayClusterShaders/Public/ShaderParameters/DisplayClusterShaderParameters_UVLightCards.h b/Engine/Plugins/Runtime/nDisplay/Source/DisplayClusterShaders/Public/ShaderParameters/DisplayClusterShaderParameters_UVLightCards.h
new file mode 100644
--- /dev/null
+++ b/Engine/Plugins/Runtime/nDisplay/Source/DisplayClusterShaders/Public/ShaderParameters/DisplayClusterShaderParameters_UVLightCards.h
@@ -0,0 +1,20 @@
+﻿// Copyright Epic Games, Inc. All Rights Reserved.
+
+#pragma once
+
+#include "CoreMinimal.h"
+
+class FPrimitiveSceneProxy;
+
+/** A data struct that contains parameters for performing the UV light card render pass */
+struct FDisplayClusterShaderParameters_UVLightCards
+{
+	/** A list of primitive scene proxies to render */
+	TArray<FPrimitiveSceneProxy*> PrimitivesToRender;
+
+	/** The size of the plane in world units to project the UV space onto */
+	float ProjectionPlaneSize;
+
+	/** Indicates whether to render the final color of the light cards or not */
+	bool bRenderFinalColor;
+};
