Fix compile errors after merge conflict

ROBOMERGE-AUTHOR: trystan.binkley-Jone

diff --git a/Engine/Plugins/Runtime/nDisplay/Source/DisplayClusterShaders/Private/Shaders/DisplayClusterShadersPreprocess_UVLightCards.cpp b/Engine/Plugins/Runtime/nDisplay/Source/DisplayClusterShaders/Private/Shaders/DisplayClusterShadersPreprocess_UVLightCards.cpp
--- a/Engine/Plugins/Runtime/nDisplay/Source/DisplayClusterShaders/Private/Shaders/DisplayClusterShadersPreprocess_UVLightCards.cpp
+++ b/Engine/Plugins/Runtime/nDisplay/Source/DisplayClusterShaders/Private/Shaders/DisplayClusterShadersPreprocess_UVLightCards.cpp
@@ -16,6 +16,7 @@
 #include "RenderGraphResources.h"
 #include "SceneInterface.h"
 #include "SceneRendererInterface.h"
+#include "SceneTexturesConfig.h"
 #include "ShaderParameters/DisplayClusterShaderParameters_UVLightCards.h"
 #include "UnrealClient.h"
 
