nDisplay: Fix various crashes with delete and undo of viewports.
- Fix crash when undoing the deletion of multiple viewports.
- Fix property propagation crash when deleting, compiling, and undoing.
- Fix crash on deletion caused by viewport proxies not being removed from the viewport manager view extension. They would maintain an invalid resource that could crash during render.

diff --git a/Engine/Plugins/Runtime/nDisplay/Source/DisplayCluster/Private/Game/EngineClasses/Scene/DisplayClusterRootActor.cpp b/Engine/Plugins/Runtime/nDisplay/Source/DisplayCluster/Private/Game/EngineClasses/Scene/DisplayClusterRootActor.cpp
--- a/Engine/Plugins/Runtime/nDisplay/Source/DisplayCluster/Private/Game/EngineClasses/Scene/DisplayClusterRootActor.cpp
+++ b/Engine/Plugins/Runtime/nDisplay/Source/DisplayCluster/Private/Game/EngineClasses/Scene/DisplayClusterRootActor.cpp
@@ -1146,7 +1146,13 @@ static void PropagateDefaultMapToInstancedMap(const FMapProperty* MapProperty, U
 				if (!bArchetypeCorrect)
 				{
 					const ADisplayClusterRootActor* RootActor =
-					CastChecked<ADisplayClusterRootActor>(InstancedObject->GetTypedOuter(ADisplayClusterRootActor::StaticClass()));
+						Cast<ADisplayClusterRootActor>(InstancedObject->GetTypedOuter(ADisplayClusterRootActor::StaticClass()));
+					if (!RootActor)
+					{
+						// Undo transactions can potentially trigger this while an object was renamed to the transient package.
+						check(InstancedObject->GetPackage() == GetTransientPackage());
+						continue;
+					}
 #if WITH_EDITOR
 					if (GEditor)
 					{
diff --git a/Engine/Plugins/Runtime/nDisplay/Source/DisplayCluster/Private/Game/EngineClasses/Scene/DisplayClusterRootActorEditor.cpp b/Engine/Plugins/Runtime/nDisplay/Source/DisplayCluster/Private/Game/EngineClasses/Scene/DisplayClusterRootActorEditor.cpp
--- a/Engine/Plugins/Runtime/nDisplay/Source/DisplayCluster/Private/Game/EngineClasses/Scene/DisplayClusterRootActorEditor.cpp
+++ b/Engine/Plugins/Runtime/nDisplay/Source/DisplayCluster/Private/Game/EngineClasses/Scene/DisplayClusterRootActorEditor.cpp
@@ -939,6 +939,17 @@ void ADisplayClusterRootActor::UpdatePreviewComponents()
 	ImplUpdatePreviewConfiguration_Editor(DisplayClusterConfigurationStrings::gui::preview::PreviewNodeAll);
 
 	TArray<UDisplayClusterPreviewComponent*> IteratedPreviewComponents;
+
+	auto DestroyPreviewComponent = [this](UDisplayClusterPreviewComponent* PreviewComp)
+	{
+		PreviewComponents.Remove(PreviewComp->GetName());
+			
+		// Reset preview components before unregister
+		PreviewComp->ResetPreviewComponent(true);
+
+		PreviewComp->UnregisterComponent();
+		PreviewComp->DestroyComponent();
+	};
 	
 	if (CurrentConfigData != nullptr && IsPreviewEnabled())
 	{
@@ -953,6 +964,17 @@ void ADisplayClusterRootActor::UpdatePreviewComponents()
 			{
 				const FString PreviewCompId = GeneratePreviewComponentName_Editor(Node.Key, Viewport.Key);
 				UDisplayClusterPreviewComponent* PreviewComp = PreviewComponents.FindRef(PreviewCompId);
+
+				if (PreviewComp && PreviewComp->GetOwner() != this)
+				{
+					// In this case a viewport was likely deleted, the blueprint compiled, then the deletion undone.
+					// Just destroy the preview component and start over.
+					ensure(!PreviewComp->GetOwner() || PreviewComp->GetOwner()->GetName().StartsWith(TEXT("REINST_")));
+					DestroyPreviewComponent(PreviewComp);
+					PreviewComp = nullptr;
+				}
+
+				Modify();
 				if (!PreviewComp)
 				{
 					PreviewComp = NewObject<UDisplayClusterPreviewComponent>(this, FName(*PreviewCompId), RF_DuplicateTransient | RF_Transactional | RF_NonPIEDuplicateTransient);
@@ -964,6 +986,9 @@ void ADisplayClusterRootActor::UpdatePreviewComponents()
 					PreviewClusterNodeIndex = 0;
 				}
 
+				// Make sure we're an owned component. Possible this can be lost on undo/redo without a recompile.
+				AddOwnedComponent(PreviewComp);
+				
 				if (GetWorld() && !PreviewComp->IsRegistered())
 				{
 					PreviewComp->RegisterComponent();
@@ -989,13 +1014,7 @@ void ADisplayClusterRootActor::UpdatePreviewComponents()
 	{
 		if (!IteratedPreviewComponents.Contains(ExistingComp))
 		{
-			PreviewComponents.Remove(ExistingComp->GetName());
-			
-			// Reset preview components before unregister
-			ExistingComp->ResetPreviewComponent(true);
-
-			ExistingComp->UnregisterComponent();
-			ExistingComp->DestroyComponent();
+			DestroyPreviewComponent(ExistingComp);
 		}
 	}
 }
diff --git a/Engine/Plugins/Runtime/nDisplay/Source/DisplayCluster/Private/Game/EngineClasses/Scene/DisplayClusterStageGeometryComponent.cpp b/Engine/Plugins/Runtime/nDisplay/Source/DisplayCluster/Private/Game/EngineClasses/Scene/DisplayClusterStageGeometryComponent.cpp
--- a/Engine/Plugins/Runtime/nDisplay/Source/DisplayCluster/Private/Game/EngineClasses/Scene/DisplayClusterStageGeometryComponent.cpp
+++ b/Engine/Plugins/Runtime/nDisplay/Source/DisplayCluster/Private/Game/EngineClasses/Scene/DisplayClusterStageGeometryComponent.cpp
@@ -298,7 +298,7 @@ void UDisplayClusterStageGeometryComponent::GenerateGeometryMap(bool bIsNorthMap
 		GetSceneViewInitOptions(NormalMapDirection, RenderSettings.ViewInitOptions);
 		GeometryMap.ViewMatrices = FViewMatrices(RenderSettings.ViewInitOptions);
 
-		FTextureRenderTarget2DResource* TexResource = (FTextureRenderTarget2DResource*)GeometryMap.RenderTarget->GetResource();
+		FTextureRenderTarget2DResource* TexResource = GeometryMap.RenderTarget ? (FTextureRenderTarget2DResource*)GeometryMap.RenderTarget->GetResource() : nullptr;
 		if (TexResource)
 		{
 			FCanvas Canvas(TexResource, nullptr, World, World->Scene->GetFeatureLevel());
diff --git a/Engine/Plugins/Runtime/nDisplay/Source/DisplayCluster/Private/Render/Viewport/DisplayClusterViewportManagerProxy.cpp b/Engine/Plugins/Runtime/nDisplay/Source/DisplayCluster/Private/Render/Viewport/DisplayClusterViewportManagerProxy.cpp
--- a/Engine/Plugins/Runtime/nDisplay/Source/DisplayCluster/Private/Render/Viewport/DisplayClusterViewportManagerProxy.cpp
+++ b/Engine/Plugins/Runtime/nDisplay/Source/DisplayCluster/Private/Render/Viewport/DisplayClusterViewportManagerProxy.cpp
@@ -7,6 +7,7 @@
 
 #include "IDisplayCluster.h"
 #include "IDisplayClusterCallbacks.h"
+#include "DisplayClusterViewportManagerViewExtension.h"
 #include "Render/IDisplayClusterRenderManager.h"
 #include "Render/Viewport/DisplayClusterViewportManager.h"
 #include "Render/Projection/IDisplayClusterProjectionPolicyFactory.h"
@@ -109,6 +110,12 @@ void FDisplayClusterViewportManagerProxy::Initialize(FDisplayClusterViewportMana
 	ViewportManagerViewExtension = InViewportManager.ViewportManagerViewExtension;
 }
 
+void FDisplayClusterViewportManagerProxy::SetViewportManagerViewExtension(
+	const TSharedPtr<FDisplayClusterViewportManagerViewExtension>& InExtension)
+{
+	ViewportManagerViewExtension = InExtension;
+}
+
 void FDisplayClusterViewportManagerProxy::DeleteResource_RenderThread(FDisplayClusterViewportResource* InDeletedResourcePtr)
 {
 	if (InDeletedResourcePtr)
@@ -153,17 +160,23 @@ void FDisplayClusterViewportManagerProxy::DeleteViewport_RenderThread(const TSha
 	check(IsInRenderingThread());
 
 	// Remove viewport obj from manager
-	int32 ViewportProxyIndex = ViewportProxies.Find(InViewportProxy);
+	const int32 ViewportProxyIndex = ViewportProxies.Find(InViewportProxy);
 	if (ViewportProxyIndex != INDEX_NONE)
 	{
-		ViewportProxies[ViewportProxyIndex].Reset();
+		if (ViewportManagerViewExtension.IsValid())
+		{
+			ViewportManagerViewExtension->DeleteViewportProxy(ViewportProxies[ViewportProxyIndex]);
+		}
 		ViewportProxies.RemoveAt(ViewportProxyIndex);
 	}
 
-	int32 ClusterViewportProxyIndex = ClusterNodeViewportProxies.Find(InViewportProxy);
+	const int32 ClusterViewportProxyIndex = ClusterNodeViewportProxies.Find(InViewportProxy);
 	if (ClusterViewportProxyIndex != INDEX_NONE)
 	{
-		ClusterNodeViewportProxies[ClusterViewportProxyIndex].Reset();
+		if (ViewportManagerViewExtension.IsValid())
+		{
+			ViewportManagerViewExtension->DeleteViewportProxy(ClusterNodeViewportProxies[ClusterViewportProxyIndex]);
+		}
 		ClusterNodeViewportProxies.RemoveAt(ClusterViewportProxyIndex);
 	}
 }
diff --git a/Engine/Plugins/Runtime/nDisplay/Source/DisplayCluster/Private/Render/Viewport/DisplayClusterViewportManagerProxy.h b/Engine/Plugins/Runtime/nDisplay/Source/DisplayCluster/Private/Render/Viewport/DisplayClusterViewportManagerProxy.h
--- a/Engine/Plugins/Runtime/nDisplay/Source/DisplayCluster/Private/Render/Viewport/DisplayClusterViewportManagerProxy.h
+++ b/Engine/Plugins/Runtime/nDisplay/Source/DisplayCluster/Private/Render/Viewport/DisplayClusterViewportManagerProxy.h
@@ -75,6 +75,9 @@ public:
 	/** Initialize ViewportManagerProxy from ViewportManager. */
 	void Initialize(FDisplayClusterViewportManager& InViewportManager);
 
+	/** Set the viewport manager view extension pointer. */
+	void SetViewportManagerViewExtension(const TSharedPtr<class FDisplayClusterViewportManagerViewExtension>& InExtension);
+
 	/** Get LightCardManager proxy object. */
 	TSharedPtr<FDisplayClusterViewportLightCardManagerProxy, ESPMode::ThreadSafe> GetLightCardManagerProxy_RenderThread() const
 	{ return LightCardManagerProxy; }
diff --git a/Engine/Plugins/Runtime/nDisplay/Source/DisplayCluster/Private/Render/Viewport/DisplayClusterViewportManagerViewExtension.cpp b/Engine/Plugins/Runtime/nDisplay/Source/DisplayCluster/Private/Render/Viewport/DisplayClusterViewportManagerViewExtension.cpp
--- a/Engine/Plugins/Runtime/nDisplay/Source/DisplayCluster/Private/Render/Viewport/DisplayClusterViewportManagerViewExtension.cpp
+++ b/Engine/Plugins/Runtime/nDisplay/Source/DisplayCluster/Private/Render/Viewport/DisplayClusterViewportManagerViewExtension.cpp
@@ -120,6 +120,20 @@ FScreenPassTexture FDisplayClusterViewportManagerViewExtension::ReturnUntouchedS
 	}
 }
 
+void FDisplayClusterViewportManagerViewExtension::DeleteViewportProxy(
+	const TSharedPtr<FDisplayClusterViewportProxy>& InProxy)
+{
+	const int32 Index = ViewportProxies.IndexOfByPredicate([&](const FViewportProxy& ExtensionProxy)
+	{
+		return ExtensionProxy.ViewportProxy == InProxy;
+	});
+
+	if (Index != INDEX_NONE)
+	{
+		ViewportProxies.RemoveAt(Index);
+	}
+}
+
 bool FDisplayClusterViewportManagerViewExtension::IsActiveThisFrame_Internal(const FSceneViewExtensionContext& Context) const
 {
 	static const FDisplayClusterSceneViewExtensionContext DCViewExtensionContext;
diff --git a/Engine/Plugins/Runtime/nDisplay/Source/DisplayCluster/Private/Render/Viewport/DisplayClusterViewportManagerViewExtension.h b/Engine/Plugins/Runtime/nDisplay/Source/DisplayCluster/Private/Render/Viewport/DisplayClusterViewportManagerViewExtension.h
--- a/Engine/Plugins/Runtime/nDisplay/Source/DisplayCluster/Private/Render/Viewport/DisplayClusterViewportManagerViewExtension.h
+++ b/Engine/Plugins/Runtime/nDisplay/Source/DisplayCluster/Private/Render/Viewport/DisplayClusterViewportManagerViewExtension.h
@@ -53,6 +53,8 @@ public:
 	 */
 	static FScreenPassTexture ReturnUntouchedSceneColorForPostProcessing(const FPostProcessMaterialInputs& InOutInputs);
 
+	/** Remove an internally managed viewport proxy */
+	void DeleteViewportProxy(const TSharedPtr<FDisplayClusterViewportProxy>& InProxy);
 private:
 	/** nDisplay VE Callback [subscribed to Renderer:ResolvedSceneColorCallbacks].
 	 *
diff --git a/Engine/Plugins/Runtime/nDisplay/Source/DisplayClusterConfigurator/Private/ClusterConfiguration/DisplayClusterConfiguratorClusterUtils.cpp b/Engine/Plugins/Runtime/nDisplay/Source/DisplayClusterConfigurator/Private/ClusterConfiguration/DisplayClusterConfiguratorClusterUtils.cpp
--- a/Engine/Plugins/Runtime/nDisplay/Source/DisplayClusterConfigurator/Private/ClusterConfiguration/DisplayClusterConfiguratorClusterUtils.cpp
+++ b/Engine/Plugins/Runtime/nDisplay/Source/DisplayClusterConfigurator/Private/ClusterConfiguration/DisplayClusterConfiguratorClusterUtils.cpp
@@ -886,8 +886,6 @@ bool FDisplayClusterConfiguratorClusterUtils::RemoveViewportFromClusterNode(UDis
 			
 			Viewport->SetFlags(RF_Transient);
 			
-			FDisplayClusterConfiguratorUtils::MarkDisplayClusterBlueprintAsModified(ViewportParent, true);
-			
 			return true;
 		}
 	}
diff --git a/Engine/Plugins/Runtime/nDisplay/Source/DisplayClusterConfigurator/Private/Views/OutputMapping/GraphNodes/SDisplayClusterConfiguratorViewportNode.cpp b/Engine/Plugins/Runtime/nDisplay/Source/DisplayClusterConfigurator/Private/Views/OutputMapping/GraphNodes/SDisplayClusterConfiguratorViewportNode.cpp
--- a/Engine/Plugins/Runtime/nDisplay/Source/DisplayClusterConfigurator/Private/Views/OutputMapping/GraphNodes/SDisplayClusterConfiguratorViewportNode.cpp
+++ b/Engine/Plugins/Runtime/nDisplay/Source/DisplayClusterConfigurator/Private/Views/OutputMapping/GraphNodes/SDisplayClusterConfiguratorViewportNode.cpp
@@ -462,7 +462,7 @@ void SDisplayClusterConfiguratorViewportNode::UpdatePreviewTexture()
 		CachedTexture = CurrentTexture;
 		if (CachedTexture != nullptr)
 		{
-			if (BackgroundActiveBrush.GetResourceObject() != CachedTexture)
+			if (BackgroundActiveBrush.GetResourceObject() != CachedTexture && CachedTexture->GetResource() != nullptr)
 			{
 				BackgroundActiveBrush = FSlateBrush();
 				BackgroundActiveBrush.SetResourceObject(CachedTexture);
diff --git a/Engine/Plugins/Runtime/nDisplay/Source/DisplayClusterConfigurator/Private/Views/OutputMapping/SDisplayClusterConfiguratorGraphEditor.cpp b/Engine/Plugins/Runtime/nDisplay/Source/DisplayClusterConfigurator/Private/Views/OutputMapping/SDisplayClusterConfiguratorGraphEditor.cpp
--- a/Engine/Plugins/Runtime/nDisplay/Source/DisplayClusterConfigurator/Private/Views/OutputMapping/SDisplayClusterConfiguratorGraphEditor.cpp
+++ b/Engine/Plugins/Runtime/nDisplay/Source/DisplayClusterConfigurator/Private/Views/OutputMapping/SDisplayClusterConfiguratorGraphEditor.cpp
@@ -505,7 +505,8 @@ void SDisplayClusterConfiguratorGraphEditor::DeleteSelectedNodes()
 
 		TSharedPtr<FDisplayClusterConfiguratorBlueprintEditor> Toolkit = ToolkitPtr.Pin();
 		Toolkit->GetEditorData()->MarkPackageDirty();
-		Toolkit->ClusterChanged();
+		constexpr bool bStructureChange = true;
+		Toolkit->ClusterChanged(bStructureChange);
 		ClearSelectionSet();
 	}
 }
diff --git a/Engine/Plugins/Runtime/nDisplay/Source/DisplayClusterConfigurator/Private/Views/TreeViews/Cluster/DisplayClusterConfiguratorViewCluster.cpp b/Engine/Plugins/Runtime/nDisplay/Source/DisplayClusterConfigurator/Private/Views/TreeViews/Cluster/DisplayClusterConfiguratorViewCluster.cpp
--- a/Engine/Plugins/Runtime/nDisplay/Source/DisplayClusterConfigurator/Private/Views/TreeViews/Cluster/DisplayClusterConfiguratorViewCluster.cpp
+++ b/Engine/Plugins/Runtime/nDisplay/Source/DisplayClusterConfigurator/Private/Views/TreeViews/Cluster/DisplayClusterConfiguratorViewCluster.cpp
@@ -376,7 +376,7 @@ void FDisplayClusterConfiguratorViewCluster::DeleteSelectedNodes()
 
 	const FScopedTransaction Transaction(FText::Format(LOCTEXT("DeleteClusterNodes", "Remove Cluster {0}|plural(one=Node, other=Nodes)"), SelectedTreeItems.Num()));
 
-	for (TSharedPtr<IDisplayClusterConfiguratorTreeItem> TreeItem : SelectedTreeItems)
+	for (const TSharedPtr<IDisplayClusterConfiguratorTreeItem>& TreeItem : SelectedTreeItems)
 	{
 		if (TreeItem->CanDeleteItem())
 		{
@@ -387,7 +387,8 @@ void FDisplayClusterConfiguratorViewCluster::DeleteSelectedNodes()
 	// Mark the cluster configuration data as dirty, allowing user to save the changes, and fire off a cluster changed event to let other
 	// parts of the UI update as well
 	GetEditorData()->MarkPackageDirty();
-	ToolkitPtr.Pin()->ClusterChanged();
+	constexpr bool bStructureChange = true;
+	ToolkitPtr.Pin()->ClusterChanged(bStructureChange);
 
 	// Clear the selection to ensure the deleted cluster item isn't still being displayed in the details panel
 	ClearSelection();
