nDisplay: Crash fix for invalid resource on viewport deletion. This was not fully solved due to a missing file from CL 24295899. The fix has been also been simplified.

diff --git a/Engine/Plugins/Runtime/nDisplay/Source/DisplayCluster/Private/Render/Viewport/DisplayClusterViewportManager.cpp b/Engine/Plugins/Runtime/nDisplay/Source/DisplayCluster/Private/Render/Viewport/DisplayClusterViewportManager.cpp
--- a/Engine/Plugins/Runtime/nDisplay/Source/DisplayCluster/Private/Render/Viewport/DisplayClusterViewportManager.cpp
+++ b/Engine/Plugins/Runtime/nDisplay/Source/DisplayCluster/Private/Render/Viewport/DisplayClusterViewportManager.cpp
@@ -645,7 +645,7 @@ void FDisplayClusterViewportManager::FinalizeNewFrame()
 	}
 
 	// Send render frame settings to rendering thread
-	ViewportManagerProxy->ImplUpdateRenderFrameSettings(GetRenderFrameSettings());
+	ViewportManagerProxy->ImplUpdateRenderFrameSettings(GetRenderFrameSettings(), ViewportManagerViewExtension);
 
 	// Send updated viewports data to render thread proxy
 	ViewportManagerProxy->ImplUpdateViewports(ClusterNodeViewports);
diff --git a/Engine/Plugins/Runtime/nDisplay/Source/DisplayCluster/Private/Render/Viewport/DisplayClusterViewportManagerProxy.cpp b/Engine/Plugins/Runtime/nDisplay/Source/DisplayCluster/Private/Render/Viewport/DisplayClusterViewportManagerProxy.cpp
--- a/Engine/Plugins/Runtime/nDisplay/Source/DisplayCluster/Private/Render/Viewport/DisplayClusterViewportManagerProxy.cpp
+++ b/Engine/Plugins/Runtime/nDisplay/Source/DisplayCluster/Private/Render/Viewport/DisplayClusterViewportManagerProxy.cpp
@@ -7,7 +7,6 @@
 
 #include "IDisplayCluster.h"
 #include "IDisplayClusterCallbacks.h"
-#include "DisplayClusterViewportManagerViewExtension.h"
 #include "Render/IDisplayClusterRenderManager.h"
 #include "Render/Viewport/DisplayClusterViewportManager.h"
 #include "Render/Projection/IDisplayClusterProjectionPolicyFactory.h"
@@ -110,12 +109,6 @@ void FDisplayClusterViewportManagerProxy::Initialize(FDisplayClusterViewportMana
 	ViewportManagerViewExtension = InViewportManager.ViewportManagerViewExtension;
 }
 
-void FDisplayClusterViewportManagerProxy::SetViewportManagerViewExtension(
-	const TSharedPtr<FDisplayClusterViewportManagerViewExtension>& InExtension)
-{
-	ViewportManagerViewExtension = InExtension;
-}
-
 void FDisplayClusterViewportManagerProxy::DeleteResource_RenderThread(FDisplayClusterViewportResource* InDeletedResourcePtr)
 {
 	if (InDeletedResourcePtr)
@@ -163,25 +156,18 @@ void FDisplayClusterViewportManagerProxy::DeleteViewport_RenderThread(const TSha
 	const int32 ViewportProxyIndex = ViewportProxies.Find(InViewportProxy);
 	if (ViewportProxyIndex != INDEX_NONE)
 	{
-		if (ViewportManagerViewExtension.IsValid())
-		{
-			ViewportManagerViewExtension->DeleteViewportProxy(ViewportProxies[ViewportProxyIndex]);
-		}
 		ViewportProxies.RemoveAt(ViewportProxyIndex);
 	}
 
 	const int32 ClusterViewportProxyIndex = ClusterNodeViewportProxies.Find(InViewportProxy);
 	if (ClusterViewportProxyIndex != INDEX_NONE)
 	{
-		if (ViewportManagerViewExtension.IsValid())
-		{
-			ViewportManagerViewExtension->DeleteViewportProxy(ClusterNodeViewportProxies[ClusterViewportProxyIndex]);
-		}
 		ClusterNodeViewportProxies.RemoveAt(ClusterViewportProxyIndex);
 	}
 }
 
-void FDisplayClusterViewportManagerProxy::ImplUpdateRenderFrameSettings(const FDisplayClusterRenderFrameSettings& InRenderFrameSettings)
+void FDisplayClusterViewportManagerProxy::ImplUpdateRenderFrameSettings(const FDisplayClusterRenderFrameSettings& InRenderFrameSettings,
+	const TSharedPtr<FDisplayClusterViewportManagerViewExtension, ESPMode::ThreadSafe>& InViewportManagerViewExtension)
 {
 	check(IsInGameThread());
 
@@ -189,9 +175,10 @@ void FDisplayClusterViewportManagerProxy::ImplUpdateRenderFrameSettings(const FD
 
 	// Send frame settings to renderthread
 	ENQUEUE_RENDER_COMMAND(DisplayClusterUpdateRenderFrameSettings)(
-		[ViewportManagerProxy = SharedThis(this), Settings](FRHICommandListImmediate& RHICmdList)
+		[ViewportManagerProxy = SharedThis(this), Settings, InViewportManagerViewExtension](FRHICommandListImmediate& RHICmdList)
 	{
 		ViewportManagerProxy->RenderFrameSettings = *Settings;
+		ViewportManagerProxy->ViewportManagerViewExtension = InViewportManagerViewExtension;
 		delete Settings;
 
 		// After updated settings we need update cluster node viewports
diff --git a/Engine/Plugins/Runtime/nDisplay/Source/DisplayCluster/Private/Render/Viewport/DisplayClusterViewportManagerProxy.h b/Engine/Plugins/Runtime/nDisplay/Source/DisplayCluster/Private/Render/Viewport/DisplayClusterViewportManagerProxy.h
--- a/Engine/Plugins/Runtime/nDisplay/Source/DisplayCluster/Private/Render/Viewport/DisplayClusterViewportManagerProxy.h
+++ b/Engine/Plugins/Runtime/nDisplay/Source/DisplayCluster/Private/Render/Viewport/DisplayClusterViewportManagerProxy.h
@@ -13,6 +13,7 @@
 class FDisplayClusterRenderTargetManager;
 class FDisplayClusterViewportPostProcessManager;
 class FDisplayClusterViewportManager;
+class FDisplayClusterViewportManagerViewExtension;
 class FDisplayClusterViewportResource;
 class FDisplayClusterViewportLightCardManagerProxy;
 class IDisplayClusterProjectionPolicy;
@@ -75,9 +76,6 @@ public:
 	/** Initialize ViewportManagerProxy from ViewportManager. */
 	void Initialize(FDisplayClusterViewportManager& InViewportManager);
 
-	/** Set the viewport manager view extension pointer. */
-	void SetViewportManagerViewExtension(const TSharedPtr<class FDisplayClusterViewportManagerViewExtension>& InExtension);
-
 	/** Get LightCardManager proxy object. */
 	TSharedPtr<FDisplayClusterViewportLightCardManagerProxy, ESPMode::ThreadSafe> GetLightCardManagerProxy_RenderThread() const
 	{ return LightCardManagerProxy; }
@@ -97,7 +95,8 @@ public:
 	FDisplayClusterViewportProxy* ImplFindViewport_RenderThread(const int32 StereoViewIndex, uint32* OutContextNum = nullptr) const;
 
 	/** Copy RenderFrameSettings from game to render thread. */
-	void ImplUpdateRenderFrameSettings(const FDisplayClusterRenderFrameSettings& InRenderFrameSettings);
+	void ImplUpdateRenderFrameSettings(const FDisplayClusterRenderFrameSettings& InRenderFrameSettings,
+		const TSharedPtr<FDisplayClusterViewportManagerViewExtension, ESPMode::ThreadSafe>& InViewportManagerViewExtension);
 
 	/** Copy Viewports data to Proxies from game to render thread. */
 	void ImplUpdateViewports(const TArray<FDisplayClusterViewport*>& InViewports);
diff --git a/Engine/Plugins/Runtime/nDisplay/Source/DisplayCluster/Private/Render/Viewport/DisplayClusterViewportManagerViewExtension.cpp b/Engine/Plugins/Runtime/nDisplay/Source/DisplayCluster/Private/Render/Viewport/DisplayClusterViewportManagerViewExtension.cpp
--- a/Engine/Plugins/Runtime/nDisplay/Source/DisplayCluster/Private/Render/Viewport/DisplayClusterViewportManagerViewExtension.cpp
+++ b/Engine/Plugins/Runtime/nDisplay/Source/DisplayCluster/Private/Render/Viewport/DisplayClusterViewportManagerViewExtension.cpp
@@ -120,20 +120,6 @@ FScreenPassTexture FDisplayClusterViewportManagerViewExtension::ReturnUntouchedS
 	}
 }
 
-void FDisplayClusterViewportManagerViewExtension::DeleteViewportProxy(
-	const TSharedPtr<FDisplayClusterViewportProxy>& InProxy)
-{
-	const int32 Index = ViewportProxies.IndexOfByPredicate([&](const FViewportProxy& ExtensionProxy)
-	{
-		return ExtensionProxy.ViewportProxy == InProxy;
-	});
-
-	if (Index != INDEX_NONE)
-	{
-		ViewportProxies.RemoveAt(Index);
-	}
-}
-
 bool FDisplayClusterViewportManagerViewExtension::IsActiveThisFrame_Internal(const FSceneViewExtensionContext& Context) const
 {
 	static const FDisplayClusterSceneViewExtensionContext DCViewExtensionContext;
@@ -183,6 +169,8 @@ void FDisplayClusterViewportManagerViewExtension::PostRenderViewFamily_RenderThr
 			ViewportIt.ViewportProxy->OnPostRenderViewFamily_RenderThread(GraphBuilder, InViewFamily, *SceneView, ViewportIt.ViewportProxyContext);
 		}
 	}
+
+	ViewportProxies.Empty();
 }
 
 //////////////////////////////////////////////////////////////////////////////////////////////
diff --git a/Engine/Plugins/Runtime/nDisplay/Source/DisplayCluster/Private/Render/Viewport/DisplayClusterViewportManagerViewExtension.h b/Engine/Plugins/Runtime/nDisplay/Source/DisplayCluster/Private/Render/Viewport/DisplayClusterViewportManagerViewExtension.h
--- a/Engine/Plugins/Runtime/nDisplay/Source/DisplayCluster/Private/Render/Viewport/DisplayClusterViewportManagerViewExtension.h
+++ b/Engine/Plugins/Runtime/nDisplay/Source/DisplayCluster/Private/Render/Viewport/DisplayClusterViewportManagerViewExtension.h
@@ -52,9 +52,6 @@ public:
 	 * @return - Screen pass texture
 	 */
 	static FScreenPassTexture ReturnUntouchedSceneColorForPostProcessing(const FPostProcessMaterialInputs& InOutInputs);
-
-	/** Remove an internally managed viewport proxy */
-	void DeleteViewportProxy(const TSharedPtr<FDisplayClusterViewportProxy>& InProxy);
 private:
 	/** nDisplay VE Callback [subscribed to Renderer:ResolvedSceneColorCallbacks].
 	 *
