nDisplay: Added ability for light cards to stay in place when their orbit origin, the viewport origin, is moved around.

diff --git a/Engine/Plugins/Runtime/nDisplay/Source/DisplayCluster/Private/Game/EngineClasses/Scene/DisplayClusterRootActor.cpp b/Engine/Plugins/Runtime/nDisplay/Source/DisplayCluster/Private/Game/EngineClasses/Scene/DisplayClusterRootActor.cpp
--- a/Engine/Plugins/Runtime/nDisplay/Source/DisplayCluster/Private/Game/EngineClasses/Scene/DisplayClusterRootActor.cpp
+++ b/Engine/Plugins/Runtime/nDisplay/Source/DisplayCluster/Private/Game/EngineClasses/Scene/DisplayClusterRootActor.cpp
@@ -783,8 +783,11 @@ void ADisplayClusterRootActor::UpdateLightCardPositions()
 	}
 
 	const FVector Location = ViewOriginComponent ? ViewOriginComponent->GetComponentLocation() : GetActorLocation();
+	const FVector RelativeLocation = GetTransform().InverseTransformPosition(Location);
 	const FRotator Rotation = GetActorRotation();
 
+	const bool bRepositionLightCards = RelativeLocation != StageActorOrbitLocation;
+
 	// Iterate over all the light cards referenced in the root actor's config data and update their location and rotation to match the view origin
 	if (UDisplayClusterConfigurationData* CurrentData = GetConfigData())
 	{
@@ -800,8 +803,53 @@ void ADisplayClusterRootActor::UpdateLightCardPositions()
 				LightCardActor->SetRootActorOwner(this);
 				if (LightCardActor->bLockToOwningRootActor)
 				{
+					const FTransform OldLightCardTransform = LightCardActor->GetStageActorTransform();
+
 					LightCardActor->SetActorLocation(Location);
 					LightCardActor->SetActorRotation(Rotation);
+
+					if (bRepositionLightCards)
+					{
+						const FTransform StageActorTransform = LightCardActor->GetOrigin();
+						FDisplayClusterPositionalParams PositionalParams = LightCardActor->GetPositionalParams();
+
+						// We want to adjust the positional params so that the light card stays in the same world position.
+						// First, compute the desired relative location based on the new orbit origin and the light card's previous world position
+						// Then convert that to longitude and latitude values
+						const FVector OrbitLocation = Rotation.UnrotateVector(OldLightCardTransform.GetLocation() - Location);
+
+						PositionalParams.DistanceFromCenter = OrbitLocation.Length();
+
+						if (PositionalParams.DistanceFromCenter > UE_SMALL_NUMBER)
+						{
+							PositionalParams.Latitude = FMath::RadiansToDegrees(FMath::Asin(OrbitLocation.Z / PositionalParams.DistanceFromCenter));
+							PositionalParams.Longitude = FMath::RadiansToDegrees(FMath::Atan2(OrbitLocation.Y, OrbitLocation.X)) + 180.0;
+						}
+						else
+						{
+							PositionalParams.Latitude = 0.0;
+							PositionalParams.Longitude = 180.0;
+						}
+
+						// The light card's orientation needs to be adjusted as well, since the orientation parameters are defined in radial space,
+						// which is relative to the orbit origin
+						const FVector WorldNormal = OldLightCardTransform.GetRotation().RotateVector(FVector::ForwardVector);
+						const FVector LocalNormal = FRotationMatrix::MakeFromX(OrbitLocation.GetSafeNormal()).InverseTransformVector(StageActorTransform.InverseTransformVectorNoScale(WorldNormal));
+						const FRotator NormalRotation = FRotationMatrix::MakeFromX(-LocalNormal).Rotator();
+
+						PositionalParams.Pitch = NormalRotation.Pitch;
+						PositionalParams.Yaw = NormalRotation.Yaw;
+
+						// The spin also needs to be adjusted. This can be done by transforming the old spin vector from world space into the new normal space,
+						// and then converting to an angle
+						const FVector WorldSpin = OldLightCardTransform.GetRotation().RotateVector(FVector::UpVector);
+						const FVector LocalSpin = FRotationMatrix::MakeFromX(OrbitLocation.GetSafeNormal()).InverseTransformVector(StageActorTransform.InverseTransformVectorNoScale(WorldSpin));
+						const FVector NormalSpin = FRotationMatrix::MakeFromX(LocalNormal).InverseTransformVector(LocalSpin);
+
+						PositionalParams.Spin = FMath::RadiansToDegrees(FMath::Atan2(NormalSpin.Y, NormalSpin.Z));
+
+						LightCardActor->SetPositionalParams(PositionalParams);
+					}
 				}
 			}
 		}
@@ -830,6 +878,8 @@ void ADisplayClusterRootActor::UpdateLightCardPositions()
 			}
 		}
 	}
+
+	StageActorOrbitLocation = RelativeLocation;
 }
 
 void ADisplayClusterRootActor::SetChromakeyCardsOwner()
diff --git a/Engine/Plugins/Runtime/nDisplay/Source/DisplayCluster/Private/Game/EngineClasses/Scene/DisplayClusterStageGeometryComponent.cpp b/Engine/Plugins/Runtime/nDisplay/Source/DisplayCluster/Private/Game/EngineClasses/Scene/DisplayClusterStageGeometryComponent.cpp
--- a/Engine/Plugins/Runtime/nDisplay/Source/DisplayCluster/Private/Game/EngineClasses/Scene/DisplayClusterStageGeometryComponent.cpp
+++ b/Engine/Plugins/Runtime/nDisplay/Source/DisplayCluster/Private/Game/EngineClasses/Scene/DisplayClusterStageGeometryComponent.cpp
@@ -269,6 +269,9 @@ void UDisplayClusterStageGeometryComponent::UpdateStageGeometry()
 			}
 		});
 
+		// Make sure to add the viewport the stage geometry has been rendered from, in case that viewpoint is outside the stage geometry bounds
+		StageBoundingBox += RootActor->GetCommonViewPoint()->GetComponentLocation();
+
 		StageBoundingBox = StageBoundingBox.ShiftBy(-RootActor->GetActorLocation());
 		StageBoundingRadius = FMath::Max(StageBoundingBox.Min.Length(), StageBoundingBox.Max.Length());
 	}
diff --git a/Engine/Plugins/Runtime/nDisplay/Source/DisplayCluster/Public/DisplayClusterRootActor.h b/Engine/Plugins/Runtime/nDisplay/Source/DisplayCluster/Public/DisplayClusterRootActor.h
--- a/Engine/Plugins/Runtime/nDisplay/Source/DisplayCluster/Public/DisplayClusterRootActor.h
+++ b/Engine/Plugins/Runtime/nDisplay/Source/DisplayCluster/Public/DisplayClusterRootActor.h
@@ -316,6 +316,9 @@ private:
 
 	float LastDeltaSecondsValue = 0.f;
 
+	/** Stores the location relative to the actor's origin that stage actors like light cards or CCWs orbit around when owned by this root actor. */
+	FVector StageActorOrbitLocation = FVector::ZeroVector;
+
 private:
 	template <typename TComp>
 	void GetTypedPrimitives(TSet<FPrimitiveComponentId>& OutPrimitives, const TArray<FString>* InCompNames = nullptr, bool bCollectChildrenVisualizationComponent = true) const;
