nDisplay: Additional fixes when deleting and undoing viewports.
- Remove an ensure which was possible and valid to be hit during undo of viewport deletion.
- Fix a crash when deleting a viewport, closing the BP editor, then undoing the deletion.
- Fix a crash when GC cleans up a deleted viewport.
- Fix incorrect map iteration when removing bulk viewports through custom property propagation.
- Add hack to solve an issue where a viewport policy parameters could become out of sync with the archetype through undo/redo.

diff --git a/Engine/Plugins/Runtime/nDisplay/Source/DisplayCluster/Private/Game/EngineClasses/Scene/DisplayClusterRootActor.cpp b/Engine/Plugins/Runtime/nDisplay/Source/DisplayCluster/Private/Game/EngineClasses/Scene/DisplayClusterRootActor.cpp
--- a/Engine/Plugins/Runtime/nDisplay/Source/DisplayCluster/Private/Game/EngineClasses/Scene/DisplayClusterRootActor.cpp
+++ b/Engine/Plugins/Runtime/nDisplay/Source/DisplayCluster/Private/Game/EngineClasses/Scene/DisplayClusterRootActor.cpp
@@ -1163,6 +1163,25 @@ static void PropagateDefaultMapToInstancedMap(const FMapProperty* MapProperty, U
 			ObjectToAdd = NewObject<UObject>(InstanceOwner, ArchetypeToUse->GetClass(),
 			ArchetypeToUse->GetFName(), RF_Transactional, ArchetypeToUse);
 		}
+#if WITH_EDITOR
+		else if (GIsTransacting)
+		{
+			// HACK: Projection policy parameters can become cleared when a VP is deleted, undone, redone, undone.
+			// The policy params are meant to only be set from the DCRA CDO, so resetting them back to the VP archetype is safe.
+			// @todo figure out how the parameters are being cleared between the redo/undo. No other properties seem to be impacted. 
+			if (UDisplayClusterConfigurationViewport* InstanceViewport = Cast<UDisplayClusterConfigurationViewport>(ObjectToAdd))
+			{
+				const UDisplayClusterConfigurationViewport* ArchetypeViewport = CastChecked<UDisplayClusterConfigurationViewport>(ArchetypeToUse);
+				if (InstanceViewport->ProjectionPolicy.Parameters.IsEmpty() && !ArchetypeViewport->ProjectionPolicy.Parameters.IsEmpty())
+				{
+					UE_LOG(LogDisplayClusterBlueprint, Warning, TEXT("Projection Policy mismatch from archetype on viewport, correcting. Instance: %s, Archetype: %s."),
+						*ObjectToAdd->GetName(), *ObjectToAdd->GetArchetype()->GetName());
+
+					InstanceViewport->ProjectionPolicy.Parameters = ArchetypeViewport->ProjectionPolicy.Parameters;
+				}
+			}
+		}
+#endif
 		
 		MapInstanceHelper.AddPair(Key, &ObjectToAdd);
 		
@@ -1217,7 +1236,7 @@ static void PropagateDefaultMapToInstancedMap(const FMapProperty* MapProperty, U
 						ensure(bArchetypeCorrect || bIsMultiUserSession);
 					}
 #endif
-					UE_LOG(LogDisplayClusterBlueprint, Warning, TEXT("Archetype mismatch on nDisplay config %s. Make sure the config is compiled and saved. Property: %s, Instance: %s, Archetype: %s, Default %s."),
+					UE_LOG(LogDisplayClusterBlueprint, Warning, TEXT("Archetype mismatch on nDisplay config %s. Make sure the config is compiled and saved. Property: %s, Instance: %s, Archetype: %s, Default: %s."),
 						*RootActor->GetName(), *MapProperty->GetName(), *InstancedObject->GetName(), *InstancedObject->GetArchetype()->GetName(), *DefaultObject->GetName());
 				}
 				continue;
@@ -1228,7 +1247,7 @@ static void PropagateDefaultMapToInstancedMap(const FMapProperty* MapProperty, U
 	}
 
 	// Look for elements that should be removed.
-	for (FScriptMapHelper::FIterator InstanceIt = MapInstanceHelper.CreateIterator(); InstanceIt;)
+	for (FScriptMapHelper::FIterator InstanceIt = MapInstanceHelper.CreateIterator(); InstanceIt; ++InstanceIt)
 	{
 		uint8* InstancePairPtr = MapInstanceHelper.GetPairPtr(*InstanceIt);
 		check(InstancePairPtr);
@@ -1254,10 +1273,6 @@ static void PropagateDefaultMapToInstancedMap(const FMapProperty* MapProperty, U
 			MapInstanceHelper.RemoveAt(*InstanceIt);
 			bHasChanged = true;
 		}
-		else
-		{
-			++InstanceIt;
-		}
 	}
 	
 	if (bHasChanged)
diff --git a/Engine/Plugins/Runtime/nDisplay/Source/DisplayCluster/Private/Misc/DisplayClusterObjectRef.cpp b/Engine/Plugins/Runtime/nDisplay/Source/DisplayCluster/Private/Misc/DisplayClusterObjectRef.cpp
--- a/Engine/Plugins/Runtime/nDisplay/Source/DisplayCluster/Private/Misc/DisplayClusterObjectRef.cpp
+++ b/Engine/Plugins/Runtime/nDisplay/Source/DisplayCluster/Private/Misc/DisplayClusterObjectRef.cpp
@@ -20,7 +20,8 @@ USceneComponent* FDisplayClusterSceneComponentRef::GetOrFindSceneComponent() con
 		{
 			for (UActorComponent* ItActorComponent : Actor->GetComponents())
 			{
-				if (ItActorComponent->GetFName() == ComponentName)
+				// Component can be null when a viewport is deleted and GC occurs.
+				if (ItActorComponent && ItActorComponent->GetFName() == ComponentName)
 				{
 					USceneComponent* SceneComponent = Cast<USceneComponent>(ItActorComponent);
 					if (SceneComponent)
diff --git a/Engine/Plugins/Runtime/nDisplay/Source/DisplayClusterConfiguration/Private/DisplayClusterConfigurationTypes.cpp b/Engine/Plugins/Runtime/nDisplay/Source/DisplayClusterConfiguration/Private/DisplayClusterConfigurationTypes.cpp
--- a/Engine/Plugins/Runtime/nDisplay/Source/DisplayClusterConfiguration/Private/DisplayClusterConfigurationTypes.cpp
+++ b/Engine/Plugins/Runtime/nDisplay/Source/DisplayClusterConfiguration/Private/DisplayClusterConfigurationTypes.cpp
@@ -245,10 +245,11 @@ void UDisplayClusterConfigurationViewport::OnPreCompile(FCompilerResultsLog& Mes
 {
 	Super::OnPreCompile(MessageLog);
 
-	if (!ensure(bAllowPreviewTexture))
+	if (!bAllowPreviewTexture)
 	{
 		// Verify correct rendering value is applied. This branch shouldn't be hit as long as the Region
-		// struct always has a final PostEditChangeProperty called without a change type of Interactive.
+		// struct always has a final PostEditChangeProperty called without a change type of Interactive, but may
+		// be hit when compiling after this viewport was deleted and then the deletion undone.
 		EnablePreviewTexture();
 	}
 }
diff --git a/Engine/Plugins/Runtime/nDisplay/Source/DisplayClusterConfigurator/Private/Views/OutputMapping/EdNodes/DisplayClusterConfiguratorBaseNode.cpp b/Engine/Plugins/Runtime/nDisplay/Source/DisplayClusterConfigurator/Private/Views/OutputMapping/EdNodes/DisplayClusterConfiguratorBaseNode.cpp
--- a/Engine/Plugins/Runtime/nDisplay/Source/DisplayClusterConfigurator/Private/Views/OutputMapping/EdNodes/DisplayClusterConfiguratorBaseNode.cpp
+++ b/Engine/Plugins/Runtime/nDisplay/Source/DisplayClusterConfigurator/Private/Views/OutputMapping/EdNodes/DisplayClusterConfiguratorBaseNode.cpp
@@ -31,7 +31,10 @@ void UDisplayClusterConfiguratorBaseNode::PostEditUndo()
 
 	// Make sure the undo operation isn't going to a state where this node or the object it edits isn't valid
 	// before attempting to update the object or its children
-	if (IsObjectValid())
+	//
+	// The ToolkitPtr could be invalid if the node was deleted, the BP editor was closed, then the operation undone.
+	// This node is invalid in this case and will be recreated in the new BP editor instance.
+	if (IsObjectValid() && ToolkitPtr.IsValid())
 	{
 		CleanupChildrenNodes();
 		
diff --git a/Engine/Plugins/Runtime/nDisplay/Source/DisplayClusterConfigurator/Private/Views/OutputMapping/EdNodes/DisplayClusterConfiguratorViewportNode.cpp b/Engine/Plugins/Runtime/nDisplay/Source/DisplayClusterConfigurator/Private/Views/OutputMapping/EdNodes/DisplayClusterConfiguratorViewportNode.cpp
--- a/Engine/Plugins/Runtime/nDisplay/Source/DisplayClusterConfigurator/Private/Views/OutputMapping/EdNodes/DisplayClusterConfiguratorViewportNode.cpp
+++ b/Engine/Plugins/Runtime/nDisplay/Source/DisplayClusterConfigurator/Private/Views/OutputMapping/EdNodes/DisplayClusterConfiguratorViewportNode.cpp
@@ -36,7 +36,7 @@ void UDisplayClusterConfiguratorViewportNode::Cleanup()
 
 TSharedPtr<SGraphNode> UDisplayClusterConfiguratorViewportNode::CreateVisualWidget()
 {
-	return SNew(SDisplayClusterConfiguratorViewportNode, this, ToolkitPtr.Pin().ToSharedRef());;
+	return SNew(SDisplayClusterConfiguratorViewportNode, this, ToolkitPtr.Pin().ToSharedRef());
 }
 
 void UDisplayClusterConfiguratorViewportNode::ResizeNode(const FVector2D& NewSize)
