nDisplay: fixed several preview render issues
 - updating viewport proxy data now happens before rendering starts in the InitializeNewFrame() function
 - The InitializeNewFrame() function call was skipped in DCRA pre-rendering.
 - The ViewportManager functions were used in the wrong order.

diff --git a/Engine/Plugins/Runtime/nDisplay/Source/DisplayCluster/Private/Game/EngineClasses/Scene/DisplayClusterRootActorEditor.cpp b/Engine/Plugins/Runtime/nDisplay/Source/DisplayCluster/Private/Game/EngineClasses/Scene/DisplayClusterRootActorEditor.cpp
--- a/Engine/Plugins/Runtime/nDisplay/Source/DisplayCluster/Private/Game/EngineClasses/Scene/DisplayClusterRootActorEditor.cpp
+++ b/Engine/Plugins/Runtime/nDisplay/Source/DisplayCluster/Private/Game/EngineClasses/Scene/DisplayClusterRootActorEditor.cpp
@@ -92,7 +92,8 @@ void ADisplayClusterRootActor::ResetPreviewInternals_Editor()
 
 	TickPerFrameCounter = 0;
 	PreviewClusterNodeIndex = 0;
-	PreviewViewportIndex = 0;
+
+	ResetClusterNodePreviewRendering_Editor();
 }
 
 void ADisplayClusterRootActor::Constructor_Editor()
@@ -376,10 +377,15 @@ bool ADisplayClusterRootActor::ImplUpdatePreviewConfiguration_Editor(const FStri
 {
 	DECLARE_SCOPE_CYCLE_COUNTER(TEXT("ADisplayClusterRootActor::ImplUpdatePreviewConfiguration_Editor"), STAT_ImplUpdatePreviewConfiguration_Editor, STATGROUP_NDisplay);
 	
-	PreviewRenderFrame.Reset();
+	// The function UpdateConfiguration() rebuilds viewports and their internals, so rendering is no longer valid for now.
+	// Reset current preview rendering.
+	ResetClusterNodePreviewRendering_Editor();
 
 	if (IsPreviewEnabled() && ViewportManager.IsValid())
 	{
+		// Now we render this node
+		PreviewRenderFrameClusterNodeId = InClusterNodeId;
+
 		FDisplayClusterPreviewSettings PreviewSettings;
 		PreviewSettings.PreviewRenderTargetRatioMult = PreviewRenderTargetRatioMult;
 
@@ -398,54 +404,69 @@ bool ADisplayClusterRootActor::ImplUpdatePreviewConfiguration_Editor(const FStri
 	return false;
 }
 
+bool ADisplayClusterRootActor::IsActiveClusterNodePreviewRendering_Editor() const
+{
+	return PreviewRenderFrame.IsValid() && PreviewViewportIndex >= 0 && !PreviewRenderFrameClusterNodeId.IsEmpty();
+}
+
+void ADisplayClusterRootActor::ResetClusterNodePreviewRendering_Editor()
+{
+	PreviewViewportIndex = -1;
+	PreviewRenderFrame.Reset();
+	PreviewRenderFrameClusterNodeId.Empty();
+}
+
 bool ADisplayClusterRootActor::ImplUpdatePreviewRenderFrame_Editor(const FString& InClusterNodeId)
 {
+	// Update cluster node for render:
+	if (PreviewRenderFrameClusterNodeId != InClusterNodeId)
+	{
+		ResetClusterNodePreviewRendering_Editor();
+	}
+
 	if (!PreviewRenderFrame.IsValid())
 	{
+		// Now always use RootActor world to preview.
+		UWorld* CurrentPreviewWorld = GetWorld();
+
 		// Begin render new frame for cluster node
-		if (ImplUpdatePreviewConfiguration_Editor(InClusterNodeId))
+		if (CurrentPreviewWorld == nullptr || !ImplUpdatePreviewConfiguration_Editor(InClusterNodeId))
+		{
+			return false;
+		}
+
+		// Update all preview components resources before render
+		for (const TTuple<FString, TObjectPtr<UDisplayClusterPreviewComponent>>& PreviewComponentIt : PreviewComponents)
 		{
-			// Update all preview components resources before render
-			for (const TTuple<FString, TObjectPtr<UDisplayClusterPreviewComponent>>& PreviewComponentIt : PreviewComponents)
+			if (PreviewComponentIt.Value && (InClusterNodeId.IsEmpty() || PreviewComponentIt.Value->GetClusterNodeId() == InClusterNodeId))
 			{
-				if (PreviewComponentIt.Value && (InClusterNodeId.IsEmpty() || PreviewComponentIt.Value->GetClusterNodeId() == InClusterNodeId))
-				{
-					PreviewComponentIt.Value->UpdatePreviewResources();
-				}
+				PreviewComponentIt.Value->UpdatePreviewResources();
 			}
+		}
 
-			// Now always use RootActor world to preview. 
-			UWorld* PreviewWorld = GetWorld();
-			if (PreviewWorld)
-			{
-				PreviewRenderFrame = MakeUnique<FDisplayClusterRenderFrame>();
+		PreviewRenderFrame = MakeUnique<FDisplayClusterRenderFrame>();
 
-				const bool bSceneNeedsStarting = !ViewportManager->GetCurrentWorld();
-				
-				// Update preview viewports from settings
-				if (ViewportManager->BeginNewFrame(nullptr, PreviewWorld, *PreviewRenderFrame))
-				{
-					if (bSceneNeedsStarting)
-					{
-						// Fix inner frustum possibly using the wrong post process render target
-						ImplUpdatePreviewConfiguration_Editor(InClusterNodeId);
-					}
-					
-					PreviewViewportIndex = 0;
-					return true;
-				}
+		// Update preview viewports from settings
+		if (!ViewportManager->BeginNewFrame(nullptr, CurrentPreviewWorld, *PreviewRenderFrame))
+		{
+			PreviewRenderFrame.Reset();
 
-				return false;
-			}
+			return false;
 		}
+
+		// Begin Render Preview For Cluster Node
+		PreviewViewportIndex = 0;
+
+		// Initialize frame for render
+		ViewportManager->InitializeNewFrame();
 	}
 
-	return true;
+	return IsActiveClusterNodePreviewRendering_Editor();
 }
 
-bool ADisplayClusterRootActor::ImplRenderPassPreviewClusterNode_Editor()
+bool ADisplayClusterRootActor::ImplRenderPassPreviewClusterNode_Editor(const FString& InClusterNodeId)
 {
-	if (!PreviewRenderFrame.IsValid() || PreviewViewportIndex < 0)
+	if(!ImplUpdatePreviewRenderFrame_Editor(InClusterNodeId))
 	{
 		return false;
 	}
@@ -471,8 +492,7 @@ bool ADisplayClusterRootActor::ImplRenderPassPreviewClusterNode_Editor()
 	if (bFrameRendered)
 	{
 		// current cluster node is composed
-		PreviewViewportIndex = -1;
-		PreviewRenderFrame.Reset();
+		ResetClusterNodePreviewRendering_Editor();
 
 		// Send event about RTT changed
 		OnPreviewGenerated.ExecuteIfBound();
@@ -525,19 +545,19 @@ void ADisplayClusterRootActor::ImplRenderPreview_Editor()
 			return;
 		}
 
-		const FString& ClusterNodeId = ExistClusterNodesIDs[PreviewClusterNodeIndex];
-
-		if (PreviewRenderFrameClusterNodeId != ClusterNodeId)
-		{
-			PreviewRenderFrameClusterNodeId = ClusterNodeId;
-			PreviewRenderFrame.Reset();
-			PreviewViewportIndex = -1;
-		}
-
-		ImplUpdatePreviewRenderFrame_Editor(PreviewRenderFrameClusterNodeId);
-
-		// Render this cluster node viewports
-		if (!ImplRenderPassPreviewClusterNode_Editor())
+		/**
+		 * Render this cluster node viewports
+		 * Note: ViewportManager should be used correctly:
+		 * 1. ViewportManager->UpdateConfiguration(ClusterNode, Configuration) - only when changing DCRA or rendering settings (or cluster node)
+		 * 2. ViewportManager->BeginNewFrame(PreviewWorld, RenderSettings) - once per frame for the entire cluster node
+		 * 3. The cluster node will be rendered using the RenderInEditor() function, which will be called multiple times until the node is completely rendered.
+		 * 3.1. ViewportManager->RenderInEditor(PreviewInfo)
+		 * ...
+		 * 3.N. ViewportManager->RenderInEditor(PreviewInfo)
+		 * 
+		 * Warning: When any function of this workflow is called in the wrong order, the entire workflow must be restarted.
+		 */
+		if (!ImplRenderPassPreviewClusterNode_Editor(ExistClusterNodesIDs[PreviewClusterNodeIndex]))
 		{
 			// Cluster node render still in progress..
 			return;
@@ -936,7 +956,11 @@ void ADisplayClusterRootActor::UpdatePreviewComponents()
 		return;
 	}
 
-	ImplUpdatePreviewConfiguration_Editor(DisplayClusterConfigurationStrings::gui::preview::PreviewNodeAll);
+	// Do not updated inside preview rendering workflow
+	if (!IsActiveClusterNodePreviewRendering_Editor())
+	{
+		ImplUpdatePreviewConfiguration_Editor(DisplayClusterConfigurationStrings::gui::preview::PreviewNodeAll);
+	}
 
 	TArray<UDisplayClusterPreviewComponent*> IteratedPreviewComponents;
 
diff --git a/Engine/Plugins/Runtime/nDisplay/Source/DisplayCluster/Public/DisplayClusterRootActor.h b/Engine/Plugins/Runtime/nDisplay/Source/DisplayCluster/Public/DisplayClusterRootActor.h
--- a/Engine/Plugins/Runtime/nDisplay/Source/DisplayCluster/Public/DisplayClusterRootActor.h
+++ b/Engine/Plugins/Runtime/nDisplay/Source/DisplayCluster/Public/DisplayClusterRootActor.h
@@ -418,7 +418,14 @@ public:
 
 	/** If editor rendering is enabled. */
 	bool IsEditorRenderEnabled() const { return bEnableEditorRender; }
-	
+
+protected:
+	/** Reset preview rendering for cluster node. */
+	void ResetClusterNodePreviewRendering_Editor();
+
+	/** Is preview rendering for cluster node in progress. */
+	bool IsActiveClusterNodePreviewRendering_Editor() const;
+
 private:
 	/** Is editor rendering enabled? This can be false and the preview still enabled. */
 	bool bEnableEditorRender = true;
@@ -520,7 +527,7 @@ protected:
 	bool ImplUpdatePreviewConfiguration_Editor(const FString& InClusterNodeId);
 
 	void ImplRenderPreview_Editor();
-	bool ImplRenderPassPreviewClusterNode_Editor();
+	bool ImplRenderPassPreviewClusterNode_Editor(const FString& InClusterNodeId);
 
 	bool ImplUpdatePreviewRenderFrame_Editor(const FString& InClusterNodeId);
 
@@ -562,9 +569,11 @@ private:
 	int32 TickPerFrameCounter = 0;
 
 	int32 PreviewClusterNodeIndex = 0;
-	int32 PreviewViewportIndex = 0;
+	
+	int32 PreviewViewportIndex = -1;
 	TUniquePtr<FDisplayClusterRenderFrame> PreviewRenderFrame;
 	FString PreviewRenderFrameClusterNodeId;
+
 	int32 PreviewViewportsRenderedInThisFrameCnt = 0;
 
 	FOnPreviewUpdated OnPreviewGenerated;
