[nDisplay][MediaIO][Rivermax] Media capture synchronization

diff --git a/Engine/Plugins/Media/MediaIOFramework/Source/MediaIOCore/Public/MediaCapture.h b/Engine/Plugins/Media/MediaIOFramework/Source/MediaIOCore/Public/MediaCapture.h
--- a/Engine/Plugins/Media/MediaIOFramework/Source/MediaIOCore/Public/MediaCapture.h
+++ b/Engine/Plugins/Media/MediaIOFramework/Source/MediaIOCore/Public/MediaCapture.h
@@ -211,6 +211,7 @@ public:
 /** Delegate signatures */
 DECLARE_DYNAMIC_MULTICAST_DELEGATE(FMediaCaptureStateChangedSignature);
 DECLARE_MULTICAST_DELEGATE(FMediaCaptureStateChangedSignatureNative);
+DECLARE_DELEGATE(FMediaCaptureOutputSynchronization)
 
 /**
  * Abstract base class for media capture.
@@ -334,6 +335,20 @@ public:
 	 */
 	FMediaCaptureStateChangedSignatureNative OnStateChangedNative;
 
+	/**
+	 * Returns true if media capture implementation supports synchronization logic. If so,
+	 * OnOutputSynchronization delegate below must be called before each captured frame is exported.
+	 */
+	virtual bool IsOutputSynchronizationSupported() const { return false; }
+
+	/**
+	 * Called when output data is ready to be sent out. This blocking call allows to
+	 * postpone data export based on the synchronization logic.
+	 *
+	 * This delegate must be called if IsOutputSynchronizationSupported() returns true.
+	 */
+	FMediaCaptureOutputSynchronization OnOutputSynchronization;
+
 	void SetFixedViewportSize(TSharedPtr<FSceneViewport> InSceneViewport, FIntPoint InSize);
 	void ResetFixedViewportSize(TSharedPtr<FSceneViewport> InViewport, bool bInFlushRenderingCommands);
 
diff --git a/Engine/Plugins/Runtime/nDisplay/Source/DisplayCluster/Private/Cluster/Controller/DisplayClusterClusterNodeCtrlPrimary.cpp b/Engine/Plugins/Runtime/nDisplay/Source/DisplayCluster/Private/Cluster/Controller/DisplayClusterClusterNodeCtrlPrimary.cpp
--- a/Engine/Plugins/Runtime/nDisplay/Source/DisplayCluster/Private/Cluster/Controller/DisplayClusterClusterNodeCtrlPrimary.cpp
+++ b/Engine/Plugins/Runtime/nDisplay/Source/DisplayCluster/Private/Cluster/Controller/DisplayClusterClusterNodeCtrlPrimary.cpp
@@ -12,6 +12,7 @@
 #include "Network/Service/RenderSync/DisplayClusterRenderSyncService.h"
 #include "Network/Service/ClusterEventsJson/DisplayClusterClusterEventsJsonService.h"
 #include "Network/Service/ClusterEventsBinary/DisplayClusterClusterEventsBinaryService.h"
+#include "Network/Service/GenericBarrier/DisplayClusterGenericBarrierService.h"
 #include "Network/Listener/DisplayClusterTcpListener.h"
 #include "Network/DisplayClusterNetworkTypes.h"
 
@@ -115,12 +116,13 @@ bool FDisplayClusterClusterNodeCtrlPrimary::InitializeServers()
 	UE_LOG(LogDisplayClusterCluster, Log, TEXT("%s - initializing primary node servers..."), *GetControllerName());
 
 	// Instantiate node servers
-	ClusterSyncServer         = MakeUnique<FDisplayClusterClusterSyncService>       ();
-	RenderSyncServer          = MakeUnique<FDisplayClusterRenderSyncService>        ();
+	ClusterSyncServer         = MakeUnique<FDisplayClusterClusterSyncService>();
+	RenderSyncServer          = MakeUnique<FDisplayClusterRenderSyncService>();
+	GenericBarriersServer     = MakeUnique<FDisplayClusterGenericBarrierService>();
 	ClusterEventsJsonServer   = MakeUnique<FDisplayClusterClusterEventsJsonService> ();
 	ClusterEventsBinaryServer = MakeUnique<FDisplayClusterClusterEventsBinaryService>();
 
-	return ClusterSyncServer && RenderSyncServer && ClusterEventsJsonServer && ClusterEventsBinaryServer;
+	return ClusterSyncServer && RenderSyncServer && GenericBarriersServer && ClusterEventsJsonServer && ClusterEventsBinaryServer;
 }
 
 bool FDisplayClusterClusterNodeCtrlPrimary::StartServers()
@@ -176,6 +178,7 @@ bool FDisplayClusterClusterNodeCtrlPrimary::StartServers()
 	// Start the servers
 	return StartServerWithLogs(ClusterSyncServer.Get(),         TcpListener) // Start with shared listener
 		&& StartServerWithLogs(RenderSyncServer.Get(),          TcpListener) // Start with shared listener
+		&& StartServerWithLogs(GenericBarriersServer.Get(),     TcpListener) // Start with shared listener
 		&& StartServerWithLogs(ClusterEventsJsonServer.Get(),   CfgPrimary->Host, Ports.ClusterEventsJson)
 		&& StartServerWithLogs(ClusterEventsBinaryServer.Get(), CfgPrimary->Host, Ports.ClusterEventsBinary)
 		&& TcpListener->StartListening(CfgPrimary->Host, Ports.ClusterSync); // Start shared listener as well
@@ -191,6 +194,7 @@ void FDisplayClusterClusterNodeCtrlPrimary::StopServers()
 	// Stop services
 	ClusterSyncServer->Shutdown();
 	RenderSyncServer->Shutdown();
+	GenericBarriersServer->Shutdown();
 	ClusterEventsJsonServer->Shutdown();
 	ClusterEventsBinaryServer->Shutdown();
 
diff --git a/Engine/Plugins/Runtime/nDisplay/Source/DisplayCluster/Private/Cluster/Controller/DisplayClusterClusterNodeCtrlPrimary.h b/Engine/Plugins/Runtime/nDisplay/Source/DisplayCluster/Private/Cluster/Controller/DisplayClusterClusterNodeCtrlPrimary.h
--- a/Engine/Plugins/Runtime/nDisplay/Source/DisplayCluster/Private/Cluster/Controller/DisplayClusterClusterNodeCtrlPrimary.h
+++ b/Engine/Plugins/Runtime/nDisplay/Source/DisplayCluster/Private/Cluster/Controller/DisplayClusterClusterNodeCtrlPrimary.h
@@ -10,6 +10,7 @@ class FDisplayClusterClusterSyncService;
 class FDisplayClusterRenderSyncService;
 class FDisplayClusterClusterEventsJsonService;
 class FDisplayClusterClusterEventsBinaryService;
+class FDisplayClusterGenericBarrierService;
 class FDisplayClusterTcpListener;
 
 
@@ -63,6 +64,7 @@ private:
 	// Node servers
 	TUniquePtr<FDisplayClusterClusterSyncService>         ClusterSyncServer;
 	TUniquePtr<FDisplayClusterRenderSyncService>          RenderSyncServer;
+	TUniquePtr<FDisplayClusterGenericBarrierService>      GenericBarriersServer;
 	TUniquePtr<FDisplayClusterClusterEventsJsonService>   ClusterEventsJsonServer;
 	TUniquePtr<FDisplayClusterClusterEventsBinaryService> ClusterEventsBinaryServer;
 
diff --git a/Engine/Plugins/Runtime/nDisplay/Source/DisplayCluster/Private/Cluster/DisplayClusterClusterManager.cpp b/Engine/Plugins/Runtime/nDisplay/Source/DisplayCluster/Private/Cluster/DisplayClusterClusterManager.cpp
--- a/Engine/Plugins/Runtime/nDisplay/Source/DisplayCluster/Private/Cluster/DisplayClusterClusterManager.cpp
+++ b/Engine/Plugins/Runtime/nDisplay/Source/DisplayCluster/Private/Cluster/DisplayClusterClusterManager.cpp
@@ -14,6 +14,8 @@
 #include "Cluster/Failover/DisplayClusterFailoverNodeCtrlPrimary.h"
 #include "Cluster/Failover/DisplayClusterFailoverNodeCtrlSecondary.h"
 
+#include "Network/Service/GenericBarrier/DisplayClusterGenericBarrierClient.h"
+
 #include "Config/IPDisplayClusterConfigManager.h"
 #include "DisplayClusterConfigurationTypes.h"
 
@@ -306,6 +308,11 @@ void FDisplayClusterClusterManager::UnregisterSyncObject(IDisplayClusterClusterS
 	}
 }
 
+TUniquePtr<IDisplayClusterGenericBarriersClient> FDisplayClusterClusterManager::CreateGenericBarriersClient(const FString& ClientName)
+{
+	return MakeUnique<FDisplayClusterGenericBarrierClient>(ClientName);
+}
+
 void FDisplayClusterClusterManager::AddClusterEventListener(TScriptInterface<IDisplayClusterClusterEventListener> Listener)
 {
 	FScopeLock Lock(&ClusterEventListenersCS);
diff --git a/Engine/Plugins/Runtime/nDisplay/Source/DisplayCluster/Private/Cluster/DisplayClusterClusterManager.h b/Engine/Plugins/Runtime/nDisplay/Source/DisplayCluster/Private/Cluster/DisplayClusterClusterManager.h
--- a/Engine/Plugins/Runtime/nDisplay/Source/DisplayCluster/Private/Cluster/DisplayClusterClusterManager.h
+++ b/Engine/Plugins/Runtime/nDisplay/Source/DisplayCluster/Private/Cluster/DisplayClusterClusterManager.h
@@ -62,6 +62,8 @@ public:
 	virtual void RegisterSyncObject(IDisplayClusterClusterSyncObject* SyncObj, EDisplayClusterSyncGroup SyncGroup) override;
 	virtual void UnregisterSyncObject(IDisplayClusterClusterSyncObject* SyncObj) override;
 
+	virtual TUniquePtr<IDisplayClusterGenericBarriersClient> CreateGenericBarriersClient(const FString& ClientName) override;
+
 	virtual void AddClusterEventListener(TScriptInterface<IDisplayClusterClusterEventListener>) override;
 	virtual void RemoveClusterEventListener(TScriptInterface<IDisplayClusterClusterEventListener>) override;
 
diff --git a/Engine/Plugins/Runtime/nDisplay/Source/DisplayCluster/Private/Misc/DisplayClusterLog.cpp b/Engine/Plugins/Runtime/nDisplay/Source/DisplayCluster/Private/Misc/DisplayClusterLog.cpp
--- a/Engine/Plugins/Runtime/nDisplay/Source/DisplayCluster/Private/Misc/DisplayClusterLog.cpp
+++ b/Engine/Plugins/Runtime/nDisplay/Source/DisplayCluster/Private/Misc/DisplayClusterLog.cpp
@@ -7,6 +7,7 @@ DEFINE_LOG_CATEGORY(LogDisplayClusterEngine);
 DEFINE_LOG_CATEGORY(LogDisplayClusterModule);
 DEFINE_LOG_CATEGORY(LogDisplayClusterCluster);
 DEFINE_LOG_CATEGORY(LogDisplayClusterConfig);
+DEFINE_LOG_CATEGORY(LogDisplayClusterBarrier);
 DEFINE_LOG_CATEGORY(LogDisplayClusterGame);
 DEFINE_LOG_CATEGORY(LogDisplayClusterNetwork);
 DEFINE_LOG_CATEGORY(LogDisplayClusterRender);
diff --git a/Engine/Plugins/Runtime/nDisplay/Source/DisplayCluster/Private/Misc/DisplayClusterLog.h b/Engine/Plugins/Runtime/nDisplay/Source/DisplayCluster/Private/Misc/DisplayClusterLog.h
--- a/Engine/Plugins/Runtime/nDisplay/Source/DisplayCluster/Private/Misc/DisplayClusterLog.h
+++ b/Engine/Plugins/Runtime/nDisplay/Source/DisplayCluster/Private/Misc/DisplayClusterLog.h
@@ -10,6 +10,7 @@ DECLARE_LOG_CATEGORY_EXTERN(LogDisplayClusterEngine,     Log, All);
 DECLARE_LOG_CATEGORY_EXTERN(LogDisplayClusterModule,     Log, All);
 DECLARE_LOG_CATEGORY_EXTERN(LogDisplayClusterCluster,    Log, All);
 DECLARE_LOG_CATEGORY_EXTERN(LogDisplayClusterConfig,     Log, All);
+DECLARE_LOG_CATEGORY_EXTERN(LogDisplayClusterBarrier,    Log, All);
 DECLARE_LOG_CATEGORY_EXTERN(LogDisplayClusterGame,       Log, All);
 DECLARE_LOG_CATEGORY_EXTERN(LogDisplayClusterNetwork,    Log, All);
 DECLARE_LOG_CATEGORY_EXTERN(LogDisplayClusterRender,     Log, All);
diff --git a/Engine/Plugins/Runtime/nDisplay/Source/DisplayCluster/Private/Network/Barrier/DisplayClusterBarrier.cpp b/Engine/Plugins/Runtime/nDisplay/Source/DisplayCluster/Private/Network/Barrier/DisplayClusterBarrier.cpp
new file mode 100644
--- /dev/null
+++ b/Engine/Plugins/Runtime/nDisplay/Source/DisplayCluster/Private/Network/Barrier/DisplayClusterBarrier.cpp
@@ -0,0 +1,292 @@
+// Copyright Epic Games, Inc. All Rights Reserved.
+
+#include "Network/Barrier/DisplayClusterBarrier.h"
+
+#include "Misc/DisplayClusterLog.h"
+#include "Misc/ScopeLock.h"
+
+
+FDisplayClusterBarrier::FDisplayClusterBarrier(const FString& InName, const TArray<FString>& InCallersAllowed, const uint32 InTimeout)
+	: Name(InName)
+	, CallersAllowed(InCallersAllowed)
+	, Timeout(InTimeout)
+	, ThreadLimit(static_cast<uint32>(CallersAllowed.Num()))
+	, WatchdogTimer(InName + FString("_watchdog"))
+{
+	UE_LOG(LogDisplayClusterBarrier, Log, TEXT("Initialized barrier '%s' with timeout %u ms and threads limit: %u"), *Name, Timeout, ThreadLimit);
+
+	for (int32 Idx = 0; Idx < CallersAllowed.Num(); ++Idx)
+	{
+		UE_LOG(LogDisplayClusterBarrier, Verbose, TEXT("Barrier '%s': client (%d): '%s'"), *Name, Idx, *CallersAllowed[Idx]);
+	}
+
+	// Subscribe for timeout events
+	WatchdogTimer.OnWatchdogTimeOut().AddRaw(this, &FDisplayClusterBarrier::HandleBarrierTimeout);
+}
+
+FDisplayClusterBarrier::~FDisplayClusterBarrier()
+{
+	// Release threads if there any
+	Deactivate();
+}
+
+
+bool FDisplayClusterBarrier::Activate()
+{
+	FScopeLock Lock(&DataCS);
+
+	UE_LOG(LogDisplayClusterBarrier, Log, TEXT("Barrier '%s': activating..."), *Name);
+
+	if (!bActive)
+	{
+		ThreadCount = 0;
+		bActive = true;
+		CallersAwaiting.Reset();
+
+		// No exit allowed
+		EventOutputGateOpen->Reset();
+		// Allow join
+		EventInputGateOpen->Trigger();
+	}
+
+	return true;
+}
+
+void FDisplayClusterBarrier::Deactivate()
+{
+	FScopeLock Lock(&DataCS);
+
+	if (bActive)
+	{
+		UE_LOG(LogDisplayClusterBarrier, Log, TEXT("Barrier '%s': deactivating..."), *Name);
+
+		bActive = false;
+
+		// Release all threads that are currently at the barrier
+		EventInputGateOpen->Trigger();
+		EventOutputGateOpen->Trigger();
+
+		// No more threads awaiting
+		CallersAwaiting.Reset();
+
+		// And reset timer of course
+		WatchdogTimer.ResetTimer();
+	}
+}
+
+bool FDisplayClusterBarrier::IsActivated() const
+{
+	FScopeLock Lock(&DataCS);
+	return bActive;
+}
+
+
+EDisplayClusterBarrierWaitResult FDisplayClusterBarrier::Wait(const FString& CallerId, double* ThreadWaitTime /*= nullptr*/, double* BarrierWaitTime /*= nullptr*/)
+{
+	UE_LOG(LogDisplayClusterBarrier, VeryVerbose, TEXT("Barrier '%s': caller arrived '%s'"), *Name, *CallerId);
+
+	double ThreadWaitTimeStart = 0;
+
+	{
+		FScopeLock LockEntrance(&EntranceCS);
+
+		// Wait unless the barrier allows new threads to join. This will happen
+		// once all threads from the previous sync iteration leave the barrier,
+		// or the barrier gets deactivated.
+		EventInputGateOpen->Wait();
+
+		{
+			FScopeLock LockData(&DataCS);
+
+			// Check if this thread has been previously dropped
+			if (CallersTimedout.Contains(CallerId))
+			{
+				UE_LOG(LogDisplayClusterBarrier, Verbose, TEXT("Barrier '%s': caller '%s' not allowed to join, it has been timed out previously"), *Name, *CallerId);
+				return EDisplayClusterBarrierWaitResult::TimeOut;
+			}
+
+			// Check if the barrier is active
+			if (!bActive)
+			{
+				UE_LOG(LogDisplayClusterBarrier, Verbose, TEXT("Barrier '%s': not active"), *Name);
+				return EDisplayClusterBarrierWaitResult::NotActive;
+			}
+
+			// Check if this thread is allowed to sync at this barrier
+			if (!CallersAllowed.Contains(CallerId))
+			{
+				UE_LOG(LogDisplayClusterBarrier, Verbose, TEXT("Barrier '%s': caller '%s' not allowed to join, no permission"), *Name, *CallerId);
+				return EDisplayClusterBarrierWaitResult::NotAllowed;
+			}
+
+			// Register caller
+			CallersAwaiting.Add(CallerId);
+
+			// Make sure the barrier was not deactivated while this thread was waiting at the input gate
+			if (bActive)
+			{
+				// Fixate awaiting start for this particular thread
+				ThreadWaitTimeStart = FPlatformTime::Seconds();
+
+				// One more thread awaiting
+				++ThreadCount;
+
+				// In case this thread came first to the barrier, we need:
+				// - to fixate barrier awaiting start time
+				// - start watchdog timer
+				if (ThreadCount == 1)
+				{
+					UE_LOG(LogDisplayClusterBarrier, Verbose, TEXT("Barrier '%s': sync start"), *Name);
+
+					BarrierWaitTimeStart = ThreadWaitTimeStart;
+					WatchdogTimer.SetTimer(Timeout);
+				}
+
+				UE_LOG(LogDisplayClusterBarrier, Verbose, TEXT("Barrier '%s': awaiting threads amount - %d"), *Name, ThreadCount);
+
+				// In case this thread is the last one the barrier is awaiting for, we need:
+				// - to fixate barrier awaiting finish time
+				// - to open the output gate (release the barrier)
+				// - to close the input gate
+				// - reset watchdog timer
+				if (ThreadCount == ThreadLimit)
+				{
+					BarrierWaitTimeFinish = FPlatformTime::Seconds();
+					BarrierWaitTimeOverall = BarrierWaitTimeFinish - BarrierWaitTimeStart;
+
+					UE_LOG(LogDisplayClusterBarrier, Verbose, TEXT("Barrier '%s': sync end, barrier wait time %f"), *Name, BarrierWaitTimeOverall);
+
+					// All callers are here, pet the watchdog
+					WatchdogTimer.ResetTimer();
+
+					// Close input gate, open output gate
+					EventInputGateOpen->Reset();
+					EventOutputGateOpen->Trigger();
+				}
+			}
+		}
+	}
+
+	// Wait for the barrier to open
+	EventOutputGateOpen->Wait();
+
+	// Fixate awaiting finish for this particular thread
+	const double ThreadWaitTimeFinish = FPlatformTime::Seconds();
+
+	{
+		UE_LOG(LogDisplayClusterBarrier, Verbose, TEXT("Barrier '%s': caller '%s' is leaving the barrier"), *Name, *CallerId);
+
+		{
+			FScopeLock LockData(&DataCS);
+
+			// Unregister caller
+			CallersAwaiting.Remove(CallerId);
+
+			// Make sure the barrier was not deactivated while this thread was waiting at the output gate
+			if (bActive)
+			{
+				// In case this thread is leaving last, close output and open input
+				if (--ThreadCount == 0)
+				{
+					EventOutputGateOpen->Reset();
+					EventInputGateOpen->Trigger();
+				}
+			}
+		}
+	}
+
+	// Export barrier overall waiting time
+	if (BarrierWaitTime)
+	{
+		*BarrierWaitTime = BarrierWaitTimeOverall;
+	}
+
+	// Export thread waiting time
+	if (ThreadWaitTime)
+	{
+		*ThreadWaitTime = ThreadWaitTimeFinish - ThreadWaitTimeStart;
+	}
+
+	UE_LOG(LogDisplayClusterBarrier, VeryVerbose, TEXT("Barrier '%s': caller left '%s'"), *Name, *CallerId);
+
+	return EDisplayClusterBarrierWaitResult::Ok;
+}
+
+void FDisplayClusterBarrier::UnregisterSyncCaller(const FString& CallerId)
+{
+	UE_LOG(LogDisplayClusterBarrier, Log, TEXT("Barrier '%s': unregistering caller '%s'..."), *Name, *CallerId);
+
+	FScopeLock Lock(&DataCS);
+
+	// Remove specified callers from the management lists
+	if (CallersAllowed.Remove(CallerId) > 0 || CallersTimedout.Remove(CallerId) > 0)
+	{
+		// Update thread limit
+		ThreadLimit = static_cast<uint32>(CallersAllowed.Num());
+
+		UE_LOG(LogDisplayClusterBarrier, Log, TEXT("Barrier '%s': caller '%s' has been unregistered, new thread limit is %u"), *Name, *CallerId, ThreadLimit);
+
+		// In case it's a last missing caller, we need to open the barrier
+		if (ThreadCount == ThreadLimit)
+		{
+			BarrierWaitTimeFinish  = FPlatformTime::Seconds();
+			BarrierWaitTimeOverall = BarrierWaitTimeFinish - BarrierWaitTimeStart;
+
+			UE_LOG(LogDisplayClusterBarrier, Verbose, TEXT("Barrier '%s': sync end, barrier wait time %f"), *Name, BarrierWaitTimeOverall);
+
+			// All callers here, pet the watchdog
+			WatchdogTimer.ResetTimer();
+
+			// Close the input gate, and open the output gate
+			EventInputGateOpen->Reset();
+			EventOutputGateOpen->Trigger();
+		}
+	}
+	else
+	{
+		UE_LOG(LogDisplayClusterBarrier, Log, TEXT("Barrier '%s': caller '%s' not found"), *Name, *CallerId);
+	}
+}
+
+void FDisplayClusterBarrier::HandleBarrierTimeout()
+{
+	// Being here means some callers have not come to the barrier yet during specific time period. Those
+	// missing callers will be considered as the lost ones. The barrier will continue working with the remaining callers only.
+
+	FScopeLock Lock(&DataCS);
+
+	UE_LOG(LogDisplayClusterBarrier, Log, TEXT("Barrier '%s': Time out! %d callers missing"), *Name, CallersAllowed.Num() - CallersAwaiting.Num());
+
+	// First of all, update the time variables
+	BarrierWaitTimeFinish = FPlatformTime::Seconds();
+	BarrierWaitTimeOverall = BarrierWaitTimeFinish - BarrierWaitTimeStart;
+
+	// List of timed out threads
+	TArray<FString> CallersTimedOutOnLastSync;
+
+	// Update the list of timed out threads. Copy all missing callers to the list of timed out ones.
+	for (const FString& CallerAllowed : CallersAllowed)
+	{
+		if (!CallersAwaiting.Contains(CallerAllowed))
+		{
+			UE_LOG(LogDisplayClusterBarrier, Log, TEXT("Barrier '%s': caller '%s' was moved to the 'TimedOut' list"), *Name, *CallerAllowed);
+			CallersTimedOutOnLastSync.Add(CallerAllowed);
+		}
+	}
+
+	// Update timedout list
+	CallersTimedout.Append(CallersTimedOutOnLastSync);
+	// Update the list of permitted callers
+	CallersAllowed = CallersAwaiting;
+	// Update thread limit
+	ThreadLimit = ThreadCount; // Same as CallersAllowed.Num()
+
+	UE_LOG(LogDisplayClusterBarrier, Log, TEXT("Barrier '%s': new threads limit %d"), *Name, ThreadLimit);
+
+	// Notify listeners
+	OnBarrierTimeout().Broadcast(Name, CallersTimedOutOnLastSync);
+
+	// Close the input gate, and open the output gate to let the remaining callers go
+	EventInputGateOpen->Reset();
+	EventOutputGateOpen->Trigger();
+}
diff --git a/Engine/Plugins/Runtime/nDisplay/Source/DisplayCluster/Private/Network/Barrier/DisplayClusterBarrier.h b/Engine/Plugins/Runtime/nDisplay/Source/DisplayCluster/Private/Network/Barrier/DisplayClusterBarrier.h
new file mode 100644
--- /dev/null
+++ b/Engine/Plugins/Runtime/nDisplay/Source/DisplayCluster/Private/Network/Barrier/DisplayClusterBarrier.h
@@ -0,0 +1,98 @@
+// Copyright Epic Games, Inc. All Rights Reserved.
+
+#pragma once
+
+#include "CoreMinimal.h"
+#include "Network/Barrier/IDisplayClusterBarrier.h"
+#include "Misc/DisplayClusterWatchdogTimer.h"
+
+#include "HAL/Event.h"
+
+
+/**
+ * Thread barrier v2
+ */
+class FDisplayClusterBarrier
+	: public IDisplayClusterBarrier
+{
+public:
+	FDisplayClusterBarrier(const FString& Name, const TArray<FString>& CallersAllowed, const uint32 Timeout);
+	virtual  ~FDisplayClusterBarrier();
+
+public:
+	// Barrier name
+	virtual const FString& GetName() const override
+	{
+		return Name;
+	}
+
+	// Activate barrier
+	virtual bool Activate() override;
+	// Deactivate barrier, no threads will be blocked
+	virtual void Deactivate() override;
+	// Returns true if the barrier has been activated
+	virtual bool IsActivated() const override;
+
+	// Wait until all caller threads arrived
+	virtual EDisplayClusterBarrierWaitResult Wait(const FString& CallerId, double* OutThreadWaitTime = nullptr, double* OutBarrierWaitTime = nullptr) override;
+
+	// Remove specified caller from the sync pipeline
+	virtual void UnregisterSyncCaller(const FString& CallerId) override;
+
+	// Barrier timout notification
+	virtual FDisplayClusterBarrierTimeoutEvent& OnBarrierTimeout() override
+	{
+		return BarrierTimeoutEvent;
+	}
+
+private:
+	// Handler for barrier timouts
+	void HandleBarrierTimeout();
+
+private:
+	// Barrier name
+	const FString Name;
+	// Barrier state
+	bool bActive = false;
+
+
+	// Caller that are allowed to join the barrier
+	TArray<FString> CallersAllowed;
+
+	// Callers that are already waiting at the barrier
+	TArray<FString> CallersAwaiting;
+
+	// Cluster Callers that have been timed out previously
+	TArray<FString> CallersTimedout;
+
+
+	// Timeout for the barrier
+	const uint32 Timeout = 0;
+
+	// Overall amount of threads to wait at the barrier
+	uint32 ThreadLimit = 0;
+
+	// Amount of threads at the barrier
+	uint32 ThreadCount = 0;
+
+
+	// Events to control when threads can join (input) and leave (output) the barrier
+	FEventRef EventInputGateOpen { EEventMode::ManualReset };
+	FEventRef EventOutputGateOpen{ EEventMode::ManualReset };
+
+	// Watchdog timer to detect barrier waiting timeouts
+	FDisplayClusterWatchdogTimer WatchdogTimer;
+
+	// Barrier timeout event
+	FDisplayClusterBarrierTimeoutEvent BarrierTimeoutEvent;
+
+	// Diagnostics data
+	double BarrierWaitTimeStart   = 0;
+	double BarrierWaitTimeFinish  = 0;
+	double BarrierWaitTimeOverall = 0;
+
+	// Barrier state CS
+	mutable FCriticalSection DataCS;
+	// Barrier entrance CS
+	mutable FCriticalSection EntranceCS;
+};
diff --git a/Engine/Plugins/Runtime/nDisplay/Source/DisplayCluster/Private/Network/Barrier/DisplayClusterBarrierFactory.cpp b/Engine/Plugins/Runtime/nDisplay/Source/DisplayCluster/Private/Network/Barrier/DisplayClusterBarrierFactory.cpp
--- a/Engine/Plugins/Runtime/nDisplay/Source/DisplayCluster/Private/Network/Barrier/DisplayClusterBarrierFactory.cpp
+++ b/Engine/Plugins/Runtime/nDisplay/Source/DisplayCluster/Private/Network/Barrier/DisplayClusterBarrierFactory.cpp
@@ -1,23 +1,13 @@
 // Copyright Epic Games, Inc. All Rights Reserved.
 
 #include "Network/Barrier/DisplayClusterBarrierFactory.h"
-#include "Network/Barrier/DisplayClusterBarrierV1.h"
-#include "Network/Barrier/DisplayClusterBarrierV2.h"
+#include "Network/Barrier/DisplayClusterBarrier.h"
 
 #include "Misc/DisplayClusterLog.h"
 
 
-TUniquePtr<IDisplayClusterBarrier> FDisplayClusterBarrierFactory::CreateBarrier(const TArray<FString>& ThreadMarkers, const uint32 Timeout, const FString& Name)
+IDisplayClusterBarrier* FDisplayClusterBarrierFactory::CreateBarrier(const FString& BarrierId, const TArray<FString>& ThreadMarkers, const uint32 Timeout)
 {
-	// The old barriers v1 are not compatible with failover feature. I left them in the codebase
-	// in case it might be useful in the future.
-#if 0
-	// v1 barrier (old)
-	UE_LOG(LogDisplayClusterNetwork, Log, TEXT("Instantiating barrier '%s' of type 'v1': Threads=%d, Timeout=%u ms"), *Name, ThreadsAmount, Timeout);
-	return MakeUnique<FDisplayClusterBarrierV1>(ThreadMarkers.Num(), Timeout, Name);
-#else
-	// v2 barrier (new)
-	UE_LOG(LogDisplayClusterNetwork, Log, TEXT("Instantiating barrier '%s' of type 'v2': Threads=%d, Timeout=%u ms"), *Name, ThreadMarkers.Num(), Timeout);
-	return MakeUnique<FDisplayClusterBarrierV2>(ThreadMarkers, Timeout, Name);
-#endif
+	UE_LOG(LogDisplayClusterBarrier, Log, TEXT("Instantiating unique barrier '%s': Threads=%d, Timeout=%u ms"), *BarrierId, ThreadMarkers.Num(), Timeout);
+	return new FDisplayClusterBarrier(BarrierId, ThreadMarkers, Timeout);
 }
diff --git a/Engine/Plugins/Runtime/nDisplay/Source/DisplayCluster/Private/Network/Barrier/DisplayClusterBarrierFactory.h b/Engine/Plugins/Runtime/nDisplay/Source/DisplayCluster/Private/Network/Barrier/DisplayClusterBarrierFactory.h
--- a/Engine/Plugins/Runtime/nDisplay/Source/DisplayCluster/Private/Network/Barrier/DisplayClusterBarrierFactory.h
+++ b/Engine/Plugins/Runtime/nDisplay/Source/DisplayCluster/Private/Network/Barrier/DisplayClusterBarrierFactory.h
@@ -17,5 +17,5 @@ private:
 
 public:
 	// Factory method to instantiate a barrier based on the user settings
-	static TUniquePtr<IDisplayClusterBarrier> CreateBarrier(const TArray<FString>& ThreadMarkers, const uint32 Timeout, const FString& Name);
+	static IDisplayClusterBarrier* CreateBarrier(const FString& BarrierId, const TArray<FString>& ThreadMarkers, const uint32 Timeout);
 };
diff --git a/Engine/Plugins/Runtime/nDisplay/Source/DisplayCluster/Private/Network/Barrier/DisplayClusterBarrierV1.cpp b/Engine/Plugins/Runtime/nDisplay/Source/DisplayCluster/Private/Network/Barrier/DisplayClusterBarrierV1.cpp
deleted file mode 100644
diff --git a/Engine/Plugins/Runtime/nDisplay/Source/DisplayCluster/Private/Network/Barrier/DisplayClusterBarrierV1.h b/Engine/Plugins/Runtime/nDisplay/Source/DisplayCluster/Private/Network/Barrier/DisplayClusterBarrierV1.h
deleted file mode 100644
diff --git a/Engine/Plugins/Runtime/nDisplay/Source/DisplayCluster/Private/Network/Barrier/DisplayClusterBarrierV2.cpp b/Engine/Plugins/Runtime/nDisplay/Source/DisplayCluster/Private/Network/Barrier/DisplayClusterBarrierV2.cpp
deleted file mode 100644
diff --git a/Engine/Plugins/Runtime/nDisplay/Source/DisplayCluster/Private/Network/Barrier/DisplayClusterBarrierV2.h b/Engine/Plugins/Runtime/nDisplay/Source/DisplayCluster/Private/Network/Barrier/DisplayClusterBarrierV2.h
deleted file mode 100644
diff --git a/Engine/Plugins/Runtime/nDisplay/Source/DisplayCluster/Private/Network/Barrier/IDisplayClusterBarrier.h b/Engine/Plugins/Runtime/nDisplay/Source/DisplayCluster/Private/Network/Barrier/IDisplayClusterBarrier.h
--- a/Engine/Plugins/Runtime/nDisplay/Source/DisplayCluster/Private/Network/Barrier/IDisplayClusterBarrier.h
+++ b/Engine/Plugins/Runtime/nDisplay/Source/DisplayCluster/Private/Network/Barrier/IDisplayClusterBarrier.h
@@ -37,13 +37,16 @@ public:
 	// Deactivate barrier, no threads will be blocked
 	virtual void Deactivate() = 0;
 
+	// Returns true if the barrier has been activated
+	virtual bool IsActivated() const = 0;
+
 	// Wait until all threads arrive
 	virtual EDisplayClusterBarrierWaitResult Wait(const FString& ThreadMarker, double* OutThreadWaitTime = nullptr, double* OutBarrierWaitTime = nullptr) = 0;
 
-	// Remove specified node from sync pipeline
-	virtual void UnregisterSyncNode(const FString& NodeId) = 0;
+	// Remove specified caller from the sync pipeline
+	virtual void UnregisterSyncCaller(const FString& CallerId) = 0;
 
-	// Barrier timout notification (provides BarrierName and NodesTimedOut in parameters)
+	// Barrier timout notification (provides BarrierName and CallersTimedOut in parameters)
 	DECLARE_EVENT_TwoParams(IDisplayClusterBarrier, FDisplayClusterBarrierTimeoutEvent, const FString&, const TArray<FString>&);
 	virtual FDisplayClusterBarrierTimeoutEvent& OnBarrierTimeout() = 0;
 };
diff --git a/Engine/Plugins/Runtime/nDisplay/Source/DisplayCluster/Private/Network/DisplayClusterClient.h b/Engine/Plugins/Runtime/nDisplay/Source/DisplayCluster/Private/Network/DisplayClusterClient.h
--- a/Engine/Plugins/Runtime/nDisplay/Source/DisplayCluster/Private/Network/DisplayClusterClient.h
+++ b/Engine/Plugins/Runtime/nDisplay/Source/DisplayCluster/Private/Network/DisplayClusterClient.h
@@ -32,7 +32,7 @@ public:
 	// Connects to a server
 	virtual bool Connect(const FString& Address, const uint16 Port, const uint32 ConnectRetriesAmount, const uint32 ConnectRetryDelay) override;
 	// Terminates current connection
-	virtual void Disconnect() override final;
+	virtual void Disconnect() override;
 
 	// Provides with net unit name
 	virtual FString GetName() const override
diff --git a/Engine/Plugins/Runtime/nDisplay/Source/DisplayCluster/Private/Network/Protocol/IDisplayClusterProtocolClusterSync.h b/Engine/Plugins/Runtime/nDisplay/Source/DisplayCluster/Private/Network/Protocol/IDisplayClusterProtocolClusterSync.h
--- a/Engine/Plugins/Runtime/nDisplay/Source/DisplayCluster/Private/Network/Protocol/IDisplayClusterProtocolClusterSync.h
+++ b/Engine/Plugins/Runtime/nDisplay/Source/DisplayCluster/Private/Network/Protocol/IDisplayClusterProtocolClusterSync.h
@@ -2,6 +2,8 @@
 
 #pragma once
 
+#include "CoreMinimal.h"
+
 #include "DisplayClusterEnums.h"
 #include "Misc/Optional.h"
 #include "Misc/QualifiedFrameTime.h"
diff --git a/Engine/Plugins/Runtime/nDisplay/Source/DisplayCluster/Private/Network/Protocol/IDisplayClusterProtocolEventsBinary.h b/Engine/Plugins/Runtime/nDisplay/Source/DisplayCluster/Private/Network/Protocol/IDisplayClusterProtocolEventsBinary.h
--- a/Engine/Plugins/Runtime/nDisplay/Source/DisplayCluster/Private/Network/Protocol/IDisplayClusterProtocolEventsBinary.h
+++ b/Engine/Plugins/Runtime/nDisplay/Source/DisplayCluster/Private/Network/Protocol/IDisplayClusterProtocolEventsBinary.h
@@ -2,6 +2,7 @@
 
 #pragma once
 
+#include "CoreMinimal.h"
 #include "Network/DisplayClusterNetworkTypes.h"
 
 struct FDisplayClusterClusterEventBinary;
diff --git a/Engine/Plugins/Runtime/nDisplay/Source/DisplayCluster/Private/Network/Protocol/IDisplayClusterProtocolEventsJson.h b/Engine/Plugins/Runtime/nDisplay/Source/DisplayCluster/Private/Network/Protocol/IDisplayClusterProtocolEventsJson.h
--- a/Engine/Plugins/Runtime/nDisplay/Source/DisplayCluster/Private/Network/Protocol/IDisplayClusterProtocolEventsJson.h
+++ b/Engine/Plugins/Runtime/nDisplay/Source/DisplayCluster/Private/Network/Protocol/IDisplayClusterProtocolEventsJson.h
@@ -2,6 +2,7 @@
 
 #pragma once
 
+#include "CoreMinimal.h"
 #include "Network/DisplayClusterNetworkTypes.h"
 
 struct FDisplayClusterClusterEventJson;
diff --git a/Engine/Plugins/Runtime/nDisplay/Source/DisplayCluster/Private/Network/Protocol/IDisplayClusterProtocolGenericBarrier.h b/Engine/Plugins/Runtime/nDisplay/Source/DisplayCluster/Private/Network/Protocol/IDisplayClusterProtocolGenericBarrier.h
new file mode 100644
--- /dev/null
+++ b/Engine/Plugins/Runtime/nDisplay/Source/DisplayCluster/Private/Network/Protocol/IDisplayClusterProtocolGenericBarrier.h
@@ -0,0 +1,45 @@
+// Copyright Epic Games, Inc. All Rights Reserved.
+
+#pragma once
+
+#include "CoreMinimal.h"
+#include "Network/DisplayClusterNetworkTypes.h"
+
+
+/**
+ * Generic barriers protocol.
+ * Provides generic barrier sync mechanism.
+ */
+class IDisplayClusterProtocolGenericBarrier
+{
+public:
+	virtual ~IDisplayClusterProtocolGenericBarrier() = default;
+
+public:
+	// Just a union of different barrier control operation results to simplify management API
+	enum class EBarrierControlResult : uint8
+	{
+		CreatedSuccessfully = 0,
+		AlreadyExists,
+		NotFound,
+		ReleasedSuccessfully,
+		SynchronizedSuccessfully,
+		UnknownError
+	};
+
+public:
+	// Creates new barrier
+	virtual EDisplayClusterCommResult CreateBarrier(const FString& BarrierId, const TArray<FString>& UniqueThreadMarkers, uint32 Timeout, EBarrierControlResult& Result) = 0;
+
+	// Wait until a barrier with specific ID is created and ready to go
+	virtual EDisplayClusterCommResult WaitUntilBarrierIsCreated(const FString& BarrierId, EBarrierControlResult& Result) = 0;
+
+	// Checks if a specific barrier exists
+	virtual EDisplayClusterCommResult IsBarrierAvailable(const FString& BarrierId, EBarrierControlResult& Result) = 0;
+
+	// Releases specific barrier
+	virtual EDisplayClusterCommResult ReleaseBarrier(const FString& BarrierId, EBarrierControlResult& Result) = 0;
+
+	// Synchronize calling thread on a specific barrier
+	virtual EDisplayClusterCommResult SyncOnBarrier(const FString& BarrierId, const FString& UniqueThreadMarker, EBarrierControlResult& Result) = 0;
+};
diff --git a/Engine/Plugins/Runtime/nDisplay/Source/DisplayCluster/Private/Network/Protocol/IDisplayClusterProtocolRenderSync.h b/Engine/Plugins/Runtime/nDisplay/Source/DisplayCluster/Private/Network/Protocol/IDisplayClusterProtocolRenderSync.h
--- a/Engine/Plugins/Runtime/nDisplay/Source/DisplayCluster/Private/Network/Protocol/IDisplayClusterProtocolRenderSync.h
+++ b/Engine/Plugins/Runtime/nDisplay/Source/DisplayCluster/Private/Network/Protocol/IDisplayClusterProtocolRenderSync.h
@@ -2,6 +2,7 @@
 
 #pragma once
 
+#include "CoreMinimal.h"
 #include "Network/DisplayClusterNetworkTypes.h"
 
 
diff --git a/Engine/Plugins/Runtime/nDisplay/Source/DisplayCluster/Private/Network/Service/ClusterSync/DisplayClusterClusterSyncService.cpp b/Engine/Plugins/Runtime/nDisplay/Source/DisplayCluster/Private/Network/Service/ClusterSync/DisplayClusterClusterSyncService.cpp
--- a/Engine/Plugins/Runtime/nDisplay/Source/DisplayCluster/Private/Network/Service/ClusterSync/DisplayClusterClusterSyncService.cpp
+++ b/Engine/Plugins/Runtime/nDisplay/Source/DisplayCluster/Private/Network/Service/ClusterSync/DisplayClusterClusterSyncService.cpp
@@ -38,9 +38,9 @@ FDisplayClusterClusterSyncService::FDisplayClusterClusterSyncService()
 	const FDisplayClusterConfigurationNetworkSettings& NetworkSettings = GDisplayCluster->GetConfigMgr()->GetConfig()->Cluster->Network;
 
 	// Instantiate service barriers
-	BarrierGameStart  = FDisplayClusterBarrierFactory::CreateBarrier(NodeIds, NetworkSettings.GameStartBarrierTimeout,  TEXT("GameStart_barrier"));
-	BarrierFrameStart = FDisplayClusterBarrierFactory::CreateBarrier(NodeIds, NetworkSettings.FrameStartBarrierTimeout, TEXT("FrameStart_barrier"));
-	BarrierFrameEnd   = FDisplayClusterBarrierFactory::CreateBarrier(NodeIds, NetworkSettings.FrameEndBarrierTimeout,   TEXT("FrameEnd_barrier"));
+	BarrierGameStart.Reset(FDisplayClusterBarrierFactory::CreateBarrier(TEXT("GameStart_barrier"), NodeIds, NetworkSettings.GameStartBarrierTimeout));
+	BarrierFrameStart.Reset(FDisplayClusterBarrierFactory::CreateBarrier(TEXT("FrameStart_barrier"), NodeIds, NetworkSettings.FrameStartBarrierTimeout));
+	BarrierFrameEnd.Reset(FDisplayClusterBarrierFactory::CreateBarrier(TEXT("FrameEnd_barrier"), NodeIds, NetworkSettings.FrameEndBarrierTimeout));
 
 	// Put the barriers into an aux container
 	ServiceBarriers.Emplace(BarrierGameStart->GetName(),  &BarrierGameStart);
@@ -116,7 +116,7 @@ void FDisplayClusterClusterSyncService::KillSession(const FString& NodeId)
 	// Before killing the session on the parent level, we need to unregister this node from the barriers
 	for (TPair<FString, TUniquePtr<IDisplayClusterBarrier>*>& BarrierIt : ServiceBarriers)
 	{
-		(*BarrierIt.Value)->UnregisterSyncNode(NodeId);
+		(*BarrierIt.Value)->UnregisterSyncCaller(NodeId);
 	}
 
 	// Now do the session related job
@@ -141,7 +141,7 @@ void FDisplayClusterClusterSyncService::ProcessSessionClosed(const FDisplayClust
 			// We have to unregister the node that just disconnected from all barriers
 			for (TPair<FString, TUniquePtr<IDisplayClusterBarrier>*>& BarrierIt : ServiceBarriers)
 			{
-				(*BarrierIt.Value)->UnregisterSyncNode(NodeId);
+				(*BarrierIt.Value)->UnregisterSyncCaller(NodeId);
 			}
 
 			// Notify others about node fail
@@ -157,7 +157,7 @@ void FDisplayClusterClusterSyncService::ProcessBarrierTimeout(const FString& Bar
 	{
 		for (const FString& NodeId : NodesTimedOut)
 		{
-			(*BarrierIt.Value)->UnregisterSyncNode(NodeId);
+			(*BarrierIt.Value)->UnregisterSyncCaller(NodeId);
 		}
 	}
 
diff --git a/Engine/Plugins/Runtime/nDisplay/Source/DisplayCluster/Private/Network/Service/ClusterSync/DisplayClusterClusterSyncStrings.h b/Engine/Plugins/Runtime/nDisplay/Source/DisplayCluster/Private/Network/Service/ClusterSync/DisplayClusterClusterSyncStrings.h
--- a/Engine/Plugins/Runtime/nDisplay/Source/DisplayCluster/Private/Network/Service/ClusterSync/DisplayClusterClusterSyncStrings.h
+++ b/Engine/Plugins/Runtime/nDisplay/Source/DisplayCluster/Private/Network/Service/ClusterSync/DisplayClusterClusterSyncStrings.h
@@ -10,8 +10,8 @@ namespace DisplayClusterClusterSyncStrings
 {
 	constexpr static auto ProtocolName = "ClusterSync";
 	
-	constexpr static auto TypeRequest  = "request";
-	constexpr static auto TypeResponse = "response";
+	constexpr static auto TypeRequest  = "Request";
+	constexpr static auto TypeResponse = "Response";
 
 	constexpr static auto ArgumentsDefaultCategory = "CS";
 	constexpr static auto ArgumentsJsonEvents      = "CS_JE";
diff --git a/Engine/Plugins/Runtime/nDisplay/Source/DisplayCluster/Private/Network/Service/GenericBarrier/DisplayClusterGenericBarrierClient.cpp b/Engine/Plugins/Runtime/nDisplay/Source/DisplayCluster/Private/Network/Service/GenericBarrier/DisplayClusterGenericBarrierClient.cpp
new file mode 100644
--- /dev/null
+++ b/Engine/Plugins/Runtime/nDisplay/Source/DisplayCluster/Private/Network/Service/GenericBarrier/DisplayClusterGenericBarrierClient.cpp
@@ -0,0 +1,302 @@
+// Copyright Epic Games, Inc. All Rights Reserved.
+
+#include "Network/Service/GenericBarrier/DisplayClusterGenericBarrierClient.h"
+#include "Network/Service/GenericBarrier/DisplayClusterGenericBarrierStrings.h"
+#include "Network/Packet/DisplayClusterPacketInternal.h"
+#include "Network/Listener/DisplayClusterHelloMessageStrings.h"
+
+#include "Config/IPDisplayClusterConfigManager.h"
+#include "Cluster/IPDisplayClusterClusterManager.h"
+
+#include "DisplayClusterEnums.h"
+#include "Misc/DisplayClusterGlobals.h"
+#include "Misc/DisplayClusterHelpers.h"
+#include "Misc/DisplayClusterLog.h"
+#include "Misc/DisplayClusterStrings.h"
+#include "Misc/DisplayClusterTypesConverter.h"
+
+#include "DisplayClusterConfigurationTypes.h"
+
+
+FDisplayClusterGenericBarrierClient::FDisplayClusterGenericBarrierClient()
+	: FDisplayClusterGenericBarrierClient(FString("CLN_GB"))
+{
+}
+
+FDisplayClusterGenericBarrierClient::FDisplayClusterGenericBarrierClient(const FString& InName)
+	: FDisplayClusterClient(InName)
+{
+}
+
+
+//////////////////////////////////////////////////////////////////////////////////////////////
+// IDisplayClusterClient
+//////////////////////////////////////////////////////////////////////////////////////////////
+bool FDisplayClusterGenericBarrierClient::Connect(const FString& Address, const uint16 Port, const uint32 ConnectRetriesAmount, const uint32 ConnectRetryDelay)
+{
+	// First, allow base class to perform connection
+	if (!FDisplayClusterClient::Connect(Address, Port, ConnectRetriesAmount, ConnectRetryDelay))
+	{
+		return false;
+	}
+
+	// Prepare 'hello' message
+	TSharedPtr<FDisplayClusterPacketInternal> HelloMsg = MakeShared<FDisplayClusterPacketInternal>(
+		DisplayClusterHelloMessageStrings::Hello::Name,
+		DisplayClusterGenericBarrierStrings::TypeRequest,
+		DisplayClusterGenericBarrierStrings::ProtocolName
+	);
+
+	// Fill in the message with data
+	const FString NodeId = GDisplayCluster->GetPrivateClusterMgr()->GetNodeId();
+	HelloMsg->SetTextArg(DisplayClusterHelloMessageStrings::ArgumentsDefaultCategory, DisplayClusterHelloMessageStrings::Hello::ArgNodeId, NodeId);
+
+	// Send message (no response awaiting)
+	return SendPacket(HelloMsg);
+}
+
+
+//////////////////////////////////////////////////////////////////////////////////////////////
+// IDisplayClusterGenericBarriersClient
+//////////////////////////////////////////////////////////////////////////////////////////////
+bool FDisplayClusterGenericBarrierClient::Connect()
+{
+	// Generic barriers are available in runtime only
+	if (!GDisplayCluster || (GDisplayCluster->GetOperationMode() != EDisplayClusterOperationMode::Cluster))
+	{
+		return false;
+	}
+
+	// Now we need to get server address/port data from the config
+	if (const IPDisplayClusterConfigManager* const ConfigMgr = GDisplayCluster->GetPrivateConfigMgr())
+	{
+		if (const UDisplayClusterConfigurationData* const Config = ConfigMgr->GetConfig())
+		{
+			if (const UDisplayClusterConfigurationClusterNode* const PrimaryNode = Config->Cluster->GetNode(Config->Cluster->PrimaryNode.Id))
+			{
+				// Now try to establish connection
+				return Connect(PrimaryNode->Host, Config->Cluster->PrimaryNode.Ports.ClusterSync, 3, 1000);
+			}
+		}
+	}
+
+	return false;
+}
+
+void FDisplayClusterGenericBarrierClient::Disconnect()
+{
+	// Forward IDisplayClusterGenericBarriersClient::Disconnect() to FDisplayClusterClient::Disconnect() explicitly
+	FDisplayClusterClient::Disconnect();
+}
+
+bool FDisplayClusterGenericBarrierClient::IsConnected() const
+{
+	// Forward IDisplayClusterGenericBarriersClient::IsConnected() to FDisplayClusterClient::IsConnected() explicitly
+	return FDisplayClusterClient::IsConnected();
+}
+
+FString FDisplayClusterGenericBarrierClient::GetName() const
+{
+	// Forward IDisplayClusterGenericBarriersClient::GetName() to FDisplayClusterClient::GetName() explicitly
+	return FDisplayClusterClient::GetName();
+}
+
+bool FDisplayClusterGenericBarrierClient::CreateBarrier(const FString& BarrierId, const TArray<FString>& UniqueThreadMarkers, uint32 Timeout)
+{
+	EBarrierControlResult CtrlResult = EBarrierControlResult::UnknownError;
+	const EDisplayClusterCommResult Result = CreateBarrier(BarrierId, UniqueThreadMarkers, Timeout, CtrlResult);
+	return (CtrlResult == EBarrierControlResult::CreatedSuccessfully || CtrlResult == EBarrierControlResult::AlreadyExists) && (Result == EDisplayClusterCommResult::Ok);
+}
+
+bool FDisplayClusterGenericBarrierClient::WaitUntilBarrierIsCreated(const FString& BarrierId)
+{
+	EBarrierControlResult CtrlResult = EBarrierControlResult::UnknownError;
+	const EDisplayClusterCommResult Result = WaitUntilBarrierIsCreated(BarrierId, CtrlResult);
+	return (CtrlResult == EBarrierControlResult::AlreadyExists) && (Result == EDisplayClusterCommResult::Ok);
+}
+
+bool FDisplayClusterGenericBarrierClient::IsBarrierAvailable(const FString& BarrierId)
+{
+	EBarrierControlResult CtrlResult = EBarrierControlResult::UnknownError;
+	const EDisplayClusterCommResult Result = IsBarrierAvailable(BarrierId, CtrlResult);
+	return (CtrlResult == EBarrierControlResult::AlreadyExists) && (Result == EDisplayClusterCommResult::Ok);
+}
+
+bool FDisplayClusterGenericBarrierClient::ReleaseBarrier(const FString& BarrierId)
+{
+	EBarrierControlResult CtrlResult = EBarrierControlResult::UnknownError;
+	const EDisplayClusterCommResult Result = ReleaseBarrier(BarrierId, CtrlResult);
+	return (CtrlResult == EBarrierControlResult::ReleasedSuccessfully) && (Result == EDisplayClusterCommResult::Ok);
+}
+
+bool FDisplayClusterGenericBarrierClient::Synchronize(const FString& BarrierId, const FString& UniqueThreadMarker)
+{
+	EBarrierControlResult CtrlResult = EBarrierControlResult::UnknownError;
+	const EDisplayClusterCommResult Result = SyncOnBarrier(BarrierId, UniqueThreadMarker, CtrlResult);
+	return (CtrlResult == EBarrierControlResult::SynchronizedSuccessfully) && (Result == EDisplayClusterCommResult::Ok);
+}
+
+
+//////////////////////////////////////////////////////////////////////////////////////////////
+// IDisplayClusterProtocolGenericBarrier
+//////////////////////////////////////////////////////////////////////////////////////////////
+EDisplayClusterCommResult FDisplayClusterGenericBarrierClient::CreateBarrier(const FString& BarrierId, const TArray<FString>& UniqueThreadMarkers, uint32 Timeout, EBarrierControlResult& Result)
+{
+	const TSharedPtr<FDisplayClusterPacketInternal> Request = MakeShared<FDisplayClusterPacketInternal>(
+		DisplayClusterGenericBarrierStrings::CreateBarrier::Name,
+		DisplayClusterGenericBarrierStrings::TypeRequest,
+		DisplayClusterGenericBarrierStrings::ProtocolName);
+
+	// Request arguments
+	Request->SetTextArg(DisplayClusterGenericBarrierStrings::ArgumentsDefaultCategory, DisplayClusterGenericBarrierStrings::ArgBarrierId, BarrierId);
+	Request->SetTextArg(DisplayClusterGenericBarrierStrings::ArgumentsDefaultCategory, DisplayClusterGenericBarrierStrings::CreateBarrier::ArgTimeout, Timeout);
+
+	// Stringify markers array and put as a single parameter
+	const FString ThreadMarkersPacked = DisplayClusterHelpers::str::template ArrayToStr(UniqueThreadMarkers, DisplayClusterStrings::common::ArrayValSeparator, false);
+	Request->SetTextArg(DisplayClusterGenericBarrierStrings::ArgumentsDefaultCategory, DisplayClusterGenericBarrierStrings::CreateBarrier::ArgThreadMarkers, ThreadMarkersPacked);
+
+	TSharedPtr<FDisplayClusterPacketInternal> Response;
+
+	{
+		TRACE_CPUPROFILER_EVENT_SCOPE(CLN_GB::CreateBarrier);
+		Response = SendRecvPacket(Request);
+	}
+
+	if (!Response)
+	{
+		UE_LOG(LogDisplayClusterNetwork, Warning, TEXT("Network error on '%s'"), *Request->GetName());
+		return EDisplayClusterCommResult::NetworkError;
+	}
+
+	// Extract response data
+	FString CtrlResult;
+	Response->GetTextArg(DisplayClusterGenericBarrierStrings::ArgumentsDefaultCategory, DisplayClusterGenericBarrierStrings::ArgResult, CtrlResult);
+	Result = static_cast<EBarrierControlResult>(DisplayClusterTypesConverter::template FromString<uint8>(CtrlResult));
+
+	return Response->GetCommResult();
+}
+
+EDisplayClusterCommResult FDisplayClusterGenericBarrierClient::WaitUntilBarrierIsCreated(const FString& BarrierId, EBarrierControlResult& Result)
+{
+	const TSharedPtr<FDisplayClusterPacketInternal> Request = MakeShared<FDisplayClusterPacketInternal>(
+		DisplayClusterGenericBarrierStrings::WaitUntilBarrierIsCreated::Name,
+		DisplayClusterGenericBarrierStrings::TypeRequest,
+		DisplayClusterGenericBarrierStrings::ProtocolName);
+
+	// Request arguments
+	Request->SetTextArg(DisplayClusterGenericBarrierStrings::ArgumentsDefaultCategory, DisplayClusterGenericBarrierStrings::ArgBarrierId, BarrierId);
+
+	TSharedPtr<FDisplayClusterPacketInternal> Response;
+
+	{
+		TRACE_CPUPROFILER_EVENT_SCOPE(CLN_GB::WaitUntilBarrierIsCreated);
+		Response = SendRecvPacket(Request);
+	}
+
+	if (!Response)
+	{
+		UE_LOG(LogDisplayClusterNetwork, Warning, TEXT("Network error on '%s'"), *Request->GetName());
+		return EDisplayClusterCommResult::NetworkError;
+	}
+
+	// Extract response data
+	FString CtrlResult;
+	Response->GetTextArg(DisplayClusterGenericBarrierStrings::ArgumentsDefaultCategory, DisplayClusterGenericBarrierStrings::ArgResult, CtrlResult);
+	Result = static_cast<EBarrierControlResult>(DisplayClusterTypesConverter::template FromString<uint8>(CtrlResult));
+
+	return Response->GetCommResult();
+}
+
+EDisplayClusterCommResult FDisplayClusterGenericBarrierClient::IsBarrierAvailable(const FString& BarrierId, EBarrierControlResult& Result)
+{
+	const TSharedPtr<FDisplayClusterPacketInternal> Request = MakeShared<FDisplayClusterPacketInternal>(
+		DisplayClusterGenericBarrierStrings::IsBarrierAvailable::Name,
+		DisplayClusterGenericBarrierStrings::TypeRequest,
+		DisplayClusterGenericBarrierStrings::ProtocolName);
+
+	// Request arguments
+	Request->SetTextArg(DisplayClusterGenericBarrierStrings::ArgumentsDefaultCategory, DisplayClusterGenericBarrierStrings::ArgBarrierId, BarrierId);
+
+	TSharedPtr<FDisplayClusterPacketInternal> Response;
+
+	{
+		TRACE_CPUPROFILER_EVENT_SCOPE(CLN_GB::IsBarrierAvailable);
+		Response = SendRecvPacket(Request);
+	}
+
+	if (!Response)
+	{
+		UE_LOG(LogDisplayClusterNetwork, Warning, TEXT("Network error on '%s'"), *Request->GetName());
+		return EDisplayClusterCommResult::NetworkError;
+	}
+
+	// Extract response data
+	FString CtrlResult;
+	Response->GetTextArg(DisplayClusterGenericBarrierStrings::ArgumentsDefaultCategory, DisplayClusterGenericBarrierStrings::ArgResult, CtrlResult);
+	Result = static_cast<EBarrierControlResult>(DisplayClusterTypesConverter::template FromString<uint8>(CtrlResult));
+
+	return Response->GetCommResult();
+}
+
+EDisplayClusterCommResult FDisplayClusterGenericBarrierClient::ReleaseBarrier(const FString& BarrierId, EBarrierControlResult& Result)
+{
+	const TSharedPtr<FDisplayClusterPacketInternal> Request = MakeShared<FDisplayClusterPacketInternal>(
+		DisplayClusterGenericBarrierStrings::ReleaseBarrier::Name,
+		DisplayClusterGenericBarrierStrings::TypeRequest,
+		DisplayClusterGenericBarrierStrings::ProtocolName);
+
+	// Request arguments
+	Request->SetTextArg(DisplayClusterGenericBarrierStrings::ArgumentsDefaultCategory, DisplayClusterGenericBarrierStrings::ArgBarrierId, BarrierId);
+
+	TSharedPtr<FDisplayClusterPacketInternal> Response;
+
+	{
+		TRACE_CPUPROFILER_EVENT_SCOPE(CLN_GB::ReleaseBarrier);
+		Response = SendRecvPacket(Request);
+	}
+
+	if (!Response)
+	{
+		UE_LOG(LogDisplayClusterNetwork, Warning, TEXT("Network error on '%s'"), *Request->GetName());
+		return EDisplayClusterCommResult::NetworkError;
+	}
+
+	// Extract response data
+	FString CtrlResult;
+	Response->GetTextArg(DisplayClusterGenericBarrierStrings::ArgumentsDefaultCategory, DisplayClusterGenericBarrierStrings::ArgResult, CtrlResult);
+	Result = static_cast<EBarrierControlResult>(DisplayClusterTypesConverter::template FromString<uint8>(CtrlResult));
+
+	return Response->GetCommResult();
+}
+
+EDisplayClusterCommResult FDisplayClusterGenericBarrierClient::SyncOnBarrier(const FString& BarrierId, const FString& UniqueThreadMarker, EBarrierControlResult& Result)
+{
+	const TSharedPtr<FDisplayClusterPacketInternal> Request = MakeShared<FDisplayClusterPacketInternal>(
+		DisplayClusterGenericBarrierStrings::SyncOnBarrier::Name,
+		DisplayClusterGenericBarrierStrings::TypeRequest,
+		DisplayClusterGenericBarrierStrings::ProtocolName);
+
+	// Request arguments
+	Request->SetTextArg(DisplayClusterGenericBarrierStrings::ArgumentsDefaultCategory, DisplayClusterGenericBarrierStrings::ArgBarrierId, BarrierId);
+	Request->SetTextArg(DisplayClusterGenericBarrierStrings::ArgumentsDefaultCategory, DisplayClusterGenericBarrierStrings::SyncOnBarrier::ArgThreadMarker, UniqueThreadMarker);
+
+	TSharedPtr<FDisplayClusterPacketInternal> Response;
+
+	{
+		TRACE_CPUPROFILER_EVENT_SCOPE(CLN_GB::SyncOnBarrier);
+		Response = SendRecvPacket(Request);
+	}
+
+	if (!Response)
+	{
+		UE_LOG(LogDisplayClusterNetwork, Warning, TEXT("Network error on '%s'"), *Request->GetName());
+		return EDisplayClusterCommResult::NetworkError;
+	}
+
+	// Extract response data
+	FString CtrlResult;
+	Response->GetTextArg(DisplayClusterGenericBarrierStrings::ArgumentsDefaultCategory, DisplayClusterGenericBarrierStrings::ArgResult, CtrlResult);
+	Result = static_cast<EBarrierControlResult>(DisplayClusterTypesConverter::template FromString<uint8>(CtrlResult));
+
+	return Response->GetCommResult();
+}
diff --git a/Engine/Plugins/Runtime/nDisplay/Source/DisplayCluster/Private/Network/Service/GenericBarrier/DisplayClusterGenericBarrierClient.h b/Engine/Plugins/Runtime/nDisplay/Source/DisplayCluster/Private/Network/Service/GenericBarrier/DisplayClusterGenericBarrierClient.h
new file mode 100644
--- /dev/null
+++ b/Engine/Plugins/Runtime/nDisplay/Source/DisplayCluster/Private/Network/Service/GenericBarrier/DisplayClusterGenericBarrierClient.h
@@ -0,0 +1,55 @@
+// Copyright Epic Games, Inc. All Rights Reserved.
+
+#pragma once
+
+#include "CoreMinimal.h"
+#include "Network/DisplayClusterClient.h"
+#include "Network/Packet/DisplayClusterPacketInternal.h"
+#include "Network/Protocol/IDisplayClusterProtocolGenericBarrier.h"
+#include "Cluster/IDisplayClusterGenericBarriersClient.h"
+
+
+/**
+ * Generic barriers TCP client
+ */
+class FDisplayClusterGenericBarrierClient
+	: public FDisplayClusterClient<FDisplayClusterPacketInternal>
+	, public IDisplayClusterProtocolGenericBarrier
+	, public IDisplayClusterGenericBarriersClient
+{
+public:
+	FDisplayClusterGenericBarrierClient();
+	FDisplayClusterGenericBarrierClient(const FString& InName);
+
+public:
+	//////////////////////////////////////////////////////////////////////////////////////////////
+	// IDisplayClusterClient
+	//////////////////////////////////////////////////////////////////////////////////////////////
+	virtual bool Connect(const FString& Address, const uint16 Port, const uint32 ConnectRetriesAmount, const uint32 ConnectRetryDelay) override;
+
+public:
+	//////////////////////////////////////////////////////////////////////////////////////////////
+	// IDisplayClusterGenericBarriersClient
+	//////////////////////////////////////////////////////////////////////////////////////////////
+	virtual bool Connect() override;
+	virtual void Disconnect() override;
+	virtual bool IsConnected() const override;
+
+	virtual FString GetName() const override;
+
+	virtual bool CreateBarrier(const FString& BarrierId, const TArray<FString>& UniqueThreadMarkers, uint32 Timeout) override;
+	virtual bool WaitUntilBarrierIsCreated(const FString& BarrierId) override;
+	virtual bool IsBarrierAvailable(const FString& BarrierId) override;
+	virtual bool ReleaseBarrier(const FString& BarrierId) override;
+	virtual bool Synchronize(const FString& BarrierId, const FString& UniqueThreadMarker) override;
+
+public:
+	//////////////////////////////////////////////////////////////////////////////////////////////
+	// IDisplayClusterProtocolGenericBarrier
+	//////////////////////////////////////////////////////////////////////////////////////////////
+	virtual EDisplayClusterCommResult CreateBarrier(const FString& BarrierId, const TArray<FString>& UniqueThreadMarkers, uint32 Timeout, EBarrierControlResult& Result) override;
+	virtual EDisplayClusterCommResult WaitUntilBarrierIsCreated(const FString& BarrierId, EBarrierControlResult& Result) override;
+	virtual EDisplayClusterCommResult IsBarrierAvailable(const FString& BarrierId, EBarrierControlResult& Result) override;
+	virtual EDisplayClusterCommResult ReleaseBarrier(const FString& BarrierId, EBarrierControlResult& Result) override;
+	virtual EDisplayClusterCommResult SyncOnBarrier(const FString& BarrierId, const FString& UniqueThreadMarker, EBarrierControlResult& Result) override;
+};
diff --git a/Engine/Plugins/Runtime/nDisplay/Source/DisplayCluster/Private/Network/Service/GenericBarrier/DisplayClusterGenericBarrierService.cpp b/Engine/Plugins/Runtime/nDisplay/Source/DisplayCluster/Private/Network/Service/GenericBarrier/DisplayClusterGenericBarrierService.cpp
new file mode 100644
--- /dev/null
+++ b/Engine/Plugins/Runtime/nDisplay/Source/DisplayCluster/Private/Network/Service/GenericBarrier/DisplayClusterGenericBarrierService.cpp
@@ -0,0 +1,375 @@
+// Copyright Epic Games, Inc. All Rights Reserved.
+
+#include "Network/Service/GenericBarrier/DisplayClusterGenericBarrierService.h"
+#include "Network/Service/GenericBarrier/DisplayClusterGenericBarrierStrings.h"
+#include "Network/Barrier/DisplayClusterBarrierFactory.h"
+#include "Network/Barrier/IDisplayClusterBarrier.h"
+#include "Network/Packet/DisplayClusterPacketInternal.h"
+#include "Network/Session/DisplayClusterSession.h"
+
+#include "Misc/DisplayClusterHelpers.h"
+#include "Misc/DisplayClusterLog.h"
+#include "Misc/DisplayClusterStrings.h"
+#include "Misc/DisplayClusterTypesConverter.h"
+
+#include "HAL/Event.h"
+
+
+FDisplayClusterGenericBarrierService::FDisplayClusterGenericBarrierService()
+	: FDisplayClusterService(FString("SRV_GB"))
+{
+}
+
+FDisplayClusterGenericBarrierService::~FDisplayClusterGenericBarrierService()
+{
+	Shutdown();
+}
+
+
+void FDisplayClusterGenericBarrierService::Shutdown()
+{
+	{
+		FScopeLock Lock(&BarriersCS);
+
+		// Deactivate all barriers
+		for (TPair<FString, TSharedPtr<IDisplayClusterBarrier>>& Barrier : Barriers)
+		{
+			Barrier.Value->Deactivate();
+		}
+
+		// Release events
+		for (TPair<FString, FEvent*>& EventIt : BarrierCreationEvents)
+		{
+			EventIt.Value->Trigger();
+			FPlatformProcess::ReturnSynchEventToPool(EventIt.Value);
+		}
+
+		BarrierCreationEvents.Empty();
+	}
+
+	return FDisplayClusterServer::Shutdown();
+}
+
+FString FDisplayClusterGenericBarrierService::GetProtocolName() const
+{
+	static const FString ProtocolName(DisplayClusterGenericBarrierStrings::ProtocolName);
+	return ProtocolName;
+}
+
+TSharedPtr<IDisplayClusterSession> FDisplayClusterGenericBarrierService::CreateSession(FDisplayClusterSessionInfo& SessionInfo)
+{
+	SessionInfo.SessionName = FString::Printf(TEXT("%s_session_%lu_%s"), *GetName(), SessionInfo.SessionId, *SessionInfo.Endpoint.ToString());
+	return MakeShared<FDisplayClusterSession<FDisplayClusterPacketInternal, true>>(SessionInfo, *this, *this, FDisplayClusterService::GetThreadPriority());
+}
+
+
+//////////////////////////////////////////////////////////////////////////////////////////////
+// IDisplayClusterSessionPacketHandler
+//////////////////////////////////////////////////////////////////////////////////////////////
+TSharedPtr<FDisplayClusterPacketInternal> FDisplayClusterGenericBarrierService::ProcessPacket(const TSharedPtr<FDisplayClusterPacketInternal>& Request, const FDisplayClusterSessionInfo& SessionInfo)
+{
+	// Check the pointer
+	if (!Request.IsValid())
+	{
+		UE_LOG(LogDisplayClusterNetwork, Error, TEXT("%s - Invalid request data (nullptr)"), *GetName());
+		return nullptr;
+	}
+
+	// Cache session info
+	SetSessionInfoCache(SessionInfo);
+
+	UE_LOG(LogDisplayClusterNetwork, Verbose, TEXT("%s - Processing packet: %s"), *GetName(), *Request->ToLogString());
+
+	// Check protocol and type
+	if (Request->GetProtocol() != DisplayClusterGenericBarrierStrings::ProtocolName || Request->GetType() != DisplayClusterGenericBarrierStrings::TypeRequest)
+	{
+		UE_LOG(LogDisplayClusterNetwork, Error, TEXT("%s - Unsupported packet type: %s"), *GetName(), *Request->ToLogString());
+		return nullptr;
+	}
+
+	// Initialize response packet
+	TSharedPtr<FDisplayClusterPacketInternal> Response = MakeShared<FDisplayClusterPacketInternal>(Request->GetName(), DisplayClusterGenericBarrierStrings::TypeResponse, Request->GetProtocol());
+
+	// Dispatch the packet
+	if (Request->GetName().Equals(DisplayClusterGenericBarrierStrings::SyncOnBarrier::Name, ESearchCase::IgnoreCase))
+	{
+		// Extract parameter: barrier ID
+		FString BarrierId;
+		Request->GetTextArg(DisplayClusterGenericBarrierStrings::ArgumentsDefaultCategory, DisplayClusterGenericBarrierStrings::ArgBarrierId, BarrierId);
+
+		// Extract parameter: thread marker
+		FString UniqueThreadMarker;
+		Request->GetTextArg(DisplayClusterGenericBarrierStrings::ArgumentsDefaultCategory, DisplayClusterGenericBarrierStrings::SyncOnBarrier::ArgThreadMarker, UniqueThreadMarker);
+
+		// Process command
+		EBarrierControlResult CtrlResult = EBarrierControlResult::UnknownError;
+		const EDisplayClusterCommResult Result = SyncOnBarrier(BarrierId, UniqueThreadMarker, CtrlResult);
+
+		// Set response data
+		Response->SetTextArg(DisplayClusterGenericBarrierStrings::ArgumentsDefaultCategory, DisplayClusterGenericBarrierStrings::ArgResult, DisplayClusterTypesConverter::template ToString((uint8)CtrlResult));
+		Response->SetCommResult(Result);
+
+		return Response;
+	}
+	else if (Request->GetName().Equals(DisplayClusterGenericBarrierStrings::CreateBarrier::Name, ESearchCase::IgnoreCase))
+	{
+		// Extract parameter: barrier ID
+		FString BarrierId;
+		Request->GetTextArg(DisplayClusterGenericBarrierStrings::ArgumentsDefaultCategory, DisplayClusterGenericBarrierStrings::ArgBarrierId, BarrierId);
+
+		// Extract parameter: thread markers
+		FString UniqueThreadMarkersStringified;
+		Request->GetTextArg(DisplayClusterGenericBarrierStrings::ArgumentsDefaultCategory, DisplayClusterGenericBarrierStrings::CreateBarrier::ArgThreadMarkers, UniqueThreadMarkersStringified);
+		TArray<FString> UniqueThreadMarkers;
+		DisplayClusterHelpers::str::template StrToArray(UniqueThreadMarkersStringified, DisplayClusterStrings::common::ArrayValSeparator, UniqueThreadMarkers);
+
+		// Extract parameter: timeout
+		uint32 Timeout = 0;
+		Request->GetTextArg(DisplayClusterGenericBarrierStrings::ArgumentsDefaultCategory, DisplayClusterGenericBarrierStrings::CreateBarrier::ArgTimeout, Timeout);
+
+		// Process command
+		EBarrierControlResult CtrlResult = EBarrierControlResult::UnknownError;
+		const EDisplayClusterCommResult Result = CreateBarrier(BarrierId, UniqueThreadMarkers, Timeout, CtrlResult);
+
+		// Set response data
+		Response->SetTextArg(DisplayClusterGenericBarrierStrings::ArgumentsDefaultCategory, DisplayClusterGenericBarrierStrings::ArgResult, DisplayClusterTypesConverter::template ToString((uint8)CtrlResult));
+		Response->SetCommResult(Result);
+
+		return Response;
+	}
+	else if (Request->GetName().Equals(DisplayClusterGenericBarrierStrings::WaitUntilBarrierIsCreated::Name, ESearchCase::IgnoreCase))
+	{
+		// Extract parameter: barrier ID
+		FString BarrierId;
+		Request->GetTextArg(DisplayClusterGenericBarrierStrings::ArgumentsDefaultCategory, DisplayClusterGenericBarrierStrings::ArgBarrierId, BarrierId);
+
+		// Process command
+		EBarrierControlResult CtrlResult = EBarrierControlResult::UnknownError;
+		const EDisplayClusterCommResult Result = WaitUntilBarrierIsCreated(BarrierId, CtrlResult);
+
+		// Set response data
+		Response->SetTextArg(DisplayClusterGenericBarrierStrings::ArgumentsDefaultCategory, DisplayClusterGenericBarrierStrings::ArgResult, DisplayClusterTypesConverter::template ToString((uint8)CtrlResult));
+		Response->SetCommResult(Result);
+
+		return Response;
+	}
+	else if (Request->GetName().Equals(DisplayClusterGenericBarrierStrings::ReleaseBarrier::Name, ESearchCase::IgnoreCase))
+	{
+		// Extract parameter: barrier ID
+		FString BarrierId;
+		Request->GetTextArg(DisplayClusterGenericBarrierStrings::ArgumentsDefaultCategory, DisplayClusterGenericBarrierStrings::ArgBarrierId, BarrierId);
+
+		// Process command
+		EBarrierControlResult CtrlResult = EBarrierControlResult::UnknownError;
+		const EDisplayClusterCommResult Result = ReleaseBarrier(BarrierId, CtrlResult);
+
+		// Set response data
+		Response->SetTextArg(DisplayClusterGenericBarrierStrings::ArgumentsDefaultCategory, DisplayClusterGenericBarrierStrings::ArgResult, DisplayClusterTypesConverter::template ToString((uint8)CtrlResult));
+		Response->SetCommResult(Result);
+
+		return Response;
+	}
+	else if (Request->GetName().Equals(DisplayClusterGenericBarrierStrings::IsBarrierAvailable::Name, ESearchCase::IgnoreCase))
+	{
+		// Extract parameter: barrier ID
+		FString BarrierId;
+		Request->GetTextArg(DisplayClusterGenericBarrierStrings::ArgumentsDefaultCategory, DisplayClusterGenericBarrierStrings::ArgBarrierId, BarrierId);
+
+		// Process command
+		EBarrierControlResult CtrlResult = EBarrierControlResult::UnknownError;
+		const EDisplayClusterCommResult Result = IsBarrierAvailable(BarrierId, CtrlResult);
+
+		// Set response data
+		Response->SetTextArg(DisplayClusterGenericBarrierStrings::ArgumentsDefaultCategory, DisplayClusterGenericBarrierStrings::ArgResult, DisplayClusterTypesConverter::template ToString((uint8)CtrlResult));
+		Response->SetCommResult(Result);
+
+		return Response;
+	}
+
+	// Being here means that we have no appropriate dispatch logic for this packet
+	UE_LOG(LogDisplayClusterNetwork, Warning, TEXT("%s - No dispatcher found for packet '%s'"), *GetName(), *Request->GetName());
+
+	return nullptr;
+}
+
+
+//////////////////////////////////////////////////////////////////////////////////////////////
+// IDisplayClusterProtocolGenericBarrier
+//////////////////////////////////////////////////////////////////////////////////////////////
+EDisplayClusterCommResult FDisplayClusterGenericBarrierService::CreateBarrier(const FString& BarrierId, const TArray<FString>& UniqueThreadMarkers, uint32 Timeout, EBarrierControlResult& Result)
+{
+	// Validate input data
+	if (BarrierId.IsEmpty() || UniqueThreadMarkers.Num() < 1 || Timeout == 0)
+	{
+		return EDisplayClusterCommResult::WrongRequestData;
+	}
+
+	{
+		FScopeLock Lock(&BarriersCS);
+
+		// Check if exists
+		if (Barriers.Contains(BarrierId))
+		{
+			Result = EBarrierControlResult::AlreadyExists;
+			return EDisplayClusterCommResult::Ok;
+		}
+
+		// Create a new one
+		TSharedPtr<IDisplayClusterBarrier, ESPMode::ThreadSafe> NewBarrier(FDisplayClusterBarrierFactory::CreateBarrier(BarrierId, UniqueThreadMarkers, Timeout));
+		checkSlow(NewBarrier);
+
+		// Activate barrier
+		NewBarrier->Activate();
+
+		// Store it
+		Barriers.Emplace(BarrierId, MoveTemp(NewBarrier));
+
+		// Notify listeners if there are any
+		if (FEvent** BarrierCreatedEvent = BarrierCreationEvents.Find(BarrierId))
+		{
+			(*BarrierCreatedEvent)->Trigger();
+		}
+
+		Result = EBarrierControlResult::CreatedSuccessfully;
+	}
+
+	return EDisplayClusterCommResult::Ok;
+}
+
+EDisplayClusterCommResult FDisplayClusterGenericBarrierService::WaitUntilBarrierIsCreated(const FString& BarrierId, EBarrierControlResult& Result)
+{
+	FEvent** BarrierAvailableEvent = nullptr;
+
+	{
+		FScopeLock Lock(&BarriersCS);
+
+		// Check if the barrier exists already
+		if (Barriers.Contains(BarrierId))
+		{
+			Result = EBarrierControlResult::AlreadyExists;
+			return EDisplayClusterCommResult::Ok;
+		}
+		// If no, set up a notification event
+		else
+		{
+			// Check if notification event has been created previously
+			BarrierAvailableEvent = BarrierCreationEvents.Find(BarrierId);
+			if (!BarrierAvailableEvent)
+			{
+				// It hasn't, so create it
+				BarrierAvailableEvent = &BarrierCreationEvents.Emplace(BarrierId, FPlatformProcess::GetSynchEventFromPool(true));
+			}
+		}
+	}
+
+	// So the barrier has not been created yet, we need to wait
+
+	// Make sure the event is valid
+	if (!BarrierAvailableEvent)
+	{
+		Result = EBarrierControlResult::UnknownError;
+		return EDisplayClusterCommResult::InternalError;
+	}
+
+	// Wait until created
+	(*BarrierAvailableEvent)->Wait();
+
+	Result = EBarrierControlResult::AlreadyExists;
+	return EDisplayClusterCommResult::Ok;
+}
+
+EDisplayClusterCommResult FDisplayClusterGenericBarrierService::IsBarrierAvailable(const FString& BarrierId, EBarrierControlResult& Result)
+{
+	// Validate input data
+	if (BarrierId.IsEmpty())
+	{
+		return EDisplayClusterCommResult::WrongRequestData;
+	}
+
+	{
+		FScopeLock Lock(&BarriersCS);
+		Result = (Barriers.Contains(BarrierId) ? EBarrierControlResult::AlreadyExists : EBarrierControlResult::NotFound);
+	}
+
+	return EDisplayClusterCommResult::Ok;
+}
+
+EDisplayClusterCommResult FDisplayClusterGenericBarrierService::ReleaseBarrier(const FString& BarrierId, EBarrierControlResult& Result)
+{
+	// Validate input data
+	if (BarrierId.IsEmpty())
+	{
+		return EDisplayClusterCommResult::WrongRequestData;
+	}
+
+	{
+		FScopeLock Lock(&BarriersCS);
+
+		if (TSharedPtr<IDisplayClusterBarrier, ESPMode::ThreadSafe>* Barrier = Barriers.Find(BarrierId))
+		{
+			// Deactivate barrier first because it can be used by other clients currently.
+			// In this case the destructor won't be called.
+			(*Barrier)->Deactivate();
+
+			// And forget about it. Once all the clients leave the barrier, the instance will be released.
+			Barriers.Remove(BarrierId);
+
+			// We can safely release the creation event as well
+			if (FEvent** BarrierCreatedEvent = BarrierCreationEvents.Find(BarrierId))
+			{
+				FPlatformProcess::ReturnSynchEventToPool(*BarrierCreatedEvent);
+				BarrierCreationEvents.Remove(BarrierId);
+			}
+		}
+		else
+		{
+			// Not exists
+			Result = EBarrierControlResult::NotFound;
+			return EDisplayClusterCommResult::WrongRequestData;
+		}
+	}
+
+	Result = EBarrierControlResult::ReleasedSuccessfully;
+	return EDisplayClusterCommResult::Ok;
+}
+
+EDisplayClusterCommResult FDisplayClusterGenericBarrierService::SyncOnBarrier(const FString& BarrierId, const FString& UniqueThreadMarker, EBarrierControlResult& Result)
+{
+	TRACE_CPUPROFILER_EVENT_SCOPE(SRV_GB::SyncOnBarrier);
+
+	// Validate input data
+	if (BarrierId.IsEmpty() || UniqueThreadMarker.IsEmpty())
+	{
+		Result = EBarrierControlResult::UnknownError;
+		return EDisplayClusterCommResult::WrongRequestData;
+	}
+
+	TSharedPtr<IDisplayClusterBarrier, ESPMode::ThreadSafe>* Barrier = nullptr;
+
+	{
+		FScopeLock Lock(&BarriersCS);
+
+		// Get barrier ptr to be able to use outside of the critical section
+		// so the other clients can also access it.
+		Barrier = Barriers.Find(BarrierId);
+	}
+
+	// More validation
+	if (!Barrier)
+	{
+		// Barrier not found
+		Result = EBarrierControlResult::NotFound;
+		return EDisplayClusterCommResult::WrongRequestData;
+	}
+	else if(!(*Barrier)->IsActivated())
+	{
+		Result = EBarrierControlResult::UnknownError;
+		return EDisplayClusterCommResult::NotAllowed;
+	}
+
+	// Sync on the barrier
+	(*Barrier)->Wait(UniqueThreadMarker);
+
+	Result = EBarrierControlResult::SynchronizedSuccessfully;
+	return EDisplayClusterCommResult::Ok;
+}
diff --git a/Engine/Plugins/Runtime/nDisplay/Source/DisplayCluster/Private/Network/Service/GenericBarrier/DisplayClusterGenericBarrierService.h b/Engine/Plugins/Runtime/nDisplay/Source/DisplayCluster/Private/Network/Service/GenericBarrier/DisplayClusterGenericBarrierService.h
new file mode 100644
--- /dev/null
+++ b/Engine/Plugins/Runtime/nDisplay/Source/DisplayCluster/Private/Network/Service/GenericBarrier/DisplayClusterGenericBarrierService.h
@@ -0,0 +1,67 @@
+// Copyright Epic Games, Inc. All Rights Reserved.
+
+#pragma once
+
+#include "CoreMinimal.h"
+
+#include "Network/Service/DisplayClusterService.h"
+#include "Network/Session/IDisplayClusterSessionPacketHandler.h"
+#include "Network/Protocol/IDisplayClusterProtocolGenericBarrier.h"
+#include "Network/Packet/DisplayClusterPacketInternal.h"
+
+class FEvent;
+
+
+/**
+ * Generic barriers TCP server
+ */
+class FDisplayClusterGenericBarrierService
+	: public    FDisplayClusterService
+	, public    IDisplayClusterSessionPacketHandler<FDisplayClusterPacketInternal, true>
+	, protected IDisplayClusterProtocolGenericBarrier
+{
+public:
+	FDisplayClusterGenericBarrierService();
+	virtual ~FDisplayClusterGenericBarrierService();
+
+public:
+	//////////////////////////////////////////////////////////////////////////////////////////////
+	// IDisplayClusterServer
+	//////////////////////////////////////////////////////////////////////////////////////////////
+	virtual void Shutdown() override final;
+	virtual FString GetProtocolName() const override;
+
+protected:
+	// Creates session instance for this service
+	virtual TSharedPtr<IDisplayClusterSession> CreateSession(FDisplayClusterSessionInfo& SessionInfo) override;
+
+protected:
+	//////////////////////////////////////////////////////////////////////////////////////////////
+	// IDisplayClusterSessionPacketHandler
+	//////////////////////////////////////////////////////////////////////////////////////////////
+	virtual TSharedPtr<FDisplayClusterPacketInternal> ProcessPacket(const TSharedPtr<FDisplayClusterPacketInternal>& Request, const FDisplayClusterSessionInfo& SessionInfo) override;
+
+private:
+	//////////////////////////////////////////////////////////////////////////////////////////////
+	// IDisplayClusterProtocolGenericBarrier
+	//////////////////////////////////////////////////////////////////////////////////////////////
+	virtual EDisplayClusterCommResult CreateBarrier(const FString& BarrierId, const TArray<FString>& UniqueThreadMarkers, uint32 Timeout, EBarrierControlResult& Result) override;
+	virtual EDisplayClusterCommResult WaitUntilBarrierIsCreated(const FString& BarrierId, EBarrierControlResult& Result) override;
+	virtual EDisplayClusterCommResult IsBarrierAvailable(const FString& BarrierId, EBarrierControlResult& Result) override;
+	virtual EDisplayClusterCommResult ReleaseBarrier(const FString& BarrierId, EBarrierControlResult& Result) override;
+	virtual EDisplayClusterCommResult SyncOnBarrier(const FString& BarrierId, const FString& UniqueThreadMarker, EBarrierControlResult& Result) override;
+
+private:
+	// Unregister cluster node from a barrier
+	void UnregisterClusterNode(const FString& NodeId);
+
+private:
+	// Barriers
+	TMap<FString, TSharedPtr<IDisplayClusterBarrier, ESPMode::ThreadSafe>> Barriers;
+
+	// Barrier creation events
+	TMap<FString, FEvent*> BarrierCreationEvents;
+
+	// Critical section for internal data access synchronization
+	mutable FCriticalSection BarriersCS;
+};
diff --git a/Engine/Plugins/Runtime/nDisplay/Source/DisplayCluster/Private/Network/Service/GenericBarrier/DisplayClusterGenericBarrierStrings.h b/Engine/Plugins/Runtime/nDisplay/Source/DisplayCluster/Private/Network/Service/GenericBarrier/DisplayClusterGenericBarrierStrings.h
new file mode 100644
--- /dev/null
+++ b/Engine/Plugins/Runtime/nDisplay/Source/DisplayCluster/Private/Network/Service/GenericBarrier/DisplayClusterGenericBarrierStrings.h
@@ -0,0 +1,52 @@
+// Copyright Epic Games, Inc. All Rights Reserved.
+
+#pragma once
+
+
+/**
+ * Generic barriers protocol strings
+ */
+namespace DisplayClusterGenericBarrierStrings
+{
+	constexpr static const TCHAR* ProtocolName = TEXT("GenericBarrier");
+
+	constexpr static const TCHAR* TypeRequest  = TEXT("Request");
+	constexpr static const TCHAR* TypeResponse = TEXT("Response");
+
+	constexpr static const TCHAR* ArgumentsDefaultCategory = TEXT("GB");
+
+	// Shared arguments
+	constexpr static const TCHAR* ArgBarrierId = TEXT("BarrierId");
+	constexpr static const TCHAR* ArgResult    = TEXT("CtrlResult");
+
+
+	namespace CreateBarrier
+	{
+		constexpr static const TCHAR* Name = TEXT("CreateBarrier");
+
+		constexpr static const TCHAR* ArgThreadMarkers = TEXT("ThreadMarkers");
+		constexpr static const TCHAR* ArgTimeout       = TEXT("Timeout");
+	}
+
+	namespace WaitUntilBarrierIsCreated
+	{
+		constexpr static const TCHAR* Name = TEXT("WaitUntilBarrierIsCreated");
+	}
+
+	namespace IsBarrierAvailable
+	{
+		constexpr static const TCHAR* Name = TEXT("IsBarrierAvailable");
+	}
+
+	namespace ReleaseBarrier
+	{
+		constexpr static const TCHAR* Name = TEXT("ReleaseBarrier");
+	}
+
+	namespace SyncOnBarrier
+	{
+		constexpr static const TCHAR* Name = TEXT("SyncOnBarrier");
+
+		constexpr static const TCHAR* ArgThreadMarker = TEXT("ThreadMarker");
+	}
+};
diff --git a/Engine/Plugins/Runtime/nDisplay/Source/DisplayCluster/Private/Network/Service/RenderSync/DisplayClusterRenderSyncService.cpp b/Engine/Plugins/Runtime/nDisplay/Source/DisplayCluster/Private/Network/Service/RenderSync/DisplayClusterRenderSyncService.cpp
--- a/Engine/Plugins/Runtime/nDisplay/Source/DisplayCluster/Private/Network/Service/RenderSync/DisplayClusterRenderSyncService.cpp
+++ b/Engine/Plugins/Runtime/nDisplay/Source/DisplayCluster/Private/Network/Service/RenderSync/DisplayClusterRenderSyncService.cpp
@@ -246,11 +246,13 @@ void FDisplayClusterRenderSyncService::InitializeBarriers()
 		const FString BarrierName = SyncGroup.Key + TEXT("_barrier");
 
 		// Instantiate the barrier
-		TUniquePtr<IDisplayClusterBarrier> Barrier = FDisplayClusterBarrierFactory::CreateBarrier(SyncGroup.Value, Config->Cluster->Network.RenderSyncBarrierTimeout, BarrierName);
-		// Subscribe for barrier timeout events
-		Barrier->OnBarrierTimeout().AddRaw(this, &FDisplayClusterRenderSyncService::ProcessBarrierTimeout);
-		// Store barrier instance
-		PolicyToBarrierMap.Emplace(SyncGroup.Key, MoveTemp(Barrier));
+		if (IDisplayClusterBarrier* const Barrier = FDisplayClusterBarrierFactory::CreateBarrier(BarrierName, SyncGroup.Value, Config->Cluster->Network.RenderSyncBarrierTimeout))
+		{
+			// Subscribe for barrier timeout events
+			Barrier->OnBarrierTimeout().AddRaw(this, &FDisplayClusterRenderSyncService::ProcessBarrierTimeout);
+			// Store barrier instance
+			PolicyToBarrierMap.Emplace(SyncGroup.Key, Barrier);
+		}
 	}
 }
 
@@ -295,6 +297,6 @@ void FDisplayClusterRenderSyncService::UnregisterClusterNode(const FString& Node
 {
 	if (IDisplayClusterBarrier* Barrier = GetBarrierForNode(NodeId))
 	{
-		Barrier->UnregisterSyncNode(NodeId);
+		Barrier->UnregisterSyncCaller(NodeId);
 	}
 }
diff --git a/Engine/Plugins/Runtime/nDisplay/Source/DisplayCluster/Private/Network/Service/RenderSync/DisplayClusterRenderSyncStrings.h b/Engine/Plugins/Runtime/nDisplay/Source/DisplayCluster/Private/Network/Service/RenderSync/DisplayClusterRenderSyncStrings.h
--- a/Engine/Plugins/Runtime/nDisplay/Source/DisplayCluster/Private/Network/Service/RenderSync/DisplayClusterRenderSyncStrings.h
+++ b/Engine/Plugins/Runtime/nDisplay/Source/DisplayCluster/Private/Network/Service/RenderSync/DisplayClusterRenderSyncStrings.h
@@ -10,8 +10,8 @@ namespace DisplayClusterRenderSyncStrings
 {
 	constexpr static auto ProtocolName = "RenderSync";
 
-	constexpr static auto TypeRequest  = "request";
-	constexpr static auto TypeResponse = "response";
+	constexpr static auto TypeRequest  = "Request";
+	constexpr static auto TypeResponse = "Response";
 
 	constexpr static auto ArgumentsDefaultCategory = "RS";
 
diff --git a/Engine/Plugins/Runtime/nDisplay/Source/DisplayCluster/Private/Render/DisplayClusterRenderManager.cpp b/Engine/Plugins/Runtime/nDisplay/Source/DisplayCluster/Private/Render/DisplayClusterRenderManager.cpp
--- a/Engine/Plugins/Runtime/nDisplay/Source/DisplayCluster/Private/Render/DisplayClusterRenderManager.cpp
+++ b/Engine/Plugins/Runtime/nDisplay/Source/DisplayCluster/Private/Render/DisplayClusterRenderManager.cpp
@@ -24,6 +24,7 @@
 #include "Render/Device/DisplayClusterRenderDeviceFactoryInternal.h"
 
 #include "Render/Device/IDisplayClusterRenderDeviceFactory.h"
+#include "Render/Monitoring/DisplayClusterVblankMonitor.h"
 #include "Render/PostProcess/IDisplayClusterPostProcess.h"
 #include "Render/Projection/IDisplayClusterProjectionPolicyFactory.h"
 #include "Render/Projection/IDisplayClusterProjectionPolicy.h"
@@ -53,6 +54,17 @@
 #endif
 
 
+static TAutoConsoleVariable<int32> CVarSyncDiagnosticsVBlankMonitoring(
+	TEXT("nDisplay.sync.diag.VBlankMonitoring"),
+	0,
+	TEXT("Sync diagnostics: V-blank monitoring\n")
+	TEXT("0 : disabled\n")
+	TEXT("1 : enabled (if policy supports only)\n")
+	,
+	ECVF_ReadOnly
+);
+
+
 FDisplayClusterRenderManager::FDisplayClusterRenderManager()
 {
 	// Instantiate and register internal render device factory
@@ -68,6 +80,9 @@ FDisplayClusterRenderManager::FDisplayClusterRenderManager()
 	RegisterSynchronizationPolicyFactory(DisplayClusterConfigurationStrings::config::cluster::render_sync::Ethernet,        NewSyncPolicyFactory); // Ethernet
 	RegisterSynchronizationPolicyFactory(DisplayClusterConfigurationStrings::config::cluster::render_sync::EthernetBarrier, NewSyncPolicyFactory); // Ethernet_Simple
 	RegisterSynchronizationPolicyFactory(DisplayClusterConfigurationStrings::config::cluster::render_sync::Nvidia,          NewSyncPolicyFactory); // NVIDIA
+
+	// Instantiate V-blank monitor (it won't auto-start polling)
+	VBlankMonitor = MakeShared<FDisplayClusterVBlankMonitor, ESPMode::ThreadSafe>();
 }
 
 FDisplayClusterRenderManager::~FDisplayClusterRenderManager()
@@ -121,6 +136,12 @@ bool FDisplayClusterRenderManager::StartSession(UDisplayClusterConfigurationData
 		RenderDevicePtr = NewRenderDevice.Get();
 	}
 
+	// Start v-blank monitoring if requested
+	if (!!CVarSyncDiagnosticsVBlankMonitoring.GetValueOnGameThread())
+	{
+		VBlankMonitor->StartMonitoring();
+	}
+
 	// When session is starting in Editor the device won't be initialized so we avoid nullptr access here.
 	//@todo Now we always have a device, even for Editor. Change the condition working on the EditorDevice.
 	return (RenderDevicePtr ? RenderDevicePtr->Initialize() : true);
diff --git a/Engine/Plugins/Runtime/nDisplay/Source/DisplayCluster/Private/Render/DisplayClusterRenderManager.h b/Engine/Plugins/Runtime/nDisplay/Source/DisplayCluster/Private/Render/DisplayClusterRenderManager.h
--- a/Engine/Plugins/Runtime/nDisplay/Source/DisplayCluster/Private/Render/DisplayClusterRenderManager.h
+++ b/Engine/Plugins/Runtime/nDisplay/Source/DisplayCluster/Private/Render/DisplayClusterRenderManager.h
@@ -6,6 +6,7 @@
 
 #include "Render/IPDisplayClusterRenderManager.h"
 
+
 class IDisplayClusterPostProcess;
 class IDisplayClusterPostProcessFactory; 
 class IDisplayClusterProjectionPolicy;
@@ -13,8 +14,10 @@ class IDisplayClusterProjectionPolicyFactory;
 class IDisplayClusterRenderDeviceFactory;
 class IDisplayClusterRenderSyncPolicy;
 class IDisplayClusterRenderSyncPolicyFactory;
+class IDisplayClusterVblankMonitor;
 class UDisplayClusterCameraComponent;
 
+
 /**
  * Render manager. Responsible for everything related to the visuals.
  */
@@ -61,6 +64,11 @@ public:
 	virtual void GetRegisteredPostProcess(TArray<FString>& OutPostProcessIDs) const override;
 	virtual TSharedPtr<IDisplayClusterRender_MeshComponent, ESPMode::ThreadSafe> CreateMeshComponent() const override;
 
+	virtual TSharedPtr<IDisplayClusterVblankMonitor, ESPMode::ThreadSafe> GetVblankMonitor() override
+	{
+		return VBlankMonitor;
+	}
+
 	virtual IDisplayClusterViewportManager* GetViewportManager() const override;
 
 public:
@@ -108,4 +116,8 @@ private:
 
 	// This flag is used to auto-focus the UE window once on start
 	bool bWasWindowFocused = false;
+
+private:
+	// V-blank monitor
+	TSharedPtr<IDisplayClusterVblankMonitor, ESPMode::ThreadSafe> VBlankMonitor;
 };
diff --git a/Engine/Plugins/Runtime/nDisplay/Source/DisplayCluster/Private/Render/Monitoring/DisplayClusterVblankMonitor.cpp b/Engine/Plugins/Runtime/nDisplay/Source/DisplayCluster/Private/Render/Monitoring/DisplayClusterVblankMonitor.cpp
new file mode 100644
--- /dev/null
+++ b/Engine/Plugins/Runtime/nDisplay/Source/DisplayCluster/Private/Render/Monitoring/DisplayClusterVblankMonitor.cpp
@@ -0,0 +1,166 @@
+// Copyright Epic Games, Inc. All Rights Reserved.
+
+#include "Render/Monitoring/DisplayClusterVblankMonitor.h"
+
+#include "Render/Synchronization/DisplayClusterRenderSyncHelper.h"
+#include "Misc/DisplayClusterLog.h"
+#include "Misc/ScopeLock.h"
+
+#include "HAL/RunnableThread.h"
+
+
+FDisplayClusterVBlankMonitor::FDisplayClusterVBlankMonitor()
+{
+}
+
+FDisplayClusterVBlankMonitor::~FDisplayClusterVBlankMonitor()
+{
+	StopMonitoring();
+}
+
+
+uint32 FDisplayClusterVBlankMonitor::Run()
+{
+	UE_LOG(LogDisplayClusterRenderSync, Log, TEXT("V-blank monitoring thread has started"));
+
+	while (!bThreadExitRequested)
+	{
+		// Get sync helper based on the current RHI
+		IDisplayClusterRenderSyncHelper& SyncHelper = FDisplayClusterRenderSyncHelper::Get();
+
+		// Polling...
+		if (SyncHelper.IsWaitForVBlankSupported())
+		{
+			// Wait for V-blank
+			SyncHelper.WaitForVBlank();
+
+			// Get current timestamp
+			const double CurrentTime = FPlatformTime::Seconds();
+
+			TRACE_BOOKMARK(TEXT("VBLANK %llu"), VBlankCounter);
+
+			{
+				FScopeLock Lock(&DataCS);
+
+				++VBlankCounter;
+
+				if (!bTimedataAvailable)
+				{
+					// Here we store last V-blank timestamp only
+					if (FMath::IsNearlyZero(LastVblankTime))
+					{
+						LastVblankTime = CurrentTime;
+					}
+					// Here we store both V-blank timestamp and V-blank period
+					else
+					{
+						VblankPeriod = CurrentTime - LastVblankTime;
+						LastVblankTime = CurrentTime;
+						bTimedataAvailable = true;
+					}
+				}
+				else
+				{
+					// Update V-blank timestamp and period (always average for better precision)
+					VblankPeriod = (VblankPeriod + (CurrentTime - LastVblankTime)) / 2;
+					LastVblankTime = CurrentTime;
+
+					UE_LOG(LogDisplayClusterRenderSync, VeryVerbose, TEXT("V-blank monitoring: vblank=%lf, period=%lf"), LastVblankTime, VblankPeriod);
+				}
+			}
+		}
+		else
+		{
+			TRACE_CPUPROFILER_EVENT_SCOPE(VBLANK - not available);
+			FPlatformProcess::SleepNoStats(1.f);
+		}
+	}
+
+	UE_LOG(LogDisplayClusterRenderSync, Log, TEXT("V-blank monitoring thread has finished"));
+
+	return 0;
+}
+
+bool FDisplayClusterVBlankMonitor::StartMonitoring()
+{
+	FScopeLock Lock(&ControlCS);
+
+	// Check if running already
+	if (bIsRunning)
+	{
+		return false;
+	}
+
+	// Reset internals
+	ResetData();
+
+	// Update controls
+	bIsRunning = true;
+	bThreadExitRequested = false;
+
+	// Start monitoring thread
+	WorkingThread.Reset(FRunnableThread::Create(this, TEXT("VBLANK_Monitoring_thread"), 32 * 1024, TPri_AboveNormal));
+
+	return true;
+}
+
+bool FDisplayClusterVBlankMonitor::StopMonitoring()
+{
+	FScopeLock Lock(&ControlCS);
+
+	// Check if running
+	if (!bIsRunning)
+	{
+		return false;
+	}
+
+	// Update controls
+	bThreadExitRequested = true;
+	bIsRunning = false;
+
+	// Wait until working thread is finished
+	if (WorkingThread)
+	{
+		WorkingThread->WaitForCompletion();
+		WorkingThread.Reset();
+	}
+
+	// Reset internals
+	ResetData();
+
+	return true;
+}
+
+bool FDisplayClusterVBlankMonitor::IsMonitoring()
+{
+	FScopeLock Lock(&ControlCS);
+	return bIsRunning;
+}
+
+bool FDisplayClusterVBlankMonitor::IsVblankTimeDataAvailable()
+{
+	FScopeLock Lock(&DataCS);
+	return bTimedataAvailable;
+}
+
+double FDisplayClusterVBlankMonitor::GetLastVblankTime()
+{
+	FScopeLock Lock(&DataCS);
+	return LastVblankTime;
+}
+
+double FDisplayClusterVBlankMonitor::GetNextVblankTime()
+{
+	FScopeLock Lock(&DataCS);
+	return LastVblankTime + VblankPeriod;
+}
+
+void FDisplayClusterVBlankMonitor::ResetData()
+{
+	FScopeLock Lock(&DataCS);
+
+	VBlankCounter = 0;
+	bTimedataAvailable = false;
+	LastVblankTime = 0;
+	VblankPeriod = 0;
+}
diff --git a/Engine/Plugins/Runtime/nDisplay/Source/DisplayCluster/Private/Render/Monitoring/DisplayClusterVblankMonitor.h b/Engine/Plugins/Runtime/nDisplay/Source/DisplayCluster/Private/Render/Monitoring/DisplayClusterVblankMonitor.h
new file mode 100644
--- /dev/null
+++ b/Engine/Plugins/Runtime/nDisplay/Source/DisplayCluster/Private/Render/Monitoring/DisplayClusterVblankMonitor.h
@@ -0,0 +1,68 @@
+// Copyright Epic Games, Inc. All Rights Reserved.
+
+#pragma once
+
+#include "CoreMinimal.h"
+#include "Render/Monitoring/IDisplayClusterVblankMonitor.h"
+#include "HAL/Runnable.h"
+
+
+/**
+ * V-blank monitor. It continuously polls the current sync policy
+ * for V-blanks and reports the time data to utraces.
+ */
+class FDisplayClusterVBlankMonitor
+	: public IDisplayClusterVblankMonitor
+	, protected FRunnable
+{
+public:
+	FDisplayClusterVBlankMonitor();
+	virtual ~FDisplayClusterVBlankMonitor();
+
+public:
+	//~ Begin IDisplayClusterVblankMonitor interface
+	virtual bool StartMonitoring() override;
+	virtual bool StopMonitoring() override;
+	virtual bool IsMonitoring() override;
+	virtual bool IsVblankTimeDataAvailable() override;
+	virtual double GetLastVblankTime() override;
+	virtual double GetNextVblankTime() override;
+	//~ End IDisplayClusterVblankMonitor interface
+
+protected:
+	//~ Begin FRunnable interface
+	virtual uint32 Run() override;
+	//~ End FRunnable interface
+
+private:
+	/** Resets internal data */
+	void ResetData();
+
+private:
+	/** Holds the thread object */
+	TUniquePtr<FRunnableThread> WorkingThread;
+
+	/** Keeps monitoring status (true if running) */
+	bool bIsRunning = false;
+
+	/** V - blank counter */
+	uint64 VBlankCounter = 0;
+
+	/** Flag to ask working thread to terminate */
+	bool bThreadExitRequested = false;
+
+	/* It's set to true when time data is available (last vblank, period) */
+	bool bTimedataAvailable = false;
+
+	/* A timestamp of the last V-blank */
+	double LastVblankTime = 0;
+
+	/* V-blank period (approximate) */
+	double VblankPeriod = 0;
+
+	/** CritSec for control */
+	mutable FCriticalSection ControlCS;
+	
+	/** CritSec for data */
+	mutable FCriticalSection DataCS;
+};
diff --git a/Engine/Plugins/Runtime/nDisplay/Source/DisplayCluster/Private/Render/Synchronization/DisplayClusterRenderSyncHelper.h b/Engine/Plugins/Runtime/nDisplay/Source/DisplayCluster/Private/Render/Synchronization/DisplayClusterRenderSyncHelper.h
--- a/Engine/Plugins/Runtime/nDisplay/Source/DisplayCluster/Private/Render/Synchronization/DisplayClusterRenderSyncHelper.h
+++ b/Engine/Plugins/Runtime/nDisplay/Source/DisplayCluster/Private/Render/Synchronization/DisplayClusterRenderSyncHelper.h
@@ -27,7 +27,7 @@ public:
 
 
 /**
- * Syncrhonization helper implementation. The main purpose of this class
+ * Synchronization helper implementation. The main purpose of this class
  * is to solve the combinatoric complexity Platform/RHI/SyncPolicy without
  * making the whole sync policies hierarchy overcomplicated.
  */
diff --git a/Engine/Plugins/Runtime/nDisplay/Source/DisplayCluster/Private/Render/Synchronization/DisplayClusterRenderSyncPolicyBase.cpp b/Engine/Plugins/Runtime/nDisplay/Source/DisplayCluster/Private/Render/Synchronization/DisplayClusterRenderSyncPolicyBase.cpp
--- a/Engine/Plugins/Runtime/nDisplay/Source/DisplayCluster/Private/Render/Synchronization/DisplayClusterRenderSyncPolicyBase.cpp
+++ b/Engine/Plugins/Runtime/nDisplay/Source/DisplayCluster/Private/Render/Synchronization/DisplayClusterRenderSyncPolicyBase.cpp
@@ -17,26 +17,8 @@
 #include "RHIResources.h"
 
 
-static TAutoConsoleVariable<int32> CVarSyncDiagnosticsVBlankMonitoring(
-	TEXT("nDisplay.sync.diag.VBlankMonitoring"),
-	0,
-	TEXT("Sync diagnostics: V-blank monitoring\n")
-	TEXT("0 : disabled\n")
-	TEXT("1 : enabled (if policy supports only)\n")
-	,
-	ECVF_ReadOnly
-);
-
-
 bool FDisplayClusterRenderSyncPolicyBase::Initialize()
 {
-	// In case v-blank monitoring was requested and this feature is supported by
-	// the sync policy implementation, start v-blank monitoring.
-	if (!!CVarSyncDiagnosticsVBlankMonitoring.GetValueOnGameThread() && IsWaitForVBlankFeatureSupported())
-	{
-		VBlankMonitor = MakeUnique<FDisplayClusterVBlankMonitor>(*this);
-	}
-
 	return true;
 }
 
diff --git a/Engine/Plugins/Runtime/nDisplay/Source/DisplayCluster/Private/Render/Synchronization/DisplayClusterRenderSyncPolicyBase.h b/Engine/Plugins/Runtime/nDisplay/Source/DisplayCluster/Private/Render/Synchronization/DisplayClusterRenderSyncPolicyBase.h
--- a/Engine/Plugins/Runtime/nDisplay/Source/DisplayCluster/Private/Render/Synchronization/DisplayClusterRenderSyncPolicyBase.h
+++ b/Engine/Plugins/Runtime/nDisplay/Source/DisplayCluster/Private/Render/Synchronization/DisplayClusterRenderSyncPolicyBase.h
@@ -3,7 +3,6 @@
 #pragma once
 
 #include "Render/Synchronization/IDisplayClusterRenderSyncPolicy.h"
-#include "Render/Synchronization/DisplayClusterVblankMonitor.h"
 
 
 /**
@@ -28,7 +27,7 @@ public:
 	virtual bool WaitForVBlank() override;
 
 public:
-	// Performs Ethernet barrier syncrhonization with other nodes in the cluster
+	// Performs Ethernet barrier synchronization with other nodes in the cluster
 	void SyncOnBarrier();
 
 	// Returns policy configuration parameters
@@ -44,7 +43,4 @@ protected:
 private:
 	// Sync policy configuration parameters
 	TMap<FString, FString> Parameters;
-
-	// V-blank monitoring thread. Performs V-blank polling on the top level of the sync policies hierarchy.
-	TUniquePtr<FDisplayClusterVBlankMonitor> VBlankMonitor;
 };
diff --git a/Engine/Plugins/Runtime/nDisplay/Source/DisplayCluster/Private/Render/Synchronization/DisplayClusterVblankMonitor.cpp b/Engine/Plugins/Runtime/nDisplay/Source/DisplayCluster/Private/Render/Synchronization/DisplayClusterVblankMonitor.cpp
deleted file mode 100644
diff --git a/Engine/Plugins/Runtime/nDisplay/Source/DisplayCluster/Private/Render/Synchronization/DisplayClusterVblankMonitor.h b/Engine/Plugins/Runtime/nDisplay/Source/DisplayCluster/Private/Render/Synchronization/DisplayClusterVblankMonitor.h
deleted file mode 100644
diff --git a/Engine/Plugins/Runtime/nDisplay/Source/DisplayCluster/Public/Cluster/IDisplayClusterClusterManager.h b/Engine/Plugins/Runtime/nDisplay/Source/DisplayCluster/Public/Cluster/IDisplayClusterClusterManager.h
--- a/Engine/Plugins/Runtime/nDisplay/Source/DisplayCluster/Public/Cluster/IDisplayClusterClusterManager.h
+++ b/Engine/Plugins/Runtime/nDisplay/Source/DisplayCluster/Public/Cluster/IDisplayClusterClusterManager.h
@@ -7,12 +7,13 @@
 #include "UObject/ScriptInterface.h"
 #include "DisplayClusterEnums.h"
 
-struct FDisplayClusterClusterEventJson;
 struct FDisplayClusterClusterEventBinary;
-class IDisplayClusterClusterEventJsonListener;
+struct FDisplayClusterClusterEventJson;
 class IDisplayClusterClusterEventBinaryListener;
-class IDisplayClusterClusterSyncObject;
+class IDisplayClusterClusterEventJsonListener;
 class IDisplayClusterClusterEventListener;
+class IDisplayClusterClusterSyncObject;
+class IDisplayClusterGenericBarriersClient;
 
 DECLARE_MULTICAST_DELEGATE_OneParam(FOnClusterEventJson, const FDisplayClusterClusterEventJson& /* Event */);
 typedef FOnClusterEventJson::FDelegate FOnClusterEventJsonListener;
@@ -54,6 +55,9 @@ public:
 	/** Unregisters synchronization object. */
 	virtual void UnregisterSyncObject(IDisplayClusterClusterSyncObject* SyncObj) = 0;
 
+	/** Spawns a client for generic barriers service. */
+	virtual TUniquePtr<IDisplayClusterGenericBarriersClient> CreateGenericBarriersClient(const FString& ClientName) = 0;
+
 	//////////////////////////////////////////////////////////////////////////////////////////////
 	// Cluster events
 	//////////////////////////////////////////////////////////////////////////////////////////////
diff --git a/Engine/Plugins/Runtime/nDisplay/Source/DisplayCluster/Public/Cluster/IDisplayClusterGenericBarriersClient.h b/Engine/Plugins/Runtime/nDisplay/Source/DisplayCluster/Public/Cluster/IDisplayClusterGenericBarriersClient.h
new file mode 100644
--- /dev/null
+++ b/Engine/Plugins/Runtime/nDisplay/Source/DisplayCluster/Public/Cluster/IDisplayClusterGenericBarriersClient.h
@@ -0,0 +1,46 @@
+// Copyright Epic Games, Inc. All Rights Reserved.
+
+#pragma once
+
+#include "CoreMinimal.h"
+
+
+/**
+ * Generic barriers client interface
+ */
+class IDisplayClusterGenericBarriersClient
+{
+public:
+	virtual ~IDisplayClusterGenericBarriersClient() = default;
+
+// Networking API
+public:
+	/** Connects to a server */
+	virtual bool Connect() = 0;
+
+	/** Terminates current connection */
+	virtual void Disconnect() = 0;
+
+	/** Returns connection status */
+	virtual bool IsConnected() const = 0;
+
+	/** Returns client name */
+	virtual FString GetName() const = 0;
+
+// Barrier API
+public:
+	/** Creates new barrier */
+	virtual bool CreateBarrier(const FString& BarrierId, const TArray<FString>& UniqueThreadMarkers, uint32 Timeout) = 0;
+
+	/** Wait until a barrier with specific ID is created and ready to go */
+	virtual bool WaitUntilBarrierIsCreated(const FString& BarrierId) = 0;
+
+	/** Checks if a specific barrier exists */
+	virtual bool IsBarrierAvailable(const FString& BarrierId) = 0;
+
+	/** Releases specific barrier */
+	virtual bool ReleaseBarrier(const FString& BarrierId) = 0;
+
+	/** Synchronize calling thread on a specific barrier */
+	virtual bool Synchronize(const FString& BarrierId, const FString& UniqueThreadMarker) = 0;
+};
diff --git a/Engine/Plugins/Runtime/nDisplay/Source/DisplayCluster/Public/Render/IDisplayClusterRenderManager.h b/Engine/Plugins/Runtime/nDisplay/Source/DisplayCluster/Public/Render/IDisplayClusterRenderManager.h
--- a/Engine/Plugins/Runtime/nDisplay/Source/DisplayCluster/Public/Render/IDisplayClusterRenderManager.h
+++ b/Engine/Plugins/Runtime/nDisplay/Source/DisplayCluster/Public/Render/IDisplayClusterRenderManager.h
@@ -13,6 +13,7 @@ class IDisplayClusterProjectionPolicyFactory;
 class IDisplayClusterRenderDeviceFactory;
 class IDisplayClusterRenderSyncPolicy;
 class IDisplayClusterRenderSyncPolicyFactory;
+class IDisplayClusterVblankMonitor;
 class IDisplayClusterViewportManager; 
 class UCineCameraComponent;
 struct FPostProcessSettings;
@@ -162,6 +163,11 @@ public:
 	*/
 	virtual void GetRegisteredPostProcess(TArray<FString>& OutPostProcessIDs) const = 0;
 
+	/**
+	* Returns V-blank monitor interface
+	*/
+	virtual TSharedPtr<IDisplayClusterVblankMonitor, ESPMode::ThreadSafe> GetVblankMonitor() = 0;
+
 	/**
 	* @return - Current viewport manager from root actor
 	*/
diff --git a/Engine/Plugins/Runtime/nDisplay/Source/DisplayCluster/Public/Render/Monitoring/IDisplayClusterVblankMonitor.h b/Engine/Plugins/Runtime/nDisplay/Source/DisplayCluster/Public/Render/Monitoring/IDisplayClusterVblankMonitor.h
new file mode 100644
--- /dev/null
+++ b/Engine/Plugins/Runtime/nDisplay/Source/DisplayCluster/Public/Render/Monitoring/IDisplayClusterVblankMonitor.h
@@ -0,0 +1,34 @@
+// Copyright Epic Games, Inc. All Rights Reserved.
+
+#pragma once
+
+#include "CoreMinimal.h"
+
+
+/**
+ * V-blank monitor interface
+ */
+class IDisplayClusterVblankMonitor
+{
+public:
+	virtual ~IDisplayClusterVblankMonitor() = default;
+
+public:
+	/** Start monitoring */
+	virtual bool StartMonitoring() = 0;
+
+	/** Stop monitoring */
+	virtual bool StopMonitoring() = 0;
+
+	/** Returns true if currently monitoring */
+	virtual bool IsMonitoring() = 0;
+
+	/** Returns true if last and next expected V-blank timestamps are known */
+	virtual bool IsVblankTimeDataAvailable() = 0;
+
+	/** Returns last V-blank timestamp */
+	virtual double GetLastVblankTime() = 0;
+
+	/** Returns next expected V-blank timestamp */
+	virtual double GetNextVblankTime() = 0;
+};
diff --git a/Engine/Plugins/Runtime/nDisplay/Source/DisplayClusterConfiguration/Private/DisplayClusterConfigurationTypes_Media.cpp b/Engine/Plugins/Runtime/nDisplay/Source/DisplayClusterConfiguration/Private/DisplayClusterConfigurationTypes_Media.cpp
--- a/Engine/Plugins/Runtime/nDisplay/Source/DisplayClusterConfiguration/Private/DisplayClusterConfigurationTypes_Media.cpp
+++ b/Engine/Plugins/Runtime/nDisplay/Source/DisplayClusterConfiguration/Private/DisplayClusterConfigurationTypes_Media.cpp
@@ -5,24 +5,24 @@
 #include "DisplayClusterConfigurationTypes_Media.h"
 
 
-bool FDisplayClusterConfigurationMediaICVFX::IsMediaInputAssigned(const FString& InNodeId) const
+bool FDisplayClusterConfigurationMediaICVFX::IsMediaInputAssigned(const FString& NodeId) const
 {
-	return GetMediaSource(InNodeId) != nullptr;
+	return GetMediaSource(NodeId) != nullptr;
 }
 
-bool FDisplayClusterConfigurationMediaICVFX::IsMediaOutputAssigned(const FString& InNodeId) const
+bool FDisplayClusterConfigurationMediaICVFX::IsMediaOutputAssigned(const FString& NodeId) const
 {
-	return GetMediaOutput(InNodeId) != nullptr;
+	return GetMediaOutput(NodeId) != nullptr;
 }
 
-UMediaSource* FDisplayClusterConfigurationMediaICVFX::GetMediaSource(const FString& InNodeId) const
+UMediaSource* FDisplayClusterConfigurationMediaICVFX::GetMediaSource(const FString& NodeId) const
 {
 	// Look up for a group that contains node ID specified
 	for (const FDisplayClusterConfigurationMediaInputGroup& MediaInputGroup : MediaInputGroups)
 	{
-		const bool bFound = MediaInputGroup.ClusterNodes.ItemNames.ContainsByPredicate([InNodeId](const FString& Item)
+		const bool bFound = MediaInputGroup.ClusterNodes.ItemNames.ContainsByPredicate([NodeId](const FString& Item)
 			{
-				return Item.Equals(InNodeId, ESearchCase::IgnoreCase);
+				return Item.Equals(NodeId, ESearchCase::IgnoreCase);
 			});
 
 		if (bFound)
@@ -34,14 +34,14 @@ UMediaSource* FDisplayClusterConfigurationMediaICVFX::GetMediaSource(const FStri
 	return nullptr;
 }
 
-UMediaOutput* FDisplayClusterConfigurationMediaICVFX::GetMediaOutput(const FString& InNodeId) const
+UMediaOutput* FDisplayClusterConfigurationMediaICVFX::GetMediaOutput(const FString& NodeId) const
 {
 	// Look up for a group that contains node ID specified
 	for (const FDisplayClusterConfigurationMediaOutputGroup& MediaOutputGroup : MediaOutputGroups)
 	{
-		const bool bFound = MediaOutputGroup.ClusterNodes.ItemNames.ContainsByPredicate([InNodeId](const FString& Item)
+		const bool bFound = MediaOutputGroup.ClusterNodes.ItemNames.ContainsByPredicate([NodeId](const FString& Item)
 			{
-				return Item.Equals(InNodeId, ESearchCase::IgnoreCase);
+				return Item.Equals(NodeId, ESearchCase::IgnoreCase);
 			});
 
 		if (bFound)
@@ -52,3 +52,22 @@ UMediaOutput* FDisplayClusterConfigurationMediaICVFX::GetMediaOutput(const FStri
 
 	return nullptr;
 }
+
+UDisplayClusterMediaOutputSynchronizationPolicy* FDisplayClusterConfigurationMediaICVFX::GetOutputSyncPolicy(const FString& NodeId) const
+{
+	// Look up for a group that contains node ID specified
+	for (const FDisplayClusterConfigurationMediaOutputGroup& MediaOutputGroup : MediaOutputGroups)
+	{
+		const bool bFound = MediaOutputGroup.ClusterNodes.ItemNames.ContainsByPredicate([NodeId](const FString& Item)
+			{
+				return Item.Equals(NodeId, ESearchCase::IgnoreCase);
+			});
+
+		if (bFound)
+		{
+			return MediaOutputGroup.OutputSyncPolicy;
+		}
+	}
+
+	return nullptr;
+}
diff --git a/Engine/Plugins/Runtime/nDisplay/Source/DisplayClusterConfiguration/Public/DisplayClusterConfigurationTypes.h b/Engine/Plugins/Runtime/nDisplay/Source/DisplayClusterConfiguration/Public/DisplayClusterConfigurationTypes.h
--- a/Engine/Plugins/Runtime/nDisplay/Source/DisplayClusterConfiguration/Public/DisplayClusterConfigurationTypes.h
+++ b/Engine/Plugins/Runtime/nDisplay/Source/DisplayClusterConfiguration/Public/DisplayClusterConfigurationTypes.h
@@ -171,7 +171,7 @@ public:
 	FDisplayClusterConfigurationPrimaryNodePorts();
 
 public:
-	/** Advanced: network port for Cluster Sync Events */
+	/** Advanced: TCP port for nDisplay internal services */
 	UPROPERTY(EditAnywhere, Category = NDisplay, meta = (ClampMin = "1024", ClampMax = "65535", UIMin = "1024", UIMax = "65535"))
 	uint16 ClusterSync;
 
@@ -359,7 +359,7 @@ public:
 	TMap<FString, FDisplayClusterConfigurationPostprocess> Postprocess;
 
 	// Media settings
-	UPROPERTY(EditAnywhere, BlueprintReadWrite, Category = "Media", meta = (ShowOnlyInnerProperties))
+	UPROPERTY(EditAnywhere, BlueprintReadWrite, Category = "Media")
 	FDisplayClusterConfigurationMedia Media;
 
 #if WITH_EDITORONLY_DATA
diff --git a/Engine/Plugins/Runtime/nDisplay/Source/DisplayClusterConfiguration/Public/DisplayClusterConfigurationTypes_Media.h b/Engine/Plugins/Runtime/nDisplay/Source/DisplayClusterConfiguration/Public/DisplayClusterConfigurationTypes_Media.h
--- a/Engine/Plugins/Runtime/nDisplay/Source/DisplayClusterConfiguration/Public/DisplayClusterConfigurationTypes_Media.h
+++ b/Engine/Plugins/Runtime/nDisplay/Source/DisplayClusterConfiguration/Public/DisplayClusterConfigurationTypes_Media.h
@@ -13,6 +13,34 @@
 
 #include "DisplayClusterConfigurationTypes_Media.generated.h"
 
+
+/*
+ * Base media output synchronization policy class
+ */
+UCLASS(Abstract, editinlinenew, BlueprintType, hidecategories = (Object))
+class DISPLAYCLUSTERCONFIGURATION_API UDisplayClusterMediaOutputSynchronizationPolicy
+	: public UObject
+{
+	GENERATED_BODY()
+
+public:
+	/** Returns true if specified media capture type can be synchonized by the policy implementation */
+	virtual bool IsCaptureTypeSupported(UMediaCapture* MediaCapture) const
+	{
+		return true;
+	}
+
+	/** Starts synchronization of specific output stream (capture device). Returns false if failed. */
+	virtual bool StartSynchronization(UMediaCapture* MediaCapture, const FString& MediaId) PURE_VIRTUAL(UDisplayClusterMediaOutputSynchronizationPolicy::StartSynchronization, return false;);
+
+	/** Stops synchronization of specific output stream (capture device). */
+	virtual void StopSynchronization() PURE_VIRTUAL(UDisplayClusterMediaOutputSynchronizationPolicy::StopSynchronization, );
+
+	/** Returns true if currently synchronising a media output. */
+	virtual bool IsRunning() PURE_VIRTUAL(UDisplayClusterMediaOutputSynchronizationPolicy::IsRunning, return false; );
+};
+
+
 PRAGMA_DISABLE_DEPRECATION_WARNINGS
 
 /*
@@ -33,15 +61,18 @@ public:
 	TObjectPtr<UMediaSource> MediaSource;
 
 	/** Media output to use */
-	UPROPERTY(Instanced, EditAnywhere, BlueprintReadWrite, Category = "Media")
+	UPROPERTY(Instanced, EditAnywhere, BlueprintReadWrite, Category = "Output")
 	TObjectPtr<UMediaOutput> MediaOutput;
 
-#if WITH_EDITORONLY_DATA
+	/** Media output synchronization policy */
+	UPROPERTY(Instanced, EditAnywhere, BlueprintReadWrite, Category = "Output", meta = (DisplayName = "Capture Synchronization"))
+	TObjectPtr<UDisplayClusterMediaOutputSynchronizationPolicy> OutputSyncPolicy;
 
+protected:
+#if WITH_EDITORONLY_DATA
 	UE_DEPRECATED(5.2, "This property has been deprecated")
 	UPROPERTY()
 	FString MediaSharingNode_DEPRECATED;
-
 #endif // WITH_EDITORONLY_DATA
 };
 
@@ -83,6 +114,10 @@ public:
 	/** Media output to use */
 	UPROPERTY(Instanced, EditAnywhere, BlueprintReadWrite, Category = "Media")
 	TObjectPtr<UMediaOutput> MediaOutput;
+
+	/** Media output synchronization policy */
+	UPROPERTY(Instanced, EditAnywhere, BlueprintReadWrite, Category = "Media", meta = (DisplayName = "Capture Synchronization"))
+	TObjectPtr<UDisplayClusterMediaOutputSynchronizationPolicy> OutputSyncPolicy;
 };
 
 
@@ -109,16 +144,19 @@ public:
 
 public:
 	/** Returns true if a specific cluster node has media source assigned */
-	bool IsMediaInputAssigned(const FString& InNodeId) const;
+	bool IsMediaInputAssigned(const FString& NodeId) const;
 
 	/** Returns true if a specific cluster node has media source assigned */
-	bool IsMediaOutputAssigned(const FString& InNodeId) const;
+	bool IsMediaOutputAssigned(const FString& NodeId) const;
 
 	/** Returns media source bound to a specific cluster node */
-	UMediaSource* GetMediaSource(const FString& InNodeId) const;
+	UMediaSource* GetMediaSource(const FString& NodeId) const;
 
 	/** Returns media output bound to a specific cluster node */
-	UMediaOutput* GetMediaOutput(const FString& InNodeId) const;
+	UMediaOutput* GetMediaOutput(const FString& NodeId) const;
+
+	/** Returns media output sync policy bound to a specific cluster node */
+	UDisplayClusterMediaOutputSynchronizationPolicy* GetOutputSyncPolicy(const FString& NodeId) const;
 
 	UE_DEPRECATED(5.2, "This function has beend deprecated.")
 	bool IsMediaSharingNode(const FString& InNodeId) const
diff --git a/Engine/Plugins/Runtime/nDisplay/Source/DisplayClusterMedia/DisplayClusterMedia.Build.cs b/Engine/Plugins/Runtime/nDisplay/Source/DisplayClusterMedia/DisplayClusterMedia.Build.cs
--- a/Engine/Plugins/Runtime/nDisplay/Source/DisplayClusterMedia/DisplayClusterMedia.Build.cs
+++ b/Engine/Plugins/Runtime/nDisplay/Source/DisplayClusterMedia/DisplayClusterMedia.Build.cs
@@ -13,6 +13,7 @@ public class DisplayClusterMedia : ModuleRules
 
 		PublicDependencyModuleNames.AddRange(
 			new string[] {
+				"DisplayClusterConfiguration",
 				"Media",
 				"MediaAssets",
 				"MediaIOCore",
@@ -23,11 +24,12 @@ public class DisplayClusterMedia : ModuleRules
 				"Core",
 				"CoreUObject",
 				"DisplayCluster",
-				"DisplayClusterConfiguration",
 				"Engine",
 				"Renderer",
 				"RenderCore",
 				"RHI",
+				"RivermaxCore",
+				"RivermaxMedia",
 			});
 
 		if (Target.bBuildEditor == true)
diff --git a/Engine/Plugins/Runtime/nDisplay/Source/DisplayClusterMedia/Private/Capture/DisplayClusterMediaCaptureBase.cpp b/Engine/Plugins/Runtime/nDisplay/Source/DisplayClusterMedia/Private/Capture/DisplayClusterMediaCaptureBase.cpp
--- a/Engine/Plugins/Runtime/nDisplay/Source/DisplayClusterMedia/Private/Capture/DisplayClusterMediaCaptureBase.cpp
+++ b/Engine/Plugins/Runtime/nDisplay/Source/DisplayClusterMedia/Private/Capture/DisplayClusterMediaCaptureBase.cpp
@@ -6,15 +6,18 @@
 #include "IDisplayCluster.h"
 #include "IDisplayClusterCallbacks.h"
 
+#include "DisplayClusterConfigurationTypes_Media.h"
+
 #include "MediaCapture.h"
 #include "MediaOutput.h"
 
 #include "RenderGraphBuilder.h"
 
 
-FDisplayClusterMediaCaptureBase::FDisplayClusterMediaCaptureBase(const FString& InMediaId, const FString& InClusterNodeId, UMediaOutput* InMediaOutput)
+FDisplayClusterMediaCaptureBase::FDisplayClusterMediaCaptureBase(const FString& InMediaId, const FString& InClusterNodeId, UMediaOutput* InMediaOutput, UDisplayClusterMediaOutputSynchronizationPolicy* InSyncPolicy)
 	: FDisplayClusterMediaBase(InMediaId, InClusterNodeId)
 	, MediaOutput(InMediaOutput)
+	, SyncPolicy(InSyncPolicy)
 {
 	check(InMediaOutput);
 
@@ -29,12 +32,20 @@ FDisplayClusterMediaCaptureBase::~FDisplayClusterMediaCaptureBase()
 
 void FDisplayClusterMediaCaptureBase::AddReferencedObjects(FReferenceCollector& Collector)
 {
-	Collector.AddReferencedObject(MediaOutput);
+	if (MediaOutput)
+	{
+		Collector.AddReferencedObject(MediaOutput);
+	}
 
 	if (MediaCapture)
 	{
 		Collector.AddReferencedObject(MediaCapture);
 	}
+
+	if (SyncPolicy)
+	{
+		Collector.AddReferencedObject(SyncPolicy);
+	}
 }
 
 bool FDisplayClusterMediaCaptureBase::StartCapture()
@@ -46,6 +57,15 @@ bool FDisplayClusterMediaCaptureBase::StartCapture()
 		{
 			MediaCapture->SetMediaOutput(MediaOutput);
 
+			// Initialize and start capture synchronization
+			if (SyncPolicy && SyncPolicy->IsCaptureTypeSupported(MediaCapture))
+			{
+				if (!SyncPolicy->StartSynchronization(MediaCapture, GetMediaId()))
+				{
+					UE_LOG(LogDisplayClusterMedia, Warning, TEXT("MediaCapture '%s': couldn't start synchronization."), *GetMediaId());
+				}
+			}
+
 			bWasCaptureStarted = StartMediaCapture();
 			return bWasCaptureStarted;
 		}
@@ -56,6 +76,13 @@ bool FDisplayClusterMediaCaptureBase::StartCapture()
 
 void FDisplayClusterMediaCaptureBase::StopCapture()
 {
+	// Stop synchronization
+	if (SyncPolicy)
+	{
+		SyncPolicy->StopSynchronization();
+	}
+
+	// Stop capture
 	if (MediaCapture)
 	{
 		MediaCapture->StopCapture(false);
@@ -70,6 +97,8 @@ void FDisplayClusterMediaCaptureBase::ExportMediaData(FRDGBuilder& GraphBuilder,
 
 	if (SrcTexture)
 	{
+		UE_LOG(LogDisplayClusterMedia, Verbose, TEXT("MediaCapture '%s': exporting texture on RT frame '%lu'..."), *GetMediaId(), GFrameCounterRenderThread);
+
 		MediaCapture->SetValidSourceGPUMask(GraphBuilder.RHICmdList.GetGPUMask());
 
 		LastSrcRegionSize = FIntSize(TextureInfo.Region.Size());
diff --git a/Engine/Plugins/Runtime/nDisplay/Source/DisplayClusterMedia/Private/Capture/DisplayClusterMediaCaptureBase.h b/Engine/Plugins/Runtime/nDisplay/Source/DisplayClusterMedia/Private/Capture/DisplayClusterMediaCaptureBase.h
--- a/Engine/Plugins/Runtime/nDisplay/Source/DisplayClusterMedia/Private/Capture/DisplayClusterMediaCaptureBase.h
+++ b/Engine/Plugins/Runtime/nDisplay/Source/DisplayClusterMedia/Private/Capture/DisplayClusterMediaCaptureBase.h
@@ -12,6 +12,7 @@
 class FRDGBuilder;
 class UMediaCapture;
 class UMediaOutput;
+class UDisplayClusterMediaOutputSynchronizationPolicy;
 
 
 /**
@@ -22,7 +23,7 @@ class FDisplayClusterMediaCaptureBase
 	, public FGCObject
 {
 public:
-	FDisplayClusterMediaCaptureBase(const FString& MediaId, const FString& ClusterNodeId, UMediaOutput* MediaOutput);
+	FDisplayClusterMediaCaptureBase(const FString& MediaId, const FString& ClusterNodeId, UMediaOutput* MediaOutput, UDisplayClusterMediaOutputSynchronizationPolicy* SyncPolicy = nullptr);
 	virtual ~FDisplayClusterMediaCaptureBase();
 
 public:
@@ -38,6 +39,11 @@ public:
 	virtual bool StartCapture();
 	virtual void StopCapture();
 
+	UMediaCapture* GetMediaCapture() const
+	{
+		return MediaCapture;
+	}
+
 protected:
 	void ExportMediaData(FRDGBuilder& GraphBuilder, const FMediaTextureInfo& TextureInfo);
 	void OnPostClusterTick();
@@ -64,8 +70,9 @@ private:
 
 private:
 	//~ Begin GC by AddReferencedObjects
-	UMediaOutput*           MediaOutput  = nullptr;
-	UMediaCapture*          MediaCapture = nullptr;
+	UMediaOutput*  MediaOutput  = nullptr;
+	UMediaCapture* MediaCapture = nullptr;
+	UDisplayClusterMediaOutputSynchronizationPolicy* SyncPolicy = nullptr;
 	//~ End GC by AddReferencedObjects
 
 	// Used to restart media capture in the case it falls in error
diff --git a/Engine/Plugins/Runtime/nDisplay/Source/DisplayClusterMedia/Private/Capture/DisplayClusterMediaCaptureCamera.cpp b/Engine/Plugins/Runtime/nDisplay/Source/DisplayClusterMedia/Private/Capture/DisplayClusterMediaCaptureCamera.cpp
--- a/Engine/Plugins/Runtime/nDisplay/Source/DisplayClusterMedia/Private/Capture/DisplayClusterMediaCaptureCamera.cpp
+++ b/Engine/Plugins/Runtime/nDisplay/Source/DisplayClusterMedia/Private/Capture/DisplayClusterMediaCaptureCamera.cpp
@@ -15,8 +15,8 @@
 #include "RHIResources.h"
 
 
-FDisplayClusterMediaCaptureCamera::FDisplayClusterMediaCaptureCamera(const FString& InMediaId, const FString& InClusterNodeId, const FString& InCameraId, const FString& InViewportId, UMediaOutput* InMediaOutput)
-	: FDisplayClusterMediaCaptureViewport(InMediaId, InClusterNodeId, InViewportId, InMediaOutput)
+FDisplayClusterMediaCaptureCamera::FDisplayClusterMediaCaptureCamera(const FString& InMediaId, const FString& InClusterNodeId, const FString& InCameraId, const FString& InViewportId, UMediaOutput* InMediaOutput, UDisplayClusterMediaOutputSynchronizationPolicy* SyncPolicy)
+	: FDisplayClusterMediaCaptureViewport(InMediaId, InClusterNodeId, InViewportId, InMediaOutput, SyncPolicy)
 	, CameraId(InCameraId)
 {
 	if (const ADisplayClusterRootActor* const RootActor = IDisplayCluster::Get().GetGameMgr()->GetRootActor())
diff --git a/Engine/Plugins/Runtime/nDisplay/Source/DisplayClusterMedia/Private/Capture/DisplayClusterMediaCaptureCamera.h b/Engine/Plugins/Runtime/nDisplay/Source/DisplayClusterMedia/Private/Capture/DisplayClusterMediaCaptureCamera.h
--- a/Engine/Plugins/Runtime/nDisplay/Source/DisplayClusterMedia/Private/Capture/DisplayClusterMediaCaptureCamera.h
+++ b/Engine/Plugins/Runtime/nDisplay/Source/DisplayClusterMedia/Private/Capture/DisplayClusterMediaCaptureCamera.h
@@ -16,7 +16,7 @@ class FDisplayClusterMediaCaptureCamera
 	: public FDisplayClusterMediaCaptureViewport
 {
 public:
-	FDisplayClusterMediaCaptureCamera(const FString& MediaId, const FString& ClusterNodeId, const FString& CameraId, const FString& ViewportId, UMediaOutput* MediaOutput);
+	FDisplayClusterMediaCaptureCamera(const FString& MediaId, const FString& ClusterNodeId, const FString& CameraId, const FString& ViewportId, UMediaOutput* MediaOutput, UDisplayClusterMediaOutputSynchronizationPolicy* SyncPolicy = nullptr);
 
 protected:
 	
diff --git a/Engine/Plugins/Runtime/nDisplay/Source/DisplayClusterMedia/Private/Capture/DisplayClusterMediaCaptureNode.cpp b/Engine/Plugins/Runtime/nDisplay/Source/DisplayClusterMedia/Private/Capture/DisplayClusterMediaCaptureNode.cpp
--- a/Engine/Plugins/Runtime/nDisplay/Source/DisplayClusterMedia/Private/Capture/DisplayClusterMediaCaptureNode.cpp
+++ b/Engine/Plugins/Runtime/nDisplay/Source/DisplayClusterMedia/Private/Capture/DisplayClusterMediaCaptureNode.cpp
@@ -15,8 +15,8 @@
 #include "UnrealClient.h"
 
 
-FDisplayClusterMediaCaptureNode::FDisplayClusterMediaCaptureNode(const FString& InMediaId, const FString& InClusterNodeId, UMediaOutput* InMediaOutput)
-	: FDisplayClusterMediaCaptureBase(InMediaId, InClusterNodeId, InMediaOutput)
+FDisplayClusterMediaCaptureNode::FDisplayClusterMediaCaptureNode(const FString& InMediaId, const FString& InClusterNodeId, UMediaOutput* InMediaOutput, UDisplayClusterMediaOutputSynchronizationPolicy* SyncPolicy)
+	: FDisplayClusterMediaCaptureBase(InMediaId, InClusterNodeId, InMediaOutput, SyncPolicy)
 {
 }
 
diff --git a/Engine/Plugins/Runtime/nDisplay/Source/DisplayClusterMedia/Private/Capture/DisplayClusterMediaCaptureNode.h b/Engine/Plugins/Runtime/nDisplay/Source/DisplayClusterMedia/Private/Capture/DisplayClusterMediaCaptureNode.h
--- a/Engine/Plugins/Runtime/nDisplay/Source/DisplayClusterMedia/Private/Capture/DisplayClusterMediaCaptureNode.h
+++ b/Engine/Plugins/Runtime/nDisplay/Source/DisplayClusterMedia/Private/Capture/DisplayClusterMediaCaptureNode.h
@@ -18,7 +18,7 @@ class FDisplayClusterMediaCaptureNode
 	: public FDisplayClusterMediaCaptureBase
 {
 public:
-	FDisplayClusterMediaCaptureNode(const FString& MediaId, const FString& ClusterNodeId, UMediaOutput* MediaOutput);
+	FDisplayClusterMediaCaptureNode(const FString& MediaId, const FString& ClusterNodeId, UMediaOutput* MediaOutput, UDisplayClusterMediaOutputSynchronizationPolicy* SyncPolicy = nullptr);
 
 public:
 	virtual bool StartCapture() override;
diff --git a/Engine/Plugins/Runtime/nDisplay/Source/DisplayClusterMedia/Private/Capture/DisplayClusterMediaCaptureViewport.cpp b/Engine/Plugins/Runtime/nDisplay/Source/DisplayClusterMedia/Private/Capture/DisplayClusterMediaCaptureViewport.cpp
--- a/Engine/Plugins/Runtime/nDisplay/Source/DisplayClusterMedia/Private/Capture/DisplayClusterMediaCaptureViewport.cpp
+++ b/Engine/Plugins/Runtime/nDisplay/Source/DisplayClusterMedia/Private/Capture/DisplayClusterMediaCaptureViewport.cpp
@@ -12,8 +12,8 @@
 #include "RHIResources.h"
 
 
-FDisplayClusterMediaCaptureViewport::FDisplayClusterMediaCaptureViewport(const FString& InMediaId, const FString& InClusterNodeId, const FString& InViewportId, UMediaOutput* InMediaOutput)
-	: FDisplayClusterMediaCaptureBase(InMediaId, InClusterNodeId, InMediaOutput)
+FDisplayClusterMediaCaptureViewport::FDisplayClusterMediaCaptureViewport(const FString& InMediaId, const FString& InClusterNodeId, const FString& InViewportId, UMediaOutput* InMediaOutput, UDisplayClusterMediaOutputSynchronizationPolicy* SyncPolicy)
+	: FDisplayClusterMediaCaptureBase(InMediaId, InClusterNodeId, InMediaOutput, SyncPolicy)
 	, ViewportId(InViewportId)
 {
 }
diff --git a/Engine/Plugins/Runtime/nDisplay/Source/DisplayClusterMedia/Private/Capture/DisplayClusterMediaCaptureViewport.h b/Engine/Plugins/Runtime/nDisplay/Source/DisplayClusterMedia/Private/Capture/DisplayClusterMediaCaptureViewport.h
--- a/Engine/Plugins/Runtime/nDisplay/Source/DisplayClusterMedia/Private/Capture/DisplayClusterMediaCaptureViewport.h
+++ b/Engine/Plugins/Runtime/nDisplay/Source/DisplayClusterMedia/Private/Capture/DisplayClusterMediaCaptureViewport.h
@@ -17,7 +17,7 @@ class FDisplayClusterMediaCaptureViewport
 	: public FDisplayClusterMediaCaptureBase
 {
 public:
-	FDisplayClusterMediaCaptureViewport(const FString& MediaId, const FString& ClusterNodeId, const FString& ViewportId, UMediaOutput* MediaOutput);
+	FDisplayClusterMediaCaptureViewport(const FString& MediaId, const FString& ClusterNodeId, const FString& ViewportId, UMediaOutput* MediaOutput, UDisplayClusterMediaOutputSynchronizationPolicy* SyncPolicy = nullptr);
 
 public:
 	virtual bool StartCapture() override;
diff --git a/Engine/Plugins/Runtime/nDisplay/Source/DisplayClusterMedia/Private/DisplayClusterMediaHelpers.cpp b/Engine/Plugins/Runtime/nDisplay/Source/DisplayClusterMedia/Private/DisplayClusterMediaHelpers.cpp
--- a/Engine/Plugins/Runtime/nDisplay/Source/DisplayClusterMedia/Private/DisplayClusterMediaHelpers.cpp
+++ b/Engine/Plugins/Runtime/nDisplay/Source/DisplayClusterMedia/Private/DisplayClusterMediaHelpers.cpp
@@ -10,6 +10,51 @@
 
 namespace DisplayClusterMediaHelpers
 {
+	namespace MediaId
+	{
+		// Generate media ID for a specific entity
+		FString GenerateMediaId(EMediaDeviceType DeviceType, EMediaOwnerType OwnerType, const FString& NodeId, const FString& DCRAName, const FString& OwnerName)
+		{
+			if (DeviceType == EMediaDeviceType::Input)
+			{
+				switch (OwnerType)
+				{
+				case EMediaOwnerType::Backbuffer:
+					return FString::Printf(TEXT("%s_%s_backbuffer_input"), *NodeId, *DCRAName);
+
+				case EMediaOwnerType::Viewport:
+					return FString::Printf(TEXT("%s_%s_%s_viewport_input"), *NodeId, *DCRAName, *OwnerName);
+
+				case EMediaOwnerType::ICVFXCamera:
+					return FString::Printf(TEXT("%s_%s_%s_icvfx_input"), *NodeId, *DCRAName, *OwnerName);
+
+				default:
+					unimplemented();
+				}
+			}
+			else if (DeviceType == EMediaDeviceType::Output)
+			{
+				switch (OwnerType)
+				{
+				case EMediaOwnerType::Backbuffer:
+					return FString::Printf(TEXT("%s_%s_backbuffer_capture"), *NodeId, *DCRAName);
+
+				case EMediaOwnerType::Viewport:
+					return FString::Printf(TEXT("%s_%s_%s_viewport_capture"), *NodeId, *DCRAName, *OwnerName);
+
+				case EMediaOwnerType::ICVFXCamera:
+					return FString::Printf(TEXT("%s_%s_%s_icvfx_capture"), *NodeId, *DCRAName, *OwnerName);
+
+				default:
+					unimplemented();
+				}
+			}
+
+			// Should never get here
+			return FString();
+		}
+	}
+
 	//@todo This needs to be exposed from the DisplayCluster core module after its refactoring
 	FString GenerateICVFXViewportName(const FString& ClusterNodeId, const FString& ICVFXCameraName)
 	{
diff --git a/Engine/Plugins/Runtime/nDisplay/Source/DisplayClusterMedia/Private/DisplayClusterMediaHelpers.h b/Engine/Plugins/Runtime/nDisplay/Source/DisplayClusterMedia/Private/DisplayClusterMediaHelpers.h
--- a/Engine/Plugins/Runtime/nDisplay/Source/DisplayClusterMedia/Private/DisplayClusterMediaHelpers.h
+++ b/Engine/Plugins/Runtime/nDisplay/Source/DisplayClusterMedia/Private/DisplayClusterMediaHelpers.h
@@ -10,6 +10,27 @@ class FRHITexture;
 
 namespace DisplayClusterMediaHelpers
 {
+	namespace MediaId
+	{
+		// Helper for media ID generation
+		enum class EMediaDeviceType : uint8
+		{
+			Input,
+			Output
+		};
+
+		// Helper for media ID generation
+		enum class EMediaOwnerType : uint8
+		{
+			Backbuffer,
+			Viewport,
+			ICVFXCamera
+		};
+
+		// Generate media ID for a specific entity
+		FString GenerateMediaId(EMediaDeviceType DeviceType, EMediaOwnerType OwnerType, const FString& NodeId, const FString& DCRAName, const FString& OwnerName);
+	}
+
 	// Generates internal ICVFX viewport IDs
 	FString GenerateICVFXViewportName(const FString& ClusterNodeId, const FString& ICVFXCameraName);
 
diff --git a/Engine/Plugins/Runtime/nDisplay/Source/DisplayClusterMedia/Private/DisplayClusterMediaLog.cpp b/Engine/Plugins/Runtime/nDisplay/Source/DisplayClusterMedia/Private/DisplayClusterMediaLog.cpp
--- a/Engine/Plugins/Runtime/nDisplay/Source/DisplayClusterMedia/Private/DisplayClusterMediaLog.cpp
+++ b/Engine/Plugins/Runtime/nDisplay/Source/DisplayClusterMedia/Private/DisplayClusterMediaLog.cpp
@@ -5,3 +5,4 @@
 
 // Module-wide log categories
 DEFINE_LOG_CATEGORY(LogDisplayClusterMedia);
+DEFINE_LOG_CATEGORY(LogDisplayClusterMediaSync);
diff --git a/Engine/Plugins/Runtime/nDisplay/Source/DisplayClusterMedia/Private/DisplayClusterMediaLog.h b/Engine/Plugins/Runtime/nDisplay/Source/DisplayClusterMedia/Private/DisplayClusterMediaLog.h
--- a/Engine/Plugins/Runtime/nDisplay/Source/DisplayClusterMedia/Private/DisplayClusterMediaLog.h
+++ b/Engine/Plugins/Runtime/nDisplay/Source/DisplayClusterMedia/Private/DisplayClusterMediaLog.h
@@ -6,4 +6,5 @@
 
 
 // Module-wide log categories
-DECLARE_LOG_CATEGORY_EXTERN(LogDisplayClusterMedia, Log, All);
+DECLARE_LOG_CATEGORY_EXTERN(LogDisplayClusterMedia,     Log, All);
+DECLARE_LOG_CATEGORY_EXTERN(LogDisplayClusterMediaSync, Log, All);
diff --git a/Engine/Plugins/Runtime/nDisplay/Source/DisplayClusterMedia/Private/DisplayClusterMediaModule.cpp b/Engine/Plugins/Runtime/nDisplay/Source/DisplayClusterMedia/Private/DisplayClusterMediaModule.cpp
--- a/Engine/Plugins/Runtime/nDisplay/Source/DisplayClusterMedia/Private/DisplayClusterMediaModule.cpp
+++ b/Engine/Plugins/Runtime/nDisplay/Source/DisplayClusterMedia/Private/DisplayClusterMediaModule.cpp
@@ -109,29 +109,32 @@ void FDisplayClusterMediaModule::InitializeMedia()
 					// Media input
 					if (MediaSettings.MediaSource)
 					{
-						const FString MediaInputId = RootActor->GetName() + FString("_") + ClusterNodeId + FString("_backbuffer_input");
+						const FString MediaInputId = DisplayClusterMediaHelpers::MediaId::GenerateMediaId(
+							DisplayClusterMediaHelpers::MediaId::EMediaDeviceType::Input,
+							DisplayClusterMediaHelpers::MediaId::EMediaOwnerType::Backbuffer,
+							*ClusterNodeId, *RootActor->GetName(), FString());
 
-						UE_LOG(LogDisplayClusterMedia, Log, TEXT("Initializing backbuffer media input for node '%s'..."), *ClusterNodeId);
+						UE_LOG(LogDisplayClusterMedia, Log, TEXT("Initializing backbuffer media input '%s'..."), *MediaInputId);
 
-						InputNode = MakeUnique<FDisplayClusterMediaInputNode>(
+						InputNode = MakeShared<FDisplayClusterMediaInputNode>(
 							MediaInputId, ClusterNodeId,
 							MediaSettings.MediaSource);
-
-						AllInputs.Add(InputNode.Get());
 					}
 
 					// Media capture
 					if (MediaSettings.MediaOutput)
 					{
-						const FString MediaCaptureId = RootActor->GetName() + FString("_") + ClusterNodeId + FString("_backbuffer_capture");
+						const FString MediaCaptureId = DisplayClusterMediaHelpers::MediaId::GenerateMediaId(
+							DisplayClusterMediaHelpers::MediaId::EMediaDeviceType::Output,
+							DisplayClusterMediaHelpers::MediaId::EMediaOwnerType::Backbuffer,
+							*ClusterNodeId, *RootActor->GetName(), FString());
 
-						UE_LOG(LogDisplayClusterMedia, Log, TEXT("Initializing backbuffer media capture for node '%s'..."), *ClusterNodeId);
+						UE_LOG(LogDisplayClusterMedia, Log, TEXT("Initializing backbuffer media capture '%s'..."), *MediaCaptureId);
 
-						CaptureNode = MakeUnique<FDisplayClusterMediaCaptureNode>(
+						CaptureNode = MakeShared<FDisplayClusterMediaCaptureNode>(
 							MediaCaptureId, ClusterNodeId,
-							MediaSettings.MediaOutput);
-
-						AllCaptures.Add(CaptureNode.Get());
+							MediaSettings.MediaOutput,
+							MediaSettings.OutputSyncPolicy);
 					}
 				}
 			}
@@ -149,32 +152,37 @@ void FDisplayClusterMediaModule::InitializeMedia()
 						// Media input
 						if (MediaSettings.MediaSource)
 						{
-							const FString MediaInputId = RootActor->GetName() + FString("_") + ClusterNodeId + FString("_") + ViewportIt.Key + FString("_viewport_input");
+							const FString MediaInputId = DisplayClusterMediaHelpers::MediaId::GenerateMediaId(
+								DisplayClusterMediaHelpers::MediaId::EMediaDeviceType::Input,
+								DisplayClusterMediaHelpers::MediaId::EMediaOwnerType::Viewport,
+								*ClusterNodeId, *RootActor->GetName(), ViewportIt.Key);
 
-							UE_LOG(LogDisplayClusterMedia, Log, TEXT("Initializing viewport media input for viewport '%s'..."), *ViewportIt.Key);
+							UE_LOG(LogDisplayClusterMedia, Log, TEXT("Initializing viewport media input '%s' for viewport '%s'..."), *MediaInputId, *ViewportIt.Key);
 
-							TUniquePtr<FDisplayClusterMediaInputViewport> NewViewportInput = MakeUnique<FDisplayClusterMediaInputViewport>(
+							TSharedPtr<FDisplayClusterMediaInputViewport> NewViewportInput = MakeShared<FDisplayClusterMediaInputViewport>(
 								MediaInputId, ClusterNodeId,
 								ViewportIt.Key,
 								MediaSettings.MediaSource);
 
-							AllInputs.Add(NewViewportInput.Get());
 							InputViewports.Emplace(MediaInputId, MoveTemp(NewViewportInput));
 						}
 
 						// Media capture
 						if (MediaSettings.MediaOutput)
 						{
-							const FString MediaCaptureId = RootActor->GetName() + FString("_") + ClusterNodeId + FString("_") + ViewportIt.Key + FString("_viewport_capture");
+							const FString MediaCaptureId = DisplayClusterMediaHelpers::MediaId::GenerateMediaId(
+								DisplayClusterMediaHelpers::MediaId::EMediaDeviceType::Output,
+								DisplayClusterMediaHelpers::MediaId::EMediaOwnerType::Viewport,
+								*ClusterNodeId, *RootActor->GetName(), ViewportIt.Key);
 
-							UE_LOG(LogDisplayClusterMedia, Log, TEXT("Initializing viewport capture for viewport '%s'..."), *ViewportIt.Key);
+							UE_LOG(LogDisplayClusterMedia, Log, TEXT("Initializing viewport capture '%s' for viewport '%s'..."), *MediaCaptureId, *ViewportIt.Key);
 
-							TUniquePtr<FDisplayClusterMediaCaptureViewport> NewViewportCapture = MakeUnique<FDisplayClusterMediaCaptureViewport>(
+							TSharedPtr<FDisplayClusterMediaCaptureViewport> NewViewportCapture = MakeShared<FDisplayClusterMediaCaptureViewport>(
 								MediaCaptureId, ClusterNodeId,
 								ViewportIt.Key,
-								MediaSettings.MediaOutput);
+								MediaSettings.MediaOutput,
+								MediaSettings.OutputSyncPolicy);
 
-							AllCaptures.Add(NewViewportCapture.Get());
 							CaptureViewports.Emplace(MediaCaptureId, MoveTemp(NewViewportCapture));
 						}
 					}
@@ -184,49 +192,54 @@ void FDisplayClusterMediaModule::InitializeMedia()
 			///////////////////////////////
 			// ICVFX media setup
 			{
-				TArray<UActorComponent*> ICVFXCameraComponents;
-				RootActor->GetComponents(UDisplayClusterICVFXCameraComponent::StaticClass(), ICVFXCameraComponents);
+				// Get all ICVFX camera components
+				TArray<UDisplayClusterICVFXCameraComponent*> ICVFXCameraComponents;
+				RootActor->GetComponents(ICVFXCameraComponents);
 
-				for (const UActorComponent* const Component : ICVFXCameraComponents)
+				for (const UDisplayClusterICVFXCameraComponent* const ICVFXCameraComponent : ICVFXCameraComponents)
 				{
-					if (const UDisplayClusterICVFXCameraComponent* const ICVFXCamera = Cast<UDisplayClusterICVFXCameraComponent>(Component))
+					const FDisplayClusterConfigurationMediaICVFX& MediaSettings = ICVFXCameraComponent->CameraSettings.RenderSettings.Media;
+
+					if (MediaSettings.bEnable)
 					{
-						const FDisplayClusterConfigurationMediaICVFX& MediaSettings = ICVFXCamera->CameraSettings.RenderSettings.Media;
+						const FString ICVFXCameraName = ICVFXCameraComponent->GetName();
+						const FString ICVFXViewportId = DisplayClusterMediaHelpers::GenerateICVFXViewportName(ClusterNodeId, ICVFXCameraName);
 
-						if (MediaSettings.bEnable)
+						// Media input only
+						if (UMediaSource* MediaSource = MediaSettings.GetMediaSource(ClusterNodeId))
 						{
-							const FString ICVFXCameraName = ICVFXCamera->GetName();
-							const FString ICVFXViewportId = DisplayClusterMediaHelpers::GenerateICVFXViewportName(ClusterNodeId, ICVFXCameraName);
-							const FString MediaCaptureId  = RootActor->GetName() + FString("_") + ICVFXCameraName + FString("_icvfx_capture");
-							const FString MediaInputId    = RootActor->GetName() + FString("_") + ICVFXCameraName + FString("_icvfx_input");
-
-							// Media input only
-							if (UMediaSource* MediaSource = MediaSettings.GetMediaSource(ClusterNodeId))
-							{
-								UE_LOG(LogDisplayClusterMedia, Log, TEXT("Initializing ICVFX media input for camera '%s'..."), *ICVFXCameraName);
-
-								TUniquePtr<FDisplayClusterMediaInputViewport> NewICVFXInput = MakeUnique<FDisplayClusterMediaInputViewport>(
-									MediaInputId, ClusterNodeId,
-									ICVFXViewportId,
-									MediaSource);
-
-								AllInputs.Add(NewICVFXInput.Get());
-								InputViewports.Emplace(MediaInputId, MoveTemp(NewICVFXInput));
-							}
-
-							// Media capture only
-							if (UMediaOutput* MediaOutput = MediaSettings.GetMediaOutput(ClusterNodeId))
-							{
-								UE_LOG(LogDisplayClusterMedia, Log, TEXT("Initializing ICVFX capture for camera '%s'..."), *ICVFXCameraName);
-
-								TUniquePtr<FDisplayClusterMediaCaptureViewport> NewICVFXCapture = MakeUnique<FDisplayClusterMediaCaptureCamera>(
-									MediaCaptureId, ClusterNodeId,
-									ICVFXCameraName, ICVFXViewportId,
-									MediaOutput);
-
-								AllCaptures.Add(NewICVFXCapture.Get());
-								CaptureViewports.Emplace(MediaCaptureId, MoveTemp(NewICVFXCapture));
-							}
+							const FString MediaInputId = DisplayClusterMediaHelpers::MediaId::GenerateMediaId(
+								DisplayClusterMediaHelpers::MediaId::EMediaDeviceType::Input,
+								DisplayClusterMediaHelpers::MediaId::EMediaOwnerType::ICVFXCamera,
+								*ClusterNodeId, *RootActor->GetName(), ICVFXCameraName);
+
+							UE_LOG(LogDisplayClusterMedia, Log, TEXT("Initializing ICVFX media input '%s' for camera '%s'..."), *MediaInputId, *ICVFXCameraName);
+
+							TSharedPtr<FDisplayClusterMediaInputViewport> NewICVFXInput = MakeShared<FDisplayClusterMediaInputViewport>(
+								MediaInputId, ClusterNodeId,
+								ICVFXViewportId,
+								MediaSource);
+
+							InputViewports.Emplace(MediaInputId, MoveTemp(NewICVFXInput));
+						}
+
+						// Media capture only
+						if (UMediaOutput* MediaOutput = MediaSettings.GetMediaOutput(ClusterNodeId))
+						{
+							const FString MediaCaptureId = DisplayClusterMediaHelpers::MediaId::GenerateMediaId(
+								DisplayClusterMediaHelpers::MediaId::EMediaDeviceType::Output,
+								DisplayClusterMediaHelpers::MediaId::EMediaOwnerType::ICVFXCamera,
+								*ClusterNodeId, *RootActor->GetName(), ICVFXCameraName);
+
+							UE_LOG(LogDisplayClusterMedia, Log, TEXT("Initializing ICVFX capture '%s' for camera '%s'..."), *MediaCaptureId, *ICVFXCameraName);
+
+							TSharedPtr<FDisplayClusterMediaCaptureViewport> NewICVFXCapture = MakeShared<FDisplayClusterMediaCaptureCamera>(
+								MediaCaptureId, ClusterNodeId,
+								ICVFXCameraName, ICVFXViewportId,
+								MediaOutput,
+								MediaSettings.GetOutputSyncPolicy(ClusterNodeId));
+
+							CaptureViewports.Emplace(MediaCaptureId, MoveTemp(NewICVFXCapture));
 						}
 					}
 				}
@@ -242,44 +255,70 @@ void FDisplayClusterMediaModule::ReleaseMedia()
 
 	CaptureViewports.Reset();
 	CaptureNode.Reset();
-	AllCaptures.Reset();
 
 	InputViewports.Reset();
 	InputNode.Reset();
-	AllInputs.Reset();
 
 	FrameQueue.Release();
 }
 
 void FDisplayClusterMediaModule::StartCapture()
 {
-	for (FDisplayClusterMediaCaptureBase* Capture : AllCaptures)
+	// Start viewports capture
+	for (TPair<FString, TSharedPtr<FDisplayClusterMediaCaptureViewport>>& Capture : CaptureViewports)
+	{
+		Capture.Value->StartCapture();
+	}
+
+	// Start backbuffer capture
+	if (CaptureNode)
 	{
-		Capture->StartCapture();
+		CaptureNode->StartCapture();
 	}
 }
 
 void FDisplayClusterMediaModule::StopCapture()
 {
-	for (FDisplayClusterMediaCaptureBase* Capture : AllCaptures)
+	// Stop viewports capture
+	for (TPair<FString, TSharedPtr<FDisplayClusterMediaCaptureViewport>>& Capture : CaptureViewports)
+	{
+		Capture.Value->StopCapture();
+	}
+
+	// Stop backbuffer capture
+	if (CaptureNode)
 	{
-		Capture->StopCapture();
+		CaptureNode->StopCapture();
 	}
 }
 
 void FDisplayClusterMediaModule::PlayMedia()
 {
-	for (FDisplayClusterMediaInputBase* MediaInput : AllInputs)
+	// Start playback on viewports
+	for (TPair<FString, TSharedPtr<FDisplayClusterMediaInputViewport>>& MediaInput : InputViewports)
 	{
-		MediaInput->Play();
+		MediaInput.Value->Play();
+	}
+
+	// Start playback to the backbuffer
+	if (InputNode)
+	{
+		InputNode->Play();
 	}
 }
 
 void FDisplayClusterMediaModule::StopMedia()
 {
-	for (FDisplayClusterMediaInputBase* MediaInput : AllInputs)
+	// Stop playback on viewports
+	for (TPair<FString, TSharedPtr<FDisplayClusterMediaInputViewport>>& MediaInput : InputViewports)
+	{
+		MediaInput.Value->Stop();
+	}
+
+	// Stop playback to the backbuffer
+	if (InputNode)
 	{
-		MediaInput->Stop();
+		InputNode->Stop();
 	}
 }
 
diff --git a/Engine/Plugins/Runtime/nDisplay/Source/DisplayClusterMedia/Private/DisplayClusterMediaModule.h b/Engine/Plugins/Runtime/nDisplay/Source/DisplayClusterMedia/Private/DisplayClusterMediaModule.h
--- a/Engine/Plugins/Runtime/nDisplay/Source/DisplayClusterMedia/Private/DisplayClusterMediaModule.h
+++ b/Engine/Plugins/Runtime/nDisplay/Source/DisplayClusterMedia/Private/DisplayClusterMediaModule.h
@@ -5,18 +5,15 @@
 #include "CoreMinimal.h"
 #include "Modules/ModuleInterface.h"
 
-#include "Synchronization/DisplayClusterFrameQueue.h"
+#include "Synchronization/LatencyQueue/DisplayClusterFrameQueue.h"
 
-class FDisplayClusterMediaCaptureBase;
-class FDisplayClusterMediaCaptureICVFX;
 class FDisplayClusterMediaCaptureNode;
 class FDisplayClusterMediaCaptureViewport;
-class FDisplayClusterMediaInputBase;
-class FDisplayClusterMediaInputICVFX;
 class FDisplayClusterMediaInputNode;
 class FDisplayClusterMediaInputViewport;
 class IMediaPlayerFactory;
 
+
 /**
  * Media module
  */
@@ -53,13 +50,11 @@ private:
 	void OnEnginePreExit();
 
 private:
-	TMap<FString, TUniquePtr<FDisplayClusterMediaCaptureViewport>> CaptureViewports;
-	TUniquePtr<FDisplayClusterMediaCaptureNode>                    CaptureNode;
-	TArray<FDisplayClusterMediaCaptureBase*>                       AllCaptures;
+	TMap<FString, TSharedPtr<FDisplayClusterMediaCaptureViewport>> CaptureViewports;
+	TSharedPtr<FDisplayClusterMediaCaptureNode>                    CaptureNode;
 
-	TMap<FString, TUniquePtr<FDisplayClusterMediaInputViewport>>   InputViewports;
-	TUniquePtr<FDisplayClusterMediaInputNode>                      InputNode;
-	TArray<FDisplayClusterMediaInputBase*>                         AllInputs;
+	TMap<FString, TSharedPtr<FDisplayClusterMediaInputViewport>>   InputViewports;
+	TSharedPtr<FDisplayClusterMediaInputNode>                      InputNode;
 
 	// Latency queue
 	FDisplayClusterFrameQueue FrameQueue;
diff --git a/Engine/Plugins/Runtime/nDisplay/Source/DisplayClusterMedia/Private/Input/DisplayClusterMediaInputBase.cpp b/Engine/Plugins/Runtime/nDisplay/Source/DisplayClusterMedia/Private/Input/DisplayClusterMediaInputBase.cpp
--- a/Engine/Plugins/Runtime/nDisplay/Source/DisplayClusterMedia/Private/Input/DisplayClusterMediaInputBase.cpp
+++ b/Engine/Plugins/Runtime/nDisplay/Source/DisplayClusterMedia/Private/Input/DisplayClusterMediaInputBase.cpp
@@ -72,6 +72,8 @@ void FDisplayClusterMediaInputBase::Stop()
 
 void FDisplayClusterMediaInputBase::ImportMediaData(FRHICommandListImmediate& RHICmdList, const FMediaTextureInfo& TextureInfo)
 {
+	UE_LOG(LogDisplayClusterMedia, Verbose, TEXT("MediaInput '%s': importing texture on RT frame '%lu'..."), *GetMediaId(), GFrameCounterRenderThread);
+
 	MediaTexture->JustInTimeRender();
 
 	FRHITexture* const SrcTexture = MediaTexture->GetResource() ? MediaTexture->GetResource()->GetTextureRHI() : nullptr;
diff --git a/Engine/Plugins/Runtime/nDisplay/Source/DisplayClusterMedia/Private/Synchronization/DisplayClusterFrameQueue.cpp b/Engine/Plugins/Runtime/nDisplay/Source/DisplayClusterMedia/Private/Synchronization/DisplayClusterFrameQueue.cpp
deleted file mode 100644
diff --git a/Engine/Plugins/Runtime/nDisplay/Source/DisplayClusterMedia/Private/Synchronization/DisplayClusterFrameQueue.h b/Engine/Plugins/Runtime/nDisplay/Source/DisplayClusterMedia/Private/Synchronization/DisplayClusterFrameQueue.h
deleted file mode 100644
diff --git a/Engine/Plugins/Runtime/nDisplay/Source/DisplayClusterMedia/Private/Synchronization/DisplayClusterFrameQueueItem.cpp b/Engine/Plugins/Runtime/nDisplay/Source/DisplayClusterMedia/Private/Synchronization/DisplayClusterFrameQueueItem.cpp
deleted file mode 100644
diff --git a/Engine/Plugins/Runtime/nDisplay/Source/DisplayClusterMedia/Private/Synchronization/DisplayClusterFrameQueueItem.h b/Engine/Plugins/Runtime/nDisplay/Source/DisplayClusterMedia/Private/Synchronization/DisplayClusterFrameQueueItem.h
deleted file mode 100644
diff --git a/Engine/Plugins/Runtime/nDisplay/Source/DisplayClusterMedia/Private/Synchronization/LatencyQueue/DisplayClusterFrameQueue.cpp b/Engine/Plugins/Runtime/nDisplay/Source/DisplayClusterMedia/Private/Synchronization/LatencyQueue/DisplayClusterFrameQueue.cpp
new file mode 100644
--- /dev/null
+++ b/Engine/Plugins/Runtime/nDisplay/Source/DisplayClusterMedia/Private/Synchronization/LatencyQueue/DisplayClusterFrameQueue.cpp
@@ -0,0 +1,304 @@
+// Copyright Epic Games, Inc. All Rights Reserved.
+
+#include "Synchronization/LatencyQueue/DisplayClusterFrameQueue.h"
+
+#include "DisplayClusterConfigurationTypes.h"
+#include "DisplayClusterMediaLog.h"
+#include "DisplayClusterRootActor.h"
+
+#include "IDisplayCluster.h"
+#include "IDisplayClusterCallbacks.h"
+#include "Game/IDisplayClusterGameManager.h"
+#include "Render/IDisplayClusterRenderManager.h"
+#include "Render/Viewport/IDisplayClusterViewportManager.h"
+#include "Render/Viewport/IDisplayClusterViewportManagerProxy.h"
+#include "Render/Viewport/IDisplayClusterViewportProxy.h"
+#include "ShaderParameters/DisplayClusterShaderParameters_ICVFX.h"
+#include "ShaderParameters/DisplayClusterShaderParameters_WarpBlend.h"
+
+#include "Engine/Engine.h"
+#include "Engine/GameViewportClient.h"
+
+
+void FDisplayClusterFrameQueue::Init()
+{
+	// Subscribe for events
+	if (IDisplayCluster::Get().GetOperationMode() == EDisplayClusterOperationMode::Cluster)
+	{
+		if (GEngine && GEngine->GameViewport)
+		{
+			GEngine->GameViewport->OnBeginDraw().AddRaw(this, &FDisplayClusterFrameQueue::HandleBeginDraw);
+			GEngine->GameViewport->OnEndDraw().AddRaw(this, &FDisplayClusterFrameQueue::HandleEndDraw);
+		}
+
+		IDisplayCluster::Get().GetCallbacks().OnDisplayClusterProcessLatency_RenderThread().AddRaw(this, &FDisplayClusterFrameQueue::HandleProcessLatency_RenderThread);
+		IDisplayCluster::Get().GetCallbacks().OnDisplayClusterPreProcessIcvfx_RenderThread().AddRaw(this, &FDisplayClusterFrameQueue::HandlePreProcessIcvfx_RenderThread);
+	}
+}
+
+void FDisplayClusterFrameQueue::Release()
+{
+	// Unsubscribe from events
+	if (IDisplayCluster::Get().GetOperationMode() == EDisplayClusterOperationMode::Cluster)
+	{
+		if (GEngine && GEngine->GameViewport)
+		{
+			GEngine->GameViewport->OnBeginDraw().RemoveAll(this);
+			GEngine->GameViewport->OnEndDraw().RemoveAll(this);
+		}
+
+		IDisplayCluster::Get().GetCallbacks().OnDisplayClusterProcessLatency_RenderThread().RemoveAll(this);
+		IDisplayCluster::Get().GetCallbacks().OnDisplayClusterPreProcessIcvfx_RenderThread().RemoveAll(this);
+	}
+
+	// We need to properly release render thread data
+	ENQUEUE_RENDER_COMMAND(FDisplayClusterFrameQueue_Release)(
+		[this](FRHICommandListImmediate& RHICmdList)
+		{
+			CleanQueue_RenderThread();
+		});
+}
+
+void FDisplayClusterFrameQueue::HandleBeginDraw()
+{
+	// Update latency value before rendering new frame
+	if (ADisplayClusterRootActor* RootActor = IDisplayCluster::Get().GetGameMgr()->GetRootActor())
+	{
+		if (const UDisplayClusterConfigurationData* const ConfigData = RootActor->GetConfigData())
+		{
+			// Get latency from config
+			const int32 NewLatency = FMath::Clamp(ConfigData->MediaSettings.Latency, 0, 10);
+			UE_LOG(LogDisplayClusterMedia, VeryVerbose, TEXT("New latency is %u"), NewLatency);
+
+			// Update latency data on RT
+			ENQUEUE_RENDER_COMMAND(FDisplayClusterFrameQueue_SetLatency)(
+				[this, NewLatency](FRHICommandListImmediate& RHICmdList)
+				{
+					SetLatency_RenderThread(NewLatency);
+				});
+		}
+	}
+
+	// Process pending frames adding
+	ENQUEUE_RENDER_COMMAND(FDisplayClusterFrameQueue_ProcessAddingPendingFrames)(
+		[this](FRHICommandListImmediate& RHICmdList)
+		{
+			ProcessAddingPendingFrames_RenderThread();
+		});
+}
+
+void FDisplayClusterFrameQueue::HandleEndDraw()
+{
+	// We need to update head and tail indices before rendering next frame
+	ENQUEUE_RENDER_COMMAND(FDisplayClusterFrameQueue_StepQueueIndices)(
+		[this](FRHICommandListImmediate& RHICmdList)
+		{
+			StepQueueIndices_RenderThread();
+		});
+}
+
+void FDisplayClusterFrameQueue::HandleProcessLatency_RenderThread(FRHICommandListImmediate& RHICmdList, const IDisplayClusterViewportManagerProxy* ViewportManagerProxy, FViewport* Viewport)
+{
+	if (Frames.Num() > 0)
+	{
+		// At this point, all the views for current frame have been rendered already
+		if (IDisplayClusterRenderManager* RenderMgr = IDisplayCluster::Get().GetRenderMgr())
+		{
+			if (IDisplayClusterViewportManager* ViewportMgr = RenderMgr->GetViewportManager())
+			{
+				if (IDisplayClusterViewportManagerProxy* ViewportMgrProxy = ViewportMgr->GetProxy())
+				{
+					TArrayView<TSharedPtr<IDisplayClusterViewportProxy, ESPMode::ThreadSafe>> ViewportProxies = ViewportMgrProxy->GetViewports_RenderThread();
+					for (TSharedPtr<IDisplayClusterViewportProxy, ESPMode::ThreadSafe>& ViewportProxyIt : ViewportProxies)
+					{
+						// Handle textures of all viewports that don't have any media input attached. Those don't get rendered.
+						if (ViewportProxyIt.IsValid() && !ViewportProxyIt->GetRenderSettings_RenderThread().bSkipSceneRenderingButLeaveResourcesAvailable)
+						{
+							// Get viewport texture
+							TArray<FRHITexture*> Textures;
+							ViewportProxyIt->GetResources_RenderThread(EDisplayClusterViewportResourceType::InternalRenderTargetResource, Textures);
+							check(Textures.Num() > 0);
+
+							const FString ViewportId(ViewportProxyIt->GetId());
+
+							// Enqueue this view (head)
+							Frames[IdxHead].SaveView(RHICmdList, ViewportId, Textures[0]);
+							// Dequeue from the past (tail)
+							Frames[IdxTail].LoadView(RHICmdList, ViewportId, Textures[0]);
+						}
+					}
+				}
+			}
+		}
+	}
+}
+
+void FDisplayClusterFrameQueue::HandlePreProcessIcvfx_RenderThread(FRHICommandListImmediate& RHICmdList, const IDisplayClusterViewportProxy* ViewportProxy, FDisplayClusterShaderParameters_WarpBlend& WarpBlendParameters, FDisplayClusterShaderParameters_ICVFX& ICVFXParameters)
+{
+	// Substitute shader data
+	if (Frames.Num() > 0)
+	{
+		const FString ViewportId(ViewportProxy->GetId());
+
+		Frames[IdxHead].SaveData(RHICmdList, ViewportId, WarpBlendParameters, ICVFXParameters);
+		Frames[IdxTail].LoadData(RHICmdList, ViewportId, WarpBlendParameters, ICVFXParameters);
+	}
+}
+
+void FDisplayClusterFrameQueue::SetLatency_RenderThread(int32 NewLatency)
+{
+	// Nothing to do if requested latency is the same as previously (no changes since last frame)
+	if (LastRequestedLatency != NewLatency && NewLatency >= 0)
+	{
+		// Check if latency just was disabled
+		if (NewLatency == 0)
+		{
+			// Clear the queue
+			CleanQueue_RenderThread();
+		}
+		// Check if new latency is higher
+		else if (NewLatency > LastRequestedLatency)
+		{
+			FramesPendingToAdd = (NewLatency + 1) - Frames.Num();
+		}
+		// Check if new latency is lower
+		else if (NewLatency < LastRequestedLatency)
+		{
+			// Depending on the current queue size, we need to find out if the queue needs to be resized
+			const int32 TargetQueueSize = NewLatency + 1;
+
+			// Yeah, we need to remove unnecessary elements from the queue
+			if (TargetQueueSize < Frames.Num())
+			{
+				RemoveFromQueue_RenderThread(Frames.Num() - TargetQueueSize);
+				FramesPendingToAdd = 0;
+			}
+			// No need to remove anything from the queue, just recalculate new amount of pending frames
+			else
+			{
+				FramesPendingToAdd = TargetQueueSize - Frames.Num();
+			}
+		}
+
+		// Finally, remember latency requested
+		LastRequestedLatency = NewLatency;
+	}
+}
+
+void FDisplayClusterFrameQueue::ProcessAddingPendingFrames_RenderThread()
+{
+	if (FramesPendingToAdd > 0)
+	{
+		// Increase queue size step by step, 1 item per frame. This allows to avoid
+		// situations where we have empty queue items.
+		FramesAdded = 1;
+		AddToQueue_RenderThread(FramesAdded);
+	}
+	else
+	{
+		// Queue length remains unchanged
+		FramesAdded = 0;
+	}
+}
+
+void FDisplayClusterFrameQueue::AddToQueue_RenderThread(int32 FramesAmount)
+{
+	checkSlow(FramesAmount > 0);
+
+	// Duplicate the head frame to increase the queue
+	for (int32 Idx = 0; Idx < FramesAmount; ++Idx)
+	{
+		// In case the queue is empty, just instantiate new item
+		if (Frames.IsEmpty())
+		{
+			Frames.AddDefaulted();
+
+			// Reset indices
+			IdxHead = 0;
+			IdxTail = 0;
+		}
+		// Otherwise duplicate head item to prevent 'black' frames in the future
+		else
+		{
+			// For some reason, the following doesn't work (there will be check fail inside Insert)
+			// >>>> Frames.Insert(Frames[IdxHead], IdxHead)
+			//
+			// To achieve the same avoiding unnecessary data copying, I use the approach where data copying
+			// is performed in copy constructor, and the instance is pushed to the container using move constructor.
+			FDisplayClusterFrameQueueItem ClonedItem(Frames[IdxHead]);
+			Frames.Insert(MoveTemp(ClonedItem), IdxHead);
+
+			// Recalculate tail index
+			if (IdxTail > IdxHead)
+			{
+				++IdxTail;
+			}
+		}
+	}
+}
+
+void FDisplayClusterFrameQueue::RemoveFromQueue_RenderThread(int32 FramesAmout)
+{
+	checkSlow(FramesAmout < Frames.Num());
+
+	// We need to remove specified amount of frames starting from the tail
+	for (int32 Idx = 0; Idx < FramesAmout; ++Idx)
+	{
+		// Remove the item
+		Frames.RemoveAt(IdxTail);
+
+		// Recalculate head index
+		if (IdxHead > IdxTail)
+		{
+			--IdxHead;
+		}
+
+		// Recalculate tail index
+		if (IdxTail >= Frames.Num())
+		{
+			IdxTail = 0;
+		}
+	}
+}
+
+void FDisplayClusterFrameQueue::CleanQueue_RenderThread()
+{
+	// Clean up the container
+	Frames.Empty();
+
+	// Reset internals
+	IdxHead = 0;
+	IdxTail = 0;
+	FramesPendingToAdd = 0;
+}
+
+void FDisplayClusterFrameQueue::StepQueueIndices_RenderThread()
+{
+	// Nothing to do if latency is zero
+	if (Frames.IsEmpty())
+	{
+		UE_LOG(LogDisplayClusterMedia, VeryVerbose, TEXT("Queue indices weren't updated, latency=0"));
+		return;
+	}
+
+	// Head index
+	if (++IdxHead >= Frames.Num())
+	{
+		IdxHead = 0;
+	}
+
+	// Tail index doesn't move while we the queue being increased
+	if (FramesAdded == 0)
+	{
+		if (++IdxTail >= Frames.Num())
+		{
+			IdxTail = 0;
+		}
+	}
+	// Otherwise update pending amount
+	else
+	{
+		FramesPendingToAdd -= FramesAdded;
+		checkSlow(FramesPendingToAdd >= 0);
+	}
+}
diff --git a/Engine/Plugins/Runtime/nDisplay/Source/DisplayClusterMedia/Private/Synchronization/LatencyQueue/DisplayClusterFrameQueue.h b/Engine/Plugins/Runtime/nDisplay/Source/DisplayClusterMedia/Private/Synchronization/LatencyQueue/DisplayClusterFrameQueue.h
new file mode 100644
--- /dev/null
+++ b/Engine/Plugins/Runtime/nDisplay/Source/DisplayClusterMedia/Private/Synchronization/LatencyQueue/DisplayClusterFrameQueue.h
@@ -0,0 +1,76 @@
+// Copyright Epic Games, Inc. All Rights Reserved.
+
+#pragma once
+
+#include "CoreMinimal.h"
+
+#include "Synchronization/LatencyQueue/DisplayClusterFrameQueueItem.h"
+
+class FDisplayClusterMediaInputViewport;
+class FDisplayClusterShaderParameters_ICVFX;
+class FViewport;
+class IDisplayClusterViewportManagerProxy;
+class IDisplayClusterViewportProxy;
+class RHICmdList;
+struct FDisplayClusterShaderParameters_WarpBlend;
+
+/**
+ * Latency queue class
+ *
+ * Implements artificial latency feature. It's responsible for caching visual and shaders data to be used later.
+ *
+ * All the data that comes out of the rendering pipeline for current frame is cached to the HEAD of the queue. The final
+ * output is built from the TAIL view data. Every queue item contains all necessary information for every nD viewport.
+ *
+ * Note: this queue needs to be an intermediate part of the nD rendering pipeline. It should be moved there
+ * at some point during refactoring.
+ */
+class FDisplayClusterFrameQueue
+{
+public:
+	// Initializes latency feature
+	void Init();
+	// Releases latency feature
+	void Release();
+
+private:
+	// Handles BeginDraw event
+	void HandleBeginDraw();
+	// Handles EndDraw event
+	void HandleEndDraw();
+	// Handles latency processing event
+	void HandleProcessLatency_RenderThread(FRHICommandListImmediate& RHICmdList, const IDisplayClusterViewportManagerProxy* ViewportManagerProxy, FViewport* Viewport);
+	// Handles ICVFX pre-processing event
+	void HandlePreProcessIcvfx_RenderThread(FRHICommandListImmediate& RHICmdList, const IDisplayClusterViewportProxy* ViewportProxy, FDisplayClusterShaderParameters_WarpBlend& WarpBlendParameters, FDisplayClusterShaderParameters_ICVFX& ICVFXParameters);
+
+private:
+	// Changes current latency
+	void SetLatency_RenderThread(int32 NewLatency);
+	// Auxiliary method to process queue resize
+	void ProcessAddingPendingFrames_RenderThread();
+	// Auxiliary method to increase queue length by amount of frames
+	void AddToQueue_RenderThread(int32 FramesAmount);
+	// Auxiliary method to decrease queue length by amount of frames
+	void RemoveFromQueue_RenderThread(int32 FramesAmout);
+	// Cleans the queue container
+	void CleanQueue_RenderThread();
+	// Updates head and tail indices every frame
+	void StepQueueIndices_RenderThread();
+
+private:
+	// Frames buffer
+	TArray<FDisplayClusterFrameQueueItem> Frames;
+
+	// Latency value requested last time
+	int32 LastRequestedLatency = 0;
+
+	// Amount of frames pending to add to the latency queue
+	int32 FramesPendingToAdd = 0;
+	// Amount of frames that were added at current frame
+	int32 FramesAdded = 0;
+
+	// Head queue index
+	int32 IdxHead = 0;
+	// Tail queue index
+	int32 IdxTail = 0;
+};
diff --git a/Engine/Plugins/Runtime/nDisplay/Source/DisplayClusterMedia/Private/Synchronization/LatencyQueue/DisplayClusterFrameQueueItem.cpp b/Engine/Plugins/Runtime/nDisplay/Source/DisplayClusterMedia/Private/Synchronization/LatencyQueue/DisplayClusterFrameQueueItem.cpp
new file mode 100644
--- /dev/null
+++ b/Engine/Plugins/Runtime/nDisplay/Source/DisplayClusterMedia/Private/Synchronization/LatencyQueue/DisplayClusterFrameQueueItem.cpp
@@ -0,0 +1,152 @@
+// Copyright Epic Games, Inc. All Rights Reserved.
+
+#include "Synchronization/LatencyQueue/DisplayClusterFrameQueueItem.h"
+
+#include "DisplayClusterConfigurationTypes.h"
+#include "DisplayClusterMediaLog.h"
+#include "DisplayClusterRootActor.h"
+
+#include "IDisplayCluster.h"
+#include "Render/IDisplayClusterRenderManager.h"
+#include "Render/Viewport/IDisplayClusterViewportManager.h"
+#include "Render/Viewport/IDisplayClusterViewportManagerProxy.h"
+
+#include "RHICommandList.h"
+
+
+FDisplayClusterFrameQueueItem::FDisplayClusterFrameQueueItem(const FDisplayClusterFrameQueueItem& Other)
+{
+	for (const TPair<FString, FDisplayClusterFrameQueueItemView>& View : Other.Views)
+	{
+		FDisplayClusterFrameQueueItemView& NewItem = this->Views.Emplace(View.Key, View.Value);
+		if (View.Value.Texture.IsValid())
+		{
+			NewItem.Texture = CreateTexture(View.Value.Texture);
+			FRHICommandListImmediate& RHICmdList = GetImmediateCommandList_ForRenderCommand();
+			RHICmdList.CopyTexture(View.Value.Texture.GetReference(), NewItem.Texture, FRHICopyTextureInfo());
+		}
+	}
+}
+
+void FDisplayClusterFrameQueueItem::SaveView(FRHICommandListImmediate& RHICmdList, const FString& ViewportId, FRHITexture* Texture)
+{
+	checkSlow(Texture);
+
+	// Find proper view and copy texture data
+	FDisplayClusterFrameQueueItemView* View = Views.Find(ViewportId);
+	if (!View)
+	{
+		View = &Views.Emplace(ViewportId);
+	}
+
+	// Create texture if not yet available. Or re-create if the source texture has been updated (re-sized, other format, etc.)
+	if (!View->Texture.IsValid() ||
+		Texture->GetFormat() != View->Texture.GetReference()->GetFormat() ||
+		Texture->GetDesc().Extent != View->Texture.GetReference()->GetDesc().Extent)
+	{
+		View->Texture = CreateTexture(Texture);
+		check(View->Texture.IsValid());
+	}
+
+	// Copy texture data
+	if (View->Texture.IsValid())
+	{
+		RHICmdList.CopyTexture(Texture, View->Texture.GetReference(), FRHICopyTextureInfo());
+	}
+}
+
+void FDisplayClusterFrameQueueItem::LoadView(FRHICommandListImmediate& RHICmdList, const FString& ViewportId, FRHITexture* Texture)
+{
+	checkSlow(Texture);
+
+	// Find proper view and copy texture data
+	const FDisplayClusterFrameQueueItemView* const View = Views.Find(ViewportId);
+	if (View && View->Texture.IsValid())
+	{
+		RHICmdList.CopyTexture(View->Texture.GetReference(), Texture, FRHICopyTextureInfo());
+	}
+}
+
+void FDisplayClusterFrameQueueItem::SaveData(FRHICommandListImmediate& RHICmdList, const FString& ViewportId, FDisplayClusterShaderParameters_WarpBlend& WarpBlendParameters, FDisplayClusterShaderParameters_ICVFX& ICVFXParameters)
+{
+	if (FDisplayClusterFrameQueueItemView* View = Views.Find(ViewportId))
+	{
+		// Warp/blend - cherry pick necessary parameters only
+		View->WarpBlendData.bRenderAlphaChannel = WarpBlendParameters.bRenderAlphaChannel;
+		View->WarpBlendData.Context = WarpBlendParameters.Context;
+
+		// ICVFX - cherry pick necessary parameters only
+		View->IcvfxData.Cameras = ICVFXParameters.Cameras;
+	}
+}
+
+void FDisplayClusterFrameQueueItem::LoadData(FRHICommandListImmediate& RHICmdList, const FString& ViewportId, FDisplayClusterShaderParameters_WarpBlend& WarpBlendParameters, FDisplayClusterShaderParameters_ICVFX& ICVFXParameters)
+{
+	if (const FDisplayClusterFrameQueueItemView* const View = Views.Find(ViewportId))
+	{
+		// Warp/blend - cherry pick necessary parameters only
+		WarpBlendParameters.bRenderAlphaChannel = View->WarpBlendData.bRenderAlphaChannel;
+		WarpBlendParameters.Context = View->WarpBlendData.Context;
+
+		// ICVFX - cherry pick necessary parameters only
+		for (const FDisplayClusterShaderParameters_ICVFX::FCameraSettings& CameraSettings : View->IcvfxData.Cameras)
+		{
+			// Look up for proper camera settings object. Here we use ViewportId as a key to compare.
+			FDisplayClusterShaderParameters_ICVFX::FCameraSettings* TargetCameraSettings = ICVFXParameters.Cameras.FindByPredicate([CameraSettings](const FDisplayClusterShaderParameters_ICVFX::FCameraSettings& Item)
+				{
+					return Item.Resource.ViewportId.Equals(CameraSettings.Resource.ViewportId, ESearchCase::IgnoreCase);
+				});
+
+			// Update camera settings
+			if (TargetCameraSettings)
+			{
+				TargetCameraSettings->SoftEdge                 = CameraSettings.SoftEdge;
+				TargetCameraSettings->InnerCameraBorderColor      = CameraSettings.InnerCameraBorderColor;
+				TargetCameraSettings->InnerCameraBorderThickness  = CameraSettings.InnerCameraBorderThickness;
+				TargetCameraSettings->InnerCameraFrameAspectRatio = CameraSettings.InnerCameraFrameAspectRatio;
+				TargetCameraSettings->Local2WorldTransform     = CameraSettings.Local2WorldTransform;
+				TargetCameraSettings->CameraViewRotation       = CameraSettings.CameraViewRotation;
+				TargetCameraSettings->CameraViewLocation       = CameraSettings.CameraViewLocation;
+				TargetCameraSettings->CameraPrjMatrix          = CameraSettings.CameraPrjMatrix;
+				TargetCameraSettings->ChromakeySource          = CameraSettings.ChromakeySource;
+				TargetCameraSettings->ChromakeyColor           = CameraSettings.ChromakeyColor;
+				TargetCameraSettings->ChromakeyMarkersColor    = CameraSettings.ChromakeyMarkersColor;
+				TargetCameraSettings->ChromakeyMarkersScale    = CameraSettings.ChromakeyMarkersScale;
+				TargetCameraSettings->ChromakeyMarkersDistance = CameraSettings.ChromakeyMarkersDistance;
+				TargetCameraSettings->ChromakeyMarkersOffset   = CameraSettings.ChromakeyMarkersOffset;
+				TargetCameraSettings->RenderOrder              = CameraSettings.RenderOrder;
+			}
+		}
+	}
+}
+
+FTextureRHIRef FDisplayClusterFrameQueueItem::CreateTexture(FRHITexture* ReferenceTexture)
+{
+	if (ReferenceTexture)
+	{
+		// Use original format and size
+		const int32 SizeX = ReferenceTexture->GetDesc().Extent.X;
+		const int32 SizeY = ReferenceTexture->GetDesc().Extent.Y;
+		const EPixelFormat Format = ReferenceTexture->GetFormat();
+
+		// Prepare description
+		FRHITextureCreateDesc Desc =
+			FRHITextureCreateDesc::Create2D(TEXT("DisplayClusterFrameQueueCacheTexture"), SizeX, SizeY, Format)
+			.SetClearValue(FClearValueBinding::Black)
+			.SetNumMips(1)
+			.SetFlags(ETextureCreateFlags::Dynamic)
+			.AddFlags(ETextureCreateFlags::MultiGPUGraphIgnore)
+			.SetInitialState(ERHIAccess::SRVMask);
+
+		// Leave original flags, but make sure it's ResolveTargetable but not RenderTargetable
+		ETextureCreateFlags Flags = ReferenceTexture->GetFlags();
+		Flags &= ~ETextureCreateFlags::RenderTargetable;
+		Flags |= ETextureCreateFlags::ResolveTargetable;
+		Desc.SetFlags(Flags);
+
+		// Create texture
+		return RHICreateTexture(Desc);
+	}
+
+	return nullptr;
+}
diff --git a/Engine/Plugins/Runtime/nDisplay/Source/DisplayClusterMedia/Private/Synchronization/LatencyQueue/DisplayClusterFrameQueueItem.h b/Engine/Plugins/Runtime/nDisplay/Source/DisplayClusterMedia/Private/Synchronization/LatencyQueue/DisplayClusterFrameQueueItem.h
new file mode 100644
--- /dev/null
+++ b/Engine/Plugins/Runtime/nDisplay/Source/DisplayClusterMedia/Private/Synchronization/LatencyQueue/DisplayClusterFrameQueueItem.h
@@ -0,0 +1,64 @@
+// Copyright Epic Games, Inc. All Rights Reserved.
+
+#pragma once
+
+#include "CoreMinimal.h"
+
+#include "ShaderParameters/DisplayClusterShaderParameters_ICVFX.h"
+#include "ShaderParameters/DisplayClusterShaderParameters_WarpBlend.h"
+
+class FRHICommandListImmediate;
+class FRHITexture;
+
+
+/**
+ * Viewport data cache to be used in the latency queue
+ */
+struct FDisplayClusterFrameQueueItemView
+{
+	/** Cached view texture of a viewoprt */
+	FTextureRHIRef Texture = nullptr;
+
+	/** Cached warp/blend shader data */
+	FDisplayClusterShaderParameters_WarpBlend WarpBlendData;
+
+	/** Cached ICVFX shader data */
+	FDisplayClusterShaderParameters_ICVFX IcvfxData;
+};
+
+
+/**
+ * Latency queue item class
+ *
+ * It's responsible for caching all required data for single frame.
+ * 
+ * Note: this one is render thread only. It's never used on the main thread or any other thread.
+ */
+class FDisplayClusterFrameQueueItem
+{
+public:
+	FDisplayClusterFrameQueueItem() = default;
+	FDisplayClusterFrameQueueItem(const FDisplayClusterFrameQueueItem& Other);
+	FDisplayClusterFrameQueueItem(FDisplayClusterFrameQueueItem&& Other) = default;
+
+	~FDisplayClusterFrameQueueItem() = default;
+
+public:
+	// Save texture of a viewport
+	void SaveView(FRHICommandListImmediate& RHICmdList, const FString& ViewportId, FRHITexture* Texture);
+	// Load texture of a viewport
+	void LoadView(FRHICommandListImmediate& RHICmdList, const FString& ViewportId, FRHITexture* Texture);
+
+	// Save shader data of a viewport
+	void SaveData(FRHICommandListImmediate& RHICmdList, const FString& ViewportId, FDisplayClusterShaderParameters_WarpBlend& WarpBlendParameters, FDisplayClusterShaderParameters_ICVFX& ICVFXParameters);
+	// Load shader data of a viewport
+	void LoadData(FRHICommandListImmediate& RHICmdList, const FString& ViewportId, FDisplayClusterShaderParameters_WarpBlend& WarpBlendParameters, FDisplayClusterShaderParameters_ICVFX& ICVFXParameters);
+
+private:
+	// Auxiliary function that creates a new texture based on the reference (same format, size)
+	FTextureRHIRef CreateTexture(FRHITexture* ReferenceTexture);
+
+private:
+	// Available views (viewports)
+	TMap<FString, FDisplayClusterFrameQueueItemView> Views;
+};
diff --git a/Engine/Plugins/Runtime/nDisplay/Source/DisplayClusterMedia/Private/Synchronization/SyncPolicies/DisplayClusterMediaOutputSynchronizationPolicyEthernetBarrier.cpp b/Engine/Plugins/Runtime/nDisplay/Source/DisplayClusterMedia/Private/Synchronization/SyncPolicies/DisplayClusterMediaOutputSynchronizationPolicyEthernetBarrier.cpp
new file mode 100644
--- /dev/null
+++ b/Engine/Plugins/Runtime/nDisplay/Source/DisplayClusterMedia/Private/Synchronization/SyncPolicies/DisplayClusterMediaOutputSynchronizationPolicyEthernetBarrier.cpp
@@ -0,0 +1,10 @@
+// Copyright Epic Games, Inc. All Rights Reserved.
+
+#include "Synchronization/DisplayClusterMediaOutputSynchronizationPolicyEthernetBarrier.h"
+
+
+void UDisplayClusterMediaOutputSynchronizationPolicyEthernetBarrier::Synchronize()
+{
+	// Just sync on the barrier
+	SyncThreadOnBarrier();
+}
diff --git a/Engine/Plugins/Runtime/nDisplay/Source/DisplayClusterMedia/Private/Synchronization/SyncPolicies/DisplayClusterMediaOutputSynchronizationPolicyEthernetBarrierBase.cpp b/Engine/Plugins/Runtime/nDisplay/Source/DisplayClusterMedia/Private/Synchronization/SyncPolicies/DisplayClusterMediaOutputSynchronizationPolicyEthernetBarrierBase.cpp
new file mode 100644
--- /dev/null
+++ b/Engine/Plugins/Runtime/nDisplay/Source/DisplayClusterMedia/Private/Synchronization/SyncPolicies/DisplayClusterMediaOutputSynchronizationPolicyEthernetBarrierBase.cpp
@@ -0,0 +1,289 @@
+// Copyright Epic Games, Inc. All Rights Reserved.
+
+#include "Synchronization/DisplayClusterMediaOutputSynchronizationPolicyEthernetBarrierBase.h"
+
+#include "DisplayClusterMediaHelpers.h"
+#include "DisplayClusterMediaLog.h"
+
+#include "MediaCapture.h"
+
+#include "DisplayClusterConfigurationTypes.h"
+#include "DisplayClusterEnums.h"
+#include "DisplayClusterRootActor.h"
+
+#include "Components/DisplayClusterICVFXCameraComponent.h"
+
+#include "IDisplayCluster.h"
+#include "Cluster/IDisplayClusterClusterManager.h"
+#include "Game/IDisplayClusterGameManager.h"
+
+
+bool UDisplayClusterMediaOutputSynchronizationPolicyEthernetBarrierBase::StartSynchronization(UMediaCapture* MediaCapture, const FString& MediaId)
+{
+	// Cluster mode only
+	if (IDisplayCluster::Get().GetOperationMode() != EDisplayClusterOperationMode::Cluster)
+	{
+		UE_LOG(LogDisplayClusterMediaSync, Warning, TEXT("'%s': Media synchronization is available in cluster mode only"), *MediaId);
+		return false;
+	}
+
+	// Nothing to do if already running
+	if (bIsRunning)
+	{
+		UE_LOG(LogDisplayClusterMediaSync, Warning, TEXT("'%s': Synchronization is on already"), *MediaId);
+		return true;
+	}
+
+	if (!MediaCapture)
+	{
+		UE_LOG(LogDisplayClusterMediaSync, Warning, TEXT("'%s': Invalid capture device (nullptr)"), *MediaId);
+		return false;
+	}
+
+	if (!IsCaptureTypeSupported(MediaCapture))
+	{
+		UE_LOG(LogDisplayClusterMediaSync, Warning, TEXT("'%s': Synchronization of media capture '%s' is not supported by this sync policy"), *MediaId, *MediaCapture->GetName());
+		return false;
+	}
+
+	// Store capture device
+	CapturingDevice = MediaCapture;
+	MediaDeviceId   = MediaId;
+
+	// Initialize dynamic barrier first
+	if (!InitializeBarrier(MediaId))
+	{
+		UE_LOG(LogDisplayClusterMediaSync, Warning, TEXT("'%s': Couldn't initialize barrier client"), *MediaId);
+		return false;
+	}
+
+	// And subscribe for sync callbacks
+	CapturingDevice->OnOutputSynchronization.BindUObject(this, &UDisplayClusterMediaOutputSynchronizationPolicyEthernetBarrierBase::ProcessMediaSynchronizationCallback);
+
+	// Update state
+	bIsRunning = true;
+
+	return true;
+}
+
+void UDisplayClusterMediaOutputSynchronizationPolicyEthernetBarrierBase::StopSynchronization()
+{
+	if (bIsRunning)
+	{
+		// Don't reference capture device
+		if (CapturingDevice)
+		{
+			CapturingDevice->OnOutputSynchronization.Unbind();
+			CapturingDevice = nullptr;
+		}
+
+		// Release barrier client
+		ReleaseBarrier();
+
+		// Update state
+		bIsRunning = false;
+	}
+}
+
+bool UDisplayClusterMediaOutputSynchronizationPolicyEthernetBarrierBase::IsRunning()
+{
+	return bIsRunning;
+}
+
+FString UDisplayClusterMediaOutputSynchronizationPolicyEthernetBarrierBase::GetMediaDeviceId() const
+{
+	return MediaDeviceId;
+}
+
+void UDisplayClusterMediaOutputSynchronizationPolicyEthernetBarrierBase::SyncThreadOnBarrier()
+{
+	// Sync on the barrier if everything is good
+	if (bIsRunning && EthernetBarrierClient && EthernetBarrierClient->IsConnected())
+	{
+		UE_LOG(LogDisplayClusterMediaSync, Verbose, TEXT("'%s': Synchronizing caller '%s' at the barrier '%s'"), *GetMediaDeviceId(), *ThreadMarker, *BarrierId);
+		EthernetBarrierClient->Synchronize(BarrierId, ThreadMarker);
+	}
+}
+
+bool UDisplayClusterMediaOutputSynchronizationPolicyEthernetBarrierBase::InitializeBarrier(const FString& MediaId)
+{
+	if (MediaId.IsEmpty())
+	{
+		UE_LOG(LogDisplayClusterMediaSync, Warning, TEXT("'%s': Wrong MediaId"), *MediaId);
+		return false;
+	}
+
+	IDisplayClusterClusterManager& ClusterMgr = *IDisplayCluster::Get().GetClusterMgr();
+
+	// Barriers client ID
+	const FString ClientId = FString::Printf(TEXT("CLN_GB_%s"), *MediaId);
+	UE_LOG(LogDisplayClusterMediaSync, VeryVerbose, TEXT("'%s': Requesting barrier client '%s'"), *MediaId, *ClientId);
+
+	// Instantiate barrier client
+	EthernetBarrierClient = IDisplayCluster::Get().GetClusterMgr()->CreateGenericBarriersClient(ClientId);
+	if (!EthernetBarrierClient)
+	{
+		UE_LOG(LogDisplayClusterMediaSync, Warning, TEXT("'%s': Couldn't instantiate barrier client '%s'"), *MediaId, *ClientId);
+		return false;
+	}
+
+	UE_LOG(LogDisplayClusterMediaSync, Verbose, TEXT("'%s': Instantiated barrier client '%s'"), *MediaId, *ClientId);
+
+	if (!EthernetBarrierClient->Connect())
+	{
+		UE_LOG(LogDisplayClusterMediaSync, Warning, TEXT("'%s': Barrier client '%s' couldn't establish connection"), *MediaId, *ClientId);
+		return false;
+	}
+
+	BarrierId     = GenerateBarrierName();
+	ThreadMarker  = MediaId;
+
+	// Thread markers
+	TArray<FString> ThreadMarkers;
+	GenerateListOfThreadMarkers(ThreadMarkers);
+
+	// Create sync barrier
+	if (!EthernetBarrierClient->CreateBarrier(BarrierId, ThreadMarkers, BarrierTimeoutMs))
+	{
+		UE_LOG(LogDisplayClusterMediaSync, Warning, TEXT("'%s': Barrier client '%s' couldn't create barrier '%s'"), *MediaId, *ClientId, *BarrierId);
+		EthernetBarrierClient.Reset();
+		return false;
+	}
+
+	return true;
+}
+
+void UDisplayClusterMediaOutputSynchronizationPolicyEthernetBarrierBase::ReleaseBarrier()
+{
+	if (EthernetBarrierClient)
+	{
+		if (EthernetBarrierClient->IsConnected() && !BarrierId.IsEmpty())
+		{
+			EthernetBarrierClient->ReleaseBarrier(BarrierId);
+		}
+
+		EthernetBarrierClient.Reset();
+	}
+}
+
+FString UDisplayClusterMediaOutputSynchronizationPolicyEthernetBarrierBase::GenerateBarrierName() const
+{
+	// Currently we don't have any synchronization groups. This means all the sync policy instances of the same
+	// class use the same barrier. If we want to introduce sync groups in the future, the barrier ID should
+	// take that group ID/number into account, and encode it into the barrier name.
+	//
+	// For example, we want two sets of capture devices to run with different output framerate. In this case, we would
+	// need to split those sets into different sync groups.
+	//
+	// However! All media captures are locked to UE rendering pipeline. This means all the captures will run
+	// with the same framerate. Therefore we don't need any sync groups so far.
+	return GetClass()->GetName();
+}
+
+void UDisplayClusterMediaOutputSynchronizationPolicyEthernetBarrierBase::GenerateListOfThreadMarkers(TArray<FString>& OutMarkers) const
+{
+	OutMarkers.Empty();
+
+	UE_LOG(LogDisplayClusterMediaSync, Verbose, TEXT("'%s': Generating thread markers for barrier '%s'..."), *GetMediaDeviceId(), *BarrierId);
+
+	// Get active DCRA
+	if (const ADisplayClusterRootActor* const RootActor = IDisplayCluster::Get().GetGameMgr()->GetRootActor())
+	{
+		// Get config data
+		if (const UDisplayClusterConfigurationData* const CfgData = RootActor->GetConfigData())
+		{
+			// Iterate over cluster nodes
+			for (const TPair<FString, TObjectPtr<UDisplayClusterConfigurationClusterNode>>& NodeIt : CfgData->Cluster->Nodes)
+			{
+				/////////////////////
+				// Backbuffer capture
+				{
+					const FDisplayClusterConfigurationMedia& MediaSettings = NodeIt.Value->Media;
+					if (MediaSettings.bEnable && MediaSettings.MediaOutput && MediaSettings.OutputSyncPolicy)
+					{
+						// Pick the same sync policy only
+						if (MediaSettings.OutputSyncPolicy->GetClass() == GetClass())
+						{
+							const FString BackbufferCaptureId = DisplayClusterMediaHelpers::MediaId::GenerateMediaId(
+								DisplayClusterMediaHelpers::MediaId::EMediaDeviceType::Output,
+								DisplayClusterMediaHelpers::MediaId::EMediaOwnerType::Backbuffer,
+								NodeIt.Key, RootActor->GetName(), FString());
+
+							OutMarkers.Add(BackbufferCaptureId);
+						}
+					}
+				}
+
+				///////////////////
+				// Viewport capture
+				{
+					// Iterate over viewports
+					for (const TPair<FString, TObjectPtr<UDisplayClusterConfigurationViewport>>& ViewportIt : NodeIt.Value->Viewports)
+					{
+						const FDisplayClusterConfigurationMedia& MediaSettings = ViewportIt.Value->RenderSettings.Media;
+						if (MediaSettings.bEnable && MediaSettings.MediaOutput && MediaSettings.OutputSyncPolicy)
+						{
+							// Pick the same sync policy only
+							if (MediaSettings.OutputSyncPolicy->GetClass() == GetClass())
+							{
+								const FString ViewportCaptureId = DisplayClusterMediaHelpers::MediaId::GenerateMediaId(
+									DisplayClusterMediaHelpers::MediaId::EMediaDeviceType::Output,
+									DisplayClusterMediaHelpers::MediaId::EMediaOwnerType::Viewport,
+									NodeIt.Key, RootActor->GetName(), ViewportIt.Key);
+
+								OutMarkers.Add(ViewportCaptureId);
+							}
+						}
+					}
+				}
+			}
+		}
+
+		////////////////
+		// ICVFX capture
+		{
+			// Get all ICVFX camera components
+			TArray<UDisplayClusterICVFXCameraComponent*> ICVFXCameraComponents;
+			RootActor->GetComponents(ICVFXCameraComponents);
+
+			// Iterate over ICVFX cameras
+			for (UDisplayClusterICVFXCameraComponent* const ICVFXCameraComponent : ICVFXCameraComponents)
+			{
+				const FDisplayClusterConfigurationMediaICVFX& MediaSettings = ICVFXCameraComponent->CameraSettings.RenderSettings.Media;
+				if (MediaSettings.bEnable)
+				{
+					for (const FDisplayClusterConfigurationMediaOutputGroup& MediaOutputGroup : MediaSettings.MediaOutputGroups)
+					{
+						// Pick the same sync policy only
+						if (MediaOutputGroup.OutputSyncPolicy->GetClass() == GetClass())
+						{
+							for (const FString& NodeId : MediaOutputGroup.ClusterNodes.ItemNames)
+							{
+								const FString ICVFXCaptureId = DisplayClusterMediaHelpers::MediaId::GenerateMediaId(
+									DisplayClusterMediaHelpers::MediaId::EMediaDeviceType::Output,
+									DisplayClusterMediaHelpers::MediaId::EMediaOwnerType::ICVFXCamera,
+									NodeId, RootActor->GetName(), ICVFXCameraComponent->GetName());
+
+								OutMarkers.Add(ICVFXCaptureId);
+							}
+						}
+					}
+				}
+			}
+		}
+	}
+
+	UE_LOG(LogDisplayClusterMediaSync, Verbose, TEXT("'%s': Generated %d thread markers for barrier '%s' "), *GetMediaDeviceId(), OutMarkers.Num(), *BarrierId);
+	for (int32 Idx = 0; Idx < OutMarkers.Num(); ++Idx)
+	{
+		UE_LOG(LogDisplayClusterMediaSync, Verbose, TEXT("'%s': Barrier '%s', marker %d: %s"), *GetMediaDeviceId(), *BarrierId, Idx, *OutMarkers[Idx]);
+	}
+}
+
+void UDisplayClusterMediaOutputSynchronizationPolicyEthernetBarrierBase::ProcessMediaSynchronizationCallback()
+{
+	UE_LOG(LogDisplayClusterMediaSync, VeryVerbose, TEXT("'%s': Synchronizing capture..."), *GetMediaDeviceId());
+
+	// Pass to the policy implementations
+	Synchronize();
+}
diff --git a/Engine/Plugins/Runtime/nDisplay/Source/DisplayClusterMedia/Private/Synchronization/SyncPolicies/DisplayClusterMediaOutputSynchronizationPolicyRivermax.cpp b/Engine/Plugins/Runtime/nDisplay/Source/DisplayClusterMedia/Private/Synchronization/SyncPolicies/DisplayClusterMediaOutputSynchronizationPolicyRivermax.cpp
new file mode 100644
--- /dev/null
+++ b/Engine/Plugins/Runtime/nDisplay/Source/DisplayClusterMedia/Private/Synchronization/SyncPolicies/DisplayClusterMediaOutputSynchronizationPolicyRivermax.cpp
@@ -0,0 +1,70 @@
+// Copyright Epic Games, Inc. All Rights Reserved.
+
+#include "Synchronization/DisplayClusterMediaOutputSynchronizationPolicyRivermax.h"
+
+#include "IRivermaxCoreModule.h"
+#include "IRivermaxManager.h"
+#include "RivermaxMediaCapture.h"
+#include "RivermaxPTPUtils.h"
+
+
+bool UDisplayClusterMediaOutputSynchronizationPolicyRivermax::IsCaptureTypeSupported(UMediaCapture* MediaCapture) const
+{
+	// We need to make sure:
+	// - it's RivermaxCapture
+	// - it uses PTP or System time source
+	// - it uses AlignmentPoint alignment mode
+	if (URivermaxMediaCapture* RmaxCapture = Cast<URivermaxMediaCapture>(MediaCapture))
+	{
+		using namespace UE::RivermaxCore;
+
+		if (TSharedPtr<IRivermaxManager> RivermaxMgr = IRivermaxCoreModule::Get().GetRivermaxManager())
+		{
+			const ERivermaxTimeSource TimeSource = RivermaxMgr->GetTimeSource();
+
+			if (TimeSource == ERivermaxTimeSource::PTP || TimeSource == ERivermaxTimeSource::System)
+			{
+				const FRivermaxOutputStreamOptions Options = RmaxCapture->GetOutputStreamOptions();
+
+				if (Options.AlignmentMode == ERivermaxAlignmentMode::AlignmentPoint)
+				{
+					return true;
+				}
+			}
+		}
+	}
+
+	return false;
+}
+
+double UDisplayClusterMediaOutputSynchronizationPolicyRivermax::GetTimeBeforeNextSyncPoint()
+{
+	if (URivermaxMediaCapture* RmaxCapture = Cast<URivermaxMediaCapture>(CapturingDevice))
+	{
+		if (RmaxCapture->GetState() == EMediaCaptureState::Capturing)
+		{
+			using namespace UE::RivermaxCore;
+
+			if (TSharedPtr<IRivermaxManager> RivermaxMgr = IRivermaxCoreModule::Get().GetRivermaxManager())
+			{
+				// Get current time
+				const uint64 CurrentTimeNanosec = RivermaxMgr->GetTime();
+
+				// Get next alignment timepoint
+				FRivermaxOutputStreamOptions Options = RmaxCapture->GetOutputStreamOptions();
+				const uint64 NextAlignmentTimeNanosec = GetNextAlignmentPoint(CurrentTimeNanosec, Options.FrameRate);
+
+				// Time left
+				checkSlow(NextAlignmentTimeNanosec > CurrentTimeNanosec);
+				const uint64 TimeLeftNanosec = NextAlignmentTimeNanosec - CurrentTimeNanosec;
+
+				// Return remaining time in seconds
+				return double(TimeLeftNanosec * 1E9);
+			}
+		}
+	}
+
+	// Normally we should never get here. As a fallback approach, return some big time interval
+	// to prevent calling thread blocking. 1 second is more than any possible threshold.
+	return 1.f;
+}
diff --git a/Engine/Plugins/Runtime/nDisplay/Source/DisplayClusterMedia/Private/Synchronization/SyncPolicies/DisplayClusterMediaOutputSynchronizationPolicyThresholdBase.cpp b/Engine/Plugins/Runtime/nDisplay/Source/DisplayClusterMedia/Private/Synchronization/SyncPolicies/DisplayClusterMediaOutputSynchronizationPolicyThresholdBase.cpp
new file mode 100644
--- /dev/null
+++ b/Engine/Plugins/Runtime/nDisplay/Source/DisplayClusterMedia/Private/Synchronization/SyncPolicies/DisplayClusterMediaOutputSynchronizationPolicyThresholdBase.cpp
@@ -0,0 +1,31 @@
+// Copyright Epic Games, Inc. All Rights Reserved.
+
+#include "Synchronization/DisplayClusterMediaOutputSynchronizationPolicyThresholdBase.h"
+
+#include "DisplayClusterMediaLog.h"
+
+
+void UDisplayClusterMediaOutputSynchronizationPolicyThresholdBase::Synchronize()
+{
+	// Ask the sync implementation about how much time we have before next synchronization timepoint
+	const double TimeLeftSeconds = GetTimeBeforeNextSyncPoint();
+	// Convert to seconds
+	const double MarginSeconds = double(MarginMs) / 1000;
+
+	// In case we're unsafe, skip the upcoming sync timepoint
+	if (TimeLeftSeconds < MarginSeconds)
+	{
+		// Sleep for a bit longer to skip the alignment timepoint
+		const float SleepTime = (MarginSeconds - TimeLeftSeconds) * 1.01f;
+
+		UE_LOG(LogDisplayClusterMediaSync, VeryVerbose, TEXT("'%s': TimeLeft(%lf) < Margin(%lf) --> Sleeping for %lf..."),
+			*GetMediaDeviceId(), TimeLeftSeconds, MarginSeconds, SleepTime);
+
+		FPlatformProcess::SleepNoStats(SleepTime);
+	}
+
+	// Now make sure all the outputs have skipped the sync point, then continue output
+	SyncThreadOnBarrier();
+
+	// Leaving this function means the capture thread can continue sending the frame out
+}
diff --git a/Engine/Plugins/Runtime/nDisplay/Source/DisplayClusterMedia/Private/Synchronization/SyncPolicies/DisplayClusterMediaOutputSynchronizationPolicyVblank.cpp b/Engine/Plugins/Runtime/nDisplay/Source/DisplayClusterMedia/Private/Synchronization/SyncPolicies/DisplayClusterMediaOutputSynchronizationPolicyVblank.cpp
new file mode 100644
--- /dev/null
+++ b/Engine/Plugins/Runtime/nDisplay/Source/DisplayClusterMedia/Private/Synchronization/SyncPolicies/DisplayClusterMediaOutputSynchronizationPolicyVblank.cpp
@@ -0,0 +1,44 @@
+// Copyright Epic Games, Inc. All Rights Reserved.
+
+#include "Synchronization/DisplayClusterMediaOutputSynchronizationPolicyVblank.h"
+
+#include "IDisplayCluster.h"
+#include "Render/IDisplayClusterRenderManager.h"
+#include "Render/Monitoring/IDisplayClusterVblankMonitor.h"
+
+#include "DisplayClusterMediaLog.h"
+
+
+bool UDisplayClusterMediaOutputSynchronizationPolicyVblank::StartSynchronization(UMediaCapture* MediaCapture, const FString& MediaId)
+{
+	// Get monitor interface if not yet available
+	if (!VblankMonitor)
+	{
+		if (IDisplayClusterRenderManager* RenderMgr = IDisplayCluster::Get().GetRenderMgr())
+		{
+			VblankMonitor = RenderMgr->GetVblankMonitor();
+		}
+	}
+
+	// Start monitoring if not running already
+	if (!VblankMonitor->IsMonitoring())
+	{
+		VblankMonitor->StartMonitoring();
+	}
+
+	return Super::StartSynchronization(MediaCapture, MediaId);
+}
+
+double UDisplayClusterMediaOutputSynchronizationPolicyVblank::GetTimeBeforeNextSyncPoint()
+{
+	// Compute time remaining
+	if(VblankMonitor->IsVblankTimeDataAvailable())
+	{
+		const double TimeRemaining = VblankMonitor->GetNextVblankTime() - FPlatformTime::Seconds();
+		UE_LOG(LogDisplayClusterMediaSync, VeryVerbose, TEXT("V-blank media sync: time before V-blank - %lf"), TimeRemaining);
+		return TimeRemaining;
+	}
+
+	// Return 1sec (a huge duration) to prevent calling thread blocking
+	return 1;
+}
diff --git a/Engine/Plugins/Runtime/nDisplay/Source/DisplayClusterMedia/Public/Synchronization/DisplayClusterMediaOutputSynchronizationPolicyEthernetBarrier.h b/Engine/Plugins/Runtime/nDisplay/Source/DisplayClusterMedia/Public/Synchronization/DisplayClusterMediaOutputSynchronizationPolicyEthernetBarrier.h
new file mode 100644
--- /dev/null
+++ b/Engine/Plugins/Runtime/nDisplay/Source/DisplayClusterMedia/Public/Synchronization/DisplayClusterMediaOutputSynchronizationPolicyEthernetBarrier.h
@@ -0,0 +1,23 @@
+// Copyright Epic Games, Inc. All Rights Reserved.
+
+#pragma once
+
+#include "CoreMinimal.h"
+#include "Synchronization/DisplayClusterMediaOutputSynchronizationPolicyEthernetBarrierBase.h"
+
+#include "DisplayClusterMediaOutputSynchronizationPolicyEthernetBarrier.generated.h"
+
+
+/*
+ * EthernetBarrier media synchronization policy implementation
+ */
+UCLASS(editinlinenew, Blueprintable, meta = (DisplayName = "Ethernet Barrier"))
+class DISPLAYCLUSTERMEDIA_API UDisplayClusterMediaOutputSynchronizationPolicyEthernetBarrier
+	: public UDisplayClusterMediaOutputSynchronizationPolicyEthernetBarrierBase
+{
+	GENERATED_BODY()
+
+protected:
+	/** Synchronization procedure implementation. */
+	virtual void Synchronize() override;
+};
diff --git a/Engine/Plugins/Runtime/nDisplay/Source/DisplayClusterMedia/Public/Synchronization/DisplayClusterMediaOutputSynchronizationPolicyEthernetBarrierBase.h b/Engine/Plugins/Runtime/nDisplay/Source/DisplayClusterMedia/Public/Synchronization/DisplayClusterMediaOutputSynchronizationPolicyEthernetBarrierBase.h
new file mode 100644
--- /dev/null
+++ b/Engine/Plugins/Runtime/nDisplay/Source/DisplayClusterMedia/Public/Synchronization/DisplayClusterMediaOutputSynchronizationPolicyEthernetBarrierBase.h
@@ -0,0 +1,85 @@
+// Copyright Epic Games, Inc. All Rights Reserved.
+
+#pragma once
+
+#include "CoreMinimal.h"
+#include "DisplayClusterConfigurationTypes_Media.h"
+
+#include "Cluster/IDisplayClusterGenericBarriersClient.h"
+
+#include "DisplayClusterMediaOutputSynchronizationPolicyEthernetBarrierBase.generated.h"
+
+
+/*
+ * Base class for Ethernet barrier based media synchronization policies.
+ * 
+ * It encapsulates network barriers related logic.
+ */
+UCLASS(Abstract, editinlinenew, BlueprintType, hidecategories = (Object))
+class DISPLAYCLUSTERMEDIA_API UDisplayClusterMediaOutputSynchronizationPolicyEthernetBarrierBase
+	: public UDisplayClusterMediaOutputSynchronizationPolicy
+{
+	GENERATED_BODY()
+
+public:
+	/** Starts synchronization of specific capture device. Returns false if failed. */
+	virtual bool StartSynchronization(UMediaCapture* MediaCapture, const FString& MediaId) override;
+
+	/** Stops synchronization of specific output stream (capture device). */
+	virtual void StopSynchronization() override;
+
+	/** Returns true if currently synchronising a media output. */
+	virtual bool IsRunning() override final;
+
+protected:
+	/** Synchronizes calling thread at the barrier. */
+	void SyncThreadOnBarrier();
+
+	/** Children implement their own sync approaches. */
+	virtual void Synchronize() PURE_VIRTUAL(UDisplayClusterMediaOutputSynchronizationPolicyEthernetBarrierBase::Synchronize, )
+
+protected:
+	/** Capture device being used. */
+	UPROPERTY(Transient)
+	TObjectPtr<UMediaCapture> CapturingDevice;
+
+	/** Barrier timeout (ms) */
+	UPROPERTY(EditAnywhere, BlueprintReadWrite, Category = "Settings", meta = (DisplayName = "Barrier Timeout (ms)", ClampMin = "1", ClampMax = "10000", UIMin = "1", UIMax = "10000"))
+	int32 BarrierTimeoutMs = 3000;
+
+protected:
+	/** Returns media device ID being synchronized */
+	FString GetMediaDeviceId() const;
+
+private:
+	/** Initializes dynamic barrier on the primary node. */
+	bool InitializeBarrier(const FString& SyncInstanceId);
+
+	/** Releases dynamic barrier on the primary node. */
+	void ReleaseBarrier();
+
+	/** Generates name of the dynamic barrier. */
+	FString GenerateBarrierName() const;
+
+	/** Generates array of thread markers that are going to use a barrier. */
+	void GenerateListOfThreadMarkers(TArray<FString>& OutMarkers) const;
+
+	/** Handles media capture sync callbacks. */
+	void ProcessMediaSynchronizationCallback();
+
+private:
+	/** Is synchronization currently active. */
+	bool bIsRunning = false;
+
+	/** ID of media device being synchronized. */
+	FString MediaDeviceId;
+
+	/** Unique barrier name to use. */
+	FString BarrierId;
+
+	/** Unique thread (caller) marker to be used on the barrier. */
+	FString ThreadMarker;
+
+	/** Barrier sync client. */
+	TUniquePtr<IDisplayClusterGenericBarriersClient> EthernetBarrierClient;
+};
diff --git a/Engine/Plugins/Runtime/nDisplay/Source/DisplayClusterMedia/Public/Synchronization/DisplayClusterMediaOutputSynchronizationPolicyRivermax.h b/Engine/Plugins/Runtime/nDisplay/Source/DisplayClusterMedia/Public/Synchronization/DisplayClusterMediaOutputSynchronizationPolicyRivermax.h
new file mode 100644
--- /dev/null
+++ b/Engine/Plugins/Runtime/nDisplay/Source/DisplayClusterMedia/Public/Synchronization/DisplayClusterMediaOutputSynchronizationPolicyRivermax.h
@@ -0,0 +1,27 @@
+// Copyright Epic Games, Inc. All Rights Reserved.
+
+#pragma once
+
+#include "CoreMinimal.h"
+#include "Synchronization/DisplayClusterMediaOutputSynchronizationPolicyThresholdBase.h"
+
+#include "DisplayClusterMediaOutputSynchronizationPolicyRivermax.generated.h"
+
+
+/*
+ * Rivermax media synchronization policy implementation
+ */
+UCLASS(editinlinenew, Blueprintable, meta = (DisplayName = "Rivermax (PTP)"))
+class DISPLAYCLUSTERMEDIA_API UDisplayClusterMediaOutputSynchronizationPolicyRivermax
+	: public UDisplayClusterMediaOutputSynchronizationPolicyThresholdBase
+{
+	GENERATED_BODY()
+
+public:
+	/** Returns true if specified media capture type can be synchonized by the policy implementation */
+	virtual bool IsCaptureTypeSupported(UMediaCapture* MediaCapture) const override;
+
+protected:
+	/** Returns amount of time before next synchronization point. */
+	virtual double GetTimeBeforeNextSyncPoint() override;
+};
diff --git a/Engine/Plugins/Runtime/nDisplay/Source/DisplayClusterMedia/Public/Synchronization/DisplayClusterMediaOutputSynchronizationPolicyThresholdBase.h b/Engine/Plugins/Runtime/nDisplay/Source/DisplayClusterMedia/Public/Synchronization/DisplayClusterMediaOutputSynchronizationPolicyThresholdBase.h
new file mode 100644
--- /dev/null
+++ b/Engine/Plugins/Runtime/nDisplay/Source/DisplayClusterMedia/Public/Synchronization/DisplayClusterMediaOutputSynchronizationPolicyThresholdBase.h
@@ -0,0 +1,33 @@
+// Copyright Epic Games, Inc. All Rights Reserved.
+
+#pragma once
+
+#include "CoreMinimal.h"
+#include "Synchronization/DisplayClusterMediaOutputSynchronizationPolicyEthernetBarrierBase.h"
+
+#include "DisplayClusterMediaOutputSynchronizationPolicyThresholdBase.generated.h"
+
+
+/*
+ * Base class for threshold based media synchronization policies.
+ * 
+ * Basically it uses the same approach that we use in 'Ethernet' sync policy where v-blanks are used as the timepoints.
+ */
+UCLASS(Abstract, editinlinenew, BlueprintType, hidecategories = (Object))
+class DISPLAYCLUSTERMEDIA_API UDisplayClusterMediaOutputSynchronizationPolicyThresholdBase
+	: public UDisplayClusterMediaOutputSynchronizationPolicyEthernetBarrierBase
+{
+	GENERATED_BODY()
+
+protected:
+	/** Synchronization procedure implementation. */
+	virtual void Synchronize() override;
+
+	/** Returns amount of time before next synchronization point. */
+	virtual double GetTimeBeforeNextSyncPoint() PURE_VIRTUAL(UDisplayClusterMediaOutputSynchronizationPolicyThresholdBase::GetTimeBeforeNextSyncPoint, return 0;)
+
+protected:
+	/** Synchronization margin (ms) */
+	UPROPERTY(EditAnywhere, BlueprintReadWrite, Category = "Settings", meta = (DisplayName = "Margin (ms)", ClampMin = "1", ClampMax = "20", UIMin = "1", UIMax = "20"))
+	int32 MarginMs = 5;
+};
diff --git a/Engine/Plugins/Runtime/nDisplay/Source/DisplayClusterMedia/Public/Synchronization/DisplayClusterMediaOutputSynchronizationPolicyVblank.h b/Engine/Plugins/Runtime/nDisplay/Source/DisplayClusterMedia/Public/Synchronization/DisplayClusterMediaOutputSynchronizationPolicyVblank.h
new file mode 100644
--- /dev/null
+++ b/Engine/Plugins/Runtime/nDisplay/Source/DisplayClusterMedia/Public/Synchronization/DisplayClusterMediaOutputSynchronizationPolicyVblank.h
@@ -0,0 +1,32 @@
+// Copyright Epic Games, Inc. All Rights Reserved.
+
+#pragma once
+
+#include "CoreMinimal.h"
+#include "Synchronization/DisplayClusterMediaOutputSynchronizationPolicyThresholdBase.h"
+
+#include "DisplayClusterMediaOutputSynchronizationPolicyVblank.generated.h"
+
+class IDisplayClusterVblankMonitor;
+
+
+/*
+ * Vblank media synchronization policy implementation
+ */
+UCLASS(editinlinenew, Blueprintable, meta = (DisplayName = "V-blank"))
+class DISPLAYCLUSTERMEDIA_API UDisplayClusterMediaOutputSynchronizationPolicyVblank
+	: public UDisplayClusterMediaOutputSynchronizationPolicyThresholdBase
+{
+	GENERATED_BODY()
+
+public:
+	virtual bool StartSynchronization(UMediaCapture* MediaCapture, const FString& MediaId) override;
+
+protected:
+	/** Returns amount of time before next synchronization point. */
+	virtual double GetTimeBeforeNextSyncPoint() override;
+
+private:
+	// V-blank monitor
+	TSharedPtr<IDisplayClusterVblankMonitor, ESPMode::ThreadSafe> VblankMonitor;
+};
diff --git a/Engine/Plugins/Runtime/nDisplay/Source/DisplayClusterMediaEditor/Private/DisplayClusterMediaOutputSynchronizationPolicyFactories.cpp b/Engine/Plugins/Runtime/nDisplay/Source/DisplayClusterMediaEditor/Private/DisplayClusterMediaOutputSynchronizationPolicyFactories.cpp
new file mode 100644
--- /dev/null
+++ b/Engine/Plugins/Runtime/nDisplay/Source/DisplayClusterMediaEditor/Private/DisplayClusterMediaOutputSynchronizationPolicyFactories.cpp
@@ -0,0 +1,90 @@
+// Copyright Epic Games, Inc. All Rights Reserved.
+
+#include "DisplayClusterMediaOutputSynchronizationPolicyFactories.h"
+
+#include "AssetTypeCategories.h"
+
+#include "Synchronization/DisplayClusterMediaOutputSynchronizationPolicyEthernetBarrier.h"
+#include "Synchronization/DisplayClusterMediaOutputSynchronizationPolicyRivermax.h"
+#include "Synchronization/DisplayClusterMediaOutputSynchronizationPolicyVblank.h"
+
+
+/**
+ * UDisplayClusterMediaOutputSynchronizationPolicyEthernetBarrier
+ */
+UDisplayClusterMediaOutputSynchronizationPolicyEthernetBarrierFactory::UDisplayClusterMediaOutputSynchronizationPolicyEthernetBarrierFactory(const FObjectInitializer& ObjectInitializer)
+	: Super(ObjectInitializer)
+{
+	SupportedClass = UDisplayClusterMediaOutputSynchronizationPolicyEthernetBarrier::StaticClass();
+	bCreateNew = true;
+	bEditAfterNew = true;
+}
+
+UObject* UDisplayClusterMediaOutputSynchronizationPolicyEthernetBarrierFactory::FactoryCreateNew(UClass* InClass, UObject* InParent, FName InName, EObjectFlags Flags, UObject* Context, FFeedbackContext* Warn)
+{
+	return NewObject<UDisplayClusterMediaOutputSynchronizationPolicyEthernetBarrier>(InParent, InClass, InName, Flags);
+}
+
+uint32 UDisplayClusterMediaOutputSynchronizationPolicyEthernetBarrierFactory::GetMenuCategories() const
+{
+	return EAssetTypeCategories::None;
+}
+
+bool UDisplayClusterMediaOutputSynchronizationPolicyEthernetBarrierFactory::ShouldShowInNewMenu() const
+{
+	return true;
+}
+
+
+/**
+ * UDisplayClusterMediaOutputSynchronizationPolicyRivermax
+ */
+UDisplayClusterMediaOutputSynchronizationPolicyRivermaxFactory::UDisplayClusterMediaOutputSynchronizationPolicyRivermaxFactory(const FObjectInitializer& ObjectInitializer)
+	: Super(ObjectInitializer)
+{
+	SupportedClass = UDisplayClusterMediaOutputSynchronizationPolicyRivermax::StaticClass();
+	bCreateNew = true;
+	bEditAfterNew = true;
+}
+
+UObject* UDisplayClusterMediaOutputSynchronizationPolicyRivermaxFactory::FactoryCreateNew(UClass* InClass, UObject* InParent, FName InName, EObjectFlags Flags, UObject* Context, FFeedbackContext* Warn)
+{
+	return NewObject<UDisplayClusterMediaOutputSynchronizationPolicyRivermax>(InParent, InClass, InName, Flags);
+}
+
+uint32 UDisplayClusterMediaOutputSynchronizationPolicyRivermaxFactory::GetMenuCategories() const
+{
+	return EAssetTypeCategories::None;
+}
+
+bool UDisplayClusterMediaOutputSynchronizationPolicyRivermaxFactory::ShouldShowInNewMenu() const
+{
+	return true;
+}
+
+
+/**
+ * UDisplayClusterMediaOutputSynchronizationPolicyVblank
+ */
+UDisplayClusterMediaOutputSynchronizationPolicyVblankFactory::UDisplayClusterMediaOutputSynchronizationPolicyVblankFactory(const FObjectInitializer& ObjectInitializer)
+	: Super(ObjectInitializer)
+{
+	SupportedClass = UDisplayClusterMediaOutputSynchronizationPolicyVblank::StaticClass();
+	bCreateNew = true;
+	bEditAfterNew = true;
+}
+
+UObject* UDisplayClusterMediaOutputSynchronizationPolicyVblankFactory::FactoryCreateNew(UClass* InClass, UObject* InParent, FName InName, EObjectFlags Flags, UObject* Context, FFeedbackContext* Warn)
+{
+	return NewObject<UDisplayClusterMediaOutputSynchronizationPolicyVblank>(InParent, InClass, InName, Flags);
+}
+
+uint32 UDisplayClusterMediaOutputSynchronizationPolicyVblankFactory::GetMenuCategories() const
+{
+	return EAssetTypeCategories::None;
+}
+
+bool UDisplayClusterMediaOutputSynchronizationPolicyVblankFactory::ShouldShowInNewMenu() const
+{
+	return true;
+}
diff --git a/Engine/Plugins/Runtime/nDisplay/Source/DisplayClusterMediaEditor/Private/DisplayClusterMediaOutputSynchronizationPolicyFactories.h b/Engine/Plugins/Runtime/nDisplay/Source/DisplayClusterMediaEditor/Private/DisplayClusterMediaOutputSynchronizationPolicyFactories.h
new file mode 100644
--- /dev/null
+++ b/Engine/Plugins/Runtime/nDisplay/Source/DisplayClusterMediaEditor/Private/DisplayClusterMediaOutputSynchronizationPolicyFactories.h
@@ -0,0 +1,60 @@
+// Copyright Epic Games, Inc. All Rights Reserved.
+
+#pragma once
+
+#include "CoreMinimal.h"
+#include "UObject/ObjectMacros.h"
+#include "Factories/Factory.h"
+
+#include "DisplayClusterMediaOutputSynchronizationPolicyFactories.generated.h"
+
+
+/**
+ * Factory for UDisplayClusterMediaSynchronizationPolicyEthernetBarrier objects.
+ */
+UCLASS(hidecategories=Object)
+class UDisplayClusterMediaOutputSynchronizationPolicyEthernetBarrierFactory
+	: public UFactory
+{
+	GENERATED_UCLASS_BODY()
+
+public:
+	//~ UFactory Interface
+	virtual UObject* FactoryCreateNew(UClass* InClass, UObject* InParent, FName InName, EObjectFlags Flags, UObject* Context, FFeedbackContext* Warn) override;
+	virtual uint32 GetMenuCategories() const override;
+	virtual bool ShouldShowInNewMenu() const override;
+};
+
+
+/**
+ * Factory for UDisplayClusterMediaSynchronizationPolicyRivermax objects.
+ */
+UCLASS(hidecategories = Object)
+class UDisplayClusterMediaOutputSynchronizationPolicyRivermaxFactory
+	: public UFactory
+{
+	GENERATED_UCLASS_BODY()
+
+public:
+	//~ UFactory Interface
+	virtual UObject* FactoryCreateNew(UClass* InClass, UObject* InParent, FName InName, EObjectFlags Flags, UObject* Context, FFeedbackContext* Warn) override;
+	virtual uint32 GetMenuCategories() const override;
+	virtual bool ShouldShowInNewMenu() const override;
+};
+
+
+/**
+ * Factory for UDisplayClusterMediaSynchronizationPolicyVblank objects.
+ */
+UCLASS(hidecategories = Object)
+class UDisplayClusterMediaOutputSynchronizationPolicyVblankFactory
+	: public UFactory
+{
+	GENERATED_UCLASS_BODY()
+
+public:
+	//~ UFactory Interface
+	virtual UObject* FactoryCreateNew(UClass* InClass, UObject* InParent, FName InName, EObjectFlags Flags, UObject* Context, FFeedbackContext* Warn) override;
+	virtual uint32 GetMenuCategories() const override;
+	virtual bool ShouldShowInNewMenu() const override;
+};
diff --git a/Engine/Plugins/VirtualProduction/Rivermax/RivermaxCore/Source/RivermaxCore/Private/RivermaxOutputStream.cpp b/Engine/Plugins/VirtualProduction/Rivermax/RivermaxCore/Source/RivermaxCore/Private/RivermaxOutputStream.cpp
--- a/Engine/Plugins/VirtualProduction/Rivermax/RivermaxCore/Source/RivermaxCore/Private/RivermaxOutputStream.cpp
+++ b/Engine/Plugins/VirtualProduction/Rivermax/RivermaxCore/Source/RivermaxCore/Private/RivermaxOutputStream.cpp
@@ -230,6 +230,9 @@ namespace UE::RivermaxCore::Private
 	{
 		TRACE_CPUPROFILER_EVENT_SCOPE(FRivermaxOutputStream::PushVideoFrame);
 
+		// Notify owner that new frame is about to be enqueued
+		Listener->OnPreFrameEnqueue();
+
 		if (TSharedPtr<FRivermaxOutputFrame> AvailableFrame = GetNextAvailableFrame(NewFrame.FrameIdentifier))
 		{
 			const int32 Stride = GetStride();
@@ -242,7 +245,7 @@ namespace UE::RivermaxCore::Private
 			{
 				const FString TraceName = FString::Format(TEXT("FRivermaxOutputStream::PushFrame {0}"), { AvailableFrame->FrameIndex });
 				TRACE_CPUPROFILER_EVENT_SCOPE_TEXT(*TraceName);
-				
+
 				MarkFrameToBeSent(MoveTemp(AvailableFrame));
 			}
 
@@ -1011,6 +1014,9 @@ namespace UE::RivermaxCore::Private
 
 				FRivermaxOutputStream* Stream = reinterpret_cast<FRivermaxOutputStream*>(userData);
 
+				// Notify owner that new frame is about to be enqueued
+				Stream->Listener->OnPreFrameEnqueue();
+
 				TOptional<uint32> TargetFrameIdentifier = Stream->PendingIdentifiers.Dequeue();
 				if (TargetFrameIdentifier.IsSet())
 				{
diff --git a/Engine/Plugins/VirtualProduction/Rivermax/RivermaxCore/Source/RivermaxCore/Public/IRivermaxOutputStream.h b/Engine/Plugins/VirtualProduction/Rivermax/RivermaxCore/Source/RivermaxCore/Public/IRivermaxOutputStream.h
--- a/Engine/Plugins/VirtualProduction/Rivermax/RivermaxCore/Source/RivermaxCore/Public/IRivermaxOutputStream.h
+++ b/Engine/Plugins/VirtualProduction/Rivermax/RivermaxCore/Source/RivermaxCore/Public/IRivermaxOutputStream.h
@@ -29,6 +29,9 @@ namespace UE::RivermaxCore
 
 		/** Called when stream has encountered an error and has to stop */
 		virtual void OnStreamError() = 0;
+
+		/** Called when stream is about to enqueue new frame */
+		virtual void OnPreFrameEnqueue() = 0;
 	};
 
 	class RIVERMAXCORE_API IRivermaxOutputStream
diff --git a/Engine/Plugins/VirtualProduction/Rivermax/RivermaxMedia/Source/RivermaxMedia/Private/RivermaxMediaCapture.cpp b/Engine/Plugins/VirtualProduction/Rivermax/RivermaxMedia/Source/RivermaxMedia/Private/RivermaxMediaCapture.cpp
--- a/Engine/Plugins/VirtualProduction/Rivermax/RivermaxMedia/Source/RivermaxMedia/Private/RivermaxMediaCapture.cpp
+++ b/Engine/Plugins/VirtualProduction/Rivermax/RivermaxMedia/Source/RivermaxMedia/Private/RivermaxMediaCapture.cpp
@@ -144,6 +144,11 @@ namespace UE::RivermaxMediaCaptureUtil
 ///* URivermaxMediaCapture implementation
 //*****************************************************************************/
 
+UE::RivermaxCore::FRivermaxOutputStreamOptions URivermaxMediaCapture::GetOutputStreamOptions() const
+{
+	return Options;
+}
+
 bool URivermaxMediaCapture::ValidateMediaOutput() const
 {
 	URivermaxMediaOutput* RivermaxMediaOutput = Cast<URivermaxMediaOutput>(MediaOutput);
@@ -570,3 +575,11 @@ void URivermaxMediaCapture::OnStreamError()
 		SetState(EMediaCaptureState::Error);
 	}
 }
+
+void URivermaxMediaCapture::OnPreFrameEnqueue()
+{
+	if (IsOutputSynchronizationSupported())
+	{
+		OnOutputSynchronization.ExecuteIfBound();
+	}
+}
diff --git a/Engine/Plugins/VirtualProduction/Rivermax/RivermaxMedia/Source/RivermaxMedia/Public/RivermaxMediaCapture.h b/Engine/Plugins/VirtualProduction/Rivermax/RivermaxMedia/Source/RivermaxMedia/Public/RivermaxMediaCapture.h
--- a/Engine/Plugins/VirtualProduction/Rivermax/RivermaxMedia/Source/RivermaxMedia/Public/RivermaxMediaCapture.h
+++ b/Engine/Plugins/VirtualProduction/Rivermax/RivermaxMedia/Source/RivermaxMedia/Public/RivermaxMediaCapture.h
@@ -22,6 +22,10 @@ class RIVERMAXMEDIA_API URivermaxMediaCapture : public UMediaCapture, public UE:
 {
 	GENERATED_BODY()
 
+public:
+	/** Rivermax capture specific API to provide stream options access */
+	UE::RivermaxCore::FRivermaxOutputStreamOptions GetOutputStreamOptions() const;
+
 public:
 
 	//~ Begin UMediaCapture interface
@@ -40,6 +44,7 @@ protected:
 	virtual EMediaCaptureResourceType GetCustomOutputResourceType() const override;
 	virtual FRDGBufferDesc GetCustomBufferDescription(const FIntPoint& InDesiredSize) const override;
 	virtual void OnCustomCapture_RenderingThread(FRDGBuilder& GraphBuilder, const FCaptureBaseData& InBaseData, TSharedPtr<FMediaCaptureUserData, ESPMode::ThreadSafe> InUserData, FRDGTextureRef InSourceTexture, FRDGBufferRef OutputBuffer, const FRHICopyTextureInfo& CopyInfo, FVector2D CropU, FVector2D CropV) override;
+	virtual bool IsOutputSynchronizationSupported() const override { return true; }
 	//~ End UMediaCapture interface
 
 	//~ Begin UObject interface
@@ -49,6 +54,7 @@ protected:
 	//~ Begin IRivermaxOutputStreamListener interface
 	virtual void OnInitializationCompleted(bool bHasSucceed) override;
 	virtual void OnStreamError() override;
+	virtual void OnPreFrameEnqueue() override;
 	//~ End IRivermaxOutputStreamListener interface
 
 private:
diff --git a/Engine/Plugins/VirtualProduction/Rivermax/RivermaxMedia/Source/RivermaxMedia/RivermaxMedia.build.cs b/Engine/Plugins/VirtualProduction/Rivermax/RivermaxMedia/Source/RivermaxMedia/RivermaxMedia.build.cs
--- a/Engine/Plugins/VirtualProduction/Rivermax/RivermaxMedia/Source/RivermaxMedia/RivermaxMedia.build.cs
+++ b/Engine/Plugins/VirtualProduction/Rivermax/RivermaxMedia/Source/RivermaxMedia/RivermaxMedia.build.cs
@@ -12,6 +12,7 @@ public class RivermaxMedia : ModuleRules
 				"CoreUObject",
 				"MediaAssets",
 				"MediaIOCore",
+				"RivermaxCore",
 				"TimeManagement"
 			}
 		);
