[Backout] - CL24465558
Original CL Desc
-----------------------------------------------------------------
[nDisplay][MediaIO][Rivermax] Media capture synchronization

diff --git a/Engine/Plugins/Media/MediaIOFramework/Source/MediaIOCore/Public/MediaCapture.h b/Engine/Plugins/Media/MediaIOFramework/Source/MediaIOCore/Public/MediaCapture.h
--- a/Engine/Plugins/Media/MediaIOFramework/Source/MediaIOCore/Public/MediaCapture.h
+++ b/Engine/Plugins/Media/MediaIOFramework/Source/MediaIOCore/Public/MediaCapture.h
@@ -211,7 +211,6 @@ public:
 /** Delegate signatures */
 DECLARE_DYNAMIC_MULTICAST_DELEGATE(FMediaCaptureStateChangedSignature);
 DECLARE_MULTICAST_DELEGATE(FMediaCaptureStateChangedSignatureNative);
-DECLARE_DELEGATE(FMediaCaptureOutputSynchronization)
 
 /**
  * Abstract base class for media capture.
@@ -335,20 +334,6 @@ public:
 	 */
 	FMediaCaptureStateChangedSignatureNative OnStateChangedNative;
 
-	/**
-	 * Returns true if media capture implementation supports synchronization logic. If so,
-	 * OnOutputSynchronization delegate below must be called before each captured frame is exported.
-	 */
-	virtual bool IsOutputSynchronizationSupported() const { return false; }
-
-	/**
-	 * Called when output data is ready to be sent out. This blocking call allows to
-	 * postpone data export based on the synchronization logic.
-	 *
-	 * This delegate must be called if IsOutputSynchronizationSupported() returns true.
-	 */
-	FMediaCaptureOutputSynchronization OnOutputSynchronization;
-
 	void SetFixedViewportSize(TSharedPtr<FSceneViewport> InSceneViewport, FIntPoint InSize);
 	void ResetFixedViewportSize(TSharedPtr<FSceneViewport> InViewport, bool bInFlushRenderingCommands);
 
diff --git a/Engine/Plugins/Runtime/nDisplay/Source/DisplayCluster/Private/Cluster/Controller/DisplayClusterClusterNodeCtrlPrimary.cpp b/Engine/Plugins/Runtime/nDisplay/Source/DisplayCluster/Private/Cluster/Controller/DisplayClusterClusterNodeCtrlPrimary.cpp
--- a/Engine/Plugins/Runtime/nDisplay/Source/DisplayCluster/Private/Cluster/Controller/DisplayClusterClusterNodeCtrlPrimary.cpp
+++ b/Engine/Plugins/Runtime/nDisplay/Source/DisplayCluster/Private/Cluster/Controller/DisplayClusterClusterNodeCtrlPrimary.cpp
@@ -12,7 +12,6 @@
 #include "Network/Service/RenderSync/DisplayClusterRenderSyncService.h"
 #include "Network/Service/ClusterEventsJson/DisplayClusterClusterEventsJsonService.h"
 #include "Network/Service/ClusterEventsBinary/DisplayClusterClusterEventsBinaryService.h"
-#include "Network/Service/GenericBarrier/DisplayClusterGenericBarrierService.h"
 #include "Network/Listener/DisplayClusterTcpListener.h"
 #include "Network/DisplayClusterNetworkTypes.h"
 
@@ -116,13 +115,12 @@ bool FDisplayClusterClusterNodeCtrlPrimary::InitializeServers()
 	UE_LOG(LogDisplayClusterCluster, Log, TEXT("%s - initializing primary node servers..."), *GetControllerName());
 
 	// Instantiate node servers
-	ClusterSyncServer         = MakeUnique<FDisplayClusterClusterSyncService>();
-	RenderSyncServer          = MakeUnique<FDisplayClusterRenderSyncService>();
-	GenericBarriersServer     = MakeUnique<FDisplayClusterGenericBarrierService>();
+	ClusterSyncServer         = MakeUnique<FDisplayClusterClusterSyncService>       ();
+	RenderSyncServer          = MakeUnique<FDisplayClusterRenderSyncService>        ();
 	ClusterEventsJsonServer   = MakeUnique<FDisplayClusterClusterEventsJsonService> ();
 	ClusterEventsBinaryServer = MakeUnique<FDisplayClusterClusterEventsBinaryService>();
 
-	return ClusterSyncServer && RenderSyncServer && GenericBarriersServer && ClusterEventsJsonServer && ClusterEventsBinaryServer;
+	return ClusterSyncServer && RenderSyncServer && ClusterEventsJsonServer && ClusterEventsBinaryServer;
 }
 
 bool FDisplayClusterClusterNodeCtrlPrimary::StartServers()
@@ -178,7 +176,6 @@ bool FDisplayClusterClusterNodeCtrlPrimary::StartServers()
 	// Start the servers
 	return StartServerWithLogs(ClusterSyncServer.Get(),         TcpListener) // Start with shared listener
 		&& StartServerWithLogs(RenderSyncServer.Get(),          TcpListener) // Start with shared listener
-		&& StartServerWithLogs(GenericBarriersServer.Get(),     TcpListener) // Start with shared listener
 		&& StartServerWithLogs(ClusterEventsJsonServer.Get(),   CfgPrimary->Host, Ports.ClusterEventsJson)
 		&& StartServerWithLogs(ClusterEventsBinaryServer.Get(), CfgPrimary->Host, Ports.ClusterEventsBinary)
 		&& TcpListener->StartListening(CfgPrimary->Host, Ports.ClusterSync); // Start shared listener as well
@@ -194,7 +191,6 @@ void FDisplayClusterClusterNodeCtrlPrimary::StopServers()
 	// Stop services
 	ClusterSyncServer->Shutdown();
 	RenderSyncServer->Shutdown();
-	GenericBarriersServer->Shutdown();
 	ClusterEventsJsonServer->Shutdown();
 	ClusterEventsBinaryServer->Shutdown();
 
diff --git a/Engine/Plugins/Runtime/nDisplay/Source/DisplayCluster/Private/Cluster/Controller/DisplayClusterClusterNodeCtrlPrimary.h b/Engine/Plugins/Runtime/nDisplay/Source/DisplayCluster/Private/Cluster/Controller/DisplayClusterClusterNodeCtrlPrimary.h
--- a/Engine/Plugins/Runtime/nDisplay/Source/DisplayCluster/Private/Cluster/Controller/DisplayClusterClusterNodeCtrlPrimary.h
+++ b/Engine/Plugins/Runtime/nDisplay/Source/DisplayCluster/Private/Cluster/Controller/DisplayClusterClusterNodeCtrlPrimary.h
@@ -10,7 +10,6 @@ class FDisplayClusterClusterSyncService;
 class FDisplayClusterRenderSyncService;
 class FDisplayClusterClusterEventsJsonService;
 class FDisplayClusterClusterEventsBinaryService;
-class FDisplayClusterGenericBarrierService;
 class FDisplayClusterTcpListener;
 
 
@@ -64,7 +63,6 @@ private:
 	// Node servers
 	TUniquePtr<FDisplayClusterClusterSyncService>         ClusterSyncServer;
 	TUniquePtr<FDisplayClusterRenderSyncService>          RenderSyncServer;
-	TUniquePtr<FDisplayClusterGenericBarrierService>      GenericBarriersServer;
 	TUniquePtr<FDisplayClusterClusterEventsJsonService>   ClusterEventsJsonServer;
 	TUniquePtr<FDisplayClusterClusterEventsBinaryService> ClusterEventsBinaryServer;
 
diff --git a/Engine/Plugins/Runtime/nDisplay/Source/DisplayCluster/Private/Cluster/DisplayClusterClusterManager.cpp b/Engine/Plugins/Runtime/nDisplay/Source/DisplayCluster/Private/Cluster/DisplayClusterClusterManager.cpp
--- a/Engine/Plugins/Runtime/nDisplay/Source/DisplayCluster/Private/Cluster/DisplayClusterClusterManager.cpp
+++ b/Engine/Plugins/Runtime/nDisplay/Source/DisplayCluster/Private/Cluster/DisplayClusterClusterManager.cpp
@@ -14,8 +14,6 @@
 #include "Cluster/Failover/DisplayClusterFailoverNodeCtrlPrimary.h"
 #include "Cluster/Failover/DisplayClusterFailoverNodeCtrlSecondary.h"
 
-#include "Network/Service/GenericBarrier/DisplayClusterGenericBarrierClient.h"
-
 #include "Config/IPDisplayClusterConfigManager.h"
 #include "DisplayClusterConfigurationTypes.h"
 
@@ -308,11 +306,6 @@ void FDisplayClusterClusterManager::UnregisterSyncObject(IDisplayClusterClusterS
 	}
 }
 
-TUniquePtr<IDisplayClusterGenericBarriersClient> FDisplayClusterClusterManager::CreateGenericBarriersClient(const FString& ClientName)
-{
-	return MakeUnique<FDisplayClusterGenericBarrierClient>(ClientName);
-}
-
 void FDisplayClusterClusterManager::AddClusterEventListener(TScriptInterface<IDisplayClusterClusterEventListener> Listener)
 {
 	FScopeLock Lock(&ClusterEventListenersCS);
diff --git a/Engine/Plugins/Runtime/nDisplay/Source/DisplayCluster/Private/Cluster/DisplayClusterClusterManager.h b/Engine/Plugins/Runtime/nDisplay/Source/DisplayCluster/Private/Cluster/DisplayClusterClusterManager.h
--- a/Engine/Plugins/Runtime/nDisplay/Source/DisplayCluster/Private/Cluster/DisplayClusterClusterManager.h
+++ b/Engine/Plugins/Runtime/nDisplay/Source/DisplayCluster/Private/Cluster/DisplayClusterClusterManager.h
@@ -62,8 +62,6 @@ public:
 	virtual void RegisterSyncObject(IDisplayClusterClusterSyncObject* SyncObj, EDisplayClusterSyncGroup SyncGroup) override;
 	virtual void UnregisterSyncObject(IDisplayClusterClusterSyncObject* SyncObj) override;
 
-	virtual TUniquePtr<IDisplayClusterGenericBarriersClient> CreateGenericBarriersClient(const FString& ClientName) override;
-
 	virtual void AddClusterEventListener(TScriptInterface<IDisplayClusterClusterEventListener>) override;
 	virtual void RemoveClusterEventListener(TScriptInterface<IDisplayClusterClusterEventListener>) override;
 
diff --git a/Engine/Plugins/Runtime/nDisplay/Source/DisplayCluster/Private/Misc/DisplayClusterLog.cpp b/Engine/Plugins/Runtime/nDisplay/Source/DisplayCluster/Private/Misc/DisplayClusterLog.cpp
--- a/Engine/Plugins/Runtime/nDisplay/Source/DisplayCluster/Private/Misc/DisplayClusterLog.cpp
+++ b/Engine/Plugins/Runtime/nDisplay/Source/DisplayCluster/Private/Misc/DisplayClusterLog.cpp
@@ -7,7 +7,6 @@ DEFINE_LOG_CATEGORY(LogDisplayClusterEngine);
 DEFINE_LOG_CATEGORY(LogDisplayClusterModule);
 DEFINE_LOG_CATEGORY(LogDisplayClusterCluster);
 DEFINE_LOG_CATEGORY(LogDisplayClusterConfig);
-DEFINE_LOG_CATEGORY(LogDisplayClusterBarrier);
 DEFINE_LOG_CATEGORY(LogDisplayClusterGame);
 DEFINE_LOG_CATEGORY(LogDisplayClusterNetwork);
 DEFINE_LOG_CATEGORY(LogDisplayClusterRender);
diff --git a/Engine/Plugins/Runtime/nDisplay/Source/DisplayCluster/Private/Misc/DisplayClusterLog.h b/Engine/Plugins/Runtime/nDisplay/Source/DisplayCluster/Private/Misc/DisplayClusterLog.h
--- a/Engine/Plugins/Runtime/nDisplay/Source/DisplayCluster/Private/Misc/DisplayClusterLog.h
+++ b/Engine/Plugins/Runtime/nDisplay/Source/DisplayCluster/Private/Misc/DisplayClusterLog.h
@@ -10,7 +10,6 @@ DECLARE_LOG_CATEGORY_EXTERN(LogDisplayClusterEngine,     Log, All);
 DECLARE_LOG_CATEGORY_EXTERN(LogDisplayClusterModule,     Log, All);
 DECLARE_LOG_CATEGORY_EXTERN(LogDisplayClusterCluster,    Log, All);
 DECLARE_LOG_CATEGORY_EXTERN(LogDisplayClusterConfig,     Log, All);
-DECLARE_LOG_CATEGORY_EXTERN(LogDisplayClusterBarrier,    Log, All);
 DECLARE_LOG_CATEGORY_EXTERN(LogDisplayClusterGame,       Log, All);
 DECLARE_LOG_CATEGORY_EXTERN(LogDisplayClusterNetwork,    Log, All);
 DECLARE_LOG_CATEGORY_EXTERN(LogDisplayClusterRender,     Log, All);
diff --git a/Engine/Plugins/Runtime/nDisplay/Source/DisplayCluster/Private/Network/Barrier/DisplayClusterBarrier.cpp b/Engine/Plugins/Runtime/nDisplay/Source/DisplayCluster/Private/Network/Barrier/DisplayClusterBarrier.cpp
deleted file mode 100644
diff --git a/Engine/Plugins/Runtime/nDisplay/Source/DisplayCluster/Private/Network/Barrier/DisplayClusterBarrier.h b/Engine/Plugins/Runtime/nDisplay/Source/DisplayCluster/Private/Network/Barrier/DisplayClusterBarrier.h
deleted file mode 100644
diff --git a/Engine/Plugins/Runtime/nDisplay/Source/DisplayCluster/Private/Network/Barrier/DisplayClusterBarrierFactory.cpp b/Engine/Plugins/Runtime/nDisplay/Source/DisplayCluster/Private/Network/Barrier/DisplayClusterBarrierFactory.cpp
--- a/Engine/Plugins/Runtime/nDisplay/Source/DisplayCluster/Private/Network/Barrier/DisplayClusterBarrierFactory.cpp
+++ b/Engine/Plugins/Runtime/nDisplay/Source/DisplayCluster/Private/Network/Barrier/DisplayClusterBarrierFactory.cpp
@@ -1,13 +1,23 @@
 // Copyright Epic Games, Inc. All Rights Reserved.
 
 #include "Network/Barrier/DisplayClusterBarrierFactory.h"
-#include "Network/Barrier/DisplayClusterBarrier.h"
+#include "Network/Barrier/DisplayClusterBarrierV1.h"
+#include "Network/Barrier/DisplayClusterBarrierV2.h"
 
 #include "Misc/DisplayClusterLog.h"
 
 
-IDisplayClusterBarrier* FDisplayClusterBarrierFactory::CreateBarrier(const FString& BarrierId, const TArray<FString>& ThreadMarkers, const uint32 Timeout)
+TUniquePtr<IDisplayClusterBarrier> FDisplayClusterBarrierFactory::CreateBarrier(const TArray<FString>& ThreadMarkers, const uint32 Timeout, const FString& Name)
 {
-	UE_LOG(LogDisplayClusterBarrier, Log, TEXT("Instantiating unique barrier '%s': Threads=%d, Timeout=%u ms"), *BarrierId, ThreadMarkers.Num(), Timeout);
-	return new FDisplayClusterBarrier(BarrierId, ThreadMarkers, Timeout);
+	// The old barriers v1 are not compatible with failover feature. I left them in the codebase
+	// in case it might be useful in the future.
+#if 0
+	// v1 barrier (old)
+	UE_LOG(LogDisplayClusterNetwork, Log, TEXT("Instantiating barrier '%s' of type 'v1': Threads=%d, Timeout=%u ms"), *Name, ThreadsAmount, Timeout);
+	return MakeUnique<FDisplayClusterBarrierV1>(ThreadMarkers.Num(), Timeout, Name);
+#else
+	// v2 barrier (new)
+	UE_LOG(LogDisplayClusterNetwork, Log, TEXT("Instantiating barrier '%s' of type 'v2': Threads=%d, Timeout=%u ms"), *Name, ThreadMarkers.Num(), Timeout);
+	return MakeUnique<FDisplayClusterBarrierV2>(ThreadMarkers, Timeout, Name);
+#endif
 }
diff --git a/Engine/Plugins/Runtime/nDisplay/Source/DisplayCluster/Private/Network/Barrier/DisplayClusterBarrierFactory.h b/Engine/Plugins/Runtime/nDisplay/Source/DisplayCluster/Private/Network/Barrier/DisplayClusterBarrierFactory.h
--- a/Engine/Plugins/Runtime/nDisplay/Source/DisplayCluster/Private/Network/Barrier/DisplayClusterBarrierFactory.h
+++ b/Engine/Plugins/Runtime/nDisplay/Source/DisplayCluster/Private/Network/Barrier/DisplayClusterBarrierFactory.h
@@ -17,5 +17,5 @@ private:
 
 public:
 	// Factory method to instantiate a barrier based on the user settings
-	static IDisplayClusterBarrier* CreateBarrier(const FString& BarrierId, const TArray<FString>& ThreadMarkers, const uint32 Timeout);
+	static TUniquePtr<IDisplayClusterBarrier> CreateBarrier(const TArray<FString>& ThreadMarkers, const uint32 Timeout, const FString& Name);
 };
diff --git a/Engine/Plugins/Runtime/nDisplay/Source/DisplayCluster/Private/Network/Barrier/DisplayClusterBarrierV1.cpp b/Engine/Plugins/Runtime/nDisplay/Source/DisplayCluster/Private/Network/Barrier/DisplayClusterBarrierV1.cpp
new file mode 100644
--- /dev/null
+++ b/Engine/Plugins/Runtime/nDisplay/Source/DisplayCluster/Private/Network/Barrier/DisplayClusterBarrierV1.cpp
@@ -0,0 +1,111 @@
+// Copyright Epic Games, Inc. All Rights Reserved.
+
+#include "Network/Barrier/DisplayClusterBarrierV1.h"
+
+#include "Misc/DisplayClusterLog.h"
+
+#include <chrono>
+
+
+FDisplayClusterBarrierV1::FDisplayClusterBarrierV1(const uint32 InThreadsAmount, const uint32 InTimeout, const FString& InName)
+	: Name(InName)
+	, ThreadsAmount(InThreadsAmount)
+	, ThreadsLeft(InThreadsAmount)
+	, IterationCounter(0)
+	, Timeout(InTimeout)
+	, WatchdogTimer(FString("stub_watchdog_not_used"))
+{
+	UE_LOG(LogDisplayClusterNetwork, Log, TEXT("Initialized barrier %s with timeout %u for threads count: %u"), *Name, Timeout, ThreadsAmount);
+}
+
+FDisplayClusterBarrierV1::~FDisplayClusterBarrierV1()
+{
+	// Free currently blocked threads
+	Deactivate();
+}
+
+
+bool FDisplayClusterBarrierV1::Activate()
+{
+	std::unique_lock<std::mutex> lock{ Mutex };
+
+	IterationCounter = 0;
+	ThreadsLeft = ThreadsAmount;
+	bEnabled = true;
+	CondVar.notify_all();
+
+	return true;
+}
+
+void FDisplayClusterBarrierV1::Deactivate()
+{
+	std::unique_lock<std::mutex> lock{ Mutex };
+
+	bEnabled = false;
+	CondVar.notify_all();
+}
+
+EDisplayClusterBarrierWaitResult FDisplayClusterBarrierV1::Wait(const FString& ThreadMarker, double* ThreadWaitTime /*= nullptr*/, double* BarrierWaitTime /*= nullptr*/)
+{
+	if (bEnabled == false)
+	{
+		UE_LOG(LogDisplayClusterNetwork, Verbose, TEXT("%s barrier is not active"), *Name);
+		return EDisplayClusterBarrierWaitResult::NotActive;
+	}
+
+	const double ThreadWaitTimeStart = FPlatformTime::Seconds();
+
+	{
+		std::unique_lock<std::mutex> lock{ Mutex };
+
+		size_t curIter = IterationCounter;
+
+		if (ThreadsLeft == ThreadsAmount)
+		{
+			WaitTimeStart = FPlatformTime::Seconds();
+			UE_LOG(LogDisplayClusterNetwork, VeryVerbose, TEXT("%s barrier start time: %lf"), *Name, WaitTimeStart);
+		}
+
+		// Check if all threads are in front of the barrier
+		if (--ThreadsLeft == 0)
+		{
+			UE_LOG(LogDisplayClusterNetwork, Verbose, TEXT("%s barrier trigger!"), *Name);
+			++IterationCounter;
+			ThreadsLeft = ThreadsAmount;
+
+			WaitTimeFinish = FPlatformTime::Seconds();
+			UE_LOG(LogDisplayClusterNetwork, VeryVerbose, TEXT("%s barrier finish time: %lf"), *Name, WaitTimeFinish);
+
+			WaitTimeOverall = WaitTimeFinish - WaitTimeStart;
+			UE_LOG(LogDisplayClusterNetwork, VeryVerbose, TEXT("%s barrier overall wait time: %lf"), *Name, WaitTimeOverall);
+
+			// This is the last node. Unblock the barrier.
+			CondVar.notify_all();
+		}
+		else
+		{
+			UE_LOG(LogDisplayClusterNetwork, VeryVerbose, TEXT("%s barrier waiting, %u threads left"), *Name, ThreadsLeft);
+			// Not all of threads have came here. Wait.
+			if (!CondVar.wait_for(lock, std::chrono::milliseconds(Timeout), [this, curIter] { return curIter != IterationCounter || bEnabled == false; }))
+			{
+				UE_LOG(LogDisplayClusterNetwork, Warning, TEXT("%s barrier waiting timeout"), *Name);
+				return EDisplayClusterBarrierWaitResult::TimeOut;
+			}
+		}
+	}
+
+	const double ThreadWaitTimeFinish = FPlatformTime::Seconds();
+
+	if (BarrierWaitTime)
+	{
+		*BarrierWaitTime = WaitTimeOverall;
+	}
+
+	if (ThreadWaitTime)
+	{
+		*ThreadWaitTime = ThreadWaitTimeFinish - ThreadWaitTimeStart;
+	}
+
+	// Go ahead
+	return EDisplayClusterBarrierWaitResult::Ok;
+}
diff --git a/Engine/Plugins/Runtime/nDisplay/Source/DisplayCluster/Private/Network/Barrier/DisplayClusterBarrierV1.h b/Engine/Plugins/Runtime/nDisplay/Source/DisplayCluster/Private/Network/Barrier/DisplayClusterBarrierV1.h
new file mode 100644
--- /dev/null
+++ b/Engine/Plugins/Runtime/nDisplay/Source/DisplayCluster/Private/Network/Barrier/DisplayClusterBarrierV1.h
@@ -0,0 +1,77 @@
+// Copyright Epic Games, Inc. All Rights Reserved.
+
+#pragma once
+
+#include "CoreMinimal.h"
+#include "Network/Barrier/IDisplayClusterBarrier.h"
+#include "Misc/DisplayClusterWatchdogTimer.h"
+
+#include <condition_variable>
+#include <mutex>
+
+
+/**
+ * Thread barrier v1
+ */
+class FDisplayClusterBarrierV1
+	: public IDisplayClusterBarrier
+{
+public:
+	FDisplayClusterBarrierV1(const uint32 ThreadsAmount, const uint32 Timeout, const FString& Name);
+	virtual  ~FDisplayClusterBarrierV1();
+
+public:
+	// Barrier name
+	virtual const FString& GetName() const override
+	{
+		return Name;
+	}
+
+	// Activate barrier
+	virtual bool Activate() override;
+	// Deactivate barrier, no threads will be blocked
+	virtual void Deactivate() override;
+	// Wait until all threads arrive
+	virtual EDisplayClusterBarrierWaitResult Wait(const FString& ThreadMarker, double* OutThreadWaitTime = nullptr, double* OutBarrierWaitTime = nullptr) override;
+
+	// Remove specified node from sync pipeline
+	virtual void UnregisterSyncNode(const FString& NodeId) override
+	{
+		/* Not supported in barriers v1 */
+	}
+
+	// Barrier timout notification
+	virtual FDisplayClusterBarrierTimeoutEvent& OnBarrierTimeout() override
+	{
+		return BarrierTimeoutEvent;
+	}
+
+private:
+	// Barrier name for logging
+	const FString Name;
+
+	// Barrier state
+	bool bEnabled = true;
+
+	// Amount of threads to wait at the barrier
+	const uint32 ThreadsAmount;
+	// Waiting threads amount
+	uint32 ThreadsLeft;
+	// Iteration counter (kind of barrier sync transaction)
+	size_t IterationCounter;
+
+	std::condition_variable CondVar;
+	mutable std::mutex Mutex;
+
+	uint32 Timeout = 0;
+
+	// Watchdog timer to detect barrier waiting timeouts
+	FDisplayClusterWatchdogTimer WatchdogTimer;
+
+	// Barrier timeout event
+	FDisplayClusterBarrierTimeoutEvent BarrierTimeoutEvent;
+
+	double WaitTimeStart   = 0;
+	double WaitTimeFinish  = 0;
+	double WaitTimeOverall = 0;
+};
diff --git a/Engine/Plugins/Runtime/nDisplay/Source/DisplayCluster/Private/Network/Barrier/DisplayClusterBarrierV2.cpp b/Engine/Plugins/Runtime/nDisplay/Source/DisplayCluster/Private/Network/Barrier/DisplayClusterBarrierV2.cpp
new file mode 100644
--- /dev/null
+++ b/Engine/Plugins/Runtime/nDisplay/Source/DisplayCluster/Private/Network/Barrier/DisplayClusterBarrierV2.cpp
@@ -0,0 +1,276 @@
+// Copyright Epic Games, Inc. All Rights Reserved.
+
+#include "Network/Barrier/DisplayClusterBarrierV2.h"
+
+#include "Misc/DisplayClusterLog.h"
+#include "Misc/ScopeLock.h"
+
+
+FDisplayClusterBarrierV2::FDisplayClusterBarrierV2(const TArray<FString>& InNodesAllowed, const uint32 InTimeout, const FString& InName)
+	: Name(InName)
+	, NodesAllowed(InNodesAllowed)
+	, Timeout(InTimeout)
+	, ThreadLimit(static_cast<uint32>(NodesAllowed.Num()))
+	, WatchdogTimer(InName + FString("_watchdog"))
+{
+	UE_LOG(LogDisplayClusterNetwork, Log, TEXT("Initialized barrier %s with timeout %u ms and threads limit: %u"), *Name, Timeout, ThreadLimit);
+
+	// Subscribe for timeout events
+	WatchdogTimer.OnWatchdogTimeOut().AddRaw(this, &FDisplayClusterBarrierV2::HandleBarrierTimeout);
+}
+
+FDisplayClusterBarrierV2::~FDisplayClusterBarrierV2()
+{
+	// Release threads if there any
+	Deactivate();
+}
+
+
+bool FDisplayClusterBarrierV2::Activate()
+{
+	FScopeLock Lock(&DataCS);
+
+	UE_LOG(LogDisplayClusterNetwork, Log, TEXT("Barrier %s: activating..."), *Name);
+
+	if (!bActive)
+	{
+		ThreadCount = 0;
+		bActive = true;
+		NodesAwaiting.Reset();
+
+		// No exit allowed
+		EventOutputGateOpen->Reset();
+		// Allow join
+		EventInputGateOpen->Trigger();
+	}
+
+	return true;
+}
+
+void FDisplayClusterBarrierV2::Deactivate()
+{
+	FScopeLock Lock(&DataCS);
+
+	if (bActive)
+	{
+		UE_LOG(LogDisplayClusterNetwork, Log, TEXT("Barrier %s: deactivating..."), *Name);
+
+		bActive = false;
+
+		// Release all threads that are currently at the barrier
+		EventInputGateOpen->Trigger();
+		EventOutputGateOpen->Trigger();
+
+		// No more threads awaiting
+		NodesAwaiting.Reset();
+
+		// And reset timer of course
+		WatchdogTimer.ResetTimer();
+	}
+}
+
+EDisplayClusterBarrierWaitResult FDisplayClusterBarrierV2::Wait(const FString& NodeId, double* ThreadWaitTime /*= nullptr*/, double* BarrierWaitTime /*= nullptr*/)
+{
+	double ThreadWaitTimeStart = 0;
+
+	{
+		FScopeLock LockEntrance(&EntranceCS);
+
+		// Wait unless the barrier allows new threads to join. This will happen
+		// once all threads from the previous sync iteration leave the barrier,
+		// or the barrier gets deactivated.
+		EventInputGateOpen->Wait();
+
+		{
+			FScopeLock LockData(&DataCS);
+
+			// Check if this thread has been previously dropped
+			if (NodesTimedout.Contains(NodeId))
+			{
+				UE_LOG(LogDisplayClusterNetwork, Verbose, TEXT("Barrier %s: node %s not allowed to join, it has been timed out previously"), *Name, *NodeId);
+				return EDisplayClusterBarrierWaitResult::TimeOut;
+			}
+
+			// Check if the barrier is active
+			if (!bActive)
+			{
+				UE_LOG(LogDisplayClusterNetwork, Verbose, TEXT("Barrier %s: not active"), *Name);
+				return EDisplayClusterBarrierWaitResult::NotActive;
+			}
+
+			// Check if this thread is allowed to sync at this barrier
+			if (!NodesAllowed.Contains(NodeId))
+			{
+				UE_LOG(LogDisplayClusterNetwork, Verbose, TEXT("Barrier %s: node %s not allowed to join, no permission"), *Name, *NodeId);
+				return EDisplayClusterBarrierWaitResult::NotAllowed;
+			}
+
+			// Register node
+			NodesAwaiting.Add(NodeId);
+
+			// Make sure the barrier was not deactivated while this thread was waiting at the input gate
+			if (bActive)
+			{
+				// Fixate awaiting start for this particular thread
+				ThreadWaitTimeStart = FPlatformTime::Seconds();
+
+				// One more thread awaiting
+				++ThreadCount;
+
+				// In case this thread came first to the barrier, we need:
+				// - to fixate barrier awaiting start time
+				// - start watchdog timer
+				if (ThreadCount == 1)
+				{
+					UE_LOG(LogDisplayClusterNetwork, Verbose, TEXT("Barrier %s: sync start"), *Name);
+
+					BarrierWaitTimeStart = ThreadWaitTimeStart;
+					WatchdogTimer.SetTimer(Timeout);
+				}
+
+				UE_LOG(LogDisplayClusterNetwork, Verbose, TEXT("Barrier %s: awaiting threads amount - %d"), *Name, ThreadCount);
+
+				// In case this thread is the last one the barrier is awaiting for, we need:
+				// - to fixate barrier awaiting finish time
+				// - to open the output gate (release the barrier)
+				// - to close the input gate
+				// - reset watchdog timer
+				if (ThreadCount == ThreadLimit)
+				{
+					BarrierWaitTimeFinish = FPlatformTime::Seconds();
+					BarrierWaitTimeOverall = BarrierWaitTimeFinish - BarrierWaitTimeStart;
+
+					UE_LOG(LogDisplayClusterNetwork, Verbose, TEXT("Barrier %s: sync end, barrier wait time %f"), *Name, BarrierWaitTimeOverall);
+
+					// All nodes here, pet the watchdog
+					WatchdogTimer.ResetTimer();
+
+					// Close input gate, open output gate
+					EventInputGateOpen->Reset();
+					EventOutputGateOpen->Trigger();
+				}
+			}
+		}
+	}
+
+	// Wait for the barrier to open
+	EventOutputGateOpen->Wait();
+
+	// Fixate awaiting finish for this particular thread
+	const double ThreadWaitTimeFinish = FPlatformTime::Seconds();
+
+	{
+		UE_LOG(LogDisplayClusterNetwork, Verbose, TEXT("Barrier %s: node %s is leaving the barrier"), *Name, *NodeId);
+
+		{
+			FScopeLock LockData(&DataCS);
+
+			// Unregister node
+			NodesAwaiting.Remove(NodeId);
+
+			// Make sure the barrier was not deactivated while this thread was waiting at the output gate
+			if (bActive)
+			{
+				// In case this thread is leaving last, close output and open input
+				if (--ThreadCount == 0)
+				{
+					EventOutputGateOpen->Reset();
+					EventInputGateOpen->Trigger();
+				}
+			}
+		}
+	}
+
+	// Export barrier overall waiting time
+	if (BarrierWaitTime)
+	{
+		*BarrierWaitTime = BarrierWaitTimeOverall;
+	}
+
+	// Export thread waiting time
+	if (ThreadWaitTime)
+	{
+		*ThreadWaitTime = ThreadWaitTimeFinish - ThreadWaitTimeStart;
+	}
+
+	return EDisplayClusterBarrierWaitResult::Ok;
+}
+
+void FDisplayClusterBarrierV2::UnregisterSyncNode(const FString& NodeId)
+{
+	UE_LOG(LogDisplayClusterNetwork, Log, TEXT("Barrier %s: unregistering node %s..."), *Name, *NodeId);
+
+	FScopeLock Lock(&DataCS);
+
+	// Remove specified nodes from the management lists
+	if (NodesAllowed.Remove(NodeId) > 0 || NodesTimedout.Remove(NodeId) > 0)
+	{
+		// Update thread limit
+		ThreadLimit = static_cast<uint32>(NodesAllowed.Num());
+
+		UE_LOG(LogDisplayClusterNetwork, Log, TEXT("Barrier %s: node %s has been unregistered, new thread limit is %u"), *Name, *NodeId, ThreadLimit);
+
+		// In case it's a last missing node, we need to open the barrier
+		if (ThreadCount == ThreadLimit)
+		{
+			BarrierWaitTimeFinish  = FPlatformTime::Seconds();
+			BarrierWaitTimeOverall = BarrierWaitTimeFinish - BarrierWaitTimeStart;
+
+			UE_LOG(LogDisplayClusterNetwork, Verbose, TEXT("Barrier %s: sync end, barrier wait time %f"), *Name, BarrierWaitTimeOverall);
+
+			// All nodes here, pet the watchdog
+			WatchdogTimer.ResetTimer();
+
+			// Close the input gate, and open the output gate
+			EventInputGateOpen->Reset();
+			EventOutputGateOpen->Trigger();
+		}
+	}
+	else
+	{
+		UE_LOG(LogDisplayClusterNetwork, Log, TEXT("Barrier %s: node %s not found"), *Name, *NodeId);
+	}
+}
+
+void FDisplayClusterBarrierV2::HandleBarrierTimeout()
+{
+	// Being here means some nodes have not come to the barrier yet during specific time period. Those
+	// missing nodes will be considered as the lost ones. The barrier will continue working with the remaining nodes only.
+
+	FScopeLock Lock(&DataCS);
+
+	UE_LOG(LogDisplayClusterNetwork, Log, TEXT("Barrier %s: Time out! %d nodes missing"), *Name, NodesAllowed.Num() - NodesAwaiting.Num());
+
+	// First of all, update the time variables
+	BarrierWaitTimeFinish = FPlatformTime::Seconds();
+	BarrierWaitTimeOverall = BarrierWaitTimeFinish - BarrierWaitTimeStart;
+
+	// List of timed out threads
+	TArray<FString> NodesTimedOutOnLastSync;
+
+	// Update the list of timed out threads. Copy all missing nodes to the list of timed out ones.
+	for (const FString& NodeAllowed : NodesAllowed)
+	{
+		if (!NodesAwaiting.Contains(NodeAllowed))
+		{
+			UE_LOG(LogDisplayClusterNetwork, Log, TEXT("Barrier %s: node %s was moved to the 'TimedOut' list"), *Name, *NodeAllowed);
+			NodesTimedOutOnLastSync.Add(NodeAllowed);
+		}
+	}
+
+	// Update timedout list
+	NodesTimedout.Append(NodesTimedOutOnLastSync);
+	// Update the list of permitted nodes
+	NodesAllowed = NodesAwaiting;
+	// Update thread limit
+	ThreadLimit = ThreadCount; // Same as NodesAllowed.Num()
+
+	UE_LOG(LogDisplayClusterNetwork, Log, TEXT("Barrier %s: new threads limit %d"), *Name, ThreadLimit);
+
+	// Notify listeners
+	OnBarrierTimeout().Broadcast(Name, NodesTimedOutOnLastSync);
+
+	// Close the input gate, and open the output gate to let the remaining nodes go
+	EventInputGateOpen->Reset();
+	EventOutputGateOpen->Trigger();
+}
diff --git a/Engine/Plugins/Runtime/nDisplay/Source/DisplayCluster/Private/Network/Barrier/DisplayClusterBarrierV2.h b/Engine/Plugins/Runtime/nDisplay/Source/DisplayCluster/Private/Network/Barrier/DisplayClusterBarrierV2.h
new file mode 100644
--- /dev/null
+++ b/Engine/Plugins/Runtime/nDisplay/Source/DisplayCluster/Private/Network/Barrier/DisplayClusterBarrierV2.h
@@ -0,0 +1,95 @@
+// Copyright Epic Games, Inc. All Rights Reserved.
+
+#pragma once
+
+#include "CoreMinimal.h"
+#include "Network/Barrier/IDisplayClusterBarrier.h"
+#include "Misc/DisplayClusterWatchdogTimer.h"
+
+#include "HAL/Event.h"
+
+
+/**
+ * Thread barrier v2
+ */
+class FDisplayClusterBarrierV2
+	: public IDisplayClusterBarrier
+{
+public:
+	FDisplayClusterBarrierV2(const TArray<FString>& NodesAllowed, const uint32 Timeout, const FString& Name);
+	virtual  ~FDisplayClusterBarrierV2();
+
+public:
+	// Barrier name
+	virtual const FString& GetName() const override
+	{
+		return Name;
+	}
+
+	// Activate barrier
+	virtual bool Activate() override;
+	// Deactivate barrier, no threads will be blocked
+	virtual void Deactivate() override;
+	// Wait until all node threads arrived
+	virtual EDisplayClusterBarrierWaitResult Wait(const FString& NodeId, double* OutThreadWaitTime = nullptr, double* OutBarrierWaitTime = nullptr) override;
+
+	// Remove specified node from sync pipeline
+	virtual void UnregisterSyncNode(const FString& NodeId) override;
+
+	// Barrier timout notification
+	virtual FDisplayClusterBarrierTimeoutEvent& OnBarrierTimeout() override
+	{
+		return BarrierTimeoutEvent;
+	}
+
+private:
+	// Handler for barrier timouts
+	void HandleBarrierTimeout();
+
+private:
+	// Barrier name
+	const FString Name;
+	// Barrier state
+	bool bActive = false;
+
+
+	// Cluster nodes that are allowed to join the barrier
+	TArray<FString> NodesAllowed;
+
+	// Cluster nodes that are already waiting at the barrier
+	TArray<FString> NodesAwaiting;
+
+	// Cluster nodes that have been timed out previously
+	TArray<FString> NodesTimedout;
+
+
+	// Timeout for the barrier
+	const uint32 Timeout = 0;
+
+	// Overall amount of threads to wait at the barrier
+	uint32 ThreadLimit = 0;
+
+	// Amount of threads at the barrier
+	uint32 ThreadCount = 0;
+
+
+	// Events to control when threads can join (input) and leave (output) the barrier
+	FEventRef EventInputGateOpen { EEventMode::ManualReset };
+	FEventRef EventOutputGateOpen{ EEventMode::ManualReset };
+
+	// Watchdog timer to detect barrier waiting timeouts
+	FDisplayClusterWatchdogTimer WatchdogTimer;
+
+	// Barrier timeout event
+	FDisplayClusterBarrierTimeoutEvent BarrierTimeoutEvent;
+
+	// Diagnostics data
+	double BarrierWaitTimeStart   = 0;
+	double BarrierWaitTimeFinish  = 0;
+	double BarrierWaitTimeOverall = 0;
+
+	// Barrier state CS
+	mutable FCriticalSection DataCS;
+	// Barrier entrance CS
+	mutable FCriticalSection EntranceCS;
+};
diff --git a/Engine/Plugins/Runtime/nDisplay/Source/DisplayCluster/Private/Network/Barrier/IDisplayClusterBarrier.h b/Engine/Plugins/Runtime/nDisplay/Source/DisplayCluster/Private/Network/Barrier/IDisplayClusterBarrier.h
--- a/Engine/Plugins/Runtime/nDisplay/Source/DisplayCluster/Private/Network/Barrier/IDisplayClusterBarrier.h
+++ b/Engine/Plugins/Runtime/nDisplay/Source/DisplayCluster/Private/Network/Barrier/IDisplayClusterBarrier.h
@@ -37,16 +37,13 @@ public:
 	// Deactivate barrier, no threads will be blocked
 	virtual void Deactivate() = 0;
 
-	// Returns true if the barrier has been activated
-	virtual bool IsActivated() const = 0;
-
 	// Wait until all threads arrive
 	virtual EDisplayClusterBarrierWaitResult Wait(const FString& ThreadMarker, double* OutThreadWaitTime = nullptr, double* OutBarrierWaitTime = nullptr) = 0;
 
-	// Remove specified caller from the sync pipeline
-	virtual void UnregisterSyncCaller(const FString& CallerId) = 0;
+	// Remove specified node from sync pipeline
+	virtual void UnregisterSyncNode(const FString& NodeId) = 0;
 
-	// Barrier timout notification (provides BarrierName and CallersTimedOut in parameters)
+	// Barrier timout notification (provides BarrierName and NodesTimedOut in parameters)
 	DECLARE_EVENT_TwoParams(IDisplayClusterBarrier, FDisplayClusterBarrierTimeoutEvent, const FString&, const TArray<FString>&);
 	virtual FDisplayClusterBarrierTimeoutEvent& OnBarrierTimeout() = 0;
 };
diff --git a/Engine/Plugins/Runtime/nDisplay/Source/DisplayCluster/Private/Network/DisplayClusterClient.h b/Engine/Plugins/Runtime/nDisplay/Source/DisplayCluster/Private/Network/DisplayClusterClient.h
--- a/Engine/Plugins/Runtime/nDisplay/Source/DisplayCluster/Private/Network/DisplayClusterClient.h
+++ b/Engine/Plugins/Runtime/nDisplay/Source/DisplayCluster/Private/Network/DisplayClusterClient.h
@@ -32,7 +32,7 @@ public:
 	// Connects to a server
 	virtual bool Connect(const FString& Address, const uint16 Port, const uint32 ConnectRetriesAmount, const uint32 ConnectRetryDelay) override;
 	// Terminates current connection
-	virtual void Disconnect() override;
+	virtual void Disconnect() override final;
 
 	// Provides with net unit name
 	virtual FString GetName() const override
diff --git a/Engine/Plugins/Runtime/nDisplay/Source/DisplayCluster/Private/Network/Protocol/IDisplayClusterProtocolClusterSync.h b/Engine/Plugins/Runtime/nDisplay/Source/DisplayCluster/Private/Network/Protocol/IDisplayClusterProtocolClusterSync.h
--- a/Engine/Plugins/Runtime/nDisplay/Source/DisplayCluster/Private/Network/Protocol/IDisplayClusterProtocolClusterSync.h
+++ b/Engine/Plugins/Runtime/nDisplay/Source/DisplayCluster/Private/Network/Protocol/IDisplayClusterProtocolClusterSync.h
@@ -2,8 +2,6 @@
 
 #pragma once
 
-#include "CoreMinimal.h"
-
 #include "DisplayClusterEnums.h"
 #include "Misc/Optional.h"
 #include "Misc/QualifiedFrameTime.h"
diff --git a/Engine/Plugins/Runtime/nDisplay/Source/DisplayCluster/Private/Network/Protocol/IDisplayClusterProtocolEventsBinary.h b/Engine/Plugins/Runtime/nDisplay/Source/DisplayCluster/Private/Network/Protocol/IDisplayClusterProtocolEventsBinary.h
--- a/Engine/Plugins/Runtime/nDisplay/Source/DisplayCluster/Private/Network/Protocol/IDisplayClusterProtocolEventsBinary.h
+++ b/Engine/Plugins/Runtime/nDisplay/Source/DisplayCluster/Private/Network/Protocol/IDisplayClusterProtocolEventsBinary.h
@@ -2,7 +2,6 @@
 
 #pragma once
 
-#include "CoreMinimal.h"
 #include "Network/DisplayClusterNetworkTypes.h"
 
 struct FDisplayClusterClusterEventBinary;
diff --git a/Engine/Plugins/Runtime/nDisplay/Source/DisplayCluster/Private/Network/Protocol/IDisplayClusterProtocolEventsJson.h b/Engine/Plugins/Runtime/nDisplay/Source/DisplayCluster/Private/Network/Protocol/IDisplayClusterProtocolEventsJson.h
--- a/Engine/Plugins/Runtime/nDisplay/Source/DisplayCluster/Private/Network/Protocol/IDisplayClusterProtocolEventsJson.h
+++ b/Engine/Plugins/Runtime/nDisplay/Source/DisplayCluster/Private/Network/Protocol/IDisplayClusterProtocolEventsJson.h
@@ -2,7 +2,6 @@
 
 #pragma once
 
-#include "CoreMinimal.h"
 #include "Network/DisplayClusterNetworkTypes.h"
 
 struct FDisplayClusterClusterEventJson;
diff --git a/Engine/Plugins/Runtime/nDisplay/Source/DisplayCluster/Private/Network/Protocol/IDisplayClusterProtocolGenericBarrier.h b/Engine/Plugins/Runtime/nDisplay/Source/DisplayCluster/Private/Network/Protocol/IDisplayClusterProtocolGenericBarrier.h
deleted file mode 100644
diff --git a/Engine/Plugins/Runtime/nDisplay/Source/DisplayCluster/Private/Network/Protocol/IDisplayClusterProtocolRenderSync.h b/Engine/Plugins/Runtime/nDisplay/Source/DisplayCluster/Private/Network/Protocol/IDisplayClusterProtocolRenderSync.h
--- a/Engine/Plugins/Runtime/nDisplay/Source/DisplayCluster/Private/Network/Protocol/IDisplayClusterProtocolRenderSync.h
+++ b/Engine/Plugins/Runtime/nDisplay/Source/DisplayCluster/Private/Network/Protocol/IDisplayClusterProtocolRenderSync.h
@@ -2,7 +2,6 @@
 
 #pragma once
 
-#include "CoreMinimal.h"
 #include "Network/DisplayClusterNetworkTypes.h"
 
 
diff --git a/Engine/Plugins/Runtime/nDisplay/Source/DisplayCluster/Private/Network/Service/ClusterSync/DisplayClusterClusterSyncService.cpp b/Engine/Plugins/Runtime/nDisplay/Source/DisplayCluster/Private/Network/Service/ClusterSync/DisplayClusterClusterSyncService.cpp
--- a/Engine/Plugins/Runtime/nDisplay/Source/DisplayCluster/Private/Network/Service/ClusterSync/DisplayClusterClusterSyncService.cpp
+++ b/Engine/Plugins/Runtime/nDisplay/Source/DisplayCluster/Private/Network/Service/ClusterSync/DisplayClusterClusterSyncService.cpp
@@ -38,9 +38,9 @@ FDisplayClusterClusterSyncService::FDisplayClusterClusterSyncService()
 	const FDisplayClusterConfigurationNetworkSettings& NetworkSettings = GDisplayCluster->GetConfigMgr()->GetConfig()->Cluster->Network;
 
 	// Instantiate service barriers
-	BarrierGameStart.Reset(FDisplayClusterBarrierFactory::CreateBarrier(TEXT("GameStart_barrier"), NodeIds, NetworkSettings.GameStartBarrierTimeout));
-	BarrierFrameStart.Reset(FDisplayClusterBarrierFactory::CreateBarrier(TEXT("FrameStart_barrier"), NodeIds, NetworkSettings.FrameStartBarrierTimeout));
-	BarrierFrameEnd.Reset(FDisplayClusterBarrierFactory::CreateBarrier(TEXT("FrameEnd_barrier"), NodeIds, NetworkSettings.FrameEndBarrierTimeout));
+	BarrierGameStart  = FDisplayClusterBarrierFactory::CreateBarrier(NodeIds, NetworkSettings.GameStartBarrierTimeout,  TEXT("GameStart_barrier"));
+	BarrierFrameStart = FDisplayClusterBarrierFactory::CreateBarrier(NodeIds, NetworkSettings.FrameStartBarrierTimeout, TEXT("FrameStart_barrier"));
+	BarrierFrameEnd   = FDisplayClusterBarrierFactory::CreateBarrier(NodeIds, NetworkSettings.FrameEndBarrierTimeout,   TEXT("FrameEnd_barrier"));
 
 	// Put the barriers into an aux container
 	ServiceBarriers.Emplace(BarrierGameStart->GetName(),  &BarrierGameStart);
@@ -116,7 +116,7 @@ void FDisplayClusterClusterSyncService::KillSession(const FString& NodeId)
 	// Before killing the session on the parent level, we need to unregister this node from the barriers
 	for (TPair<FString, TUniquePtr<IDisplayClusterBarrier>*>& BarrierIt : ServiceBarriers)
 	{
-		(*BarrierIt.Value)->UnregisterSyncCaller(NodeId);
+		(*BarrierIt.Value)->UnregisterSyncNode(NodeId);
 	}
 
 	// Now do the session related job
@@ -141,7 +141,7 @@ void FDisplayClusterClusterSyncService::ProcessSessionClosed(const FDisplayClust
 			// We have to unregister the node that just disconnected from all barriers
 			for (TPair<FString, TUniquePtr<IDisplayClusterBarrier>*>& BarrierIt : ServiceBarriers)
 			{
-				(*BarrierIt.Value)->UnregisterSyncCaller(NodeId);
+				(*BarrierIt.Value)->UnregisterSyncNode(NodeId);
 			}
 
 			// Notify others about node fail
@@ -157,7 +157,7 @@ void FDisplayClusterClusterSyncService::ProcessBarrierTimeout(const FString& Bar
 	{
 		for (const FString& NodeId : NodesTimedOut)
 		{
-			(*BarrierIt.Value)->UnregisterSyncCaller(NodeId);
+			(*BarrierIt.Value)->UnregisterSyncNode(NodeId);
 		}
 	}
 
diff --git a/Engine/Plugins/Runtime/nDisplay/Source/DisplayCluster/Private/Network/Service/ClusterSync/DisplayClusterClusterSyncStrings.h b/Engine/Plugins/Runtime/nDisplay/Source/DisplayCluster/Private/Network/Service/ClusterSync/DisplayClusterClusterSyncStrings.h
--- a/Engine/Plugins/Runtime/nDisplay/Source/DisplayCluster/Private/Network/Service/ClusterSync/DisplayClusterClusterSyncStrings.h
+++ b/Engine/Plugins/Runtime/nDisplay/Source/DisplayCluster/Private/Network/Service/ClusterSync/DisplayClusterClusterSyncStrings.h
@@ -10,8 +10,8 @@ namespace DisplayClusterClusterSyncStrings
 {
 	constexpr static auto ProtocolName = "ClusterSync";
 	
-	constexpr static auto TypeRequest  = "Request";
-	constexpr static auto TypeResponse = "Response";
+	constexpr static auto TypeRequest  = "request";
+	constexpr static auto TypeResponse = "response";
 
 	constexpr static auto ArgumentsDefaultCategory = "CS";
 	constexpr static auto ArgumentsJsonEvents      = "CS_JE";
diff --git a/Engine/Plugins/Runtime/nDisplay/Source/DisplayCluster/Private/Network/Service/GenericBarrier/DisplayClusterGenericBarrierClient.cpp b/Engine/Plugins/Runtime/nDisplay/Source/DisplayCluster/Private/Network/Service/GenericBarrier/DisplayClusterGenericBarrierClient.cpp
deleted file mode 100644
diff --git a/Engine/Plugins/Runtime/nDisplay/Source/DisplayCluster/Private/Network/Service/GenericBarrier/DisplayClusterGenericBarrierClient.h b/Engine/Plugins/Runtime/nDisplay/Source/DisplayCluster/Private/Network/Service/GenericBarrier/DisplayClusterGenericBarrierClient.h
deleted file mode 100644
diff --git a/Engine/Plugins/Runtime/nDisplay/Source/DisplayCluster/Private/Network/Service/GenericBarrier/DisplayClusterGenericBarrierService.cpp b/Engine/Plugins/Runtime/nDisplay/Source/DisplayCluster/Private/Network/Service/GenericBarrier/DisplayClusterGenericBarrierService.cpp
deleted file mode 100644
diff --git a/Engine/Plugins/Runtime/nDisplay/Source/DisplayCluster/Private/Network/Service/GenericBarrier/DisplayClusterGenericBarrierService.h b/Engine/Plugins/Runtime/nDisplay/Source/DisplayCluster/Private/Network/Service/GenericBarrier/DisplayClusterGenericBarrierService.h
deleted file mode 100644
diff --git a/Engine/Plugins/Runtime/nDisplay/Source/DisplayCluster/Private/Network/Service/GenericBarrier/DisplayClusterGenericBarrierStrings.h b/Engine/Plugins/Runtime/nDisplay/Source/DisplayCluster/Private/Network/Service/GenericBarrier/DisplayClusterGenericBarrierStrings.h
deleted file mode 100644
diff --git a/Engine/Plugins/Runtime/nDisplay/Source/DisplayCluster/Private/Network/Service/RenderSync/DisplayClusterRenderSyncService.cpp b/Engine/Plugins/Runtime/nDisplay/Source/DisplayCluster/Private/Network/Service/RenderSync/DisplayClusterRenderSyncService.cpp
--- a/Engine/Plugins/Runtime/nDisplay/Source/DisplayCluster/Private/Network/Service/RenderSync/DisplayClusterRenderSyncService.cpp
+++ b/Engine/Plugins/Runtime/nDisplay/Source/DisplayCluster/Private/Network/Service/RenderSync/DisplayClusterRenderSyncService.cpp
@@ -246,13 +246,11 @@ void FDisplayClusterRenderSyncService::InitializeBarriers()
 		const FString BarrierName = SyncGroup.Key + TEXT("_barrier");
 
 		// Instantiate the barrier
-		if (IDisplayClusterBarrier* const Barrier = FDisplayClusterBarrierFactory::CreateBarrier(BarrierName, SyncGroup.Value, Config->Cluster->Network.RenderSyncBarrierTimeout))
-		{
-			// Subscribe for barrier timeout events
-			Barrier->OnBarrierTimeout().AddRaw(this, &FDisplayClusterRenderSyncService::ProcessBarrierTimeout);
-			// Store barrier instance
-			PolicyToBarrierMap.Emplace(SyncGroup.Key, Barrier);
-		}
+		TUniquePtr<IDisplayClusterBarrier> Barrier = FDisplayClusterBarrierFactory::CreateBarrier(SyncGroup.Value, Config->Cluster->Network.RenderSyncBarrierTimeout, BarrierName);
+		// Subscribe for barrier timeout events
+		Barrier->OnBarrierTimeout().AddRaw(this, &FDisplayClusterRenderSyncService::ProcessBarrierTimeout);
+		// Store barrier instance
+		PolicyToBarrierMap.Emplace(SyncGroup.Key, MoveTemp(Barrier));
 	}
 }
 
@@ -297,6 +295,6 @@ void FDisplayClusterRenderSyncService::UnregisterClusterNode(const FString& Node
 {
 	if (IDisplayClusterBarrier* Barrier = GetBarrierForNode(NodeId))
 	{
-		Barrier->UnregisterSyncCaller(NodeId);
+		Barrier->UnregisterSyncNode(NodeId);
 	}
 }
diff --git a/Engine/Plugins/Runtime/nDisplay/Source/DisplayCluster/Private/Network/Service/RenderSync/DisplayClusterRenderSyncStrings.h b/Engine/Plugins/Runtime/nDisplay/Source/DisplayCluster/Private/Network/Service/RenderSync/DisplayClusterRenderSyncStrings.h
--- a/Engine/Plugins/Runtime/nDisplay/Source/DisplayCluster/Private/Network/Service/RenderSync/DisplayClusterRenderSyncStrings.h
+++ b/Engine/Plugins/Runtime/nDisplay/Source/DisplayCluster/Private/Network/Service/RenderSync/DisplayClusterRenderSyncStrings.h
@@ -10,8 +10,8 @@ namespace DisplayClusterRenderSyncStrings
 {
 	constexpr static auto ProtocolName = "RenderSync";
 
-	constexpr static auto TypeRequest  = "Request";
-	constexpr static auto TypeResponse = "Response";
+	constexpr static auto TypeRequest  = "request";
+	constexpr static auto TypeResponse = "response";
 
 	constexpr static auto ArgumentsDefaultCategory = "RS";
 
diff --git a/Engine/Plugins/Runtime/nDisplay/Source/DisplayCluster/Private/Render/DisplayClusterRenderManager.cpp b/Engine/Plugins/Runtime/nDisplay/Source/DisplayCluster/Private/Render/DisplayClusterRenderManager.cpp
--- a/Engine/Plugins/Runtime/nDisplay/Source/DisplayCluster/Private/Render/DisplayClusterRenderManager.cpp
+++ b/Engine/Plugins/Runtime/nDisplay/Source/DisplayCluster/Private/Render/DisplayClusterRenderManager.cpp
@@ -24,7 +24,6 @@
 #include "Render/Device/DisplayClusterRenderDeviceFactoryInternal.h"
 
 #include "Render/Device/IDisplayClusterRenderDeviceFactory.h"
-#include "Render/Monitoring/DisplayClusterVblankMonitor.h"
 #include "Render/PostProcess/IDisplayClusterPostProcess.h"
 #include "Render/Projection/IDisplayClusterProjectionPolicyFactory.h"
 #include "Render/Projection/IDisplayClusterProjectionPolicy.h"
@@ -54,17 +53,6 @@
 #endif
 
 
-static TAutoConsoleVariable<int32> CVarSyncDiagnosticsVBlankMonitoring(
-	TEXT("nDisplay.sync.diag.VBlankMonitoring"),
-	0,
-	TEXT("Sync diagnostics: V-blank monitoring\n")
-	TEXT("0 : disabled\n")
-	TEXT("1 : enabled (if policy supports only)\n")
-	,
-	ECVF_ReadOnly
-);
-
-
 FDisplayClusterRenderManager::FDisplayClusterRenderManager()
 {
 	// Instantiate and register internal render device factory
@@ -80,9 +68,6 @@ FDisplayClusterRenderManager::FDisplayClusterRenderManager()
 	RegisterSynchronizationPolicyFactory(DisplayClusterConfigurationStrings::config::cluster::render_sync::Ethernet,        NewSyncPolicyFactory); // Ethernet
 	RegisterSynchronizationPolicyFactory(DisplayClusterConfigurationStrings::config::cluster::render_sync::EthernetBarrier, NewSyncPolicyFactory); // Ethernet_Simple
 	RegisterSynchronizationPolicyFactory(DisplayClusterConfigurationStrings::config::cluster::render_sync::Nvidia,          NewSyncPolicyFactory); // NVIDIA
-
-	// Instantiate V-blank monitor (it won't auto-start polling)
-	VBlankMonitor = MakeShared<FDisplayClusterVBlankMonitor, ESPMode::ThreadSafe>();
 }
 
 FDisplayClusterRenderManager::~FDisplayClusterRenderManager()
@@ -136,12 +121,6 @@ bool FDisplayClusterRenderManager::StartSession(UDisplayClusterConfigurationData
 		RenderDevicePtr = NewRenderDevice.Get();
 	}
 
-	// Start v-blank monitoring if requested
-	if (!!CVarSyncDiagnosticsVBlankMonitoring.GetValueOnGameThread())
-	{
-		VBlankMonitor->StartMonitoring();
-	}
-
 	// When session is starting in Editor the device won't be initialized so we avoid nullptr access here.
 	//@todo Now we always have a device, even for Editor. Change the condition working on the EditorDevice.
 	return (RenderDevicePtr ? RenderDevicePtr->Initialize() : true);
diff --git a/Engine/Plugins/Runtime/nDisplay/Source/DisplayCluster/Private/Render/DisplayClusterRenderManager.h b/Engine/Plugins/Runtime/nDisplay/Source/DisplayCluster/Private/Render/DisplayClusterRenderManager.h
--- a/Engine/Plugins/Runtime/nDisplay/Source/DisplayCluster/Private/Render/DisplayClusterRenderManager.h
+++ b/Engine/Plugins/Runtime/nDisplay/Source/DisplayCluster/Private/Render/DisplayClusterRenderManager.h
@@ -6,7 +6,6 @@
 
 #include "Render/IPDisplayClusterRenderManager.h"
 
-
 class IDisplayClusterPostProcess;
 class IDisplayClusterPostProcessFactory; 
 class IDisplayClusterProjectionPolicy;
@@ -14,10 +13,8 @@ class IDisplayClusterProjectionPolicyFactory;
 class IDisplayClusterRenderDeviceFactory;
 class IDisplayClusterRenderSyncPolicy;
 class IDisplayClusterRenderSyncPolicyFactory;
-class IDisplayClusterVblankMonitor;
 class UDisplayClusterCameraComponent;
 
-
 /**
  * Render manager. Responsible for everything related to the visuals.
  */
@@ -64,11 +61,6 @@ public:
 	virtual void GetRegisteredPostProcess(TArray<FString>& OutPostProcessIDs) const override;
 	virtual TSharedPtr<IDisplayClusterRender_MeshComponent, ESPMode::ThreadSafe> CreateMeshComponent() const override;
 
-	virtual TSharedPtr<IDisplayClusterVblankMonitor, ESPMode::ThreadSafe> GetVblankMonitor() override
-	{
-		return VBlankMonitor;
-	}
-
 	virtual IDisplayClusterViewportManager* GetViewportManager() const override;
 
 public:
@@ -116,8 +108,4 @@ private:
 
 	// This flag is used to auto-focus the UE window once on start
 	bool bWasWindowFocused = false;
-
-private:
-	// V-blank monitor
-	TSharedPtr<IDisplayClusterVblankMonitor, ESPMode::ThreadSafe> VBlankMonitor;
 };
diff --git a/Engine/Plugins/Runtime/nDisplay/Source/DisplayCluster/Private/Render/Monitoring/DisplayClusterVblankMonitor.cpp b/Engine/Plugins/Runtime/nDisplay/Source/DisplayCluster/Private/Render/Monitoring/DisplayClusterVblankMonitor.cpp
deleted file mode 100644
diff --git a/Engine/Plugins/Runtime/nDisplay/Source/DisplayCluster/Private/Render/Monitoring/DisplayClusterVblankMonitor.h b/Engine/Plugins/Runtime/nDisplay/Source/DisplayCluster/Private/Render/Monitoring/DisplayClusterVblankMonitor.h
deleted file mode 100644
diff --git a/Engine/Plugins/Runtime/nDisplay/Source/DisplayCluster/Private/Render/Synchronization/DisplayClusterRenderSyncHelper.h b/Engine/Plugins/Runtime/nDisplay/Source/DisplayCluster/Private/Render/Synchronization/DisplayClusterRenderSyncHelper.h
--- a/Engine/Plugins/Runtime/nDisplay/Source/DisplayCluster/Private/Render/Synchronization/DisplayClusterRenderSyncHelper.h
+++ b/Engine/Plugins/Runtime/nDisplay/Source/DisplayCluster/Private/Render/Synchronization/DisplayClusterRenderSyncHelper.h
@@ -27,7 +27,7 @@ public:
 
 
 /**
- * Synchronization helper implementation. The main purpose of this class
+ * Syncrhonization helper implementation. The main purpose of this class
  * is to solve the combinatoric complexity Platform/RHI/SyncPolicy without
  * making the whole sync policies hierarchy overcomplicated.
  */
diff --git a/Engine/Plugins/Runtime/nDisplay/Source/DisplayCluster/Private/Render/Synchronization/DisplayClusterRenderSyncPolicyBase.cpp b/Engine/Plugins/Runtime/nDisplay/Source/DisplayCluster/Private/Render/Synchronization/DisplayClusterRenderSyncPolicyBase.cpp
--- a/Engine/Plugins/Runtime/nDisplay/Source/DisplayCluster/Private/Render/Synchronization/DisplayClusterRenderSyncPolicyBase.cpp
+++ b/Engine/Plugins/Runtime/nDisplay/Source/DisplayCluster/Private/Render/Synchronization/DisplayClusterRenderSyncPolicyBase.cpp
@@ -17,8 +17,26 @@
 #include "RHIResources.h"
 
 
+static TAutoConsoleVariable<int32> CVarSyncDiagnosticsVBlankMonitoring(
+	TEXT("nDisplay.sync.diag.VBlankMonitoring"),
+	0,
+	TEXT("Sync diagnostics: V-blank monitoring\n")
+	TEXT("0 : disabled\n")
+	TEXT("1 : enabled (if policy supports only)\n")
+	,
+	ECVF_ReadOnly
+);
+
+
 bool FDisplayClusterRenderSyncPolicyBase::Initialize()
 {
+	// In case v-blank monitoring was requested and this feature is supported by
+	// the sync policy implementation, start v-blank monitoring.
+	if (!!CVarSyncDiagnosticsVBlankMonitoring.GetValueOnGameThread() && IsWaitForVBlankFeatureSupported())
+	{
+		VBlankMonitor = MakeUnique<FDisplayClusterVBlankMonitor>(*this);
+	}
+
 	return true;
 }
 
diff --git a/Engine/Plugins/Runtime/nDisplay/Source/DisplayCluster/Private/Render/Synchronization/DisplayClusterRenderSyncPolicyBase.h b/Engine/Plugins/Runtime/nDisplay/Source/DisplayCluster/Private/Render/Synchronization/DisplayClusterRenderSyncPolicyBase.h
--- a/Engine/Plugins/Runtime/nDisplay/Source/DisplayCluster/Private/Render/Synchronization/DisplayClusterRenderSyncPolicyBase.h
+++ b/Engine/Plugins/Runtime/nDisplay/Source/DisplayCluster/Private/Render/Synchronization/DisplayClusterRenderSyncPolicyBase.h
@@ -3,6 +3,7 @@
 #pragma once
 
 #include "Render/Synchronization/IDisplayClusterRenderSyncPolicy.h"
+#include "Render/Synchronization/DisplayClusterVblankMonitor.h"
 
 
 /**
@@ -27,7 +28,7 @@ public:
 	virtual bool WaitForVBlank() override;
 
 public:
-	// Performs Ethernet barrier synchronization with other nodes in the cluster
+	// Performs Ethernet barrier syncrhonization with other nodes in the cluster
 	void SyncOnBarrier();
 
 	// Returns policy configuration parameters
@@ -43,4 +44,7 @@ protected:
 private:
 	// Sync policy configuration parameters
 	TMap<FString, FString> Parameters;
+
+	// V-blank monitoring thread. Performs V-blank polling on the top level of the sync policies hierarchy.
+	TUniquePtr<FDisplayClusterVBlankMonitor> VBlankMonitor;
 };
diff --git a/Engine/Plugins/Runtime/nDisplay/Source/DisplayCluster/Private/Render/Synchronization/DisplayClusterVblankMonitor.cpp b/Engine/Plugins/Runtime/nDisplay/Source/DisplayCluster/Private/Render/Synchronization/DisplayClusterVblankMonitor.cpp
new file mode 100644
--- /dev/null
+++ b/Engine/Plugins/Runtime/nDisplay/Source/DisplayCluster/Private/Render/Synchronization/DisplayClusterVblankMonitor.cpp
@@ -0,0 +1,59 @@
+// Copyright Epic Games, Inc. All Rights Reserved.
+
+#include "Render/Synchronization/DisplayClusterVblankMonitor.h"
+
+#include "Render/Synchronization/IDisplayClusterRenderSyncPolicy.h"
+#include "HAL/RunnableThread.h"
+#include "Misc/DisplayClusterLog.h"
+#include "Misc/ScopeLock.h"
+
+
+FDisplayClusterVBlankMonitor::FDisplayClusterVBlankMonitor(IDisplayClusterRenderSyncPolicy& InSyncPolicy)
+	: SyncPolicy(InSyncPolicy)
+{
+	WorkingThread.Reset(FRunnableThread::Create(this, TEXT("VBLANK_Monitoring_thread"), 32 * 1024, TPri_AboveNormal));
+}
+
+FDisplayClusterVBlankMonitor::~FDisplayClusterVBlankMonitor()
+{
+	// Stop working thread
+	Stop();
+}
+
+uint32 FDisplayClusterVBlankMonitor::Run()
+{
+	UE_LOG(LogDisplayClusterRenderSync, Log, TEXT("V-blank monitoring thread has started"));
+
+	while (!bThreadExitRequested)
+	{
+		if (SyncPolicy.WaitForVBlank())
+		{
+			TRACE_BOOKMARK(TEXT("VBLANK %llu"), VBlankCounter);
+			++VBlankCounter;
+		}
+		else
+		{
+			TRACE_CPUPROFILER_EVENT_SCOPE(VBLANK - not available);
+			FPlatformProcess::Sleep(1.f);
+		}
+	}
+
+	UE_LOG(LogDisplayClusterRenderSync, Log, TEXT("V-blank monitoring thread has finished"));
+
+	return 0;
+}
+
+void FDisplayClusterVBlankMonitor::Stop()
+{
+	FScopeLock Lock(&CritSecInternals);
+
+	// Let the working thread know it has to stop
+	bThreadExitRequested = true;
+
+	// Wait unless working thread is finished
+	if (WorkingThread)
+	{
+		WorkingThread->WaitForCompletion();
+		WorkingThread.Reset();
+	}
+}
diff --git a/Engine/Plugins/Runtime/nDisplay/Source/DisplayCluster/Private/Render/Synchronization/DisplayClusterVblankMonitor.h b/Engine/Plugins/Runtime/nDisplay/Source/DisplayCluster/Private/Render/Synchronization/DisplayClusterVblankMonitor.h
new file mode 100644
--- /dev/null
+++ b/Engine/Plugins/Runtime/nDisplay/Source/DisplayCluster/Private/Render/Synchronization/DisplayClusterVblankMonitor.h
@@ -0,0 +1,44 @@
+// Copyright Epic Games, Inc. All Rights Reserved.
+
+#pragma once
+
+#include "CoreMinimal.h"
+#include "HAL/Runnable.h"
+
+class IDisplayClusterRenderSyncPolicy;
+
+
+/**
+ * V-blank monitor. It continuously polls the current sync policy
+ * for V-blanks and reports the time data to utraces.
+ */
+class FDisplayClusterVBlankMonitor
+	: protected FRunnable
+{
+public:
+	FDisplayClusterVBlankMonitor(IDisplayClusterRenderSyncPolicy& SyncPolicy);
+	~FDisplayClusterVBlankMonitor();
+
+protected:
+	//////////////////////////////////////////////////////////////////////////////////////////////
+	// FRunnable
+	//////////////////////////////////////////////////////////////////////////////////////////////
+	virtual uint32 Run() override;
+	virtual void Stop() override final;
+
+private:
+	// Holds the thread object
+	TUniquePtr<FRunnableThread> WorkingThread;
+
+	// Current sync policy
+	IDisplayClusterRenderSyncPolicy& SyncPolicy;
+
+	// V-blank counter
+	uint64 VBlankCounter = 0;
+
+	// Flag to ask working thread to terminate
+	bool bThreadExitRequested = false;
+
+	// Access control mutex
+	mutable FCriticalSection CritSecInternals;
+};
diff --git a/Engine/Plugins/Runtime/nDisplay/Source/DisplayCluster/Public/Cluster/IDisplayClusterClusterManager.h b/Engine/Plugins/Runtime/nDisplay/Source/DisplayCluster/Public/Cluster/IDisplayClusterClusterManager.h
--- a/Engine/Plugins/Runtime/nDisplay/Source/DisplayCluster/Public/Cluster/IDisplayClusterClusterManager.h
+++ b/Engine/Plugins/Runtime/nDisplay/Source/DisplayCluster/Public/Cluster/IDisplayClusterClusterManager.h
@@ -7,13 +7,12 @@
 #include "UObject/ScriptInterface.h"
 #include "DisplayClusterEnums.h"
 
-struct FDisplayClusterClusterEventBinary;
 struct FDisplayClusterClusterEventJson;
-class IDisplayClusterClusterEventBinaryListener;
+struct FDisplayClusterClusterEventBinary;
 class IDisplayClusterClusterEventJsonListener;
-class IDisplayClusterClusterEventListener;
+class IDisplayClusterClusterEventBinaryListener;
 class IDisplayClusterClusterSyncObject;
-class IDisplayClusterGenericBarriersClient;
+class IDisplayClusterClusterEventListener;
 
 DECLARE_MULTICAST_DELEGATE_OneParam(FOnClusterEventJson, const FDisplayClusterClusterEventJson& /* Event */);
 typedef FOnClusterEventJson::FDelegate FOnClusterEventJsonListener;
@@ -55,9 +54,6 @@ public:
 	/** Unregisters synchronization object. */
 	virtual void UnregisterSyncObject(IDisplayClusterClusterSyncObject* SyncObj) = 0;
 
-	/** Spawns a client for generic barriers service. */
-	virtual TUniquePtr<IDisplayClusterGenericBarriersClient> CreateGenericBarriersClient(const FString& ClientName) = 0;
-
 	//////////////////////////////////////////////////////////////////////////////////////////////
 	// Cluster events
 	//////////////////////////////////////////////////////////////////////////////////////////////
diff --git a/Engine/Plugins/Runtime/nDisplay/Source/DisplayCluster/Public/Cluster/IDisplayClusterGenericBarriersClient.h b/Engine/Plugins/Runtime/nDisplay/Source/DisplayCluster/Public/Cluster/IDisplayClusterGenericBarriersClient.h
deleted file mode 100644
diff --git a/Engine/Plugins/Runtime/nDisplay/Source/DisplayCluster/Public/Render/IDisplayClusterRenderManager.h b/Engine/Plugins/Runtime/nDisplay/Source/DisplayCluster/Public/Render/IDisplayClusterRenderManager.h
--- a/Engine/Plugins/Runtime/nDisplay/Source/DisplayCluster/Public/Render/IDisplayClusterRenderManager.h
+++ b/Engine/Plugins/Runtime/nDisplay/Source/DisplayCluster/Public/Render/IDisplayClusterRenderManager.h
@@ -13,7 +13,6 @@ class IDisplayClusterProjectionPolicyFactory;
 class IDisplayClusterRenderDeviceFactory;
 class IDisplayClusterRenderSyncPolicy;
 class IDisplayClusterRenderSyncPolicyFactory;
-class IDisplayClusterVblankMonitor;
 class IDisplayClusterViewportManager; 
 class UCineCameraComponent;
 struct FPostProcessSettings;
@@ -163,11 +162,6 @@ public:
 	*/
 	virtual void GetRegisteredPostProcess(TArray<FString>& OutPostProcessIDs) const = 0;
 
-	/**
-	* Returns V-blank monitor interface
-	*/
-	virtual TSharedPtr<IDisplayClusterVblankMonitor, ESPMode::ThreadSafe> GetVblankMonitor() = 0;
-
 	/**
 	* @return - Current viewport manager from root actor
 	*/
diff --git a/Engine/Plugins/Runtime/nDisplay/Source/DisplayCluster/Public/Render/Monitoring/IDisplayClusterVblankMonitor.h b/Engine/Plugins/Runtime/nDisplay/Source/DisplayCluster/Public/Render/Monitoring/IDisplayClusterVblankMonitor.h
deleted file mode 100644
diff --git a/Engine/Plugins/Runtime/nDisplay/Source/DisplayClusterConfiguration/Private/DisplayClusterConfigurationTypes_Media.cpp b/Engine/Plugins/Runtime/nDisplay/Source/DisplayClusterConfiguration/Private/DisplayClusterConfigurationTypes_Media.cpp
--- a/Engine/Plugins/Runtime/nDisplay/Source/DisplayClusterConfiguration/Private/DisplayClusterConfigurationTypes_Media.cpp
+++ b/Engine/Plugins/Runtime/nDisplay/Source/DisplayClusterConfiguration/Private/DisplayClusterConfigurationTypes_Media.cpp
@@ -5,24 +5,24 @@
 #include "DisplayClusterConfigurationTypes_Media.h"
 
 
-bool FDisplayClusterConfigurationMediaICVFX::IsMediaInputAssigned(const FString& NodeId) const
+bool FDisplayClusterConfigurationMediaICVFX::IsMediaInputAssigned(const FString& InNodeId) const
 {
-	return GetMediaSource(NodeId) != nullptr;
+	return GetMediaSource(InNodeId) != nullptr;
 }
 
-bool FDisplayClusterConfigurationMediaICVFX::IsMediaOutputAssigned(const FString& NodeId) const
+bool FDisplayClusterConfigurationMediaICVFX::IsMediaOutputAssigned(const FString& InNodeId) const
 {
-	return GetMediaOutput(NodeId) != nullptr;
+	return GetMediaOutput(InNodeId) != nullptr;
 }
 
-UMediaSource* FDisplayClusterConfigurationMediaICVFX::GetMediaSource(const FString& NodeId) const
+UMediaSource* FDisplayClusterConfigurationMediaICVFX::GetMediaSource(const FString& InNodeId) const
 {
 	// Look up for a group that contains node ID specified
 	for (const FDisplayClusterConfigurationMediaInputGroup& MediaInputGroup : MediaInputGroups)
 	{
-		const bool bFound = MediaInputGroup.ClusterNodes.ItemNames.ContainsByPredicate([NodeId](const FString& Item)
+		const bool bFound = MediaInputGroup.ClusterNodes.ItemNames.ContainsByPredicate([InNodeId](const FString& Item)
 			{
-				return Item.Equals(NodeId, ESearchCase::IgnoreCase);
+				return Item.Equals(InNodeId, ESearchCase::IgnoreCase);
 			});
 
 		if (bFound)
@@ -34,14 +34,14 @@ UMediaSource* FDisplayClusterConfigurationMediaICVFX::GetMediaSource(const FStri
 	return nullptr;
 }
 
-UMediaOutput* FDisplayClusterConfigurationMediaICVFX::GetMediaOutput(const FString& NodeId) const
+UMediaOutput* FDisplayClusterConfigurationMediaICVFX::GetMediaOutput(const FString& InNodeId) const
 {
 	// Look up for a group that contains node ID specified
 	for (const FDisplayClusterConfigurationMediaOutputGroup& MediaOutputGroup : MediaOutputGroups)
 	{
-		const bool bFound = MediaOutputGroup.ClusterNodes.ItemNames.ContainsByPredicate([NodeId](const FString& Item)
+		const bool bFound = MediaOutputGroup.ClusterNodes.ItemNames.ContainsByPredicate([InNodeId](const FString& Item)
 			{
-				return Item.Equals(NodeId, ESearchCase::IgnoreCase);
+				return Item.Equals(InNodeId, ESearchCase::IgnoreCase);
 			});
 
 		if (bFound)
@@ -52,22 +52,3 @@ UMediaOutput* FDisplayClusterConfigurationMediaICVFX::GetMediaOutput(const FStri
 
 	return nullptr;
 }
-
-UDisplayClusterMediaOutputSynchronizationPolicy* FDisplayClusterConfigurationMediaICVFX::GetOutputSyncPolicy(const FString& NodeId) const
-{
-	// Look up for a group that contains node ID specified
-	for (const FDisplayClusterConfigurationMediaOutputGroup& MediaOutputGroup : MediaOutputGroups)
-	{
-		const bool bFound = MediaOutputGroup.ClusterNodes.ItemNames.ContainsByPredicate([NodeId](const FString& Item)
-			{
-				return Item.Equals(NodeId, ESearchCase::IgnoreCase);
-			});
-
-		if (bFound)
-		{
-			return MediaOutputGroup.OutputSyncPolicy;
-		}
-	}
-
-	return nullptr;
-}
diff --git a/Engine/Plugins/Runtime/nDisplay/Source/DisplayClusterConfiguration/Public/DisplayClusterConfigurationTypes.h b/Engine/Plugins/Runtime/nDisplay/Source/DisplayClusterConfiguration/Public/DisplayClusterConfigurationTypes.h
--- a/Engine/Plugins/Runtime/nDisplay/Source/DisplayClusterConfiguration/Public/DisplayClusterConfigurationTypes.h
+++ b/Engine/Plugins/Runtime/nDisplay/Source/DisplayClusterConfiguration/Public/DisplayClusterConfigurationTypes.h
@@ -171,7 +171,7 @@ public:
 	FDisplayClusterConfigurationPrimaryNodePorts();
 
 public:
-	/** Advanced: TCP port for nDisplay internal services */
+	/** Advanced: network port for Cluster Sync Events */
 	UPROPERTY(EditAnywhere, Category = NDisplay, meta = (ClampMin = "1024", ClampMax = "65535", UIMin = "1024", UIMax = "65535"))
 	uint16 ClusterSync;
 
@@ -359,7 +359,7 @@ public:
 	TMap<FString, FDisplayClusterConfigurationPostprocess> Postprocess;
 
 	// Media settings
-	UPROPERTY(EditAnywhere, BlueprintReadWrite, Category = "Media")
+	UPROPERTY(EditAnywhere, BlueprintReadWrite, Category = "Media", meta = (ShowOnlyInnerProperties))
 	FDisplayClusterConfigurationMedia Media;
 
 #if WITH_EDITORONLY_DATA
diff --git a/Engine/Plugins/Runtime/nDisplay/Source/DisplayClusterConfiguration/Public/DisplayClusterConfigurationTypes_Media.h b/Engine/Plugins/Runtime/nDisplay/Source/DisplayClusterConfiguration/Public/DisplayClusterConfigurationTypes_Media.h
--- a/Engine/Plugins/Runtime/nDisplay/Source/DisplayClusterConfiguration/Public/DisplayClusterConfigurationTypes_Media.h
+++ b/Engine/Plugins/Runtime/nDisplay/Source/DisplayClusterConfiguration/Public/DisplayClusterConfigurationTypes_Media.h
@@ -13,34 +13,6 @@
 
 #include "DisplayClusterConfigurationTypes_Media.generated.h"
 
-
-/*
- * Base media output synchronization policy class
- */
-UCLASS(Abstract, editinlinenew, BlueprintType, hidecategories = (Object))
-class DISPLAYCLUSTERCONFIGURATION_API UDisplayClusterMediaOutputSynchronizationPolicy
-	: public UObject
-{
-	GENERATED_BODY()
-
-public:
-	/** Returns true if specified media capture type can be synchonized by the policy implementation */
-	virtual bool IsCaptureTypeSupported(UMediaCapture* MediaCapture) const
-	{
-		return true;
-	}
-
-	/** Starts synchronization of specific output stream (capture device). Returns false if failed. */
-	virtual bool StartSynchronization(UMediaCapture* MediaCapture, const FString& MediaId) PURE_VIRTUAL(UDisplayClusterMediaOutputSynchronizationPolicy::StartSynchronization, return false;);
-
-	/** Stops synchronization of specific output stream (capture device). */
-	virtual void StopSynchronization() PURE_VIRTUAL(UDisplayClusterMediaOutputSynchronizationPolicy::StopSynchronization, );
-
-	/** Returns true if currently synchronising a media output. */
-	virtual bool IsRunning() PURE_VIRTUAL(UDisplayClusterMediaOutputSynchronizationPolicy::IsRunning, return false; );
-};
-
-
 PRAGMA_DISABLE_DEPRECATION_WARNINGS
 
 /*
@@ -61,18 +33,15 @@ public:
 	TObjectPtr<UMediaSource> MediaSource;
 
 	/** Media output to use */
-	UPROPERTY(Instanced, EditAnywhere, BlueprintReadWrite, Category = "Output")
+	UPROPERTY(Instanced, EditAnywhere, BlueprintReadWrite, Category = "Media")
 	TObjectPtr<UMediaOutput> MediaOutput;
 
-	/** Media output synchronization policy */
-	UPROPERTY(Instanced, EditAnywhere, BlueprintReadWrite, Category = "Output", meta = (DisplayName = "Capture Synchronization"))
-	TObjectPtr<UDisplayClusterMediaOutputSynchronizationPolicy> OutputSyncPolicy;
-
-protected:
 #if WITH_EDITORONLY_DATA
+
 	UE_DEPRECATED(5.2, "This property has been deprecated")
 	UPROPERTY()
 	FString MediaSharingNode_DEPRECATED;
+
 #endif // WITH_EDITORONLY_DATA
 };
 
@@ -114,10 +83,6 @@ public:
 	/** Media output to use */
 	UPROPERTY(Instanced, EditAnywhere, BlueprintReadWrite, Category = "Media")
 	TObjectPtr<UMediaOutput> MediaOutput;
-
-	/** Media output synchronization policy */
-	UPROPERTY(Instanced, EditAnywhere, BlueprintReadWrite, Category = "Media", meta = (DisplayName = "Capture Synchronization"))
-	TObjectPtr<UDisplayClusterMediaOutputSynchronizationPolicy> OutputSyncPolicy;
 };
 
 
@@ -144,19 +109,16 @@ public:
 
 public:
 	/** Returns true if a specific cluster node has media source assigned */
-	bool IsMediaInputAssigned(const FString& NodeId) const;
+	bool IsMediaInputAssigned(const FString& InNodeId) const;
 
 	/** Returns true if a specific cluster node has media source assigned */
-	bool IsMediaOutputAssigned(const FString& NodeId) const;
+	bool IsMediaOutputAssigned(const FString& InNodeId) const;
 
 	/** Returns media source bound to a specific cluster node */
-	UMediaSource* GetMediaSource(const FString& NodeId) const;
+	UMediaSource* GetMediaSource(const FString& InNodeId) const;
 
 	/** Returns media output bound to a specific cluster node */
-	UMediaOutput* GetMediaOutput(const FString& NodeId) const;
-
-	/** Returns media output sync policy bound to a specific cluster node */
-	UDisplayClusterMediaOutputSynchronizationPolicy* GetOutputSyncPolicy(const FString& NodeId) const;
+	UMediaOutput* GetMediaOutput(const FString& InNodeId) const;
 
 	UE_DEPRECATED(5.2, "This function has beend deprecated.")
 	bool IsMediaSharingNode(const FString& InNodeId) const
diff --git a/Engine/Plugins/Runtime/nDisplay/Source/DisplayClusterMedia/DisplayClusterMedia.Build.cs b/Engine/Plugins/Runtime/nDisplay/Source/DisplayClusterMedia/DisplayClusterMedia.Build.cs
--- a/Engine/Plugins/Runtime/nDisplay/Source/DisplayClusterMedia/DisplayClusterMedia.Build.cs
+++ b/Engine/Plugins/Runtime/nDisplay/Source/DisplayClusterMedia/DisplayClusterMedia.Build.cs
@@ -13,7 +13,6 @@ public class DisplayClusterMedia : ModuleRules
 
 		PublicDependencyModuleNames.AddRange(
 			new string[] {
-				"DisplayClusterConfiguration",
 				"Media",
 				"MediaAssets",
 				"MediaIOCore",
@@ -24,12 +23,11 @@ public class DisplayClusterMedia : ModuleRules
 				"Core",
 				"CoreUObject",
 				"DisplayCluster",
+				"DisplayClusterConfiguration",
 				"Engine",
 				"Renderer",
 				"RenderCore",
 				"RHI",
-				"RivermaxCore",
-				"RivermaxMedia",
 			});
 
 		if (Target.bBuildEditor == true)
diff --git a/Engine/Plugins/Runtime/nDisplay/Source/DisplayClusterMedia/Private/Capture/DisplayClusterMediaCaptureBase.cpp b/Engine/Plugins/Runtime/nDisplay/Source/DisplayClusterMedia/Private/Capture/DisplayClusterMediaCaptureBase.cpp
--- a/Engine/Plugins/Runtime/nDisplay/Source/DisplayClusterMedia/Private/Capture/DisplayClusterMediaCaptureBase.cpp
+++ b/Engine/Plugins/Runtime/nDisplay/Source/DisplayClusterMedia/Private/Capture/DisplayClusterMediaCaptureBase.cpp
@@ -6,18 +6,15 @@
 #include "IDisplayCluster.h"
 #include "IDisplayClusterCallbacks.h"
 
-#include "DisplayClusterConfigurationTypes_Media.h"
-
 #include "MediaCapture.h"
 #include "MediaOutput.h"
 
 #include "RenderGraphBuilder.h"
 
 
-FDisplayClusterMediaCaptureBase::FDisplayClusterMediaCaptureBase(const FString& InMediaId, const FString& InClusterNodeId, UMediaOutput* InMediaOutput, UDisplayClusterMediaOutputSynchronizationPolicy* InSyncPolicy)
+FDisplayClusterMediaCaptureBase::FDisplayClusterMediaCaptureBase(const FString& InMediaId, const FString& InClusterNodeId, UMediaOutput* InMediaOutput)
 	: FDisplayClusterMediaBase(InMediaId, InClusterNodeId)
 	, MediaOutput(InMediaOutput)
-	, SyncPolicy(InSyncPolicy)
 {
 	check(InMediaOutput);
 
@@ -32,20 +29,12 @@ FDisplayClusterMediaCaptureBase::~FDisplayClusterMediaCaptureBase()
 
 void FDisplayClusterMediaCaptureBase::AddReferencedObjects(FReferenceCollector& Collector)
 {
-	if (MediaOutput)
-	{
-		Collector.AddReferencedObject(MediaOutput);
-	}
+	Collector.AddReferencedObject(MediaOutput);
 
 	if (MediaCapture)
 	{
 		Collector.AddReferencedObject(MediaCapture);
 	}
-
-	if (SyncPolicy)
-	{
-		Collector.AddReferencedObject(SyncPolicy);
-	}
 }
 
 bool FDisplayClusterMediaCaptureBase::StartCapture()
@@ -57,15 +46,6 @@ bool FDisplayClusterMediaCaptureBase::StartCapture()
 		{
 			MediaCapture->SetMediaOutput(MediaOutput);
 
-			// Initialize and start capture synchronization
-			if (SyncPolicy && SyncPolicy->IsCaptureTypeSupported(MediaCapture))
-			{
-				if (!SyncPolicy->StartSynchronization(MediaCapture, GetMediaId()))
-				{
-					UE_LOG(LogDisplayClusterMedia, Warning, TEXT("MediaCapture '%s': couldn't start synchronization."), *GetMediaId());
-				}
-			}
-
 			bWasCaptureStarted = StartMediaCapture();
 			return bWasCaptureStarted;
 		}
@@ -76,13 +56,6 @@ bool FDisplayClusterMediaCaptureBase::StartCapture()
 
 void FDisplayClusterMediaCaptureBase::StopCapture()
 {
-	// Stop synchronization
-	if (SyncPolicy)
-	{
-		SyncPolicy->StopSynchronization();
-	}
-
-	// Stop capture
 	if (MediaCapture)
 	{
 		MediaCapture->StopCapture(false);
@@ -97,8 +70,6 @@ void FDisplayClusterMediaCaptureBase::ExportMediaData(FRDGBuilder& GraphBuilder,
 
 	if (SrcTexture)
 	{
-		UE_LOG(LogDisplayClusterMedia, Verbose, TEXT("MediaCapture '%s': exporting texture on RT frame '%lu'..."), *GetMediaId(), GFrameCounterRenderThread);
-
 		MediaCapture->SetValidSourceGPUMask(GraphBuilder.RHICmdList.GetGPUMask());
 
 		LastSrcRegionSize = FIntSize(TextureInfo.Region.Size());
diff --git a/Engine/Plugins/Runtime/nDisplay/Source/DisplayClusterMedia/Private/Capture/DisplayClusterMediaCaptureBase.h b/Engine/Plugins/Runtime/nDisplay/Source/DisplayClusterMedia/Private/Capture/DisplayClusterMediaCaptureBase.h
--- a/Engine/Plugins/Runtime/nDisplay/Source/DisplayClusterMedia/Private/Capture/DisplayClusterMediaCaptureBase.h
+++ b/Engine/Plugins/Runtime/nDisplay/Source/DisplayClusterMedia/Private/Capture/DisplayClusterMediaCaptureBase.h
@@ -12,7 +12,6 @@
 class FRDGBuilder;
 class UMediaCapture;
 class UMediaOutput;
-class UDisplayClusterMediaOutputSynchronizationPolicy;
 
 
 /**
@@ -23,7 +22,7 @@ class FDisplayClusterMediaCaptureBase
 	, public FGCObject
 {
 public:
-	FDisplayClusterMediaCaptureBase(const FString& MediaId, const FString& ClusterNodeId, UMediaOutput* MediaOutput, UDisplayClusterMediaOutputSynchronizationPolicy* SyncPolicy = nullptr);
+	FDisplayClusterMediaCaptureBase(const FString& MediaId, const FString& ClusterNodeId, UMediaOutput* MediaOutput);
 	virtual ~FDisplayClusterMediaCaptureBase();
 
 public:
@@ -39,11 +38,6 @@ public:
 	virtual bool StartCapture();
 	virtual void StopCapture();
 
-	UMediaCapture* GetMediaCapture() const
-	{
-		return MediaCapture;
-	}
-
 protected:
 	void ExportMediaData(FRDGBuilder& GraphBuilder, const FMediaTextureInfo& TextureInfo);
 	void OnPostClusterTick();
@@ -70,9 +64,8 @@ private:
 
 private:
 	//~ Begin GC by AddReferencedObjects
-	UMediaOutput*  MediaOutput  = nullptr;
-	UMediaCapture* MediaCapture = nullptr;
-	UDisplayClusterMediaOutputSynchronizationPolicy* SyncPolicy = nullptr;
+	UMediaOutput*           MediaOutput  = nullptr;
+	UMediaCapture*          MediaCapture = nullptr;
 	//~ End GC by AddReferencedObjects
 
 	// Used to restart media capture in the case it falls in error
diff --git a/Engine/Plugins/Runtime/nDisplay/Source/DisplayClusterMedia/Private/Capture/DisplayClusterMediaCaptureCamera.cpp b/Engine/Plugins/Runtime/nDisplay/Source/DisplayClusterMedia/Private/Capture/DisplayClusterMediaCaptureCamera.cpp
--- a/Engine/Plugins/Runtime/nDisplay/Source/DisplayClusterMedia/Private/Capture/DisplayClusterMediaCaptureCamera.cpp
+++ b/Engine/Plugins/Runtime/nDisplay/Source/DisplayClusterMedia/Private/Capture/DisplayClusterMediaCaptureCamera.cpp
@@ -15,8 +15,8 @@
 #include "RHIResources.h"
 
 
-FDisplayClusterMediaCaptureCamera::FDisplayClusterMediaCaptureCamera(const FString& InMediaId, const FString& InClusterNodeId, const FString& InCameraId, const FString& InViewportId, UMediaOutput* InMediaOutput, UDisplayClusterMediaOutputSynchronizationPolicy* SyncPolicy)
-	: FDisplayClusterMediaCaptureViewport(InMediaId, InClusterNodeId, InViewportId, InMediaOutput, SyncPolicy)
+FDisplayClusterMediaCaptureCamera::FDisplayClusterMediaCaptureCamera(const FString& InMediaId, const FString& InClusterNodeId, const FString& InCameraId, const FString& InViewportId, UMediaOutput* InMediaOutput)
+	: FDisplayClusterMediaCaptureViewport(InMediaId, InClusterNodeId, InViewportId, InMediaOutput)
 	, CameraId(InCameraId)
 {
 	if (const ADisplayClusterRootActor* const RootActor = IDisplayCluster::Get().GetGameMgr()->GetRootActor())
diff --git a/Engine/Plugins/Runtime/nDisplay/Source/DisplayClusterMedia/Private/Capture/DisplayClusterMediaCaptureCamera.h b/Engine/Plugins/Runtime/nDisplay/Source/DisplayClusterMedia/Private/Capture/DisplayClusterMediaCaptureCamera.h
--- a/Engine/Plugins/Runtime/nDisplay/Source/DisplayClusterMedia/Private/Capture/DisplayClusterMediaCaptureCamera.h
+++ b/Engine/Plugins/Runtime/nDisplay/Source/DisplayClusterMedia/Private/Capture/DisplayClusterMediaCaptureCamera.h
@@ -16,7 +16,7 @@ class FDisplayClusterMediaCaptureCamera
 	: public FDisplayClusterMediaCaptureViewport
 {
 public:
-	FDisplayClusterMediaCaptureCamera(const FString& MediaId, const FString& ClusterNodeId, const FString& CameraId, const FString& ViewportId, UMediaOutput* MediaOutput, UDisplayClusterMediaOutputSynchronizationPolicy* SyncPolicy = nullptr);
+	FDisplayClusterMediaCaptureCamera(const FString& MediaId, const FString& ClusterNodeId, const FString& CameraId, const FString& ViewportId, UMediaOutput* MediaOutput);
 
 protected:
 	
diff --git a/Engine/Plugins/Runtime/nDisplay/Source/DisplayClusterMedia/Private/Capture/DisplayClusterMediaCaptureNode.cpp b/Engine/Plugins/Runtime/nDisplay/Source/DisplayClusterMedia/Private/Capture/DisplayClusterMediaCaptureNode.cpp
--- a/Engine/Plugins/Runtime/nDisplay/Source/DisplayClusterMedia/Private/Capture/DisplayClusterMediaCaptureNode.cpp
+++ b/Engine/Plugins/Runtime/nDisplay/Source/DisplayClusterMedia/Private/Capture/DisplayClusterMediaCaptureNode.cpp
@@ -15,8 +15,8 @@
 #include "UnrealClient.h"
 
 
-FDisplayClusterMediaCaptureNode::FDisplayClusterMediaCaptureNode(const FString& InMediaId, const FString& InClusterNodeId, UMediaOutput* InMediaOutput, UDisplayClusterMediaOutputSynchronizationPolicy* SyncPolicy)
-	: FDisplayClusterMediaCaptureBase(InMediaId, InClusterNodeId, InMediaOutput, SyncPolicy)
+FDisplayClusterMediaCaptureNode::FDisplayClusterMediaCaptureNode(const FString& InMediaId, const FString& InClusterNodeId, UMediaOutput* InMediaOutput)
+	: FDisplayClusterMediaCaptureBase(InMediaId, InClusterNodeId, InMediaOutput)
 {
 }
 
diff --git a/Engine/Plugins/Runtime/nDisplay/Source/DisplayClusterMedia/Private/Capture/DisplayClusterMediaCaptureNode.h b/Engine/Plugins/Runtime/nDisplay/Source/DisplayClusterMedia/Private/Capture/DisplayClusterMediaCaptureNode.h
--- a/Engine/Plugins/Runtime/nDisplay/Source/DisplayClusterMedia/Private/Capture/DisplayClusterMediaCaptureNode.h
+++ b/Engine/Plugins/Runtime/nDisplay/Source/DisplayClusterMedia/Private/Capture/DisplayClusterMediaCaptureNode.h
@@ -18,7 +18,7 @@ class FDisplayClusterMediaCaptureNode
 	: public FDisplayClusterMediaCaptureBase
 {
 public:
-	FDisplayClusterMediaCaptureNode(const FString& MediaId, const FString& ClusterNodeId, UMediaOutput* MediaOutput, UDisplayClusterMediaOutputSynchronizationPolicy* SyncPolicy = nullptr);
+	FDisplayClusterMediaCaptureNode(const FString& MediaId, const FString& ClusterNodeId, UMediaOutput* MediaOutput);
 
 public:
 	virtual bool StartCapture() override;
diff --git a/Engine/Plugins/Runtime/nDisplay/Source/DisplayClusterMedia/Private/Capture/DisplayClusterMediaCaptureViewport.cpp b/Engine/Plugins/Runtime/nDisplay/Source/DisplayClusterMedia/Private/Capture/DisplayClusterMediaCaptureViewport.cpp
--- a/Engine/Plugins/Runtime/nDisplay/Source/DisplayClusterMedia/Private/Capture/DisplayClusterMediaCaptureViewport.cpp
+++ b/Engine/Plugins/Runtime/nDisplay/Source/DisplayClusterMedia/Private/Capture/DisplayClusterMediaCaptureViewport.cpp
@@ -12,8 +12,8 @@
 #include "RHIResources.h"
 
 
-FDisplayClusterMediaCaptureViewport::FDisplayClusterMediaCaptureViewport(const FString& InMediaId, const FString& InClusterNodeId, const FString& InViewportId, UMediaOutput* InMediaOutput, UDisplayClusterMediaOutputSynchronizationPolicy* SyncPolicy)
-	: FDisplayClusterMediaCaptureBase(InMediaId, InClusterNodeId, InMediaOutput, SyncPolicy)
+FDisplayClusterMediaCaptureViewport::FDisplayClusterMediaCaptureViewport(const FString& InMediaId, const FString& InClusterNodeId, const FString& InViewportId, UMediaOutput* InMediaOutput)
+	: FDisplayClusterMediaCaptureBase(InMediaId, InClusterNodeId, InMediaOutput)
 	, ViewportId(InViewportId)
 {
 }
diff --git a/Engine/Plugins/Runtime/nDisplay/Source/DisplayClusterMedia/Private/Capture/DisplayClusterMediaCaptureViewport.h b/Engine/Plugins/Runtime/nDisplay/Source/DisplayClusterMedia/Private/Capture/DisplayClusterMediaCaptureViewport.h
--- a/Engine/Plugins/Runtime/nDisplay/Source/DisplayClusterMedia/Private/Capture/DisplayClusterMediaCaptureViewport.h
+++ b/Engine/Plugins/Runtime/nDisplay/Source/DisplayClusterMedia/Private/Capture/DisplayClusterMediaCaptureViewport.h
@@ -17,7 +17,7 @@ class FDisplayClusterMediaCaptureViewport
 	: public FDisplayClusterMediaCaptureBase
 {
 public:
-	FDisplayClusterMediaCaptureViewport(const FString& MediaId, const FString& ClusterNodeId, const FString& ViewportId, UMediaOutput* MediaOutput, UDisplayClusterMediaOutputSynchronizationPolicy* SyncPolicy = nullptr);
+	FDisplayClusterMediaCaptureViewport(const FString& MediaId, const FString& ClusterNodeId, const FString& ViewportId, UMediaOutput* MediaOutput);
 
 public:
 	virtual bool StartCapture() override;
diff --git a/Engine/Plugins/Runtime/nDisplay/Source/DisplayClusterMedia/Private/DisplayClusterMediaHelpers.cpp b/Engine/Plugins/Runtime/nDisplay/Source/DisplayClusterMedia/Private/DisplayClusterMediaHelpers.cpp
--- a/Engine/Plugins/Runtime/nDisplay/Source/DisplayClusterMedia/Private/DisplayClusterMediaHelpers.cpp
+++ b/Engine/Plugins/Runtime/nDisplay/Source/DisplayClusterMedia/Private/DisplayClusterMediaHelpers.cpp
@@ -10,51 +10,6 @@
 
 namespace DisplayClusterMediaHelpers
 {
-	namespace MediaId
-	{
-		// Generate media ID for a specific entity
-		FString GenerateMediaId(EMediaDeviceType DeviceType, EMediaOwnerType OwnerType, const FString& NodeId, const FString& DCRAName, const FString& OwnerName)
-		{
-			if (DeviceType == EMediaDeviceType::Input)
-			{
-				switch (OwnerType)
-				{
-				case EMediaOwnerType::Backbuffer:
-					return FString::Printf(TEXT("%s_%s_backbuffer_input"), *NodeId, *DCRAName);
-
-				case EMediaOwnerType::Viewport:
-					return FString::Printf(TEXT("%s_%s_%s_viewport_input"), *NodeId, *DCRAName, *OwnerName);
-
-				case EMediaOwnerType::ICVFXCamera:
-					return FString::Printf(TEXT("%s_%s_%s_icvfx_input"), *NodeId, *DCRAName, *OwnerName);
-
-				default:
-					unimplemented();
-				}
-			}
-			else if (DeviceType == EMediaDeviceType::Output)
-			{
-				switch (OwnerType)
-				{
-				case EMediaOwnerType::Backbuffer:
-					return FString::Printf(TEXT("%s_%s_backbuffer_capture"), *NodeId, *DCRAName);
-
-				case EMediaOwnerType::Viewport:
-					return FString::Printf(TEXT("%s_%s_%s_viewport_capture"), *NodeId, *DCRAName, *OwnerName);
-
-				case EMediaOwnerType::ICVFXCamera:
-					return FString::Printf(TEXT("%s_%s_%s_icvfx_capture"), *NodeId, *DCRAName, *OwnerName);
-
-				default:
-					unimplemented();
-				}
-			}
-
-			// Should never get here
-			return FString();
-		}
-	}
-
 	//@todo This needs to be exposed from the DisplayCluster core module after its refactoring
 	FString GenerateICVFXViewportName(const FString& ClusterNodeId, const FString& ICVFXCameraName)
 	{
diff --git a/Engine/Plugins/Runtime/nDisplay/Source/DisplayClusterMedia/Private/DisplayClusterMediaHelpers.h b/Engine/Plugins/Runtime/nDisplay/Source/DisplayClusterMedia/Private/DisplayClusterMediaHelpers.h
--- a/Engine/Plugins/Runtime/nDisplay/Source/DisplayClusterMedia/Private/DisplayClusterMediaHelpers.h
+++ b/Engine/Plugins/Runtime/nDisplay/Source/DisplayClusterMedia/Private/DisplayClusterMediaHelpers.h
@@ -10,27 +10,6 @@ class FRHITexture;
 
 namespace DisplayClusterMediaHelpers
 {
-	namespace MediaId
-	{
-		// Helper for media ID generation
-		enum class EMediaDeviceType : uint8
-		{
-			Input,
-			Output
-		};
-
-		// Helper for media ID generation
-		enum class EMediaOwnerType : uint8
-		{
-			Backbuffer,
-			Viewport,
-			ICVFXCamera
-		};
-
-		// Generate media ID for a specific entity
-		FString GenerateMediaId(EMediaDeviceType DeviceType, EMediaOwnerType OwnerType, const FString& NodeId, const FString& DCRAName, const FString& OwnerName);
-	}
-
 	// Generates internal ICVFX viewport IDs
 	FString GenerateICVFXViewportName(const FString& ClusterNodeId, const FString& ICVFXCameraName);
 
diff --git a/Engine/Plugins/Runtime/nDisplay/Source/DisplayClusterMedia/Private/DisplayClusterMediaLog.cpp b/Engine/Plugins/Runtime/nDisplay/Source/DisplayClusterMedia/Private/DisplayClusterMediaLog.cpp
--- a/Engine/Plugins/Runtime/nDisplay/Source/DisplayClusterMedia/Private/DisplayClusterMediaLog.cpp
+++ b/Engine/Plugins/Runtime/nDisplay/Source/DisplayClusterMedia/Private/DisplayClusterMediaLog.cpp
@@ -5,4 +5,3 @@
 
 // Module-wide log categories
 DEFINE_LOG_CATEGORY(LogDisplayClusterMedia);
-DEFINE_LOG_CATEGORY(LogDisplayClusterMediaSync);
diff --git a/Engine/Plugins/Runtime/nDisplay/Source/DisplayClusterMedia/Private/DisplayClusterMediaLog.h b/Engine/Plugins/Runtime/nDisplay/Source/DisplayClusterMedia/Private/DisplayClusterMediaLog.h
--- a/Engine/Plugins/Runtime/nDisplay/Source/DisplayClusterMedia/Private/DisplayClusterMediaLog.h
+++ b/Engine/Plugins/Runtime/nDisplay/Source/DisplayClusterMedia/Private/DisplayClusterMediaLog.h
@@ -6,5 +6,4 @@
 
 
 // Module-wide log categories
-DECLARE_LOG_CATEGORY_EXTERN(LogDisplayClusterMedia,     Log, All);
-DECLARE_LOG_CATEGORY_EXTERN(LogDisplayClusterMediaSync, Log, All);
+DECLARE_LOG_CATEGORY_EXTERN(LogDisplayClusterMedia, Log, All);
diff --git a/Engine/Plugins/Runtime/nDisplay/Source/DisplayClusterMedia/Private/DisplayClusterMediaModule.cpp b/Engine/Plugins/Runtime/nDisplay/Source/DisplayClusterMedia/Private/DisplayClusterMediaModule.cpp
--- a/Engine/Plugins/Runtime/nDisplay/Source/DisplayClusterMedia/Private/DisplayClusterMediaModule.cpp
+++ b/Engine/Plugins/Runtime/nDisplay/Source/DisplayClusterMedia/Private/DisplayClusterMediaModule.cpp
@@ -109,32 +109,29 @@ void FDisplayClusterMediaModule::InitializeMedia()
 					// Media input
 					if (MediaSettings.MediaSource)
 					{
-						const FString MediaInputId = DisplayClusterMediaHelpers::MediaId::GenerateMediaId(
-							DisplayClusterMediaHelpers::MediaId::EMediaDeviceType::Input,
-							DisplayClusterMediaHelpers::MediaId::EMediaOwnerType::Backbuffer,
-							*ClusterNodeId, *RootActor->GetName(), FString());
+						const FString MediaInputId = RootActor->GetName() + FString("_") + ClusterNodeId + FString("_backbuffer_input");
 
-						UE_LOG(LogDisplayClusterMedia, Log, TEXT("Initializing backbuffer media input '%s'..."), *MediaInputId);
+						UE_LOG(LogDisplayClusterMedia, Log, TEXT("Initializing backbuffer media input for node '%s'..."), *ClusterNodeId);
 
-						InputNode = MakeShared<FDisplayClusterMediaInputNode>(
+						InputNode = MakeUnique<FDisplayClusterMediaInputNode>(
 							MediaInputId, ClusterNodeId,
 							MediaSettings.MediaSource);
+
+						AllInputs.Add(InputNode.Get());
 					}
 
 					// Media capture
 					if (MediaSettings.MediaOutput)
 					{
-						const FString MediaCaptureId = DisplayClusterMediaHelpers::MediaId::GenerateMediaId(
-							DisplayClusterMediaHelpers::MediaId::EMediaDeviceType::Output,
-							DisplayClusterMediaHelpers::MediaId::EMediaOwnerType::Backbuffer,
-							*ClusterNodeId, *RootActor->GetName(), FString());
+						const FString MediaCaptureId = RootActor->GetName() + FString("_") + ClusterNodeId + FString("_backbuffer_capture");
 
-						UE_LOG(LogDisplayClusterMedia, Log, TEXT("Initializing backbuffer media capture '%s'..."), *MediaCaptureId);
+						UE_LOG(LogDisplayClusterMedia, Log, TEXT("Initializing backbuffer media capture for node '%s'..."), *ClusterNodeId);
 
-						CaptureNode = MakeShared<FDisplayClusterMediaCaptureNode>(
+						CaptureNode = MakeUnique<FDisplayClusterMediaCaptureNode>(
 							MediaCaptureId, ClusterNodeId,
-							MediaSettings.MediaOutput,
-							MediaSettings.OutputSyncPolicy);
+							MediaSettings.MediaOutput);
+
+						AllCaptures.Add(CaptureNode.Get());
 					}
 				}
 			}
@@ -152,37 +149,32 @@ void FDisplayClusterMediaModule::InitializeMedia()
 						// Media input
 						if (MediaSettings.MediaSource)
 						{
-							const FString MediaInputId = DisplayClusterMediaHelpers::MediaId::GenerateMediaId(
-								DisplayClusterMediaHelpers::MediaId::EMediaDeviceType::Input,
-								DisplayClusterMediaHelpers::MediaId::EMediaOwnerType::Viewport,
-								*ClusterNodeId, *RootActor->GetName(), ViewportIt.Key);
+							const FString MediaInputId = RootActor->GetName() + FString("_") + ClusterNodeId + FString("_") + ViewportIt.Key + FString("_viewport_input");
 
-							UE_LOG(LogDisplayClusterMedia, Log, TEXT("Initializing viewport media input '%s' for viewport '%s'..."), *MediaInputId, *ViewportIt.Key);
+							UE_LOG(LogDisplayClusterMedia, Log, TEXT("Initializing viewport media input for viewport '%s'..."), *ViewportIt.Key);
 
-							TSharedPtr<FDisplayClusterMediaInputViewport> NewViewportInput = MakeShared<FDisplayClusterMediaInputViewport>(
+							TUniquePtr<FDisplayClusterMediaInputViewport> NewViewportInput = MakeUnique<FDisplayClusterMediaInputViewport>(
 								MediaInputId, ClusterNodeId,
 								ViewportIt.Key,
 								MediaSettings.MediaSource);
 
+							AllInputs.Add(NewViewportInput.Get());
 							InputViewports.Emplace(MediaInputId, MoveTemp(NewViewportInput));
 						}
 
 						// Media capture
 						if (MediaSettings.MediaOutput)
 						{
-							const FString MediaCaptureId = DisplayClusterMediaHelpers::MediaId::GenerateMediaId(
-								DisplayClusterMediaHelpers::MediaId::EMediaDeviceType::Output,
-								DisplayClusterMediaHelpers::MediaId::EMediaOwnerType::Viewport,
-								*ClusterNodeId, *RootActor->GetName(), ViewportIt.Key);
+							const FString MediaCaptureId = RootActor->GetName() + FString("_") + ClusterNodeId + FString("_") + ViewportIt.Key + FString("_viewport_capture");
 
-							UE_LOG(LogDisplayClusterMedia, Log, TEXT("Initializing viewport capture '%s' for viewport '%s'..."), *MediaCaptureId, *ViewportIt.Key);
+							UE_LOG(LogDisplayClusterMedia, Log, TEXT("Initializing viewport capture for viewport '%s'..."), *ViewportIt.Key);
 
-							TSharedPtr<FDisplayClusterMediaCaptureViewport> NewViewportCapture = MakeShared<FDisplayClusterMediaCaptureViewport>(
+							TUniquePtr<FDisplayClusterMediaCaptureViewport> NewViewportCapture = MakeUnique<FDisplayClusterMediaCaptureViewport>(
 								MediaCaptureId, ClusterNodeId,
 								ViewportIt.Key,
-								MediaSettings.MediaOutput,
-								MediaSettings.OutputSyncPolicy);
+								MediaSettings.MediaOutput);
 
+							AllCaptures.Add(NewViewportCapture.Get());
 							CaptureViewports.Emplace(MediaCaptureId, MoveTemp(NewViewportCapture));
 						}
 					}
@@ -192,54 +184,49 @@ void FDisplayClusterMediaModule::InitializeMedia()
 			///////////////////////////////
 			// ICVFX media setup
 			{
-				// Get all ICVFX camera components
-				TArray<UDisplayClusterICVFXCameraComponent*> ICVFXCameraComponents;
-				RootActor->GetComponents(ICVFXCameraComponents);
+				TArray<UActorComponent*> ICVFXCameraComponents;
+				RootActor->GetComponents(UDisplayClusterICVFXCameraComponent::StaticClass(), ICVFXCameraComponents);
 
-				for (const UDisplayClusterICVFXCameraComponent* const ICVFXCameraComponent : ICVFXCameraComponents)
+				for (const UActorComponent* const Component : ICVFXCameraComponents)
 				{
-					const FDisplayClusterConfigurationMediaICVFX& MediaSettings = ICVFXCameraComponent->CameraSettings.RenderSettings.Media;
-
-					if (MediaSettings.bEnable)
+					if (const UDisplayClusterICVFXCameraComponent* const ICVFXCamera = Cast<UDisplayClusterICVFXCameraComponent>(Component))
 					{
-						const FString ICVFXCameraName = ICVFXCameraComponent->GetName();
-						const FString ICVFXViewportId = DisplayClusterMediaHelpers::GenerateICVFXViewportName(ClusterNodeId, ICVFXCameraName);
+						const FDisplayClusterConfigurationMediaICVFX& MediaSettings = ICVFXCamera->CameraSettings.RenderSettings.Media;
 
-						// Media input only
-						if (UMediaSource* MediaSource = MediaSettings.GetMediaSource(ClusterNodeId))
+						if (MediaSettings.bEnable)
 						{
-							const FString MediaInputId = DisplayClusterMediaHelpers::MediaId::GenerateMediaId(
-								DisplayClusterMediaHelpers::MediaId::EMediaDeviceType::Input,
-								DisplayClusterMediaHelpers::MediaId::EMediaOwnerType::ICVFXCamera,
-								*ClusterNodeId, *RootActor->GetName(), ICVFXCameraName);
-
-							UE_LOG(LogDisplayClusterMedia, Log, TEXT("Initializing ICVFX media input '%s' for camera '%s'..."), *MediaInputId, *ICVFXCameraName);
-
-							TSharedPtr<FDisplayClusterMediaInputViewport> NewICVFXInput = MakeShared<FDisplayClusterMediaInputViewport>(
-								MediaInputId, ClusterNodeId,
-								ICVFXViewportId,
-								MediaSource);
-
-							InputViewports.Emplace(MediaInputId, MoveTemp(NewICVFXInput));
-						}
-
-						// Media capture only
-						if (UMediaOutput* MediaOutput = MediaSettings.GetMediaOutput(ClusterNodeId))
-						{
-							const FString MediaCaptureId = DisplayClusterMediaHelpers::MediaId::GenerateMediaId(
-								DisplayClusterMediaHelpers::MediaId::EMediaDeviceType::Output,
-								DisplayClusterMediaHelpers::MediaId::EMediaOwnerType::ICVFXCamera,
-								*ClusterNodeId, *RootActor->GetName(), ICVFXCameraName);
-
-							UE_LOG(LogDisplayClusterMedia, Log, TEXT("Initializing ICVFX capture '%s' for camera '%s'..."), *MediaCaptureId, *ICVFXCameraName);
-
-							TSharedPtr<FDisplayClusterMediaCaptureViewport> NewICVFXCapture = MakeShared<FDisplayClusterMediaCaptureCamera>(
-								MediaCaptureId, ClusterNodeId,
-								ICVFXCameraName, ICVFXViewportId,
-								MediaOutput,
-								MediaSettings.GetOutputSyncPolicy(ClusterNodeId));
-
-							CaptureViewports.Emplace(MediaCaptureId, MoveTemp(NewICVFXCapture));
+							const FString ICVFXCameraName = ICVFXCamera->GetName();
+							const FString ICVFXViewportId = DisplayClusterMediaHelpers::GenerateICVFXViewportName(ClusterNodeId, ICVFXCameraName);
+							const FString MediaCaptureId  = RootActor->GetName() + FString("_") + ICVFXCameraName + FString("_icvfx_capture");
+							const FString MediaInputId    = RootActor->GetName() + FString("_") + ICVFXCameraName + FString("_icvfx_input");
+
+							// Media input only
+							if (UMediaSource* MediaSource = MediaSettings.GetMediaSource(ClusterNodeId))
+							{
+								UE_LOG(LogDisplayClusterMedia, Log, TEXT("Initializing ICVFX media input for camera '%s'..."), *ICVFXCameraName);
+
+								TUniquePtr<FDisplayClusterMediaInputViewport> NewICVFXInput = MakeUnique<FDisplayClusterMediaInputViewport>(
+									MediaInputId, ClusterNodeId,
+									ICVFXViewportId,
+									MediaSource);
+
+								AllInputs.Add(NewICVFXInput.Get());
+								InputViewports.Emplace(MediaInputId, MoveTemp(NewICVFXInput));
+							}
+
+							// Media capture only
+							if (UMediaOutput* MediaOutput = MediaSettings.GetMediaOutput(ClusterNodeId))
+							{
+								UE_LOG(LogDisplayClusterMedia, Log, TEXT("Initializing ICVFX capture for camera '%s'..."), *ICVFXCameraName);
+
+								TUniquePtr<FDisplayClusterMediaCaptureViewport> NewICVFXCapture = MakeUnique<FDisplayClusterMediaCaptureCamera>(
+									MediaCaptureId, ClusterNodeId,
+									ICVFXCameraName, ICVFXViewportId,
+									MediaOutput);
+
+								AllCaptures.Add(NewICVFXCapture.Get());
+								CaptureViewports.Emplace(MediaCaptureId, MoveTemp(NewICVFXCapture));
+							}
 						}
 					}
 				}
@@ -255,70 +242,44 @@ void FDisplayClusterMediaModule::ReleaseMedia()
 
 	CaptureViewports.Reset();
 	CaptureNode.Reset();
+	AllCaptures.Reset();
 
 	InputViewports.Reset();
 	InputNode.Reset();
+	AllInputs.Reset();
 
 	FrameQueue.Release();
 }
 
 void FDisplayClusterMediaModule::StartCapture()
 {
-	// Start viewports capture
-	for (TPair<FString, TSharedPtr<FDisplayClusterMediaCaptureViewport>>& Capture : CaptureViewports)
-	{
-		Capture.Value->StartCapture();
-	}
-
-	// Start backbuffer capture
-	if (CaptureNode)
+	for (FDisplayClusterMediaCaptureBase* Capture : AllCaptures)
 	{
-		CaptureNode->StartCapture();
+		Capture->StartCapture();
 	}
 }
 
 void FDisplayClusterMediaModule::StopCapture()
 {
-	// Stop viewports capture
-	for (TPair<FString, TSharedPtr<FDisplayClusterMediaCaptureViewport>>& Capture : CaptureViewports)
-	{
-		Capture.Value->StopCapture();
-	}
-
-	// Stop backbuffer capture
-	if (CaptureNode)
+	for (FDisplayClusterMediaCaptureBase* Capture : AllCaptures)
 	{
-		CaptureNode->StopCapture();
+		Capture->StopCapture();
 	}
 }
 
 void FDisplayClusterMediaModule::PlayMedia()
 {
-	// Start playback on viewports
-	for (TPair<FString, TSharedPtr<FDisplayClusterMediaInputViewport>>& MediaInput : InputViewports)
+	for (FDisplayClusterMediaInputBase* MediaInput : AllInputs)
 	{
-		MediaInput.Value->Play();
-	}
-
-	// Start playback to the backbuffer
-	if (InputNode)
-	{
-		InputNode->Play();
+		MediaInput->Play();
 	}
 }
 
 void FDisplayClusterMediaModule::StopMedia()
 {
-	// Stop playback on viewports
-	for (TPair<FString, TSharedPtr<FDisplayClusterMediaInputViewport>>& MediaInput : InputViewports)
-	{
-		MediaInput.Value->Stop();
-	}
-
-	// Stop playback to the backbuffer
-	if (InputNode)
+	for (FDisplayClusterMediaInputBase* MediaInput : AllInputs)
 	{
-		InputNode->Stop();
+		MediaInput->Stop();
 	}
 }
 
diff --git a/Engine/Plugins/Runtime/nDisplay/Source/DisplayClusterMedia/Private/DisplayClusterMediaModule.h b/Engine/Plugins/Runtime/nDisplay/Source/DisplayClusterMedia/Private/DisplayClusterMediaModule.h
--- a/Engine/Plugins/Runtime/nDisplay/Source/DisplayClusterMedia/Private/DisplayClusterMediaModule.h
+++ b/Engine/Plugins/Runtime/nDisplay/Source/DisplayClusterMedia/Private/DisplayClusterMediaModule.h
@@ -5,15 +5,18 @@
 #include "CoreMinimal.h"
 #include "Modules/ModuleInterface.h"
 
-#include "Synchronization/LatencyQueue/DisplayClusterFrameQueue.h"
+#include "Synchronization/DisplayClusterFrameQueue.h"
 
+class FDisplayClusterMediaCaptureBase;
+class FDisplayClusterMediaCaptureICVFX;
 class FDisplayClusterMediaCaptureNode;
 class FDisplayClusterMediaCaptureViewport;
+class FDisplayClusterMediaInputBase;
+class FDisplayClusterMediaInputICVFX;
 class FDisplayClusterMediaInputNode;
 class FDisplayClusterMediaInputViewport;
 class IMediaPlayerFactory;
 
-
 /**
  * Media module
  */
@@ -50,11 +53,13 @@ private:
 	void OnEnginePreExit();
 
 private:
-	TMap<FString, TSharedPtr<FDisplayClusterMediaCaptureViewport>> CaptureViewports;
-	TSharedPtr<FDisplayClusterMediaCaptureNode>                    CaptureNode;
+	TMap<FString, TUniquePtr<FDisplayClusterMediaCaptureViewport>> CaptureViewports;
+	TUniquePtr<FDisplayClusterMediaCaptureNode>                    CaptureNode;
+	TArray<FDisplayClusterMediaCaptureBase*>                       AllCaptures;
 
-	TMap<FString, TSharedPtr<FDisplayClusterMediaInputViewport>>   InputViewports;
-	TSharedPtr<FDisplayClusterMediaInputNode>                      InputNode;
+	TMap<FString, TUniquePtr<FDisplayClusterMediaInputViewport>>   InputViewports;
+	TUniquePtr<FDisplayClusterMediaInputNode>                      InputNode;
+	TArray<FDisplayClusterMediaInputBase*>                         AllInputs;
 
 	// Latency queue
 	FDisplayClusterFrameQueue FrameQueue;
diff --git a/Engine/Plugins/Runtime/nDisplay/Source/DisplayClusterMedia/Private/Input/DisplayClusterMediaInputBase.cpp b/Engine/Plugins/Runtime/nDisplay/Source/DisplayClusterMedia/Private/Input/DisplayClusterMediaInputBase.cpp
--- a/Engine/Plugins/Runtime/nDisplay/Source/DisplayClusterMedia/Private/Input/DisplayClusterMediaInputBase.cpp
+++ b/Engine/Plugins/Runtime/nDisplay/Source/DisplayClusterMedia/Private/Input/DisplayClusterMediaInputBase.cpp
@@ -72,8 +72,6 @@ void FDisplayClusterMediaInputBase::Stop()
 
 void FDisplayClusterMediaInputBase::ImportMediaData(FRHICommandListImmediate& RHICmdList, const FMediaTextureInfo& TextureInfo)
 {
-	UE_LOG(LogDisplayClusterMedia, Verbose, TEXT("MediaInput '%s': importing texture on RT frame '%lu'..."), *GetMediaId(), GFrameCounterRenderThread);
-
 	MediaTexture->JustInTimeRender();
 
 	FRHITexture* const SrcTexture = MediaTexture->GetResource() ? MediaTexture->GetResource()->GetTextureRHI() : nullptr;
diff --git a/Engine/Plugins/Runtime/nDisplay/Source/DisplayClusterMedia/Private/Synchronization/DisplayClusterFrameQueue.cpp b/Engine/Plugins/Runtime/nDisplay/Source/DisplayClusterMedia/Private/Synchronization/DisplayClusterFrameQueue.cpp
new file mode 100644
--- /dev/null
+++ b/Engine/Plugins/Runtime/nDisplay/Source/DisplayClusterMedia/Private/Synchronization/DisplayClusterFrameQueue.cpp
@@ -0,0 +1,304 @@
+// Copyright Epic Games, Inc. All Rights Reserved.
+
+#include "Synchronization/DisplayClusterFrameQueue.h"
+
+#include "DisplayClusterConfigurationTypes.h"
+#include "DisplayClusterMediaLog.h"
+#include "DisplayClusterRootActor.h"
+
+#include "IDisplayCluster.h"
+#include "IDisplayClusterCallbacks.h"
+#include "Game/IDisplayClusterGameManager.h"
+#include "Render/IDisplayClusterRenderManager.h"
+#include "Render/Viewport/IDisplayClusterViewportManager.h"
+#include "Render/Viewport/IDisplayClusterViewportManagerProxy.h"
+#include "Render/Viewport/IDisplayClusterViewportProxy.h"
+#include "ShaderParameters/DisplayClusterShaderParameters_ICVFX.h"
+#include "ShaderParameters/DisplayClusterShaderParameters_WarpBlend.h"
+
+#include "Engine/Engine.h"
+#include "Engine/GameViewportClient.h"
+
+
+void FDisplayClusterFrameQueue::Init()
+{
+	// Subscribe for events
+	if (IDisplayCluster::Get().GetOperationMode() == EDisplayClusterOperationMode::Cluster)
+	{
+		if (GEngine && GEngine->GameViewport)
+		{
+			GEngine->GameViewport->OnBeginDraw().AddRaw(this, &FDisplayClusterFrameQueue::HandleBeginDraw);
+			GEngine->GameViewport->OnEndDraw().AddRaw(this, &FDisplayClusterFrameQueue::HandleEndDraw);
+		}
+
+		IDisplayCluster::Get().GetCallbacks().OnDisplayClusterProcessLatency_RenderThread().AddRaw(this, &FDisplayClusterFrameQueue::HandleProcessLatency_RenderThread);
+		IDisplayCluster::Get().GetCallbacks().OnDisplayClusterPreProcessIcvfx_RenderThread().AddRaw(this, &FDisplayClusterFrameQueue::HandlePreProcessIcvfx_RenderThread);
+	}
+}
+
+void FDisplayClusterFrameQueue::Release()
+{
+	// Unsubscribe from events
+	if (IDisplayCluster::Get().GetOperationMode() == EDisplayClusterOperationMode::Cluster)
+	{
+		if (GEngine && GEngine->GameViewport)
+		{
+			GEngine->GameViewport->OnBeginDraw().RemoveAll(this);
+			GEngine->GameViewport->OnEndDraw().RemoveAll(this);
+		}
+
+		IDisplayCluster::Get().GetCallbacks().OnDisplayClusterProcessLatency_RenderThread().RemoveAll(this);
+		IDisplayCluster::Get().GetCallbacks().OnDisplayClusterPreProcessIcvfx_RenderThread().RemoveAll(this);
+	}
+
+	// We need to properly release render thread data
+	ENQUEUE_RENDER_COMMAND(FDisplayClusterFrameQueue_Release)(
+		[this](FRHICommandListImmediate& RHICmdList)
+		{
+			CleanQueue_RenderThread();
+		});
+}
+
+void FDisplayClusterFrameQueue::HandleBeginDraw()
+{
+	// Update latency value before rendering new frame
+	if (ADisplayClusterRootActor* RootActor = IDisplayCluster::Get().GetGameMgr()->GetRootActor())
+	{
+		if (const UDisplayClusterConfigurationData* const ConfigData = RootActor->GetConfigData())
+		{
+			// Get latency from config
+			const int32 NewLatency = FMath::Clamp(ConfigData->MediaSettings.Latency, 0, 10);
+			UE_LOG(LogDisplayClusterMedia, VeryVerbose, TEXT("New latency is %u"), NewLatency);
+
+			// Update latency data on RT
+			ENQUEUE_RENDER_COMMAND(FDisplayClusterFrameQueue_SetLatency)(
+				[this, NewLatency](FRHICommandListImmediate& RHICmdList)
+				{
+					SetLatency_RenderThread(NewLatency);
+				});
+		}
+	}
+
+	// Process pending frames adding
+	ENQUEUE_RENDER_COMMAND(FDisplayClusterFrameQueue_ProcessAddingPendingFrames)(
+		[this](FRHICommandListImmediate& RHICmdList)
+		{
+			ProcessAddingPendingFrames_RenderThread();
+		});
+}
+
+void FDisplayClusterFrameQueue::HandleEndDraw()
+{
+	// We need to update head and tail indices before rendering next frame
+	ENQUEUE_RENDER_COMMAND(FDisplayClusterFrameQueue_StepQueueIndices)(
+		[this](FRHICommandListImmediate& RHICmdList)
+		{
+			StepQueueIndices_RenderThread();
+		});
+}
+
+void FDisplayClusterFrameQueue::HandleProcessLatency_RenderThread(FRHICommandListImmediate& RHICmdList, const IDisplayClusterViewportManagerProxy* ViewportManagerProxy, FViewport* Viewport)
+{
+	if (Frames.Num() > 0)
+	{
+		// At this point, all the views for current frame have been rendered already
+		if (IDisplayClusterRenderManager* RenderMgr = IDisplayCluster::Get().GetRenderMgr())
+		{
+			if (IDisplayClusterViewportManager* ViewportMgr = RenderMgr->GetViewportManager())
+			{
+				if (IDisplayClusterViewportManagerProxy* ViewportMgrProxy = ViewportMgr->GetProxy())
+				{
+					TArrayView<TSharedPtr<IDisplayClusterViewportProxy, ESPMode::ThreadSafe>> ViewportProxies = ViewportMgrProxy->GetViewports_RenderThread();
+					for (TSharedPtr<IDisplayClusterViewportProxy, ESPMode::ThreadSafe>& ViewportProxyIt : ViewportProxies)
+					{
+						// Handle textures of all viewports that don't have any media input attached. Those don't get rendered.
+						if (ViewportProxyIt.IsValid() && !ViewportProxyIt->GetRenderSettings_RenderThread().bSkipSceneRenderingButLeaveResourcesAvailable)
+						{
+							// Get viewport texture
+							TArray<FRHITexture*> Textures;
+							ViewportProxyIt->GetResources_RenderThread(EDisplayClusterViewportResourceType::InternalRenderTargetResource, Textures);
+							check(Textures.Num() > 0);
+
+							const FString ViewportId(ViewportProxyIt->GetId());
+
+							// Enqueue this view (head)
+							Frames[IdxHead].SaveView(RHICmdList, ViewportId, Textures[0]);
+							// Dequeue from the past (tail)
+							Frames[IdxTail].LoadView(RHICmdList, ViewportId, Textures[0]);
+						}
+					}
+				}
+			}
+		}
+	}
+}
+
+void FDisplayClusterFrameQueue::HandlePreProcessIcvfx_RenderThread(FRHICommandListImmediate& RHICmdList, const IDisplayClusterViewportProxy* ViewportProxy, FDisplayClusterShaderParameters_WarpBlend& WarpBlendParameters, FDisplayClusterShaderParameters_ICVFX& ICVFXParameters)
+{
+	// Substitute shader data
+	if (Frames.Num() > 0)
+	{
+		const FString ViewportId(ViewportProxy->GetId());
+
+		Frames[IdxHead].SaveData(RHICmdList, ViewportId, WarpBlendParameters, ICVFXParameters);
+		Frames[IdxTail].LoadData(RHICmdList, ViewportId, WarpBlendParameters, ICVFXParameters);
+	}
+}
+
+void FDisplayClusterFrameQueue::SetLatency_RenderThread(int32 NewLatency)
+{
+	// Nothing to do if requested latency is the same as previously (no changes since last frame)
+	if (LastRequestedLatency != NewLatency && NewLatency >= 0)
+	{
+		// Check if latency just was disabled
+		if (NewLatency == 0)
+		{
+			// Clear the queue
+			CleanQueue_RenderThread();
+		}
+		// Check if new latency is higher
+		else if (NewLatency > LastRequestedLatency)
+		{
+			FramesPendingToAdd = (NewLatency + 1) - Frames.Num();
+		}
+		// Check if new latency is lower
+		else if (NewLatency < LastRequestedLatency)
+		{
+			// Depending on the current queue size, we need to find out if the queue needs to be resized
+			const int32 TargetQueueSize = NewLatency + 1;
+
+			// Yeah, we need to remove unnecessary elements from the queue
+			if (TargetQueueSize < Frames.Num())
+			{
+				RemoveFromQueue_RenderThread(Frames.Num() - TargetQueueSize);
+				FramesPendingToAdd = 0;
+			}
+			// No need to remove anything from the queue, just recalculate new amount of pending frames
+			else
+			{
+				FramesPendingToAdd = TargetQueueSize - Frames.Num();
+			}
+		}
+
+		// Finally, remember latency requested
+		LastRequestedLatency = NewLatency;
+	}
+}
+
+void FDisplayClusterFrameQueue::ProcessAddingPendingFrames_RenderThread()
+{
+	if (FramesPendingToAdd > 0)
+	{
+		// Increase queue size step by step, 1 item per frame. This allows to avoid
+		// situations where we have empty queue items.
+		FramesAdded = 1;
+		AddToQueue_RenderThread(FramesAdded);
+	}
+	else
+	{
+		// Queue length remains unchanged
+		FramesAdded = 0;
+	}
+}
+
+void FDisplayClusterFrameQueue::AddToQueue_RenderThread(int32 FramesAmount)
+{
+	checkSlow(FramesAmount > 0);
+
+	// Duplicate the head frame to increase the queue
+	for (int32 Idx = 0; Idx < FramesAmount; ++Idx)
+	{
+		// In case the queue is empty, just instantiate new item
+		if (Frames.IsEmpty())
+		{
+			Frames.AddDefaulted();
+
+			// Reset indices
+			IdxHead = 0;
+			IdxTail = 0;
+		}
+		// Otherwise duplicate head item to prevent 'black' frames in the future
+		else
+		{
+			// For some reason, the following doesn't work (there will be check fail inside Insert)
+			// >>>> Frames.Insert(Frames[IdxHead], IdxHead)
+			//
+			// To achieve the same avoiding unnecessary data copying, I use the approach where data copying
+			// is performed in copy constructor, and the instance is pushed to the container using move constructor.
+			FDisplayClusterFrameQueueItem ClonedItem(Frames[IdxHead]);
+			Frames.Insert(MoveTemp(ClonedItem), IdxHead);
+
+			// Recalculate tail index
+			if (IdxTail > IdxHead)
+			{
+				++IdxTail;
+			}
+		}
+	}
+}
+
+void FDisplayClusterFrameQueue::RemoveFromQueue_RenderThread(int32 FramesAmout)
+{
+	checkSlow(FramesAmout < Frames.Num());
+
+	// We need to remove specified amount of frames starting from the tail
+	for (int32 Idx = 0; Idx < FramesAmout; ++Idx)
+	{
+		// Remove the item
+		Frames.RemoveAt(IdxTail);
+
+		// Recalculate head index
+		if (IdxHead > IdxTail)
+		{
+			--IdxHead;
+		}
+
+		// Recalculate tail index
+		if (IdxTail >= Frames.Num())
+		{
+			IdxTail = 0;
+		}
+	}
+}
+
+void FDisplayClusterFrameQueue::CleanQueue_RenderThread()
+{
+	// Clean up the container
+	Frames.Empty();
+
+	// Reset internals
+	IdxHead = 0;
+	IdxTail = 0;
+	FramesPendingToAdd = 0;
+}
+
+void FDisplayClusterFrameQueue::StepQueueIndices_RenderThread()
+{
+	// Nothing to do if latency is zero
+	if (Frames.IsEmpty())
+	{
+		UE_LOG(LogDisplayClusterMedia, VeryVerbose, TEXT("Queue indices weren't updated, latency=0"));
+		return;
+	}
+
+	// Head index
+	if (++IdxHead >= Frames.Num())
+	{
+		IdxHead = 0;
+	}
+
+	// Tail index doesn't move while we the queue being increased
+	if (FramesAdded == 0)
+	{
+		if (++IdxTail >= Frames.Num())
+		{
+			IdxTail = 0;
+		}
+	}
+	// Otherwise update pending amount
+	else
+	{
+		FramesPendingToAdd -= FramesAdded;
+		checkSlow(FramesPendingToAdd >= 0);
+	}
+}
diff --git a/Engine/Plugins/Runtime/nDisplay/Source/DisplayClusterMedia/Private/Synchronization/DisplayClusterFrameQueue.h b/Engine/Plugins/Runtime/nDisplay/Source/DisplayClusterMedia/Private/Synchronization/DisplayClusterFrameQueue.h
new file mode 100644
--- /dev/null
+++ b/Engine/Plugins/Runtime/nDisplay/Source/DisplayClusterMedia/Private/Synchronization/DisplayClusterFrameQueue.h
@@ -0,0 +1,76 @@
+// Copyright Epic Games, Inc. All Rights Reserved.
+
+#pragma once
+
+#include "CoreMinimal.h"
+
+#include "Synchronization/DisplayClusterFrameQueueItem.h"
+
+class FDisplayClusterMediaInputViewport;
+class FDisplayClusterShaderParameters_ICVFX;
+class FViewport;
+class IDisplayClusterViewportManagerProxy;
+class IDisplayClusterViewportProxy;
+class RHICmdList;
+struct FDisplayClusterShaderParameters_WarpBlend;
+
+/**
+ * Latency queue class
+ *
+ * Implements artificial latency feature. It's responsible for caching visual and shaders data to be used later.
+ *
+ * All the data that comes out of the rendering pipeline for current frame is cached to the HEAD of the queue. The final
+ * output is built from the TAIL view data. Every queue item contains all necessary information for every nD viewport.
+ *
+ * Note: this queue needs to be an intermediate part of the nD rendering pipeline. It should be moved there
+ * at some point during refactoring.
+ */
+class FDisplayClusterFrameQueue
+{
+public:
+	// Initializes latency feature
+	void Init();
+	// Releases latency feature
+	void Release();
+
+private:
+	// Handles BeginDraw event
+	void HandleBeginDraw();
+	// Handles EndDraw event
+	void HandleEndDraw();
+	// Handles latency processing event
+	void HandleProcessLatency_RenderThread(FRHICommandListImmediate& RHICmdList, const IDisplayClusterViewportManagerProxy* ViewportManagerProxy, FViewport* Viewport);
+	// Handles ICVFX pre-processing event
+	void HandlePreProcessIcvfx_RenderThread(FRHICommandListImmediate& RHICmdList, const IDisplayClusterViewportProxy* ViewportProxy, FDisplayClusterShaderParameters_WarpBlend& WarpBlendParameters, FDisplayClusterShaderParameters_ICVFX& ICVFXParameters);
+
+private:
+	// Changes current latency
+	void SetLatency_RenderThread(int32 NewLatency);
+	// Auxiliary method to process queue resize
+	void ProcessAddingPendingFrames_RenderThread();
+	// Auxiliary method to increase queue length by amount of frames
+	void AddToQueue_RenderThread(int32 FramesAmount);
+	// Auxiliary method to decrease queue length by amount of frames
+	void RemoveFromQueue_RenderThread(int32 FramesAmout);
+	// Cleans the queue container
+	void CleanQueue_RenderThread();
+	// Updates head and tail indices every frame
+	void StepQueueIndices_RenderThread();
+
+private:
+	// Frames buffer
+	TArray<FDisplayClusterFrameQueueItem> Frames;
+
+	// Latency value requested last time
+	int32 LastRequestedLatency = 0;
+
+	// Amount of frames pending to add to the latency queue
+	int32 FramesPendingToAdd = 0;
+	// Amount of frames that were added at current frame
+	int32 FramesAdded = 0;
+
+	// Head queue index
+	int32 IdxHead = 0;
+	// Tail queue index
+	int32 IdxTail = 0;
+};
diff --git a/Engine/Plugins/Runtime/nDisplay/Source/DisplayClusterMedia/Private/Synchronization/DisplayClusterFrameQueueItem.cpp b/Engine/Plugins/Runtime/nDisplay/Source/DisplayClusterMedia/Private/Synchronization/DisplayClusterFrameQueueItem.cpp
new file mode 100644
--- /dev/null
+++ b/Engine/Plugins/Runtime/nDisplay/Source/DisplayClusterMedia/Private/Synchronization/DisplayClusterFrameQueueItem.cpp
@@ -0,0 +1,152 @@
+// Copyright Epic Games, Inc. All Rights Reserved.
+
+#include "Synchronization/DisplayClusterFrameQueueItem.h"
+
+#include "DisplayClusterConfigurationTypes.h"
+#include "DisplayClusterMediaLog.h"
+#include "DisplayClusterRootActor.h"
+
+#include "IDisplayCluster.h"
+#include "Render/IDisplayClusterRenderManager.h"
+#include "Render/Viewport/IDisplayClusterViewportManager.h"
+#include "Render/Viewport/IDisplayClusterViewportManagerProxy.h"
+
+#include "RHICommandList.h"
+
+
+FDisplayClusterFrameQueueItem::FDisplayClusterFrameQueueItem(const FDisplayClusterFrameQueueItem& Other)
+{
+	for (const TPair<FString, FDisplayClusterFrameQueueItemView>& View : Other.Views)
+	{
+		FDisplayClusterFrameQueueItemView& NewItem = this->Views.Emplace(View.Key, View.Value);
+		if (View.Value.Texture.IsValid())
+		{
+			NewItem.Texture = CreateTexture(View.Value.Texture);
+			FRHICommandListImmediate& RHICmdList = GetImmediateCommandList_ForRenderCommand();
+			RHICmdList.CopyTexture(View.Value.Texture.GetReference(), NewItem.Texture, FRHICopyTextureInfo());
+		}
+	}
+}
+
+void FDisplayClusterFrameQueueItem::SaveView(FRHICommandListImmediate& RHICmdList, const FString& ViewportId, FRHITexture* Texture)
+{
+	checkSlow(Texture);
+
+	// Find proper view and copy texture data
+	FDisplayClusterFrameQueueItemView* View = Views.Find(ViewportId);
+	if (!View)
+	{
+		View = &Views.Emplace(ViewportId);
+	}
+
+	// Create texture if not yet available. Or re-create if the source texture has been updated (re-sized, other format, etc.)
+	if (!View->Texture.IsValid() ||
+		Texture->GetFormat() != View->Texture.GetReference()->GetFormat() ||
+		Texture->GetDesc().Extent != View->Texture.GetReference()->GetDesc().Extent)
+	{
+		View->Texture = CreateTexture(Texture);
+		check(View->Texture.IsValid());
+	}
+
+	// Copy texture data
+	if (View->Texture.IsValid())
+	{
+		RHICmdList.CopyTexture(Texture, View->Texture.GetReference(), FRHICopyTextureInfo());
+	}
+}
+
+void FDisplayClusterFrameQueueItem::LoadView(FRHICommandListImmediate& RHICmdList, const FString& ViewportId, FRHITexture* Texture)
+{
+	checkSlow(Texture);
+
+	// Find proper view and copy texture data
+	const FDisplayClusterFrameQueueItemView* const View = Views.Find(ViewportId);
+	if (View && View->Texture.IsValid())
+	{
+		RHICmdList.CopyTexture(View->Texture.GetReference(), Texture, FRHICopyTextureInfo());
+	}
+}
+
+void FDisplayClusterFrameQueueItem::SaveData(FRHICommandListImmediate& RHICmdList, const FString& ViewportId, FDisplayClusterShaderParameters_WarpBlend& WarpBlendParameters, FDisplayClusterShaderParameters_ICVFX& ICVFXParameters)
+{
+	if (FDisplayClusterFrameQueueItemView* View = Views.Find(ViewportId))
+	{
+		// Warp/blend - cherry pick necessary parameters only
+		View->WarpBlendData.bRenderAlphaChannel = WarpBlendParameters.bRenderAlphaChannel;
+		View->WarpBlendData.Context = WarpBlendParameters.Context;
+
+		// ICVFX - cherry pick necessary parameters only
+		View->IcvfxData.Cameras = ICVFXParameters.Cameras;
+	}
+}
+
+void FDisplayClusterFrameQueueItem::LoadData(FRHICommandListImmediate& RHICmdList, const FString& ViewportId, FDisplayClusterShaderParameters_WarpBlend& WarpBlendParameters, FDisplayClusterShaderParameters_ICVFX& ICVFXParameters)
+{
+	if (const FDisplayClusterFrameQueueItemView* const View = Views.Find(ViewportId))
+	{
+		// Warp/blend - cherry pick necessary parameters only
+		WarpBlendParameters.bRenderAlphaChannel = View->WarpBlendData.bRenderAlphaChannel;
+		WarpBlendParameters.Context = View->WarpBlendData.Context;
+
+		// ICVFX - cherry pick necessary parameters only
+		for (const FDisplayClusterShaderParameters_ICVFX::FCameraSettings& CameraSettings : View->IcvfxData.Cameras)
+		{
+			// Look up for proper camera settings object. Here we use ViewportId as a key to compare.
+			FDisplayClusterShaderParameters_ICVFX::FCameraSettings* TargetCameraSettings = ICVFXParameters.Cameras.FindByPredicate([CameraSettings](const FDisplayClusterShaderParameters_ICVFX::FCameraSettings& Item)
+				{
+					return Item.Resource.ViewportId.Equals(CameraSettings.Resource.ViewportId, ESearchCase::IgnoreCase);
+				});
+
+			// Update camera settings
+			if (TargetCameraSettings)
+			{
+				TargetCameraSettings->SoftEdge                 = CameraSettings.SoftEdge;
+				TargetCameraSettings->InnerCameraBorderColor      = CameraSettings.InnerCameraBorderColor;
+				TargetCameraSettings->InnerCameraBorderThickness  = CameraSettings.InnerCameraBorderThickness;
+				TargetCameraSettings->InnerCameraFrameAspectRatio = CameraSettings.InnerCameraFrameAspectRatio;
+				TargetCameraSettings->Local2WorldTransform     = CameraSettings.Local2WorldTransform;
+				TargetCameraSettings->CameraViewRotation       = CameraSettings.CameraViewRotation;
+				TargetCameraSettings->CameraViewLocation       = CameraSettings.CameraViewLocation;
+				TargetCameraSettings->CameraPrjMatrix          = CameraSettings.CameraPrjMatrix;
+				TargetCameraSettings->ChromakeySource          = CameraSettings.ChromakeySource;
+				TargetCameraSettings->ChromakeyColor           = CameraSettings.ChromakeyColor;
+				TargetCameraSettings->ChromakeyMarkersColor    = CameraSettings.ChromakeyMarkersColor;
+				TargetCameraSettings->ChromakeyMarkersScale    = CameraSettings.ChromakeyMarkersScale;
+				TargetCameraSettings->ChromakeyMarkersDistance = CameraSettings.ChromakeyMarkersDistance;
+				TargetCameraSettings->ChromakeyMarkersOffset   = CameraSettings.ChromakeyMarkersOffset;
+				TargetCameraSettings->RenderOrder              = CameraSettings.RenderOrder;
+			}
+		}
+	}
+}
+
+FTextureRHIRef FDisplayClusterFrameQueueItem::CreateTexture(FRHITexture* ReferenceTexture)
+{
+	if (ReferenceTexture)
+	{
+		// Use original format and size
+		const int32 SizeX = ReferenceTexture->GetDesc().Extent.X;
+		const int32 SizeY = ReferenceTexture->GetDesc().Extent.Y;
+		const EPixelFormat Format = ReferenceTexture->GetFormat();
+
+		// Prepare description
+		FRHITextureCreateDesc Desc =
+			FRHITextureCreateDesc::Create2D(TEXT("DisplayClusterFrameQueueCacheTexture"), SizeX, SizeY, Format)
+			.SetClearValue(FClearValueBinding::Black)
+			.SetNumMips(1)
+			.SetFlags(ETextureCreateFlags::Dynamic)
+			.AddFlags(ETextureCreateFlags::MultiGPUGraphIgnore)
+			.SetInitialState(ERHIAccess::SRVMask);
+
+		// Leave original flags, but make sure it's ResolveTargetable but not RenderTargetable
+		ETextureCreateFlags Flags = ReferenceTexture->GetFlags();
+		Flags &= ~ETextureCreateFlags::RenderTargetable;
+		Flags |= ETextureCreateFlags::ResolveTargetable;
+		Desc.SetFlags(Flags);
+
+		// Create texture
+		return RHICreateTexture(Desc);
+	}
+
+	return nullptr;
+}
diff --git a/Engine/Plugins/Runtime/nDisplay/Source/DisplayClusterMedia/Private/Synchronization/DisplayClusterFrameQueueItem.h b/Engine/Plugins/Runtime/nDisplay/Source/DisplayClusterMedia/Private/Synchronization/DisplayClusterFrameQueueItem.h
new file mode 100644
--- /dev/null
+++ b/Engine/Plugins/Runtime/nDisplay/Source/DisplayClusterMedia/Private/Synchronization/DisplayClusterFrameQueueItem.h
@@ -0,0 +1,64 @@
+// Copyright Epic Games, Inc. All Rights Reserved.
+
+#pragma once
+
+#include "CoreMinimal.h"
+
+#include "ShaderParameters/DisplayClusterShaderParameters_ICVFX.h"
+#include "ShaderParameters/DisplayClusterShaderParameters_WarpBlend.h"
+
+class FRHICommandListImmediate;
+class FRHITexture;
+
+
+/**
+ * Viewport data cache to be used in the latency queue
+ */
+struct FDisplayClusterFrameQueueItemView
+{
+	/** Cached view texture of a viewoprt */
+	FTextureRHIRef Texture = nullptr;
+
+	/** Cached warp/blend shader data */
+	FDisplayClusterShaderParameters_WarpBlend WarpBlendData;
+
+	/** Cached ICVFX shader data */
+	FDisplayClusterShaderParameters_ICVFX IcvfxData;
+};
+
+
+/**
+ * Latency queue item class
+ *
+ * It's responsible for caching all required data for single frame.
+ * 
+ * Note: this one is render thread only. It's never used on the main thread or any other thread.
+ */
+class FDisplayClusterFrameQueueItem
+{
+public:
+	FDisplayClusterFrameQueueItem() = default;
+	FDisplayClusterFrameQueueItem(const FDisplayClusterFrameQueueItem& Other);
+	FDisplayClusterFrameQueueItem(FDisplayClusterFrameQueueItem&& Other) = default;
+
+	~FDisplayClusterFrameQueueItem() = default;
+
+public:
+	// Save texture of a viewport
+	void SaveView(FRHICommandListImmediate& RHICmdList, const FString& ViewportId, FRHITexture* Texture);
+	// Load texture of a viewport
+	void LoadView(FRHICommandListImmediate& RHICmdList, const FString& ViewportId, FRHITexture* Texture);
+
+	// Save shader data of a viewport
+	void SaveData(FRHICommandListImmediate& RHICmdList, const FString& ViewportId, FDisplayClusterShaderParameters_WarpBlend& WarpBlendParameters, FDisplayClusterShaderParameters_ICVFX& ICVFXParameters);
+	// Load shader data of a viewport
+	void LoadData(FRHICommandListImmediate& RHICmdList, const FString& ViewportId, FDisplayClusterShaderParameters_WarpBlend& WarpBlendParameters, FDisplayClusterShaderParameters_ICVFX& ICVFXParameters);
+
+private:
+	// Auxiliary function that creates a new texture based on the reference (same format, size)
+	FTextureRHIRef CreateTexture(FRHITexture* ReferenceTexture);
+
+private:
+	// Available views (viewports)
+	TMap<FString, FDisplayClusterFrameQueueItemView> Views;
+};
diff --git a/Engine/Plugins/Runtime/nDisplay/Source/DisplayClusterMedia/Private/Synchronization/LatencyQueue/DisplayClusterFrameQueue.cpp b/Engine/Plugins/Runtime/nDisplay/Source/DisplayClusterMedia/Private/Synchronization/LatencyQueue/DisplayClusterFrameQueue.cpp
deleted file mode 100644
diff --git a/Engine/Plugins/Runtime/nDisplay/Source/DisplayClusterMedia/Private/Synchronization/LatencyQueue/DisplayClusterFrameQueue.h b/Engine/Plugins/Runtime/nDisplay/Source/DisplayClusterMedia/Private/Synchronization/LatencyQueue/DisplayClusterFrameQueue.h
deleted file mode 100644
diff --git a/Engine/Plugins/Runtime/nDisplay/Source/DisplayClusterMedia/Private/Synchronization/LatencyQueue/DisplayClusterFrameQueueItem.cpp b/Engine/Plugins/Runtime/nDisplay/Source/DisplayClusterMedia/Private/Synchronization/LatencyQueue/DisplayClusterFrameQueueItem.cpp
deleted file mode 100644
diff --git a/Engine/Plugins/Runtime/nDisplay/Source/DisplayClusterMedia/Private/Synchronization/LatencyQueue/DisplayClusterFrameQueueItem.h b/Engine/Plugins/Runtime/nDisplay/Source/DisplayClusterMedia/Private/Synchronization/LatencyQueue/DisplayClusterFrameQueueItem.h
deleted file mode 100644
diff --git a/Engine/Plugins/Runtime/nDisplay/Source/DisplayClusterMedia/Private/Synchronization/SyncPolicies/DisplayClusterMediaOutputSynchronizationPolicyEthernetBarrier.cpp b/Engine/Plugins/Runtime/nDisplay/Source/DisplayClusterMedia/Private/Synchronization/SyncPolicies/DisplayClusterMediaOutputSynchronizationPolicyEthernetBarrier.cpp
deleted file mode 100644
diff --git a/Engine/Plugins/Runtime/nDisplay/Source/DisplayClusterMedia/Private/Synchronization/SyncPolicies/DisplayClusterMediaOutputSynchronizationPolicyEthernetBarrierBase.cpp b/Engine/Plugins/Runtime/nDisplay/Source/DisplayClusterMedia/Private/Synchronization/SyncPolicies/DisplayClusterMediaOutputSynchronizationPolicyEthernetBarrierBase.cpp
deleted file mode 100644
diff --git a/Engine/Plugins/Runtime/nDisplay/Source/DisplayClusterMedia/Private/Synchronization/SyncPolicies/DisplayClusterMediaOutputSynchronizationPolicyRivermax.cpp b/Engine/Plugins/Runtime/nDisplay/Source/DisplayClusterMedia/Private/Synchronization/SyncPolicies/DisplayClusterMediaOutputSynchronizationPolicyRivermax.cpp
deleted file mode 100644
diff --git a/Engine/Plugins/Runtime/nDisplay/Source/DisplayClusterMedia/Private/Synchronization/SyncPolicies/DisplayClusterMediaOutputSynchronizationPolicyThresholdBase.cpp b/Engine/Plugins/Runtime/nDisplay/Source/DisplayClusterMedia/Private/Synchronization/SyncPolicies/DisplayClusterMediaOutputSynchronizationPolicyThresholdBase.cpp
deleted file mode 100644
diff --git a/Engine/Plugins/Runtime/nDisplay/Source/DisplayClusterMedia/Private/Synchronization/SyncPolicies/DisplayClusterMediaOutputSynchronizationPolicyVblank.cpp b/Engine/Plugins/Runtime/nDisplay/Source/DisplayClusterMedia/Private/Synchronization/SyncPolicies/DisplayClusterMediaOutputSynchronizationPolicyVblank.cpp
deleted file mode 100644
diff --git a/Engine/Plugins/Runtime/nDisplay/Source/DisplayClusterMedia/Public/Synchronization/DisplayClusterMediaOutputSynchronizationPolicyEthernetBarrier.h b/Engine/Plugins/Runtime/nDisplay/Source/DisplayClusterMedia/Public/Synchronization/DisplayClusterMediaOutputSynchronizationPolicyEthernetBarrier.h
deleted file mode 100644
diff --git a/Engine/Plugins/Runtime/nDisplay/Source/DisplayClusterMedia/Public/Synchronization/DisplayClusterMediaOutputSynchronizationPolicyEthernetBarrierBase.h b/Engine/Plugins/Runtime/nDisplay/Source/DisplayClusterMedia/Public/Synchronization/DisplayClusterMediaOutputSynchronizationPolicyEthernetBarrierBase.h
deleted file mode 100644
diff --git a/Engine/Plugins/Runtime/nDisplay/Source/DisplayClusterMedia/Public/Synchronization/DisplayClusterMediaOutputSynchronizationPolicyRivermax.h b/Engine/Plugins/Runtime/nDisplay/Source/DisplayClusterMedia/Public/Synchronization/DisplayClusterMediaOutputSynchronizationPolicyRivermax.h
deleted file mode 100644
diff --git a/Engine/Plugins/Runtime/nDisplay/Source/DisplayClusterMedia/Public/Synchronization/DisplayClusterMediaOutputSynchronizationPolicyThresholdBase.h b/Engine/Plugins/Runtime/nDisplay/Source/DisplayClusterMedia/Public/Synchronization/DisplayClusterMediaOutputSynchronizationPolicyThresholdBase.h
deleted file mode 100644
diff --git a/Engine/Plugins/Runtime/nDisplay/Source/DisplayClusterMedia/Public/Synchronization/DisplayClusterMediaOutputSynchronizationPolicyVblank.h b/Engine/Plugins/Runtime/nDisplay/Source/DisplayClusterMedia/Public/Synchronization/DisplayClusterMediaOutputSynchronizationPolicyVblank.h
deleted file mode 100644
diff --git a/Engine/Plugins/Runtime/nDisplay/Source/DisplayClusterMediaEditor/Private/DisplayClusterMediaOutputSynchronizationPolicyFactories.cpp b/Engine/Plugins/Runtime/nDisplay/Source/DisplayClusterMediaEditor/Private/DisplayClusterMediaOutputSynchronizationPolicyFactories.cpp
deleted file mode 100644
diff --git a/Engine/Plugins/Runtime/nDisplay/Source/DisplayClusterMediaEditor/Private/DisplayClusterMediaOutputSynchronizationPolicyFactories.h b/Engine/Plugins/Runtime/nDisplay/Source/DisplayClusterMediaEditor/Private/DisplayClusterMediaOutputSynchronizationPolicyFactories.h
deleted file mode 100644
diff --git a/Engine/Plugins/VirtualProduction/Rivermax/RivermaxCore/Source/RivermaxCore/Private/RivermaxOutputStream.cpp b/Engine/Plugins/VirtualProduction/Rivermax/RivermaxCore/Source/RivermaxCore/Private/RivermaxOutputStream.cpp
--- a/Engine/Plugins/VirtualProduction/Rivermax/RivermaxCore/Source/RivermaxCore/Private/RivermaxOutputStream.cpp
+++ b/Engine/Plugins/VirtualProduction/Rivermax/RivermaxCore/Source/RivermaxCore/Private/RivermaxOutputStream.cpp
@@ -230,9 +230,6 @@ namespace UE::RivermaxCore::Private
 	{
 		TRACE_CPUPROFILER_EVENT_SCOPE(FRivermaxOutputStream::PushVideoFrame);
 
-		// Notify owner that new frame is about to be enqueued
-		Listener->OnPreFrameEnqueue();
-
 		if (TSharedPtr<FRivermaxOutputFrame> AvailableFrame = GetNextAvailableFrame(NewFrame.FrameIdentifier))
 		{
 			const int32 Stride = GetStride();
@@ -245,7 +242,7 @@ namespace UE::RivermaxCore::Private
 			{
 				const FString TraceName = FString::Format(TEXT("FRivermaxOutputStream::PushFrame {0}"), { AvailableFrame->FrameIndex });
 				TRACE_CPUPROFILER_EVENT_SCOPE_TEXT(*TraceName);
-
+				
 				MarkFrameToBeSent(MoveTemp(AvailableFrame));
 			}
 
@@ -1014,9 +1011,6 @@ namespace UE::RivermaxCore::Private
 
 				FRivermaxOutputStream* Stream = reinterpret_cast<FRivermaxOutputStream*>(userData);
 
-				// Notify owner that new frame is about to be enqueued
-				Stream->Listener->OnPreFrameEnqueue();
-
 				TOptional<uint32> TargetFrameIdentifier = Stream->PendingIdentifiers.Dequeue();
 				if (TargetFrameIdentifier.IsSet())
 				{
diff --git a/Engine/Plugins/VirtualProduction/Rivermax/RivermaxCore/Source/RivermaxCore/Public/IRivermaxOutputStream.h b/Engine/Plugins/VirtualProduction/Rivermax/RivermaxCore/Source/RivermaxCore/Public/IRivermaxOutputStream.h
--- a/Engine/Plugins/VirtualProduction/Rivermax/RivermaxCore/Source/RivermaxCore/Public/IRivermaxOutputStream.h
+++ b/Engine/Plugins/VirtualProduction/Rivermax/RivermaxCore/Source/RivermaxCore/Public/IRivermaxOutputStream.h
@@ -29,9 +29,6 @@ namespace UE::RivermaxCore
 
 		/** Called when stream has encountered an error and has to stop */
 		virtual void OnStreamError() = 0;
-
-		/** Called when stream is about to enqueue new frame */
-		virtual void OnPreFrameEnqueue() = 0;
 	};
 
 	class RIVERMAXCORE_API IRivermaxOutputStream
diff --git a/Engine/Plugins/VirtualProduction/Rivermax/RivermaxMedia/Source/RivermaxMedia/Private/RivermaxMediaCapture.cpp b/Engine/Plugins/VirtualProduction/Rivermax/RivermaxMedia/Source/RivermaxMedia/Private/RivermaxMediaCapture.cpp
--- a/Engine/Plugins/VirtualProduction/Rivermax/RivermaxMedia/Source/RivermaxMedia/Private/RivermaxMediaCapture.cpp
+++ b/Engine/Plugins/VirtualProduction/Rivermax/RivermaxMedia/Source/RivermaxMedia/Private/RivermaxMediaCapture.cpp
@@ -144,11 +144,6 @@ namespace UE::RivermaxMediaCaptureUtil
 ///* URivermaxMediaCapture implementation
 //*****************************************************************************/
 
-UE::RivermaxCore::FRivermaxOutputStreamOptions URivermaxMediaCapture::GetOutputStreamOptions() const
-{
-	return Options;
-}
-
 bool URivermaxMediaCapture::ValidateMediaOutput() const
 {
 	URivermaxMediaOutput* RivermaxMediaOutput = Cast<URivermaxMediaOutput>(MediaOutput);
@@ -575,11 +570,3 @@ void URivermaxMediaCapture::OnStreamError()
 		SetState(EMediaCaptureState::Error);
 	}
 }
-
-void URivermaxMediaCapture::OnPreFrameEnqueue()
-{
-	if (IsOutputSynchronizationSupported())
-	{
-		OnOutputSynchronization.ExecuteIfBound();
-	}
-}
diff --git a/Engine/Plugins/VirtualProduction/Rivermax/RivermaxMedia/Source/RivermaxMedia/Public/RivermaxMediaCapture.h b/Engine/Plugins/VirtualProduction/Rivermax/RivermaxMedia/Source/RivermaxMedia/Public/RivermaxMediaCapture.h
--- a/Engine/Plugins/VirtualProduction/Rivermax/RivermaxMedia/Source/RivermaxMedia/Public/RivermaxMediaCapture.h
+++ b/Engine/Plugins/VirtualProduction/Rivermax/RivermaxMedia/Source/RivermaxMedia/Public/RivermaxMediaCapture.h
@@ -22,10 +22,6 @@ class RIVERMAXMEDIA_API URivermaxMediaCapture : public UMediaCapture, public UE:
 {
 	GENERATED_BODY()
 
-public:
-	/** Rivermax capture specific API to provide stream options access */
-	UE::RivermaxCore::FRivermaxOutputStreamOptions GetOutputStreamOptions() const;
-
 public:
 
 	//~ Begin UMediaCapture interface
@@ -44,7 +40,6 @@ protected:
 	virtual EMediaCaptureResourceType GetCustomOutputResourceType() const override;
 	virtual FRDGBufferDesc GetCustomBufferDescription(const FIntPoint& InDesiredSize) const override;
 	virtual void OnCustomCapture_RenderingThread(FRDGBuilder& GraphBuilder, const FCaptureBaseData& InBaseData, TSharedPtr<FMediaCaptureUserData, ESPMode::ThreadSafe> InUserData, FRDGTextureRef InSourceTexture, FRDGBufferRef OutputBuffer, const FRHICopyTextureInfo& CopyInfo, FVector2D CropU, FVector2D CropV) override;
-	virtual bool IsOutputSynchronizationSupported() const override { return true; }
 	//~ End UMediaCapture interface
 
 	//~ Begin UObject interface
@@ -54,7 +49,6 @@ protected:
 	//~ Begin IRivermaxOutputStreamListener interface
 	virtual void OnInitializationCompleted(bool bHasSucceed) override;
 	virtual void OnStreamError() override;
-	virtual void OnPreFrameEnqueue() override;
 	//~ End IRivermaxOutputStreamListener interface
 
 private:
diff --git a/Engine/Plugins/VirtualProduction/Rivermax/RivermaxMedia/Source/RivermaxMedia/RivermaxMedia.build.cs b/Engine/Plugins/VirtualProduction/Rivermax/RivermaxMedia/Source/RivermaxMedia/RivermaxMedia.build.cs
--- a/Engine/Plugins/VirtualProduction/Rivermax/RivermaxMedia/Source/RivermaxMedia/RivermaxMedia.build.cs
+++ b/Engine/Plugins/VirtualProduction/Rivermax/RivermaxMedia/Source/RivermaxMedia/RivermaxMedia.build.cs
@@ -12,7 +12,6 @@ public class RivermaxMedia : ModuleRules
 				"CoreUObject",
 				"MediaAssets",
 				"MediaIOCore",
-				"RivermaxCore",
 				"TimeManagement"
 			}
 		);
