nDisplay: Fixed issue that was causing the nDisplay config to be marked as dirty immediately when a level containing one was opened.

diff --git a/Engine/Plugins/Runtime/nDisplay/Source/DisplayCluster/Private/Game/EngineClasses/Scene/DisplayClusterRootActorEditor.cpp b/Engine/Plugins/Runtime/nDisplay/Source/DisplayCluster/Private/Game/EngineClasses/Scene/DisplayClusterRootActorEditor.cpp
--- a/Engine/Plugins/Runtime/nDisplay/Source/DisplayCluster/Private/Game/EngineClasses/Scene/DisplayClusterRootActorEditor.cpp
+++ b/Engine/Plugins/Runtime/nDisplay/Source/DisplayCluster/Private/Game/EngineClasses/Scene/DisplayClusterRootActorEditor.cpp
@@ -998,7 +998,8 @@ void ADisplayClusterRootActor::UpdatePreviewComponents()
 					PreviewComp = nullptr;
 				}
 
-				Modify();
+				const bool bMarkDirty = false;
+				Modify(bMarkDirty);
 				if (!PreviewComp)
 				{
 					PreviewComp = NewObject<UDisplayClusterPreviewComponent>(this, FName(*PreviewCompId), RF_DuplicateTransient | RF_Transactional | RF_NonPIEDuplicateTransient);
