nDisplay: Fixed workflow for updating viewport proxy data:
The GT viewport data must be sent twice to the render thread. This is because the viewport math data is updated later when the view family is initialized. However, to render a view family, the viewport proxy resources must be updated beforehand. But after all this, and before compositing the viewport, we need to update the viewport math for the proxy again.

diff --git a/Engine/Plugins/Runtime/nDisplay/Source/DisplayCluster/Private/Render/Viewport/DisplayClusterViewportManager.cpp b/Engine/Plugins/Runtime/nDisplay/Source/DisplayCluster/Private/Render/Viewport/DisplayClusterViewportManager.cpp
--- a/Engine/Plugins/Runtime/nDisplay/Source/DisplayCluster/Private/Render/Viewport/DisplayClusterViewportManager.cpp
+++ b/Engine/Plugins/Runtime/nDisplay/Source/DisplayCluster/Private/Render/Viewport/DisplayClusterViewportManager.cpp
@@ -627,26 +627,46 @@ void FDisplayClusterViewportManager::InitializeNewFrame()
 	// Send render frame settings to rendering thread
 	ViewportManagerProxy->ImplUpdateRenderFrameSettings(GetRenderFrameSettings(), ViewportManagerViewExtension);
 
-	// Send updated viewports data to render thread proxy
+	/**
+	 * [1] Send new viewport render resources to proxy:
+	 * At the moment, there are new resources in the viewport, but the math is not initialized.The math will be updated later in the game thread.
+	 * But these new resources were needed for future rendering on the render thread.
+	 * So we are sending viewports from the game stream data to the proxy.
+	 * Lastly, after the viewport math data has been updated, we need to send it again.
+	 */
 	ViewportManagerProxy->ImplUpdateViewports(ClusterNodeViewports);
+
+	// Send postprocess data to render thread
+	if (PostProcessManager.IsValid())
+	{
+		// Update postprocess data from game thread
+		PostProcessManager->Tick();
+
+		// Send updated postprocess data to rendering thread
+		PostProcessManager->FinalizeNewFrame();
+	}
 }
 
 void FDisplayClusterViewportManager::FinalizeNewFrame()
 {
 	check(IsInGameThread());
 
-	// When all viewports processed, we remove all single frame custom postprocess
+	// [2] The viewport math has been updated, we need to send it to the proxy again.
+	ViewportManagerProxy->ImplUpdateViewports(ClusterNodeViewports);
+
 	for (FDisplayClusterViewport* Viewport : ClusterNodeViewports)
 	{
 		if (Viewport)
 		{
+			// When all viewports processed, we remove all single frame custom postprocess
 			Viewport->CustomPostProcessSettings.FinalizeFrame();
 
-			// update projection policy proxy data
+			// Update projection policy proxy data
 			if (Viewport->ProjectionPolicy.IsValid())
 			{
 				Viewport->ProjectionPolicy->UpdateProxyData(Viewport);
 			}
+
 		}
 	}
 
@@ -658,15 +678,6 @@ void FDisplayClusterViewportManager::FinalizeNewFrame()
 			Viewport->RenderSettings.FinishUpdateSettings();
 		}
 	}
-
-	if (PostProcessManager.IsValid())
-	{
-		// Update postprocess data from game thread
-		PostProcessManager->Tick();
-
-		// Send updated postprocess data to rendering thread
-		PostProcessManager->FinalizeNewFrame();
-	}
 }
 
 FSceneViewFamily::ConstructionValues FDisplayClusterViewportManager::CreateViewFamilyConstructionValues(
