[nDisplay] - duplicate MediaOutput/MediaInput to prevent reference loss after MU package re-load

diff --git a/Engine/Plugins/Runtime/nDisplay/Source/DisplayClusterMedia/Private/Capture/DisplayClusterMediaCaptureBase.cpp b/Engine/Plugins/Runtime/nDisplay/Source/DisplayClusterMedia/Private/Capture/DisplayClusterMediaCaptureBase.cpp
--- a/Engine/Plugins/Runtime/nDisplay/Source/DisplayClusterMedia/Private/Capture/DisplayClusterMediaCaptureBase.cpp
+++ b/Engine/Plugins/Runtime/nDisplay/Source/DisplayClusterMedia/Private/Capture/DisplayClusterMediaCaptureBase.cpp
@@ -11,12 +11,16 @@
 
 #include "RenderGraphBuilder.h"
 
+#include "UObject/UObjectGlobals.h"
+#include "UObject/Package.h"
+
 
 FDisplayClusterMediaCaptureBase::FDisplayClusterMediaCaptureBase(const FString& InMediaId, const FString& InClusterNodeId, UMediaOutput* InMediaOutput)
 	: FDisplayClusterMediaBase(InMediaId, InClusterNodeId)
-	, MediaOutput(InMediaOutput)
 {
-	check(InMediaOutput);
+	checkSlow(InMediaOutput);
+	MediaOutput = DuplicateObject(InMediaOutput, GetTransientPackage());
+	checkSlow(MediaOutput);
 
 	IDisplayCluster::Get().GetCallbacks().OnDisplayClusterPostTick().AddRaw(this, &FDisplayClusterMediaCaptureBase::OnPostClusterTick);
 }
diff --git a/Engine/Plugins/Runtime/nDisplay/Source/DisplayClusterMedia/Private/Input/DisplayClusterMediaInputBase.cpp b/Engine/Plugins/Runtime/nDisplay/Source/DisplayClusterMedia/Private/Input/DisplayClusterMediaInputBase.cpp
--- a/Engine/Plugins/Runtime/nDisplay/Source/DisplayClusterMedia/Private/Input/DisplayClusterMediaInputBase.cpp
+++ b/Engine/Plugins/Runtime/nDisplay/Source/DisplayClusterMedia/Private/Input/DisplayClusterMediaInputBase.cpp
@@ -11,11 +11,15 @@
 #include "TextureResource.h"
 #include "RHIUtilities.h"
 
+#include "UObject/UObjectGlobals.h"
+#include "UObject/Package.h"
+
 
 FDisplayClusterMediaInputBase::FDisplayClusterMediaInputBase(const FString& InMediaId, const FString& InClusterNodeId, UMediaSource* InMediaSource)
 	: FDisplayClusterMediaBase(InMediaId, InClusterNodeId)
-	, MediaSource(InMediaSource)
 {
+	checkSlow(InMediaSource);
+	MediaSource = DuplicateObject(InMediaSource, GetTransientPackage());
 	checkSlow(MediaSource);
 
 	// Instantiate media player
