[nDisplay] Added cvar to enable/disable nDisplay media subsystem

diff --git a/Engine/Plugins/Runtime/nDisplay/Source/DisplayCluster/Private/Render/Viewport/Configuration/DisplayClusterViewportConfigurationHelpers.cpp b/Engine/Plugins/Runtime/nDisplay/Source/DisplayCluster/Private/Render/Viewport/Configuration/DisplayClusterViewportConfigurationHelpers.cpp
--- a/Engine/Plugins/Runtime/nDisplay/Source/DisplayCluster/Private/Render/Viewport/Configuration/DisplayClusterViewportConfigurationHelpers.cpp
+++ b/Engine/Plugins/Runtime/nDisplay/Source/DisplayCluster/Private/Render/Viewport/Configuration/DisplayClusterViewportConfigurationHelpers.cpp
@@ -35,6 +35,9 @@
 #include "Misc/DisplayClusterLog.h"
 #include "TextureResource.h"
 
+#include "HAL/IConsoleManager.h"
+
+
 void FDisplayClusterViewportConfigurationHelpers::UpdateViewportStereoMode(FDisplayClusterViewport& DstViewport, const EDisplayClusterConfigurationViewport_StereoMode StereoMode)
 {
 	switch (StereoMode)
@@ -211,15 +214,20 @@ void FDisplayClusterViewportConfigurationHelpers::UpdateBaseViewportSetting(FDis
 	// Set media related configuration (runtime only for now)
 	if (IDisplayCluster::Get().GetOperationMode() == EDisplayClusterOperationMode::Cluster)
 	{
-		const FDisplayClusterConfigurationMedia& MediaSettings = InConfigurationViewport.RenderSettings.Media;
-
-		if (MediaSettings.bEnable)
+		// Check if nDisplay media enabled
+		static const TConsoleVariableData<int32>* const ICVarMediaEnabled = IConsoleManager::Get().FindTConsoleVariableDataInt(TEXT("nDisplay.media.Enabled"));
+		if (ICVarMediaEnabled && !!ICVarMediaEnabled->GetValueOnGameThread())
 		{
-			// Don't render the viewport if media input assigned
-			DstViewport.RenderSettings.bSkipSceneRenderingButLeaveResourcesAvailable = !!MediaSettings.MediaSource;
+			const FDisplayClusterConfigurationMedia& MediaSettings = InConfigurationViewport.RenderSettings.Media;
 
-			// Mark this viewport is going to be captured by a capture device
-			DstViewport.RenderSettings.bIsBeingCaptured = !!MediaSettings.MediaOutput;
+			if (MediaSettings.bEnable)
+			{
+				// Don't render the viewport if media input assigned
+				DstViewport.RenderSettings.bSkipSceneRenderingButLeaveResourcesAvailable = !!MediaSettings.MediaSource;
+
+				// Mark this viewport is going to be captured by a capture device
+				DstViewport.RenderSettings.bIsBeingCaptured = !!MediaSettings.MediaOutput;
+			}
 		}
 	}
 
diff --git a/Engine/Plugins/Runtime/nDisplay/Source/DisplayCluster/Private/Render/Viewport/Configuration/DisplayClusterViewportConfigurationHelpers_ICVFX.cpp b/Engine/Plugins/Runtime/nDisplay/Source/DisplayCluster/Private/Render/Viewport/Configuration/DisplayClusterViewportConfigurationHelpers_ICVFX.cpp
--- a/Engine/Plugins/Runtime/nDisplay/Source/DisplayCluster/Private/Render/Viewport/Configuration/DisplayClusterViewportConfigurationHelpers_ICVFX.cpp
+++ b/Engine/Plugins/Runtime/nDisplay/Source/DisplayCluster/Private/Render/Viewport/Configuration/DisplayClusterViewportConfigurationHelpers_ICVFX.cpp
@@ -39,10 +39,12 @@
 
 #include "Components/DisplayClusterICVFXCameraComponent.h"
 
-
 #include "Misc/DisplayClusterLog.h"
 #include "TextureResource.h"
 
+#include "HAL/IConsoleManager.h"
+
+
 ////////////////////////////////////////////////////////////////////////////////
 // Experimental feature: to be approved after testing
 int32 GDisplayClusterPreviewEnableReuseViewportInCluster = 1;
@@ -626,17 +628,22 @@ void FDisplayClusterViewportConfigurationHelpers_ICVFX::UpdateCameraViewportSett
 	// Set media related configuration (runtime only for now)
 	if (IDisplayCluster::Get().GetOperationMode() == EDisplayClusterOperationMode::Cluster)
 	{
-		const FDisplayClusterConfigurationMediaICVFX& MediaICVFXSettings = InCameraComponent.CameraSettings.RenderSettings.Media;
-
-		if (MediaICVFXSettings.bEnable)
+		// Check if nDisplay media enabled
+		static const TConsoleVariableData<int32>* const ICVarMediaEnabled = IConsoleManager::Get().FindTConsoleVariableDataInt(TEXT("nDisplay.media.Enabled"));
+		if (ICVarMediaEnabled && !!ICVarMediaEnabled->GetValueOnGameThread())
 		{
-			const FString ThisClusterNodeId = DstViewport.GetClusterNodeId();
+			const FDisplayClusterConfigurationMediaICVFX& MediaICVFXSettings = InCameraComponent.CameraSettings.RenderSettings.Media;
 
-			// Don't render the viewport if media input assigned
-			DstViewport.RenderSettings.bSkipSceneRenderingButLeaveResourcesAvailable = MediaICVFXSettings.IsMediaInputAssigned(ThisClusterNodeId);
+			if (MediaICVFXSettings.bEnable)
+			{
+				const FString ThisClusterNodeId = DstViewport.GetClusterNodeId();
+
+				// Don't render the viewport if media input assigned
+				DstViewport.RenderSettings.bSkipSceneRenderingButLeaveResourcesAvailable = MediaICVFXSettings.IsMediaInputAssigned(ThisClusterNodeId);
 
-			// Mark this viewport is going to be captured by a capture device
-			DstViewport.RenderSettings.bIsBeingCaptured = MediaICVFXSettings.IsMediaOutputAssigned(ThisClusterNodeId);
+				// Mark this viewport is going to be captured by a capture device
+				DstViewport.RenderSettings.bIsBeingCaptured = MediaICVFXSettings.IsMediaOutputAssigned(ThisClusterNodeId);
+			}
 		}
 	}
 }
diff --git a/Engine/Plugins/Runtime/nDisplay/Source/DisplayClusterMedia/Private/DisplayClusterMediaModule.cpp b/Engine/Plugins/Runtime/nDisplay/Source/DisplayClusterMedia/Private/DisplayClusterMediaModule.cpp
--- a/Engine/Plugins/Runtime/nDisplay/Source/DisplayClusterMedia/Private/DisplayClusterMediaModule.cpp
+++ b/Engine/Plugins/Runtime/nDisplay/Source/DisplayClusterMedia/Private/DisplayClusterMediaModule.cpp
@@ -26,6 +26,17 @@
 #include "Misc/CoreDelegates.h"
 
 
+static TAutoConsoleVariable<int32> CVarMediaEnabled(
+	TEXT("nDisplay.media.Enabled"),
+	1,
+	TEXT("nDisplay media subsystem\n")
+	TEXT("0 : Disabled\n")
+	TEXT("1 : Enabled\n")
+	,
+	ECVF_ReadOnly
+);
+
+
 void FDisplayClusterMediaModule::StartupModule()
 {
 	UE_LOG(LogDisplayClusterMedia, Log, TEXT("Starting module 'DisplayClusterMedia'..."));
@@ -90,6 +101,13 @@ void FDisplayClusterMediaModule::InitializeMedia()
 		return;
 	}
 
+	// Check if media enabled
+	if (CVarMediaEnabled.GetValueOnGameThread() == 0)
+	{
+		UE_LOG(LogDisplayClusterMedia, Log, TEXT("nDisplay media subsytem is disabled by a cvar"));
+		return;
+	}
+
 	// Instantiate latency queue
 	FrameQueue.Init();
 
