nDisplay: Fixed issue in the ICVFX panel color grading drawer where the details panel for the last per-viewport and per-node color grading group was displaying the entire cluster properties instead. Also fixed issue where the color grading drawer was not displaying the ICVFX color grading properties.

diff --git a/Engine/Plugins/Runtime/nDisplay/Source/DisplayClusterColorGrading/Private/DataModelGenerators/DisplayClusterColorGradingGenerator_RootActor.cpp b/Engine/Plugins/Runtime/nDisplay/Source/DisplayClusterColorGrading/Private/DataModelGenerators/DisplayClusterColorGradingGenerator_RootActor.cpp
--- a/Engine/Plugins/Runtime/nDisplay/Source/DisplayClusterColorGrading/Private/DataModelGenerators/DisplayClusterColorGradingGenerator_RootActor.cpp
+++ b/Engine/Plugins/Runtime/nDisplay/Source/DisplayClusterColorGrading/Private/DataModelGenerators/DisplayClusterColorGradingGenerator_RootActor.cpp
@@ -243,7 +243,7 @@ private:
 		check(ArrayPropertyHandle->AsArray());
 		ArrayPropertyHandle->AsArray()->GetNumElements(ArraySize);
 
-		if (GroupIndex > 0 && GroupIndex < (int32)ArraySize)
+		if (GroupIndex > 0 && GroupIndex <= (int32)ArraySize)
 		{
 			const int32 Index = GroupIndex - 1;
 			IDetailCategoryBuilder& PerNodeSettingsCategoryBuilder = DetailBuilder.EditCategory(TEXT("DetailView_PerViewport"), LOCTEXT("DetailView_PerViewportDisplayName", "Per-Viewport Settings"));
@@ -736,8 +736,8 @@ public:
 		CategoryBuilder.AddProperty(DetailBuilder.GetProperty(GET_MEMBER_NAME_CHECKED(UDisplayClusterICVFXCameraComponent, CameraSettings.bEnable)));
 		CategoryBuilder.AddProperty(DetailBuilder.GetProperty(GET_MEMBER_NAME_CHECKED(UDisplayClusterICVFXCameraComponent, CameraSettings.Chromakey.bEnable)));
 		CategoryBuilder.AddProperty(DetailBuilder.GetProperty(GET_MEMBER_NAME_CHECKED(UDisplayClusterICVFXCameraComponent, CameraSettings.CameraOCIO.AllNodesOCIOConfiguration.bIsEnabled)));
-		CategoryBuilder.AddProperty(DetailBuilder.GetProperty(GET_MEMBER_NAME_CHECKED(UDisplayClusterICVFXCameraComponent, CameraSettings.CameraOCIO.AllNodesOCIOConfiguration)));
-		CategoryBuilder.AddProperty(DetailBuilder.GetProperty(GET_MEMBER_NAME_CHECKED(UDisplayClusterICVFXCameraComponent, CameraSettings.CameraOCIO.PerNodeOCIOProfiles)));
+		CategoryBuilder.AddProperty(DetailBuilder.GetProperty(GET_MEMBER_NAME_CHECKED(UDisplayClusterICVFXCameraComponent, CameraSettings.AllNodesColorGrading)));
+		CategoryBuilder.AddProperty(DetailBuilder.GetProperty(GET_MEMBER_NAME_CHECKED(UDisplayClusterICVFXCameraComponent, CameraSettings.PerNodeColorGrading)));
 
 		AddColorGradingDetailProperties(DetailBuilder);
 		AddDetailsPanelProperties(DetailBuilder);
@@ -794,7 +794,7 @@ private:
 		check(ArrayPropertyHandle->AsArray());
 		ArrayPropertyHandle->AsArray()->GetNumElements(ArraySize);
 
-		if (GroupIndex > 0 && GroupIndex < (int32)ArraySize)
+		if (GroupIndex > 0 && GroupIndex <= (int32)ArraySize)
 		{
 			const int32 Index = GroupIndex - 1;
 			IDetailCategoryBuilder& PerNodeSettingsCategoryBuilder = DetailBuilder.EditCategory(TEXT("DetailView_PerNode"), LOCTEXT("DetailView_PerNodeDisplayName", "Per-Node Settings"));
