nDisplay: Fix MediaInput for the sRGB input case.

When playing back recorded image sequences, they appeared darker than the original, and it was because the saved images are sRGB encoded and the reader was no converting it back to linear. W/o this, it does not work for non RGB8 inputs that are sRGB encoded.

diff --git a/Engine/Plugins/Runtime/nDisplay/Source/DisplayClusterMedia/Private/DisplayClusterMediaHelpers.cpp b/Engine/Plugins/Runtime/nDisplay/Source/DisplayClusterMedia/Private/DisplayClusterMediaHelpers.cpp
--- a/Engine/Plugins/Runtime/nDisplay/Source/DisplayClusterMedia/Private/DisplayClusterMediaHelpers.cpp
+++ b/Engine/Plugins/Runtime/nDisplay/Source/DisplayClusterMedia/Private/DisplayClusterMediaHelpers.cpp
@@ -74,6 +74,17 @@ namespace DisplayClusterMediaHelpers
 
 	void ResampleTexture_RenderThread(FRHICommandListImmediate& RHICmdList, FRHITexture* SrcTexture, FRHITexture* DstTexture, const FIntRect& SrcRect, const FIntRect& DstRect)
 	{
-		ResampleTextureImpl_RenderThread<FScreenPS>(RHICmdList, SrcTexture, DstTexture, SrcRect, DstRect);
+		check(SrcTexture);
+
+		const bool bSrcSrgb = EnumHasAnyFlags(SrcTexture->GetFlags(), TexCreate_SRGB);
+
+		if (bSrcSrgb)
+		{
+			ResampleTextureImpl_RenderThread<FScreenPSsRGBSource>(RHICmdList, SrcTexture, DstTexture, SrcRect, DstRect);
+		}
+		else
+		{
+			ResampleTextureImpl_RenderThread<FScreenPS>(RHICmdList, SrcTexture, DstTexture, SrcRect, DstRect);
+		}
 	}
 }
diff --git a/Engine/Plugins/Runtime/nDisplay/Source/DisplayClusterMedia/Private/Input/DisplayClusterMediaInputBase.cpp b/Engine/Plugins/Runtime/nDisplay/Source/DisplayClusterMedia/Private/Input/DisplayClusterMediaInputBase.cpp
--- a/Engine/Plugins/Runtime/nDisplay/Source/DisplayClusterMedia/Private/Input/DisplayClusterMediaInputBase.cpp
+++ b/Engine/Plugins/Runtime/nDisplay/Source/DisplayClusterMedia/Private/Input/DisplayClusterMediaInputBase.cpp
@@ -85,8 +85,13 @@ void FDisplayClusterMediaInputBase::ImportMediaData(FRHICommandListImmediate& RH
 	{
 		const FIntPoint DstRegionSize = TextureInfo.Region.Size();
 
-		if (SrcTexture->GetDesc().Format == DstTexture->GetDesc().Format &&
-			SrcTexture->GetDesc().Extent == DstRegionSize)
+		const bool bSrcSrgb = EnumHasAnyFlags(SrcTexture->GetFlags(), TexCreate_SRGB);
+		const bool bDstSrgb = EnumHasAnyFlags(DstTexture->GetFlags(), TexCreate_SRGB);
+
+		if ((SrcTexture->GetDesc().Format == DstTexture->GetDesc().Format)
+			&& (SrcTexture->GetDesc().Extent == DstRegionSize)
+			&& (bSrcSrgb == bDstSrgb)
+		)
 		{
 			FRHICopyTextureInfo CopyInfo;
 
