[nDisplay] fixed issue with multiple DCRAs calling CLN::GetObjectsData

diff --git a/Engine/Plugins/Runtime/nDisplay/Source/DisplayCluster/Private/Game/EngineClasses/Scene/DisplayClusterSyncTickComponent.cpp b/Engine/Plugins/Runtime/nDisplay/Source/DisplayCluster/Private/Game/EngineClasses/Scene/DisplayClusterSyncTickComponent.cpp
--- a/Engine/Plugins/Runtime/nDisplay/Source/DisplayCluster/Private/Game/EngineClasses/Scene/DisplayClusterSyncTickComponent.cpp
+++ b/Engine/Plugins/Runtime/nDisplay/Source/DisplayCluster/Private/Game/EngineClasses/Scene/DisplayClusterSyncTickComponent.cpp
@@ -3,7 +3,9 @@
 #include "Components/DisplayClusterSyncTickComponent.h"
 
 #include "IPDisplayCluster.h"
+#include "Game/IPDisplayClusterGameManager.h"
 #include "Misc/DisplayClusterGlobals.h"
+#include "DisplayClusterRootActor.h"
 
 
 UDisplayClusterSyncTickComponent::UDisplayClusterSyncTickComponent(const FObjectInitializer& ObjectInitializer)
@@ -19,7 +21,12 @@ void UDisplayClusterSyncTickComponent::TickComponent(float DeltaTime, ELevelTick
 {
 	if (GDisplayCluster->GetOperationMode() != EDisplayClusterOperationMode::Disabled)
 	{
-		GDisplayCluster->Tick(DeltaTime);
+		// Since we may have multiple DCRAs available in the world, only 'Main' DCRA
+		// should be able to throw tick call to nDisplay
+		if (GDisplayCluster->GetGameMgr()->GetRootActor() == Cast<ADisplayClusterRootActor>(GetOwner()))
+		{
+			GDisplayCluster->Tick(DeltaTime);
+		}
 	}
 
 	Super::TickComponent(DeltaTime, TickType, ThisTickFunction);
